UE4: StringAssetReferences will now properly detect and repair references to PIE objects when duplicating for PIE

diff --git a/Engine/Source/Editor/UnrealEd/Private/PlayLevel.cpp b/Engine/Source/Editor/UnrealEd/Private/PlayLevel.cpp
--- a/Engine/Source/Editor/UnrealEd/Private/PlayLevel.cpp
+++ b/Engine/Source/Editor/UnrealEd/Private/PlayLevel.cpp
@@ -2571,6 +2571,11 @@ UWorld* UEditorEngine::CreatePIEWorldByDuplication(FWorldContext &WorldContext,
 		// Reset any GUID fixups with lazy pointers
 		FLazyObjectPtr::ResetPIEFixups();
 
+		// Prepare string asset references for fixup
+		TArray<UPackage*> PackagesBeingDuplicatedForPIE;
+		PackagesBeingDuplicatedForPIE.Add(EditorWorld->GetOutermost());
+		FStringAssetReference::SetPackagesBeingDuplicatedForPIE(PackagesBeingDuplicatedForPIE);
+
 		// NULL GWorld before various PostLoad functions are called, this makes it easier to debug invalid GWorld accesses
 		GWorld = NULL;
 
@@ -2584,6 +2589,8 @@ UWorld* UEditorEngine::CreatePIEWorldByDuplication(FWorldContext &WorldContext,
 			SDO_DuplicateForPie		// bDuplicateForPIE
 			) );
 
+		FStringAssetReference::ClearPackagesBeingDuplicatedForPIE();
+
 		// Store prefix we used to rename this world and streaming levels package names
 		NewPIEWorld->StreamingLevelsPrefix = UWorld::BuildPIEPackagePrefix(WorldContext.PIEInstance);
 		// Fixup model components. The index buffers have been created for the components in the EditorWorld and the order
diff --git a/Engine/Source/Runtime/CoreUObject/Private/Misc/StringAssetReference.cpp b/Engine/Source/Runtime/CoreUObject/Private/Misc/StringAssetReference.cpp
--- a/Engine/Source/Runtime/CoreUObject/Private/Misc/StringAssetReference.cpp
+++ b/Engine/Source/Runtime/CoreUObject/Private/Misc/StringAssetReference.cpp
@@ -28,6 +28,13 @@ bool FStringAssetReference::Serialize(FArchive& Ar)
 		FCoreDelegates::StringAssetReferenceLoaded.Execute(AssetLongPathname);
 	}
 #endif // WITH_EDITOR
+
+	if (Ar.IsLoading() && (Ar.GetPortFlags()&PPF_DuplicateForPIE))
+	{
+		// Remap unique ID if necessary
+		FixupForPIE();
+	}
+
 	return true;
 }
 bool FStringAssetReference::operator==(FStringAssetReference const& Other) const
@@ -110,27 +117,6 @@ bool FStringAssetReference::SerializeFromMismatchedTag(struct FPropertyTag const
 	return false;
 }
 
-FStringAssetReference FStringAssetReference::FixupForPIE() const
-{
-	FStringAssetReference Temp(*this);
-	if (FPlatformProperties::HasEditorOnlyData() && IsValid())
-	{
-		FString Path = ToString();
-		const FString ShortPackageOuterAndName = FPackageName::GetLongPackageAssetName(Path);
-
-		if (!ShortPackageOuterAndName.StartsWith(PLAYWORLD_PACKAGE_PREFIX) && Path.Contains("TheWorld:"))
-		{
-			check(GPlayInEditorID != -1);
-
-			// Need to prepend PIE prefix, as we're in PIE and this is a level actor
-			Path = FString::Printf(TEXT("%s/%s_%d_%s"), *FPackageName::GetLongPackagePath(Path), PLAYWORLD_PACKAGE_PREFIX, GPlayInEditorID, *ShortPackageOuterAndName);
-			Temp.AssetLongPathname = Path;
-		}
-	}
-
-	return Temp;
-}
-
 UObject *FStringAssetReference::ResolveObject() const
 {
 	// Don't try to resolve if we're saving a package because StaticFindObject can't be used here
@@ -160,4 +146,51 @@ FStringAssetReference FStringAssetReference::GetOrCreateIDForObject(const class
 	return FStringAssetReference(Object);
 }
 
+void FStringAssetReference::SetPackagesBeingDuplicatedForPIE(const TArray<UPackage*>& InPackagesBeingDuplicatedForPIE)
+{
+	for ( auto PackageIt = InPackagesBeingDuplicatedForPIE.CreateConstIterator(); PackageIt; ++PackageIt )
+	{
+		UPackage* Package = *PackageIt;
+		if ( Package )
+		{
+			PackageNamesBeingDuplicatedForPIE.Add(Package->GetName());
+		}
+	}
+}
+
+void FStringAssetReference::ClearPackagesBeingDuplicatedForPIE()
+{
+	PackageNamesBeingDuplicatedForPIE.Empty();
+}
+
+void FStringAssetReference::FixupForPIE()
+{
+	if (FPlatformProperties::HasEditorOnlyData() && IsValid())
+	{
+		FString Path = ToString();
+
+		// Determine if this reference has already been fixed up for PIE
+		const FString ShortPackageOuterAndName = FPackageName::GetLongPackageAssetName(Path);
+		if (!ShortPackageOuterAndName.StartsWith(PLAYWORLD_PACKAGE_PREFIX))
+		{
+			check(GPlayInEditorID != -1);
+
+			// Determine if this refers to a package that is being duplicated for PIE
+			for (auto PackageNameIt = PackageNamesBeingDuplicatedForPIE.CreateConstIterator(); PackageNameIt; ++PackageNameIt)
+			{
+				const FString PathPrefix = (*PackageNameIt) + TEXT(".");
+				if (Path.StartsWith(PathPrefix))
+				{
+					// Need to prepend PIE prefix, as we're in PIE and this refers to an object in a PIE package
+					Path = FString::Printf(TEXT("%s/%s_%d_%s"), *FPackageName::GetLongPackagePath(Path), PLAYWORLD_PACKAGE_PREFIX, GPlayInEditorID, *ShortPackageOuterAndName);
+					AssetLongPathname = Path;
+
+					break;
+				}
+			}
+		}
+	}
+}
+
 int32 FStringAssetReference::CurrentTag = 1;
+TArray<FString> FStringAssetReference::PackageNamesBeingDuplicatedForPIE;
\ No newline at end of file
diff --git a/Engine/Source/Runtime/CoreUObject/Private/Serialization/DuplicateDataReader.cpp b/Engine/Source/Runtime/CoreUObject/Private/Serialization/DuplicateDataReader.cpp
--- a/Engine/Source/Runtime/CoreUObject/Private/Serialization/DuplicateDataReader.cpp
+++ b/Engine/Source/Runtime/CoreUObject/Private/Serialization/DuplicateDataReader.cpp
@@ -67,11 +67,6 @@ FArchive& FDuplicateDataReader::operator<<( FAssetPtr& AssetPtr)
 	FStringAssetReference ID;
 	Ar << ID;
 
-	if (Ar.GetPortFlags()&PPF_DuplicateForPIE)
-	{
-		// Remap unique ID if necessary
-		ID = ID.FixupForPIE();
-	}
 	AssetPtr = ID;
 	return Ar;
 }
diff --git a/Engine/Source/Runtime/CoreUObject/Public/Misc/StringAssetReference.h b/Engine/Source/Runtime/CoreUObject/Public/Misc/StringAssetReference.h
--- a/Engine/Source/Runtime/CoreUObject/Public/Misc/StringAssetReference.h
+++ b/Engine/Source/Runtime/CoreUObject/Public/Misc/StringAssetReference.h
@@ -43,11 +43,6 @@ struct COREUOBJECT_API FStringAssetReference
 		return AssetLongPathname;
 	}
 
-	/**
-	 * Fixes up this UniqueObjectID to add or remove the PIE prefix depending on what is currently active
-	 */
-	FStringAssetReference FixupForPIE() const;
-
 	/**
 	 * Attempts to find a currently loaded object that matches this object ID
 	 * @return Found UObject, or NULL if not currently loaded
@@ -105,10 +100,19 @@ struct COREUOBJECT_API FStringAssetReference
 
 	static FStringAssetReference GetOrCreateIDForObject(const class UObject *Object);
 
+	static void SetPackagesBeingDuplicatedForPIE(const TArray<UPackage*>& InPackagesBeingDuplicatedForPIE);
+	static void ClearPackagesBeingDuplicatedForPIE();
+
 private:
 
+	/**
+	 * Fixes up this StringAssetReference to add or remove the PIE prefix depending on what is currently active
+	 */
+	void FixupForPIE();
+
 	/** Global counter that determines when we need to re-search for GUIDs because more objects have been loaded **/
 	static int32 CurrentTag;
+	static TArray<FString> PackageNamesBeingDuplicatedForPIE;
 
 };
 
