UE4: Adjustment StringAssetReference fixups to accommodate for streaming levels that are either always loaded or loaded later.

diff --git a/Engine/Source/Editor/UnrealEd/Private/PlayLevel.cpp b/Engine/Source/Editor/UnrealEd/Private/PlayLevel.cpp
--- a/Engine/Source/Editor/UnrealEd/Private/PlayLevel.cpp
+++ b/Engine/Source/Editor/UnrealEd/Private/PlayLevel.cpp
@@ -2572,9 +2572,19 @@ UWorld* UEditorEngine::CreatePIEWorldByDuplication(FWorldContext &WorldContext,
 		FLazyObjectPtr::ResetPIEFixups();
 
 		// Prepare string asset references for fixup
-		TArray<UPackage*> PackagesBeingDuplicatedForPIE;
-		PackagesBeingDuplicatedForPIE.Add(EditorWorld->GetOutermost());
-		FStringAssetReference::SetPackagesBeingDuplicatedForPIE(PackagesBeingDuplicatedForPIE);
+		TArray<FString> PackageNamesBeingDuplicatedForPIE;
+		PackageNamesBeingDuplicatedForPIE.Add(PlayWorldMapName);
+		for ( auto LevelIt = EditorWorld->StreamingLevels.CreateConstIterator(); LevelIt; ++LevelIt )
+		{
+			ULevelStreaming* StreamingLevel = *LevelIt;
+			if ( StreamingLevel )
+			{
+				const FString StreamingLevelPIEName = UWorld::ConvertToPIEPackageName(StreamingLevel->PackageName.ToString(), WorldContext.PIEInstance);
+				PackageNamesBeingDuplicatedForPIE.Add(StreamingLevelPIEName);
+			}
+		}
+
+		FStringAssetReference::SetPackageNamesBeingDuplicatedForPIE(PackageNamesBeingDuplicatedForPIE);
 
 		// NULL GWorld before various PostLoad functions are called, this makes it easier to debug invalid GWorld accesses
 		GWorld = NULL;
@@ -2589,7 +2599,7 @@ UWorld* UEditorEngine::CreatePIEWorldByDuplication(FWorldContext &WorldContext,
 			SDO_DuplicateForPie		// bDuplicateForPIE
 			) );
 
-		FStringAssetReference::ClearPackagesBeingDuplicatedForPIE();
+		FStringAssetReference::ClearPackageNamesBeingDuplicatedForPIE();
 
 		// Store prefix we used to rename this world and streaming levels package names
 		NewPIEWorld->StreamingLevelsPrefix = UWorld::BuildPIEPackagePrefix(WorldContext.PIEInstance);
diff --git a/Engine/Source/Runtime/CoreUObject/Private/Misc/StringAssetReference.cpp b/Engine/Source/Runtime/CoreUObject/Private/Misc/StringAssetReference.cpp
--- a/Engine/Source/Runtime/CoreUObject/Private/Misc/StringAssetReference.cpp
+++ b/Engine/Source/Runtime/CoreUObject/Private/Misc/StringAssetReference.cpp
@@ -146,19 +146,12 @@ FStringAssetReference FStringAssetReference::GetOrCreateIDForObject(const class
 	return FStringAssetReference(Object);
 }
 
-void FStringAssetReference::SetPackagesBeingDuplicatedForPIE(const TArray<UPackage*>& InPackagesBeingDuplicatedForPIE)
+void FStringAssetReference::SetPackageNamesBeingDuplicatedForPIE(const TArray<FString>& InPackageNamesBeingDuplicatedForPIE)
 {
-	for ( auto PackageIt = InPackagesBeingDuplicatedForPIE.CreateConstIterator(); PackageIt; ++PackageIt )
-	{
-		UPackage* Package = *PackageIt;
-		if ( Package )
-		{
-			PackageNamesBeingDuplicatedForPIE.Add(Package->GetName());
-		}
-	}
+	PackageNamesBeingDuplicatedForPIE = InPackageNamesBeingDuplicatedForPIE;
 }
 
-void FStringAssetReference::ClearPackagesBeingDuplicatedForPIE()
+void FStringAssetReference::ClearPackageNamesBeingDuplicatedForPIE()
 {
 	PackageNamesBeingDuplicatedForPIE.Empty();
 }
@@ -175,15 +168,16 @@ void FStringAssetReference::FixupForPIE()
 		{
 			check(GPlayInEditorID != -1);
 
+			const FString PIEPath = FString::Printf(TEXT("%s/%s_%d_%s"), *FPackageName::GetLongPackagePath(Path), PLAYWORLD_PACKAGE_PREFIX, GPlayInEditorID, *ShortPackageOuterAndName);
+
 			// Determine if this refers to a package that is being duplicated for PIE
 			for (auto PackageNameIt = PackageNamesBeingDuplicatedForPIE.CreateConstIterator(); PackageNameIt; ++PackageNameIt)
 			{
 				const FString PathPrefix = (*PackageNameIt) + TEXT(".");
-				if (Path.StartsWith(PathPrefix))
+				if (PIEPath.StartsWith(PathPrefix))
 				{
 					// Need to prepend PIE prefix, as we're in PIE and this refers to an object in a PIE package
-					Path = FString::Printf(TEXT("%s/%s_%d_%s"), *FPackageName::GetLongPackagePath(Path), PLAYWORLD_PACKAGE_PREFIX, GPlayInEditorID, *ShortPackageOuterAndName);
-					AssetLongPathname = Path;
+					AssetLongPathname = PIEPath;
 
 					break;
 				}
diff --git a/Engine/Source/Runtime/CoreUObject/Public/Misc/StringAssetReference.h b/Engine/Source/Runtime/CoreUObject/Public/Misc/StringAssetReference.h
--- a/Engine/Source/Runtime/CoreUObject/Public/Misc/StringAssetReference.h
+++ b/Engine/Source/Runtime/CoreUObject/Public/Misc/StringAssetReference.h
@@ -100,8 +100,8 @@ struct COREUOBJECT_API FStringAssetReference
 
 	static FStringAssetReference GetOrCreateIDForObject(const class UObject *Object);
 
-	static void SetPackagesBeingDuplicatedForPIE(const TArray<UPackage*>& InPackagesBeingDuplicatedForPIE);
-	static void ClearPackagesBeingDuplicatedForPIE();
+	static void SetPackageNamesBeingDuplicatedForPIE(const TArray<FString>& InPackageNamesBeingDuplicatedForPIE);
+	static void ClearPackageNamesBeingDuplicatedForPIE();
 
 private:
 
diff --git a/Engine/Source/Runtime/Engine/Private/World.cpp b/Engine/Source/Runtime/Engine/Private/World.cpp
--- a/Engine/Source/Runtime/Engine/Private/World.cpp
+++ b/Engine/Source/Runtime/Engine/Private/World.cpp
@@ -1855,8 +1855,28 @@ UWorld* UWorld::DuplicateWorldForPIE(const FString& PackageName, UWorld* OwningW
 	PIELevelPackage->PIEInstanceID = WorldContext.PIEInstance;
 #endif
 
+	const FString PlayWorldMapName = ConvertToPIEPackageName(OwningWorld->GetOutermost()->GetName(), WorldContext.PIEInstance);
+
+	// Set up string asset reference fixups
+	TArray<FString> PackageNamesBeingDuplicatedForPIE;
+	PackageNamesBeingDuplicatedForPIE.Add(PlayWorldMapName);
+	for (auto LevelIt = OwningWorld->StreamingLevels.CreateConstIterator(); LevelIt; ++LevelIt)
+	{
+		ULevelStreaming* StreamingLevel = *LevelIt;
+		if (StreamingLevel)
+		{
+			const FString StreamingLevelPIEName = UWorld::ConvertToPIEPackageName(StreamingLevel->PackageName.ToString(), WorldContext.PIEInstance);
+			PackageNamesBeingDuplicatedForPIE.Add(StreamingLevelPIEName);
+		}
+	}
+
+	FStringAssetReference::SetPackageNamesBeingDuplicatedForPIE(PackageNamesBeingDuplicatedForPIE);
+
 	ULevel::StreamedLevelsOwningWorld.Add(PIELevelPackage->GetFName(), OwningWorld);
 	UWorld* PIELevelWorld = CastChecked<UWorld>(StaticDuplicateObject(EditorLevelWorld, PIELevelPackage, *EditorLevelWorld->GetName(), RF_AllFlags, NULL, SDO_DuplicateForPie));
+
+	// Clean up string asset reference fixups
+	FStringAssetReference::ClearPackageNamesBeingDuplicatedForPIE();
 	
 	PIELevelWorld->StreamingLevelsPrefix = BuildPIEPackagePrefix(WorldContext.PIEInstance);
 	{
