Adding displayname metadata for the enum EInterpCurveMode

diff --git a/Engine/Source/Runtime/CoreUObject/Classes/Object.h b/Engine/Source/Runtime/CoreUObject/Classes/Object.h
--- a/Engine/Source/Runtime/CoreUObject/Classes/Object.h
+++ b/Engine/Source/Runtime/CoreUObject/Classes/Object.h
@@ -27,12 +27,12 @@ namespace EAxis
 UENUM()
 enum EInterpCurveMode
 {
-	CIM_Linear,
-	CIM_CurveAuto,
-	CIM_Constant,
-	CIM_CurveUser,
-	CIM_CurveBreak,
-	CIM_CurveAutoClamped,
+	CIM_Linear UMETA(DisplayName="Linear"),
+	CIM_CurveAuto UMETA(DisplayName="Curve Auto"),
+	CIM_Constant UMETA(DisplayName="Constant"),
+	CIM_CurveUser UMETA(DisplayName="Curve User"),
+	CIM_CurveBreak UMETA(DisplayName="Curve Break"),
+	CIM_CurveAutoClamped UMETA(DisplayName="Curve Auto Clamped"),
 	CIM_MAX,
 };
 
