#UE4 Fix it so FStringAssetReference::SerializeFromMismatchedTag correctly calls the cooking callbacks

diff --git a/Engine/Source/Runtime/CoreUObject/Private/Misc/StringAssetReference.cpp b/Engine/Source/Runtime/CoreUObject/Private/Misc/StringAssetReference.cpp
--- a/Engine/Source/Runtime/CoreUObject/Private/Misc/StringAssetReference.cpp
+++ b/Engine/Source/Runtime/CoreUObject/Private/Misc/StringAssetReference.cpp
@@ -104,6 +104,12 @@ bool FStringAssetReference::SerializeFromMismatchedTag(struct FPropertyTag const
 		{
 			AssetLongPathname = FString();
 		}
+#if WITH_EDITOR
+		if (Ar.IsLoading() && Ar.IsPersistent() && FCoreDelegates::StringAssetReferenceLoaded.IsBound())
+		{
+			FCoreDelegates::StringAssetReferenceLoaded.Execute(AssetLongPathname);
+		}
+#endif // WITH_EDITOR
 		return true;
 	}
 	else if( Tag.Type == NAME_StrProperty )
@@ -112,6 +118,12 @@ bool FStringAssetReference::SerializeFromMismatchedTag(struct FPropertyTag const
 		Ar << String;
 
 		AssetLongPathname = String;
+#if WITH_EDITOR
+		if (Ar.IsLoading() && Ar.IsPersistent() && FCoreDelegates::StringAssetReferenceLoaded.IsBound())
+		{
+			FCoreDelegates::StringAssetReferenceLoaded.Execute(AssetLongPathname);
+		}
+#endif // WITH_EDITOR
 		return true;
 	}
 	return false;
