#ue4 Replace Blueprint Reference with Class during serializartion. Hack: the logic was added to UClassProperty::CheckValidObject

diff --git a/Engine/Source/Runtime/CoreUObject/Private/UObject/PropertyClass.cpp b/Engine/Source/Runtime/CoreUObject/Private/UObject/PropertyClass.cpp
--- a/Engine/Source/Runtime/CoreUObject/Private/UObject/PropertyClass.cpp
+++ b/Engine/Source/Runtime/CoreUObject/Private/UObject/PropertyClass.cpp
@@ -63,6 +63,59 @@ bool UClassProperty::SameType(const UProperty* Other) const
 	return Super::SameType(Other) && (MetaClass == ((UClassProperty*)Other)->MetaClass);
 }
 
+void UClassProperty::CheckValidObject(void* Value) const
+{
+#if WITH_EDITOR
+	// Ugly hack to replace Blueprint references with Class references.
+
+	struct FReplaceBlueprintWithClassHelper
+	{
+		bool bShouldReplace;
+		UClass* BlueprintClass;
+		UClassProperty* BPGeneratedClassProp;
+
+		FReplaceBlueprintWithClassHelper() : bShouldReplace(false), BlueprintClass(NULL), BPGeneratedClassProp(NULL)
+		{
+			GConfig->GetBool(TEXT("EditoronlyBP"), TEXT("bReplaceBlueprintWithClass"), bShouldReplace, GEditorIni);
+			if (bShouldReplace)
+			{
+				BlueprintClass = FindObject<UClass>(NULL, TEXT("/Script/Engine.Blueprint"));
+				ensure(BlueprintClass);
+				BPGeneratedClassProp = BlueprintClass ? FindField<UClassProperty>(BlueprintClass, TEXT("GeneratedClass")) : NULL;
+				ensure(BPGeneratedClassProp);
+			}
+		}
+
+		bool CanReplace() const
+		{
+			return bShouldReplace && BlueprintClass && BPGeneratedClassProp;
+		}
+	};
+	static FReplaceBlueprintWithClassHelper Helper;
+
+	const UObject* Object = GetObjectPropertyValue(Value);
+	Super::CheckValidObject(Value);
+	const UObject* CurrentObject = GetObjectPropertyValue(Value);
+
+	if (Helper.CanReplace()
+		&& !CurrentObject
+		&& Object 
+		&& Object->IsA(Helper.BlueprintClass)
+		&& (UObject::StaticClass() == MetaClass))
+	{
+		UObject* RecoveredObject = Helper.BPGeneratedClassProp->GetPropertyValue_InContainer(Object);
+		SetObjectPropertyValue(Value, RecoveredObject);
+		UE_LOG(LogProperty, Log,
+			TEXT("Blueprint '%s' is replaced with class '%s' in property '%s'"),
+			*Object->GetFullName(),
+			*RecoveredObject->GetFullName(),
+			*GetFullName());
+	}
+#else	// WITH_EDITOR
+	Super::CheckValidObject(Value);
+#endif	// WITH_EDITOR
+}
+
 IMPLEMENT_CORE_INTRINSIC_CLASS(UClassProperty, UObjectProperty,
 	{
 		Class->EmitObjectReference( STRUCT_OFFSET( UClassProperty, MetaClass ) );
diff --git a/Engine/Source/Runtime/CoreUObject/Public/UObject/UnrealType.h b/Engine/Source/Runtime/CoreUObject/Public/UObject/UnrealType.h
--- a/Engine/Source/Runtime/CoreUObject/Public/UObject/UnrealType.h
+++ b/Engine/Source/Runtime/CoreUObject/Public/UObject/UnrealType.h
@@ -1711,7 +1711,7 @@ protected:
 		return false;
 	}
 
-	void CheckValidObject(void* Value) const;
+	virtual void CheckValidObject(void* Value) const;
 	// End of UObjectPropertyBase interface
 
 };
@@ -1926,6 +1926,9 @@ public:
 	virtual const TCHAR* ImportText_Internal( const TCHAR* Buffer, void* Data, int32 PortFlags, UObject* OwnerObject, FOutputDevice* ErrorText ) const OVERRIDE;
 	virtual bool SameType(const UProperty* Other) const OVERRIDE;
 	// End of UProperty interface
+
+protected:
+	virtual void CheckValidObject(void* Value) const OVERRIDE;
 };
 
 /*-----------------------------------------------------------------------------
