UE4: UEnum::GetValueAsString now returns the non-localized string value of the enum. Also, GetDisplayValueAsText was added to get the localized version.

diff --git a/Engine/Source/Runtime/CoreUObject/Public/UObject/Class.h b/Engine/Source/Runtime/CoreUObject/Public/UObject/Class.h
--- a/Engine/Source/Runtime/CoreUObject/Public/UObject/Class.h
+++ b/Engine/Source/Runtime/CoreUObject/Public/UObject/Class.h
@@ -1487,6 +1487,32 @@ public:
 		out_StringValue = GetValueAsString( EnumPath, EnumValue );
 	}
 
+	/**
+	 * @param EnumPath	- Full enum path
+	 * @param EnumValue - Enum value
+	 *
+	 * @return the localized display string associated with the specified enum value for the enum specified by a path
+	 */
+	template <typename T>
+	FORCEINLINE static FText GetDisplayValueAsText( const TCHAR* EnumPath, const T EnumValue )
+	{
+		// For the C++ enum.
+		static_assert(IS_ENUM(T), "Should only call this with enum types");
+		return GetDisplayValueAsText_Internal(EnumPath, (int32)EnumValue);
+	}
+
+	template <typename T>
+	FORCEINLINE static FText GetDisplayValueAsText( const TCHAR* EnumPath, const TEnumAsByte<T> EnumValue )
+	{
+		return GetDisplayValueAsText_Internal(EnumPath, (int32)EnumValue.GetValue());
+	}
+
+	template< class T >
+	FORCEINLINE static void GetDisplayValueAsText( const TCHAR* EnumPath, const T EnumValue, FText& out_TextValue )
+	{
+		out_TextValue = GetDisplayValueAsText( EnumPath, EnumValue );
+	}
+
 private:
 	/** Map of Enum Name to Map of Old Enum entry to New Enum entry */
 	static TMap<FName,TMap<FName,FName> > EnumRedirects;
@@ -1498,7 +1524,14 @@ private:
 	{
 		UEnum* EnumClass = FindObject<UEnum>( nullptr, EnumPath );
 		UE_CLOG( !EnumClass, LogClass, Fatal, TEXT("Couldn't find enum '%s'"), EnumPath );
-		return EnumClass->GetEnumString(Value);
+		return EnumClass->GetEnumName(Value);
+	}
+
+	FORCEINLINE static FText GetDisplayValueAsText_Internal( const TCHAR* EnumPath, const int32 Value )
+	{
+		UEnum* EnumClass = FindObject<UEnum>(nullptr, EnumPath);
+		UE_CLOG(!EnumClass, LogClass, Fatal, TEXT("Couldn't find enum '%s'"), EnumPath);
+		return EnumClass->GetEnumText(Value);
 	}
 };
 
