#ue4 Editoronly BP: ConstructorHelpers::FClassFinder more usable. It accepts the path to blueprint asset.
ConstructorHelpers::FClassFinder used in QAGame and ShooterGame.
StaticLoadClass has the same default values like StaticLoadObject.

diff --git a/Engine/Source/Runtime/CoreUObject/Public/UObject/ConstructorHelpers.h b/Engine/Source/Runtime/CoreUObject/Public/UObject/ConstructorHelpers.h
--- a/Engine/Source/Runtime/CoreUObject/Public/UObject/ConstructorHelpers.h
+++ b/Engine/Source/Runtime/CoreUObject/Public/UObject/ConstructorHelpers.h
@@ -52,6 +52,27 @@ namespace ConstructorHelpersInternal
 		}
 		return PackagePtr;
 	}
+
+	inline UClass* FindOrLoadClass(FString& PathName, UClass* BaseClass)
+	{
+		// If there is no dot, add ".<object_name>_C"
+		int32 PackageDelimPos = INDEX_NONE;
+		PathName.FindChar(TCHAR('.'), PackageDelimPos);
+		if (PackageDelimPos == INDEX_NONE)
+		{
+			int32 ObjectNameStart = INDEX_NONE;
+			PathName.FindLastChar(TCHAR('/'), ObjectNameStart);
+			if (ObjectNameStart != INDEX_NONE)
+			{
+				const FString ObjectName = PathName.Mid(ObjectNameStart + 1);
+				PathName += TCHAR('.');
+				PathName += ObjectName;
+				PathName += TCHAR('_');
+				PathName += TCHAR('C');
+			}
+		}
+		return StaticLoadClass(BaseClass, NULL, *PathName);
+	}
 }
 
 struct COREUOBJECT_API ConstructorHelpers
@@ -115,11 +136,11 @@ public:
 		TSubclassOf<T> Class;
 		FClassFinder(const TCHAR* ClassToFind)
 		{
-			FString ClassName;
-			FString PathName;
-			FPackageName::ParseExportTextPath(ClassToFind, &ClassName, &PathName);
-			Class = FindObject<UClass>( ANY_PACKAGE, *PathName );
-			check(Class);
+			CheckIfIsInConstructor(ClassToFind);
+			FString PathName(ClassToFind);
+			StripObjectClass(PathName, true);
+			Class = ConstructorHelpersInternal::FindOrLoadClass(PathName, T::StaticClass());
+			ValidateObject(*Class, PathName, *PathName);
 		}
 		bool Succeeded()
 		{
diff --git a/Engine/Source/Runtime/CoreUObject/Public/UObject/UObjectGlobals.h b/Engine/Source/Runtime/CoreUObject/Public/UObject/UObjectGlobals.h
--- a/Engine/Source/Runtime/CoreUObject/Public/UObject/UObjectGlobals.h
+++ b/Engine/Source/Runtime/CoreUObject/Public/UObject/UObjectGlobals.h
@@ -184,7 +184,7 @@ COREUOBJECT_API bool ParseObject( const TCHAR* Stream, const TCHAR* Match, UClas
  * @return The object that was loaded or found. NULL for a failure.
  */
 COREUOBJECT_API UObject* StaticLoadObject( UClass* Class, UObject* InOuter, const TCHAR* Name, const TCHAR* Filename = NULL, uint32 LoadFlags = LOAD_None, UPackageMap* Sandbox = NULL, bool bAllowObjectReconciliation = true );
-COREUOBJECT_API UClass* StaticLoadClass( UClass* BaseClass, UObject* InOuter, const TCHAR* Name, const TCHAR* Filename, uint32 LoadFlags, UPackageMap* Sandbox );
+COREUOBJECT_API UClass* StaticLoadClass(UClass* BaseClass, UObject* InOuter, const TCHAR* Name, const TCHAR* Filename = NULL, uint32 LoadFlags = LOAD_None, UPackageMap* Sandbox = NULL);
 /**
  * Create a new instance of an object.  The returned object will be fully initialized.  If InFlags contains RF_NeedsLoad (indicating that the object still needs to load its object data from disk), components
  * are not instanced (this will instead occur in PostLoad()).  The different between StaticConstructObject and StaticAllocateObject is that StaticConstructObject will also call the class constructor on the object
