Fixed multiple platform data being serialized out when saving for a specific platform.

diff --git a/Engine/Source/Runtime/CoreUObject/Private/UObject/SavePackage.cpp b/Engine/Source/Runtime/CoreUObject/Private/UObject/SavePackage.cpp
--- a/Engine/Source/Runtime/CoreUObject/Private/UObject/SavePackage.cpp
+++ b/Engine/Source/Runtime/CoreUObject/Private/UObject/SavePackage.cpp
@@ -2286,10 +2286,38 @@ EObjectMark UPackage::GetObjectMarksForTargetPlatform( const class ITargetPlatfo
  *
  * @return	true if the package was saved successfully.
  */
+#define UE_PROFILE_COOKSAVE 0
+#if UE_PROFILE_COOKSAVE
+
+#define UE_START_LOG_COOK_TIME(InFilename) double PreviousTime; double StartTime; PreviousTime = StartTime = FPlatformTime::Seconds(); \
+	UE_LOG( LogSavePackage, Log, TEXT("Starting save package for %s"), InFilename);
+
+#define UE_LOG_COOK_TIME(TimeType) \
+	{\
+		double CurrentTime = FPlatformTime::Seconds(); \
+		UE_LOG( LogSavePackage, Log, TEXT("TIMING: %s took %fms"), TimeType, (CurrentTime - PreviousTime) * 1000.0f ); \
+		PreviousTime = CurrentTime; \
+	}
+
+#define UE_FINISH_LOG_COOK_TIME()\
+	UE_LOG_COOK_TIME( TEXT("Unaccounted time")); \
+	{\
+		double CurrentTime = FPlatformTime::Seconds(); \
+		UE_LOG( LogSavePackage, Log, TEXT("Total save time: %fms"), (CurrentTime - StartTime) * 1000.0f ); \
+	}
+
+#else
+// don't do anything
+#define UE_START_LOG_COOK_TIME(InFilename)
+#define UE_LOG_COOK_TIME(TimeType) 
+#define UE_FINISH_LOG_COOK_TIME()
+#endif
+
 bool UPackage::SavePackage( UPackage* InOuter, UObject* Base, EObjectFlags TopLevelFlags, const TCHAR* Filename, 
 	FOutputDevice* Error, ULinkerLoad* Conform, bool bForceByteSwapping, bool bWarnOfLongFilename, uint32 SaveFlags, 
 	const class ITargetPlatform* TargetPlatform, const FDateTime&  FinalTimeStamp)
 {
+	UE_START_LOG_COOK_TIME( Filename );
 	
 	if (FPlatformProperties::HasEditorOnlyData())
 	{
@@ -2321,6 +2349,8 @@ bool UPackage::SavePackage( UPackage* InOuter, UObject* Base, EObjectFlags TopLe
 		(*GFlushStreamingFunc)();
 		FIOSystem::Get().BlockTillAllRequestsFinishedAndFlushHandles();
 
+		UE_LOG_COOK_TIME(TEXT("Flush Async loading"));
+
 		check(InOuter);
 		check(Filename);
 		uint32 Time=0; CLOCK_CYCLES(Time);
@@ -2447,6 +2477,8 @@ bool UPackage::SavePackage( UPackage* InOuter, UObject* Base, EObjectFlags TopLe
 		UnMarkAllObjects();
 
 
+		UE_LOG_COOK_TIME(TEXT("ResetLoadersForSave UnMarkAllObjects"));
+
 		// Size of serialized out package in bytes. This is before compression.
 		int32 PackageSize = INDEX_NONE;
 		{
@@ -2465,26 +2497,32 @@ bool UPackage::SavePackage( UPackage* InOuter, UObject* Base, EObjectFlags TopLe
 			ExportTaggerArchive.SetFilterEditorOnly(FilterEditorOnly);
 
 
+			UE_LOG_COOK_TIME(TEXT("TagPackageExports"));
 
 
 			// Kick off any Precaching required for the target platform to save these objects
 			// only need to do this if we are cooking a different platform then the one which is currently running
+			// TODO: if save package is cancelled then call ClearCache on each object
+			TArray<UObject*> CachedObjects;
 			if ( !!TargetPlatform )
 			{
 				TArray<UObject*> TagExpObjects;
 				GetObjectsWithAnyMarks(TagExpObjects, OBJECTMARK_TagExp);
 				for ( int Index =0; Index < TagExpObjects.Num(); ++Index)
 				{
-					if ( TagExpObjects[Index]->HasAnyMarks( OBJECTMARK_TagExp ) )
+					UObject *ExpObject = TagExpObjects[Index];
+					if ( ExpObject->HasAnyMarks( OBJECTMARK_TagExp ) )
 					{
-						TagExpObjects[Index]->BeginCacheForCookedPlatformData( TargetPlatform );
-
-						// UE_LOG(LogSavePackage, Log, TEXT("begincache %s for %s"), *TagExpObjects[Index]->GetFullName(), *TargetPlatform->PlatformName());
+						ExpObject->BeginCacheForCookedPlatformData( TargetPlatform );
+						CachedObjects.Add( ExpObject );
 					}
 				}
 			}
 
+			
 
+
+			UE_LOG_COOK_TIME(TEXT("BeginCacheForCookedPlatformData"));
 		
 			{
 				// set GIsSavingPackage here as it is now illegal to create any new object references; they potentially wouldn't be saved correctly
@@ -2503,6 +2541,8 @@ bool UPackage::SavePackage( UPackage* InOuter, UObject* Base, EObjectFlags TopLe
 				PackageExportTagger.TagPackageExports( ExportTaggerArchive, false );
 
 
+				UE_LOG_COOK_TIME(TEXT("TagPackageExports"));
+
 				GWarn->UpdateProgress( ++CurSaveStep, TotalSaveSteps );
 
 
@@ -2606,6 +2646,7 @@ bool UPackage::SavePackage( UPackage* InOuter, UObject* Base, EObjectFlags TopLe
 					}
 				}
 
+				UE_LOG_COOK_TIME(TEXT("Serialize Imports"));
 				
 				if ( EndSavingIfCancelled( Linker, TempFilename ) ) { return false; }
 				GWarn->UpdateProgress( ++CurSaveStep, TotalSaveSteps );
@@ -2679,6 +2720,7 @@ bool UPackage::SavePackage( UPackage* InOuter, UObject* Base, EObjectFlags TopLe
 					}
 				}
 
+				UE_LOG_COOK_TIME(TEXT("Mark Names"));
 
 				if ( EndSavingIfCancelled( Linker, TempFilename ) ) { return false; }
 				GWarn->UpdateProgress( ++CurSaveStep, TotalSaveSteps );
@@ -2840,6 +2882,8 @@ bool UPackage::SavePackage( UPackage* InOuter, UObject* Base, EObjectFlags TopLe
 				SavePackageState->UpdateLinkerWithMarkedNames(Linker);
 
 
+				UE_LOG_COOK_TIME(TEXT("Serialize Summary"));
+
 				if ( EndSavingIfCancelled( Linker, TempFilename ) ) { return false; }
 				GWarn->UpdateProgress( ++CurSaveStep, TotalSaveSteps );
 
@@ -2859,6 +2903,8 @@ bool UPackage::SavePackage( UPackage* InOuter, UObject* Base, EObjectFlags TopLe
 				}
 
 				
+				UE_LOG_COOK_TIME(TEXT("Serialize Names"));
+
 				if ( EndSavingIfCancelled( Linker, TempFilename ) ) { return false; }
 				GWarn->UpdateProgress( ++CurSaveStep, TotalSaveSteps );
 
@@ -2899,6 +2945,7 @@ bool UPackage::SavePackage( UPackage* InOuter, UObject* Base, EObjectFlags TopLe
 					}
 				}
 
+				UE_LOG_COOK_TIME(TEXT("Build Export Map"));
 
 				if ( EndSavingIfCancelled( Linker, TempFilename ) ) { return false; }
 				GWarn->UpdateProgress( ++CurSaveStep, TotalSaveSteps );
@@ -2915,6 +2962,8 @@ bool UPackage::SavePackage( UPackage* InOuter, UObject* Base, EObjectFlags TopLe
 				SeekFreeSorter.SortExports( Linker, Conform );
 				Linker->Summary.ExportCount = Linker->ExportMap.Num();
 				
+				UE_LOG_COOK_TIME(TEXT("Sort Exports"));
+
 				if ( EndSavingIfCancelled( Linker, TempFilename ) ) { return false; }
 				GWarn->UpdateProgress( ++CurSaveStep, TotalSaveSteps );
 
@@ -2988,6 +3037,7 @@ bool UPackage::SavePackage( UPackage* InOuter, UObject* Base, EObjectFlags TopLe
 				}
 
 
+
 				if ( EndSavingIfCancelled( Linker, TempFilename ) ) { return false; }
 				GWarn->UpdateProgress( ++CurSaveStep, TotalSaveSteps );
 
@@ -3046,6 +3096,8 @@ bool UPackage::SavePackage( UPackage* InOuter, UObject* Base, EObjectFlags TopLe
 				}
 
 
+				UE_LOG_COOK_TIME(TEXT("Build Dependency Map"));
+
 				if ( EndSavingIfCancelled( Linker, TempFilename ) ) { return false; }
 				GWarn->UpdateProgress( ++CurSaveStep, TotalSaveSteps );
 
@@ -3082,6 +3134,9 @@ bool UPackage::SavePackage( UPackage* InOuter, UObject* Base, EObjectFlags TopLe
 				int32 OffsetAfterImportMap = Linker->Tell();
 
 
+
+				UE_LOG_COOK_TIME(TEXT("Serialize Import Map"));
+
 				if ( EndSavingIfCancelled( Linker, TempFilename ) ) { return false; }
 				GWarn->UpdateProgress( ++CurSaveStep, TotalSaveSteps );
 
@@ -3095,6 +3150,8 @@ bool UPackage::SavePackage( UPackage* InOuter, UObject* Base, EObjectFlags TopLe
 				int32 OffsetAfterExportMap = Linker->Tell();
 
 
+				UE_LOG_COOK_TIME(TEXT("Serialize Export Map"));
+
 				if ( EndSavingIfCancelled( Linker, TempFilename ) ) { return false; }
 				GWarn->UpdateProgress( ++CurSaveStep, TotalSaveSteps );
 
@@ -3106,7 +3163,9 @@ bool UPackage::SavePackage( UPackage* InOuter, UObject* Base, EObjectFlags TopLe
 					TArray<FPackageIndex>& Depends = Linker->DependsMap[ i ];
 					*Linker << Depends;
 				}
-	
+
+
+				UE_LOG_COOK_TIME(TEXT("Serialize Dependency Map"));
 	
 				// Save thumbnails
 				UPackage::SaveThumbnails( InOuter, Linker );
@@ -3132,6 +3191,7 @@ bool UPackage::SavePackage( UPackage* InOuter, UObject* Base, EObjectFlags TopLe
 					Linker->StartScriptSHAGeneration();
 				}
 
+				UE_LOG_COOK_TIME(TEXT("SaveThumbNails Save AssetRegistryData SaveWorldLevelInfo"));
 
 				// Save exports.
 				int32 LastExportSaveStep = 0;
@@ -3222,6 +3282,8 @@ bool UPackage::SavePackage( UPackage* InOuter, UObject* Base, EObjectFlags TopLe
 				}
 
 
+				UE_LOG_COOK_TIME(TEXT("Serialize Exports"));
+
 				// if we want to generate the SHA key, get it out now that the package has finished saving
 				if (ScriptSHABytes && Linker->ContainsCode())
 				{
@@ -3269,6 +3331,8 @@ bool UPackage::SavePackage( UPackage* InOuter, UObject* Base, EObjectFlags TopLe
 				uint32 Tag = PACKAGE_FILE_TAG;
 				*Linker << Tag;
 
+				UE_LOG_COOK_TIME(TEXT("Serialize BulkData"));
+
 				// We capture the package size before the first seek to work with archives that don't report the file
 				// size correctly while the file is still being written to.
 				PackageSize = Linker->Tell();
@@ -3312,6 +3376,8 @@ bool UPackage::SavePackage( UPackage* InOuter, UObject* Base, EObjectFlags TopLe
 				}
 				
 
+				UE_LOG_COOK_TIME(TEXT("Serialize Import Map"));
+
 				check( Linker->Tell() == OffsetAfterImportMap );
 
 				// Save the export map.
@@ -3321,8 +3387,9 @@ bool UPackage::SavePackage( UPackage* InOuter, UObject* Base, EObjectFlags TopLe
 					*Linker << Linker->ExportMap[ i ];
 				}
 				check( Linker->Tell() == OffsetAfterExportMap );
-			
-				
+
+				UE_LOG_COOK_TIME(TEXT("Serialize Export Map"));
+
 				if ( EndSavingIfCancelled( Linker, TempFilename ) ) { return false; }
 				GWarn->UpdateProgress( ++CurSaveStep, TotalSaveSteps );
 
@@ -3494,6 +3561,11 @@ bool UPackage::SavePackage( UPackage* InOuter, UObject* Base, EObjectFlags TopLe
 				GWarn->UpdateProgress( ++CurSaveStep, TotalSaveSteps );
 			
 			}
+
+			for ( int CachedObjectIndex = 0; CachedObjectIndex < CachedObjects.Num(); ++CachedObjectIndex )
+			{
+				CachedObjects[CachedObjectIndex]->ClearCachedCookedPlatformData(TargetPlatform);
+			}
 		}
 		if( Success == true )
 		{
@@ -3504,6 +3576,8 @@ bool UPackage::SavePackage( UPackage* InOuter, UObject* Base, EObjectFlags TopLe
 		// We're done!
 		GWarn->UpdateProgress( TotalSaveSteps, TotalSaveSteps );
 
+		UE_FINISH_LOG_COOK_TIME();
+
 		return Success;
 	}
 	else
diff --git a/Engine/Source/Runtime/CoreUObject/Public/UObject/UObject.h b/Engine/Source/Runtime/CoreUObject/Public/UObject/UObject.h
--- a/Engine/Source/Runtime/CoreUObject/Public/UObject/UObject.h
+++ b/Engine/Source/Runtime/CoreUObject/Public/UObject/UObject.h
@@ -541,10 +541,23 @@ public:
 	 * Not called during normal loading of objects
 	 * 
 	 * @param	TargetPlatform	target platform to cache platform specific data for
-	 * @return	true if everything went well, false will abort save...
 	 */
 	virtual void BeginCacheForCookedPlatformData( const ITargetPlatform *TargetPlatform ) {  }
 
+	/**
+	 * Clears cached cooked platform data for specific platform
+	 * 
+	 * @param	TargetPlatform	target platform to cache platform specific data for
+	 */
+	virtual void ClearCachedCookedPlatformData( const ITargetPlatform *TargetPlatform ) {  }
+
+	/**
+	 * Clear all cached cooked platform data
+	 * 
+	 * @param	TargetPlatform	target platform to cache platform specific data for
+	 */
+	virtual void ClearAllCachedCookedPlatformData() { }
+
 
 	/**
 	 * Determine if this object has SomeObject in its archetype chain.
diff --git a/Engine/Source/Runtime/Engine/Classes/Materials/Material.h b/Engine/Source/Runtime/Engine/Classes/Materials/Material.h
--- a/Engine/Source/Runtime/Engine/Classes/Materials/Material.h
+++ b/Engine/Source/Runtime/Engine/Classes/Materials/Material.h
@@ -699,7 +699,7 @@ private:
 	FMaterialResource* MaterialResources[EMaterialQualityLevel::Num][ERHIFeatureLevel::Num];
 
 	/** Material resources being cached for cooking. */
-	TArray<FMaterialResource*> CachedMaterialResourcesForCooking;
+	TMap<const class ITargetPlatform*, TArray<FMaterialResource*>> CachedMaterialResourcesForCooking;
 
 	/** Fence used to guarantee that the RT is finished using various resources in this UMaterial before cleanup. */
 	FRenderCommandFence ReleaseFence;
@@ -767,6 +767,9 @@ public:
 	ENGINE_API virtual void PostDuplicate(bool bDuplicateForPIE) OVERRIDE;
 	ENGINE_API virtual void PostLoad() OVERRIDE;
 	ENGINE_API virtual void BeginCacheForCookedPlatformData( const ITargetPlatform *TargetPlatform ) OVERRIDE;
+	ENGINE_API virtual void ClearCachedCookedPlatformData( const ITargetPlatform *TargetPlatform ) OVERRIDE;
+	ENGINE_API virtual void ClearAllCachedCookedPlatformData() OVERRIDE;
+
 #if WITH_EDITOR
 	ENGINE_API virtual void PreEditChange(UProperty* PropertyAboutToChange) OVERRIDE;
 	ENGINE_API virtual void PostEditChangeProperty(FPropertyChangedEvent& PropertyChangedEvent) OVERRIDE;
@@ -892,6 +895,14 @@ public:
 	/** Determines whether each quality level has different nodes by inspecting the material's expressions. */
 	void GetQualityLevelNodeUsage(TArray<bool, TInlineAllocator<EMaterialQualityLevel::Num> >& QualityLevelsUsed);
 
+
+	/**
+	 * Cache the expression texture references for this UMaterial 
+	 * if the cache is not filled then it will rebuild the texture references
+	 * see also RebuildExpressionTextureReferences
+	 */
+	void CacheExpressionTextureReferences();
+
 private:
 	/**
 	 * Flush existing resource shader maps and resets the material resource's Ids.
diff --git a/Engine/Source/Runtime/Engine/Classes/Materials/MaterialInstance.h b/Engine/Source/Runtime/Engine/Classes/Materials/MaterialInstance.h
--- a/Engine/Source/Runtime/Engine/Classes/Materials/MaterialInstance.h
+++ b/Engine/Source/Runtime/Engine/Classes/Materials/MaterialInstance.h
@@ -196,7 +196,7 @@ private:
 	FMaterialResource* StaticPermutationMaterialResources[EMaterialQualityLevel::Num][ERHIFeatureLevel::Num];
 
 	/** Material resources being cached for cooking. */
-	TArray<FMaterialResource*> CachedMaterialResourcesForCooking;
+	TMap<const class ITargetPlatform*, TArray<FMaterialResource*>> CachedMaterialResourcesForCooking;
 
 	/** Fence used to guarantee that the RT is finished using various resources in this UMaterial before cleanup. */
 	FRenderCommandFence ReleaseFence;
@@ -262,6 +262,8 @@ public:
 	virtual SIZE_T GetResourceSize(EResourceSizeMode::Type Mode) OVERRIDE;
 	virtual void PostInitProperties() OVERRIDE;	
 	virtual void BeginCacheForCookedPlatformData(  const ITargetPlatform *TargetPlatform ) OVERRIDE;
+	virtual void ClearCachedCookedPlatformData( const ITargetPlatform *TargetPlatform ) OVERRIDE;
+	virtual void ClearAllCachedCookedPlatformData() OVERRIDE;
 	virtual void Serialize(FArchive& Ar) OVERRIDE;
 	virtual void PostLoad() OVERRIDE;
 	virtual void BeginDestroy() OVERRIDE;
diff --git a/Engine/Source/Runtime/Engine/Private/Materials/Material.cpp b/Engine/Source/Runtime/Engine/Private/Materials/Material.cpp
--- a/Engine/Source/Runtime/Engine/Private/Materials/Material.cpp
+++ b/Engine/Source/Runtime/Engine/Private/Materials/Material.cpp
@@ -346,23 +346,63 @@ void UMaterialInterface::AssertDefaultMaterialsPostLoaded()
 	}
 }
 
-void SerializeInlineShaderMaps(const TArray<FMaterialResource*>& MaterialResourcesToSave, FArchive& Ar, FMaterialResource* OutMaterialResourcesLoaded[][ERHIFeatureLevel::Num])
+void SerializeInlineShaderMaps(const TMap<const ITargetPlatform*,TArray<FMaterialResource*>>& PlatformMaterialResourcesToSave, FArchive& Ar, FMaterialResource* OutMaterialResourcesLoaded[][ERHIFeatureLevel::Num])
 {
+
+	/*if ( Ar.IsLoading() || (Ar.IsCooking() == false) )
+	{
+		// shouldn't ever load into this saved material resources array
+		// if we are not cooking then SerializeInlineShaderMaps will ignore the SavedMaterialResources array
+		TArray<FMaterialResource*> SavedMaterialResources;
+		SerializeInlineShaderMaps( SavedMaterialResources, Ar, MaterialResources );
+		check( SavedMaterialResources.Num() == 0);
+	}
+	else if ( Ar.IsSaving() )
+	{
+		check( Ar.IsCooking() );
+		const TArray<FMaterialResource*> *CachedMaterialResourcesForPlatform = CachedMaterialResourcesForCooking.Find(Ar.CookingTarget());
+		// check( CachedMaterialResourcesForPlatform != NULL || (!HasAnyMarks( OBJECTMARK_TagExp ))  );
+		if ( CachedMaterialResourcesForPlatform != NULL )
+		{
+			SerializeInlineShaderMaps(*CachedMaterialResourcesForPlatform, Ar, MaterialResources);
+		}
+		else
+		{
+			// if we are saving to a linker and we are cooking then we should have definatly have had a cached material for the target platform (see UPackage::SavePackage() -> UObject::BeginCacheResourceForPlatform
+			//  if this check is triggering most likely because some object didn't have BeginCacheResourceForPlatform called from SavePackage OR
+			//  the archive isn't actually saving and GetLinker is returning a valid value for some other reason
+			check( Ar.GetLinker() == NULL );
+		}
+
+	}*/
+
+
+
 	if (Ar.IsSaving())
 	{
 		int32 NumResourcesToSave = 0;
-
+		const TArray<FMaterialResource*> *MaterialResourcesToSavePtr = NULL;
 		if (Ar.IsCooking())
 		{
-			NumResourcesToSave = MaterialResourcesToSave.Num();
+			MaterialResourcesToSavePtr = PlatformMaterialResourcesToSave.Find( Ar.CookingTarget() );
+			check( MaterialResourcesToSavePtr != NULL || (Ar.GetLinker()==NULL) );
+			if (MaterialResourcesToSavePtr!= NULL )
+			{
+				NumResourcesToSave = MaterialResourcesToSavePtr->Num();
+			}
 		}
 
 		Ar << NumResourcesToSave;
 
-		for (int32 ResourceIndex = 0; ResourceIndex < NumResourcesToSave; ResourceIndex++)
+		if ( MaterialResourcesToSavePtr )
 		{
-			MaterialResourcesToSave[ResourceIndex]->SerializeInlineShaderMap(Ar);
+			const TArray<FMaterialResource*> &MaterialResourcesToSave = *MaterialResourcesToSavePtr;
+			for (int32 ResourceIndex = 0; ResourceIndex < NumResourcesToSave; ResourceIndex++)
+			{
+				MaterialResourcesToSave[ResourceIndex]->SerializeInlineShaderMap(Ar);
+			}
 		}
+		
 	}
 	else if (Ar.IsLoading())
 	{
@@ -1499,6 +1539,14 @@ void UMaterial::RebuildMaterialParameterCollectionInfo()
 	}
 }
 
+void UMaterial::CacheExpressionTextureReferences()
+{
+	if ( ExpressionTextureReferences.Num() <= 0 )
+	{
+		RebuildExpressionTextureReferences();
+	}
+}
+
 void UMaterial::RebuildExpressionTextureReferences()
 {
 	ExpressionTextureReferences.Empty();
@@ -1550,7 +1598,38 @@ void UMaterial::Serialize(FArchive& Ar)
 
 	if (Ar.UE4Ver() >= VER_UE4_PURGED_FMATERIAL_COMPILE_OUTPUTS)
 	{
-		SerializeInlineShaderMaps(CachedMaterialResourcesForCooking, Ar, MaterialResources);
+
+		SerializeInlineShaderMaps( CachedMaterialResourcesForCooking, Ar, MaterialResources );
+
+
+
+		/*if ( Ar.IsLoading() || (Ar.IsCooking() == false) )
+		{
+			// shouldn't ever load into this saved material resources array
+			// if we are not cooking then SerializeInlineShaderMaps will ignore the SavedMaterialResources array
+			TArray<FMaterialResource*> SavedMaterialResources;
+			SerializeInlineShaderMaps( SavedMaterialResources, Ar, MaterialResources );
+			check( SavedMaterialResources.Num() == 0);
+		}
+		else if ( Ar.IsSaving() )
+		{
+			check( Ar.IsCooking() );
+			const TArray<FMaterialResource*> *CachedMaterialResourcesForPlatform = CachedMaterialResourcesForCooking.Find(Ar.CookingTarget());
+			// check( CachedMaterialResourcesForPlatform != NULL || (!HasAnyMarks( OBJECTMARK_TagExp ))  );
+			if ( CachedMaterialResourcesForPlatform != NULL )
+			{
+				SerializeInlineShaderMaps(*CachedMaterialResourcesForPlatform, Ar, MaterialResources);
+			}
+			else
+			{
+				// if we are saving to a linker and we are cooking then we should have definatly have had a cached material for the target platform (see UPackage::SavePackage() -> UObject::BeginCacheResourceForPlatform
+				//  if this check is triggering most likely because some object didn't have BeginCacheResourceForPlatform called from SavePackage OR
+				//  the archive isn't actually saving and GetLinker is returning a valid value for some other reason
+				check( Ar.GetLinker() == NULL );
+			}
+			
+		}*/
+		
 	}
 	else
 	{
@@ -1899,20 +1978,7 @@ void UMaterial::PostLoad()
 	{
 		SCOPE_SECONDS_COUNTER(MaterialLoadTime);
 
-		const TArray<FName>& DesiredShaderFormats = GetTargetShaderFormats();
-
-		if (DesiredShaderFormats.Num())
-		{
-			// Cache for all the shader formats that the cooking target requires
-			for (int32 FormatIndex = 0; FormatIndex < DesiredShaderFormats.Num(); FormatIndex++)
-			{
-				const EShaderPlatform TargetPlatform = ShaderFormatToLegacyShaderPlatform(DesiredShaderFormats[FormatIndex]);
-
-				// Begin caching shaders for the target platform and store the material resource being compiled into CachedMaterialResourcesForCooking
-				CacheResourceShadersForCooking(TargetPlatform, CachedMaterialResourcesForCooking);
-			}
-		}
-		else if (GAllowCompilationInPostLoad) //@todo, this is broken, what if we want to save cooked packages out of the editor? It can't be either cooking or not; sometimes we will want both.
+		if (FApp::CanEverRender()) 
 		{
 			CacheResourceShadersForRendering(false);
 		}
@@ -1939,17 +2005,58 @@ void UMaterial::BeginCacheForCookedPlatformData( const ITargetPlatform *TargetPl
 	TArray<FName> DesiredShaderFormats;
 	TargetPlatform->GetShaderFormats( DesiredShaderFormats );
 
-	if (DesiredShaderFormats.Num())
+	TArray<FMaterialResource*> *CachedMaterialResourcesForPlatform = CachedMaterialResourcesForCooking.Find( TargetPlatform );
+
+	if ( CachedMaterialResourcesForPlatform == NULL )
+	{
+		check( CachedMaterialResourcesForPlatform == NULL );
+
+		CachedMaterialResourcesForCooking.Add( TargetPlatform );
+		CachedMaterialResourcesForPlatform = CachedMaterialResourcesForCooking.Find( TargetPlatform );
+
+		check( CachedMaterialResourcesForPlatform != NULL );
+
+		if (DesiredShaderFormats.Num())
+		{
+			// Cache for all the shader formats that the cooking target requires
+			for (int32 FormatIndex = 0; FormatIndex < DesiredShaderFormats.Num(); FormatIndex++)
+			{
+				const EShaderPlatform TargetPlatform = ShaderFormatToLegacyShaderPlatform(DesiredShaderFormats[FormatIndex]);
+
+				// Begin caching shaders for the target platform and store the material resource being compiled into CachedMaterialResourcesForCooking
+				CacheResourceShadersForCooking(TargetPlatform, *CachedMaterialResourcesForPlatform);
+			}
+		}
+	}
+}
+
+
+
+
+void UMaterial::ClearCachedCookedPlatformData( const ITargetPlatform *TargetPlatform )
+{
+	TArray<FMaterialResource*> *CachedMaterialResourcesForPlatform = CachedMaterialResourcesForCooking.Find( TargetPlatform );
+	if ( CachedMaterialResourcesForPlatform != NULL )
 	{
-		// Cache for all the shader formats that the cooking target requires
-		for (int32 FormatIndex = 0; FormatIndex < DesiredShaderFormats.Num(); FormatIndex++)
+		for (int32 CachedResourceIndex = 0; CachedResourceIndex < CachedMaterialResourcesForPlatform->Num(); CachedResourceIndex++)
 		{
-			const EShaderPlatform TargetPlatform = ShaderFormatToLegacyShaderPlatform(DesiredShaderFormats[FormatIndex]);
+			delete (*CachedMaterialResourcesForPlatform)[CachedResourceIndex];
+		}
+	}
+	CachedMaterialResourcesForCooking.Remove( TargetPlatform );
+}
 
-			// Begin caching shaders for the target platform and store the material resource being compiled into CachedMaterialResourcesForCooking
-			CacheResourceShadersForCooking(TargetPlatform, CachedMaterialResourcesForCooking);
+void UMaterial::ClearAllCachedCookedPlatformData()
+{
+	for ( auto It : CachedMaterialResourcesForCooking )
+	{
+		TArray<FMaterialResource*> &CachedMaterialResourcesForPlatform = It.Value;
+		for (int32 CachedResourceIndex = 0; CachedResourceIndex < CachedMaterialResourcesForPlatform.Num(); CachedResourceIndex++)
+		{
+			delete (CachedMaterialResourcesForPlatform)[CachedResourceIndex];
 		}
 	}
+	CachedMaterialResourcesForCooking.Empty();
 }
 
 #if WITH_EDITOR
@@ -2461,10 +2568,7 @@ void UMaterial::ReleaseResources()
 		}
 	}
 
-	for (int32 CachedResourceIndex = 0; CachedResourceIndex < CachedMaterialResourcesForCooking.Num(); CachedResourceIndex++)
-	{
-		delete CachedMaterialResourcesForCooking[CachedResourceIndex];
-	}
+	ClearAllCachedCookedPlatformData();
 
 	for (int32 InstanceIndex = 0; InstanceIndex < 3; ++InstanceIndex)
 	{
diff --git a/Engine/Source/Runtime/Engine/Private/Materials/MaterialInstance.cpp b/Engine/Source/Runtime/Engine/Private/Materials/MaterialInstance.cpp
--- a/Engine/Source/Runtime/Engine/Private/Materials/MaterialInstance.cpp
+++ b/Engine/Source/Runtime/Engine/Private/Materials/MaterialInstance.cpp
@@ -1171,17 +1171,7 @@ void UMaterialInstance::InitStaticPermutation()
 
 	const TArray<FName>& DesiredShaderFormats = GetTargetShaderFormats();
 
-	if (DesiredShaderFormats.Num())
-	{
-		// Cache shaders for each shader format, storing the results in CachedMaterialResourcesForCooking so they will be available during saving
-		for (int32 FormatIndex = 0; FormatIndex < DesiredShaderFormats.Num(); FormatIndex++)
-		{
-			const EShaderPlatform TargetPlatform = ShaderFormatToLegacyShaderPlatform(DesiredShaderFormats[FormatIndex]);
-
-			CacheResourceShadersForCooking(TargetPlatform, CachedMaterialResourcesForCooking);
-		}
-	}
-	else //@todo, this is broken, what if we want to save cooked packages out of the editor? It can't be either cooking or not; sometimes we will want both.
+	if ( FApp::CanEverRender() ) 
 	{
 		// Cache shaders for the current platform to be used for rendering
 		CacheResourceShadersForRendering();
@@ -1334,6 +1324,8 @@ void UMaterialInstance::CacheShadersForResources(EShaderPlatform ShaderPlatform,
 
 	UMaterial* BaseMaterial = GetMaterial();
 
+	BaseMaterial->CacheExpressionTextureReferences();
+
 	for (int32 ResourceIndex = 0; ResourceIndex < ResourcesToCache.Num(); ResourceIndex++)
 	{
 		FMaterialResource* CurrentResource = ResourcesToCache[ResourceIndex];
@@ -1492,13 +1484,54 @@ void UMaterialInstance::BeginCacheForCookedPlatformData( const ITargetPlatform *
 {
 	TArray<FName> DesiredShaderFormats;
 	TargetPlatform->GetShaderFormats( DesiredShaderFormats );
-	// Cache shaders for each shader format, storing the results in CachedMaterialResourcesForCooking so they will be available during saving
-	for (int32 FormatIndex = 0; FormatIndex < DesiredShaderFormats.Num(); FormatIndex++)
+
+	TArray<FMaterialResource*> *CachedMaterialResourcesForPlatform = CachedMaterialResourcesForCooking.Find( TargetPlatform );
+
+	if ( CachedMaterialResourcesForPlatform == NULL )
+	{
+		check( CachedMaterialResourcesForPlatform == NULL );
+
+		CachedMaterialResourcesForCooking.Add( TargetPlatform );
+		CachedMaterialResourcesForPlatform = CachedMaterialResourcesForCooking.Find( TargetPlatform );
+
+		check( CachedMaterialResourcesForPlatform != NULL );
+
+		// Cache shaders for each shader format, storing the results in CachedMaterialResourcesForCooking so they will be available during saving
+		for (int32 FormatIndex = 0; FormatIndex < DesiredShaderFormats.Num(); FormatIndex++)
+		{
+			const EShaderPlatform TargetPlatform = ShaderFormatToLegacyShaderPlatform(DesiredShaderFormats[FormatIndex]);
+
+			CacheResourceShadersForCooking(TargetPlatform, *CachedMaterialResourcesForPlatform );
+		}
+	}
+}
+
+void UMaterialInstance::ClearCachedCookedPlatformData( const ITargetPlatform *TargetPlatform )
+{
+	TArray<FMaterialResource*> *CachedMaterialResourcesForPlatform = CachedMaterialResourcesForCooking.Find( TargetPlatform );
+	if ( CachedMaterialResourcesForPlatform != NULL )
 	{
-		const EShaderPlatform TargetPlatform = ShaderFormatToLegacyShaderPlatform(DesiredShaderFormats[FormatIndex]);
+		for (int32 CachedResourceIndex = 0; CachedResourceIndex < CachedMaterialResourcesForPlatform->Num(); CachedResourceIndex++)
+		{
+			delete (*CachedMaterialResourcesForPlatform)[CachedResourceIndex];
+		}
+	}
+	CachedMaterialResourcesForCooking.Remove( TargetPlatform );
+}
 
-		CacheResourceShadersForCooking(TargetPlatform, CachedMaterialResourcesForCooking);
+
+void UMaterialInstance::ClearAllCachedCookedPlatformData()
+{
+	for ( auto It : CachedMaterialResourcesForCooking )
+	{
+		TArray<FMaterialResource*> &CachedMaterialResourcesForPlatform = It.Value;
+		for (int32 CachedResourceIndex = 0; CachedResourceIndex < CachedMaterialResourcesForPlatform.Num(); CachedResourceIndex++)
+		{
+			delete (CachedMaterialResourcesForPlatform)[CachedResourceIndex];
+		}
 	}
+
+	CachedMaterialResourcesForCooking.Empty();
 }
 
 void UMaterialInstance::Serialize(FArchive& Ar)
@@ -1512,7 +1545,31 @@ void UMaterialInstance::Serialize(FArchive& Ar)
 		{
 			StaticParameters.Serialize(Ar);
 
-			SerializeInlineShaderMaps(CachedMaterialResourcesForCooking, Ar, StaticPermutationMaterialResources);
+			SerializeInlineShaderMaps( CachedMaterialResourcesForCooking, Ar, StaticPermutationMaterialResources );
+
+
+			/*if ( Ar.IsLoading() || (Ar.IsCooking() == false) )
+			{
+				// shouldn't ever load into this saved material resources array
+				// if we are not cooking then SerializeInlineShaderMaps will ignore the SavedMaterialResources array
+				TArray<FMaterialResource*> SavedMaterialResources;
+				SerializeInlineShaderMaps( SavedMaterialResources, Ar, StaticPermutationMaterialResources );
+				check( SavedMaterialResources.Num() == 0);
+			}
+			else if ( Ar.IsSaving() )
+			{
+				check( Ar.IsCooking() );
+				const TArray<FMaterialResource*> *CachedMaterialResourcesForPlatform = CachedMaterialResourcesForCooking.Find(Ar.CookingTarget());
+				// check( CachedMaterialResourcesForPlatform != NULL || (!HasAnyMarks( OBJECTMARK_TagExp ))  );
+				if ( CachedMaterialResourcesForPlatform != NULL )
+				{
+					SerializeInlineShaderMaps(*CachedMaterialResourcesForPlatform, Ar, StaticPermutationMaterialResources);
+				}
+				else
+				{
+					check( Ar.GetLinker() == NULL );
+				}
+			}*/
 		}
 		else
 		{
@@ -1671,10 +1728,7 @@ void UMaterialInstance::FinishDestroy()
 		}
 	}
 
-	for (int32 CachedResourceIndex = 0; CachedResourceIndex < CachedMaterialResourcesForCooking.Num(); CachedResourceIndex++)
-	{
-		delete CachedMaterialResourcesForCooking[CachedResourceIndex];
-	}
+	ClearAllCachedCookedPlatformData();
 
 	if( BasePropertyOverrides )
 	{
