Add templated UClass::IsChildOf

diff --git a/Engine/Source/Runtime/CoreUObject/Public/UObject/Class.h b/Engine/Source/Runtime/CoreUObject/Public/UObject/Class.h
--- a/Engine/Source/Runtime/CoreUObject/Public/UObject/Class.h
+++ b/Engine/Source/Runtime/CoreUObject/Public/UObject/Class.h
@@ -303,6 +303,12 @@ public:
 		PropertiesSize = NewSize;
 	}
 
+	template<class T>
+	bool IsChildOf() const
+	{
+		return IsChildOf(T::StaticClass());
+	}
+
 	bool IsChildOf( const UStruct* SomeBase ) const
 	{
 		for (const UStruct* Struct = this; Struct; Struct = Struct->GetSuperStruct())
