#ue4 GitHub PullReq 42:  Fix for reparenting "within" classes

diff --git a/Engine/Source/Runtime/CoreUObject/Private/UObject/Obj.cpp b/Engine/Source/Runtime/CoreUObject/Private/UObject/Obj.cpp
--- a/Engine/Source/Runtime/CoreUObject/Private/UObject/Obj.cpp
+++ b/Engine/Source/Runtime/CoreUObject/Private/UObject/Obj.cpp
@@ -57,8 +57,8 @@ bool UObject::Rename( const TCHAR* InName, UObject* NewOuter, ERenameFlags Flags
 	FMetaDataUtilities::FMoveMetadataHelperContext MoveMetaData(this, true);
 #endif //WITH_EDITOR
 
-	// Check that we are not renaming a within object into an Outer of the wrong type.
-	if( NewOuter && !NewOuter->IsA(GetClass()->ClassWithin) )	
+	// Check that we are not renaming a within object into an Outer of the wrong type, unless we're renaming the CDO of a Blueprint.
+	if( NewOuter && !NewOuter->IsA(GetClass()->ClassWithin) && !HasAnyFlags(RF_ClassDefaultObject))	
 	{
 		UE_LOG(LogObj, Fatal, TEXT("Cannot rename %s into Outer %s as it is not of type %s"), 
 			*GetFullName(), 
