#ue4 GetAuthoritativeClass cannot be used outside the editor. (making blueprints editoronly).

#codereveiw Robert.Manuszewski

diff --git a/Engine/Source/Runtime/CoreUObject/Public/UObject/Class.h b/Engine/Source/Runtime/CoreUObject/Public/UObject/Class.h
--- a/Engine/Source/Runtime/CoreUObject/Public/UObject/Class.h
+++ b/Engine/Source/Runtime/CoreUObject/Public/UObject/Class.h
@@ -1706,6 +1706,7 @@ public:
 		);
 #endif
 
+#if WITH_EDITOR
 	/**
 	 * If there are potentially multiple versions of this class (e.g. blueprint generated classes), this function will return the authoritative version, which should be used for references
 	 *
@@ -1715,6 +1716,7 @@ public:
 	{
 		return this;
 	}
+#endif
 
 	/**
 	 * Add a native function to the internal native function table
