TTP #331951: Removing a condition where we were doing an unneeded Preload() for cooked builds (cooked assets were getting loaded a second time after they had already been finalized).

diff --git a/Engine/Source/Runtime/CoreUObject/Private/UObject/LinkerLoad.cpp b/Engine/Source/Runtime/CoreUObject/Private/UObject/LinkerLoad.cpp
--- a/Engine/Source/Runtime/CoreUObject/Private/UObject/LinkerLoad.cpp
+++ b/Engine/Source/Runtime/CoreUObject/Private/UObject/LinkerLoad.cpp
@@ -3157,7 +3157,7 @@ UObject* ULinkerLoad::CreateExport( int32 Index )
  					// Preload classes on first access.  Note that this may update the Export.Object, so ClassObject is not guaranteed to be valid after this point
 					// If we're async loading on a cooked build we can skip this as there's no chance we will need to recompile the class. 
 					// Preload will be called during async package tick when the data has been precached
-					if( !FPlatformProperties::RequiresCookedData() || !GIsAsyncLoading )
+					if( !FPlatformProperties::RequiresCookedData() )
 					{
  						Preload( Export.Object );
 					}
