Add useful parameters to UObject::PostRename

diff --git a/Engine/Source/Runtime/CoreUObject/Private/UObject/Obj.cpp b/Engine/Source/Runtime/CoreUObject/Private/UObject/Obj.cpp
--- a/Engine/Source/Runtime/CoreUObject/Private/UObject/Obj.cpp
+++ b/Engine/Source/Runtime/CoreUObject/Private/UObject/Obj.cpp
@@ -164,7 +164,7 @@ bool UObject::Rename( const TCHAR* InName, UObject* NewOuter, ERenameFlags Flags
 		Redirector->DestinationObject = this;
 	}
 
-	PostRename();
+	PostRename(OldOuter, OldName);
 
 	return true;
 }
diff --git a/Engine/Source/Runtime/CoreUObject/Public/UObject/UObject.h b/Engine/Source/Runtime/CoreUObject/Public/UObject/UObject.h
--- a/Engine/Source/Runtime/CoreUObject/Public/UObject/UObject.h
+++ b/Engine/Source/Runtime/CoreUObject/Public/UObject/UObject.h
@@ -225,7 +225,7 @@ public:
 #endif // WITH_EDITOR
 
 	// @todo document
-	virtual void PostRename() {}
+	virtual void PostRename(UObject* OldOuter, const FName OldName) {}
 	
 	/**
 	 * Called after duplication & serialization and before PostLoad. Used to e.g. make sure UStaticMesh's UModel gets copied as well.
