Fix spelling error

diff --git a/Engine/Source/Runtime/CoreUObject/Private/Serialization/AsyncLoading.cpp b/Engine/Source/Runtime/CoreUObject/Private/Serialization/AsyncLoading.cpp
--- a/Engine/Source/Runtime/CoreUObject/Private/Serialization/AsyncLoading.cpp
+++ b/Engine/Source/Runtime/CoreUObject/Private/Serialization/AsyncLoading.cpp
@@ -334,7 +334,7 @@ EAsyncPackageState::Type FAsyncPackage::CreateLinker()
  */
 EAsyncPackageState::Type FAsyncPackage::FinishLinker()
 {
-	if( !Linker->HasFinishedInitializtion() )
+	if( !Linker->HasFinishedInitialization() )
 	{
 		SCOPE_CYCLE_COUNTER(STAT_FAsyncPackage_FinishLinker);
 		LastObjectWorkWasPerformedOn	= Linker->LinkerRoot;
@@ -431,7 +431,7 @@ bool FAsyncPackage::AddUniqueLinkerDependencyPackage(int32 CurrentPackageIndex,
 {
 	if (ContainsDependencyPackage(PendingImportedPackages, PendingImport.GetPackageName()) == INDEX_NONE)
 	{
-		if (PendingImport.Linker == NULL || !PendingImport.Linker->HasFinishedInitializtion())
+		if (PendingImport.Linker == NULL || !PendingImport.Linker->HasFinishedInitialization())
 		{
 			AddImportDependency(CurrentPackageIndex, PendingImport.GetPackageName());
 			UE_LOG(LogStreaming, Verbose, TEXT("  Adding linker dependency %s"), *PendingImport.GetPackageName());
@@ -505,7 +505,7 @@ EAsyncPackageState::Type FAsyncPackage::LoadImports()
 			if (PendingAsyncPackageIndex != INDEX_NONE)
 			{
 				FAsyncPackage& PendingPackage = GObjAsyncPackages[PendingAsyncPackageIndex];
-				if (PendingPackage.Linker == NULL || !PendingPackage.Linker->HasFinishedInitializtion())
+				if (PendingPackage.Linker == NULL || !PendingPackage.Linker->HasFinishedInitialization())
 				{
 					// Add this import to the dependency list.
 					AddUniqueLinkerDependencyPackage(AsyncQueueIndex, PendingPackage);
diff --git a/Engine/Source/Runtime/CoreUObject/Public/UObject/Linker.h b/Engine/Source/Runtime/CoreUObject/Public/UObject/Linker.h
--- a/Engine/Source/Runtime/CoreUObject/Public/UObject/Linker.h
+++ b/Engine/Source/Runtime/CoreUObject/Public/UObject/Linker.h
@@ -1338,7 +1338,7 @@ public:
 	 *
 	 * @return true if initialized, false if pending.
 	 */
-	bool HasFinishedInitializtion() const
+	bool HasFinishedInitialization() const
 	{
         return bHasFinishedInitialization;
 	}
