Fix opening map from command line finding an asset of the same name
Added ambiguous package warnings to FPackageName::SearchForPackageOnDisk
TTP #332500

diff --git a/Engine/Source/Runtime/CoreUObject/Private/Misc/PackageName.cpp b/Engine/Source/Runtime/CoreUObject/Private/Misc/PackageName.cpp
--- a/Engine/Source/Runtime/CoreUObject/Private/Misc/PackageName.cpp
+++ b/Engine/Source/Runtime/CoreUObject/Private/Misc/PackageName.cpp
@@ -669,7 +669,7 @@ bool FPackageName::SearchForPackageOnDisk(const FString& PackageName, FString* O
 			}
 		}
 
-		const FString PackageWildcard = PackageName + TEXT(".*");
+		const FString PackageWildcard = (PackageName.Find(TEXT(".")) != INDEX_NONE ? PackageName : PackageName + TEXT(".*"));
 		TArray<FString> Results;
 
 		for (int32 PathIndex = 0; PathIndex < Paths.Num() && !bResult; ++PathIndex)
@@ -677,7 +677,7 @@ bool FPackageName::SearchForPackageOnDisk(const FString& PackageName, FString* O
 			// Search directly on disk. Very slow!
 			IFileManager::Get().FindFilesRecursive(Results, *Paths[PathIndex], *PackageWildcard, true, false);
 
-			for (int32 FileIndex = 0; FileIndex < Results.Num() && !bResult; ++FileIndex)
+			for (int32 FileIndex = 0; FileIndex < Results.Num(); ++FileIndex)
 			{			
 				FString Filename(Results[FileIndex]);
 				if (IsPackageFilename(Results[FileIndex]))
@@ -688,12 +688,26 @@ bool FPackageName::SearchForPackageOnDisk(const FString& PackageName, FString* O
 					{
 						if (OutLongPackageName)
 						{
-							*OutLongPackageName = LongPackageName;
+							if (bResult)
+							{
+								UE_LOG(LogPackageName, Warning, TEXT("Found ambiguous long package name for '%s'. Returning '%s', but could also be '%s'."), *PackageName, **OutLongPackageName, *LongPackageName );
+							}
+							else
+							{
+								*OutLongPackageName = LongPackageName;
+							}
 						}
 						if (OutFilename)
 						{
 							FPaths::MakeStandardFilename(Filename);
-							*OutFilename = Filename;
+							if (bResult)
+							{
+								UE_LOG(LogPackageName, Warning, TEXT("Found ambiguous file name for '%s'. Returning '%s', but could also be '%s'."), *PackageName, **OutFilename, *Filename);
+							}
+							else
+							{
+								*OutFilename = Filename;
+							}
 						}
 						bResult = true;
 					}
diff --git a/Engine/Source/Runtime/Engine/Private/URL.cpp b/Engine/Source/Runtime/Engine/Private/URL.cpp
--- a/Engine/Source/Runtime/Engine/Private/URL.cpp
+++ b/Engine/Source/Runtime/Engine/Private/URL.cpp
@@ -333,7 +333,7 @@ FURL::FURL( FURL* Base, const TCHAR* TextURL, ETravelType Type )
 			{
 				Map = FPackageName::FilenameToLongPackageName(URL);
 			}
-			else if (!FPackageName::DoesPackageNameContainInvalidCharacters(URL, &MapNameError) && FPackageName::SearchForPackageOnDisk(URL, &MapFullName))
+			else if (!FPackageName::DoesPackageNameContainInvalidCharacters(URL, &MapNameError) && FPackageName::SearchForPackageOnDisk(FString(URL) + FPackageName::GetMapPackageExtension(), &MapFullName))
 			{
 				Map = MapFullName;
 			}
