fixed EditorGameAgnostic settings saving to project directory for some settings

diff --git a/Engine/Source/Runtime/Core/Public/UObject/UnrealNames.h b/Engine/Source/Runtime/Core/Public/UObject/UnrealNames.h
--- a/Engine/Source/Runtime/Core/Public/UObject/UnrealNames.h
+++ b/Engine/Source/Runtime/Core/Public/UObject/UnrealNames.h
@@ -147,6 +147,7 @@ REGISTER_NAME(247,OtherChildren)
 REGISTER_NAME(248,Location)
 REGISTER_NAME(249,Rotation)
 REGISTER_NAME(250,BSP)
+REGISTER_NAME(251,EditorGameAgnostic)
 
 // Online
 REGISTER_NAME(280,DGram)
diff --git a/Engine/Source/Runtime/CoreUObject/Private/UObject/Class.cpp b/Engine/Source/Runtime/CoreUObject/Private/UObject/Class.cpp
--- a/Engine/Source/Runtime/CoreUObject/Private/UObject/Class.cpp
+++ b/Engine/Source/Runtime/CoreUObject/Private/UObject/Class.cpp
@@ -3246,6 +3246,10 @@ const FString UClass::GetConfigName() const
 	{
 		return GGameIni;
 	}
+	else if ( ClassConfigName == NAME_EditorGameAgnostic )
+	{
+		return GEditorGameAgnosticIni;
+	}
 	else if( ClassConfigName == NAME_None )
 	{
 		UE_LOG(LogClass, Fatal,TEXT("UObject::GetConfigName() called on class with config name 'None'. Class flags = %d"), ClassFlags );
