TTP #333275: execLetWeakObjPtr in the blueprint vm was expecting to pull a TWeakObjectPtr<> from the stack, but even UWeakObjectProperty properties resolve to a raw object ptr (switching to use that).

Also, preventing users from creating their own weak_ptr BP variables.

diff --git a/Engine/Source/Editor/UnrealEd/Private/Kismet2/BlueprintEditorUtils.cpp b/Engine/Source/Editor/UnrealEd/Private/Kismet2/BlueprintEditorUtils.cpp
--- a/Engine/Source/Editor/UnrealEd/Private/Kismet2/BlueprintEditorUtils.cpp
+++ b/Engine/Source/Editor/UnrealEd/Private/Kismet2/BlueprintEditorUtils.cpp
@@ -2987,6 +2987,11 @@ bool FBlueprintEditorUtils::AddMemberVariable(UBlueprint* Blueprint, const FName
 	NewVar.Category = K2Schema->VR_DefaultCategory;
 	NewVar.DefaultValue = DefaultValue;
 
+	// user created variables should be none of these things
+	NewVar.VarType.bIsConst       = false;
+	NewVar.VarType.bIsWeakPointer = false;
+	NewVar.VarType.bIsReference   = false;
+
 	Blueprint->NewVariables.Add(NewVar);
 
 	// Potentially adjust variable names for any child blueprints
diff --git a/Engine/Source/Runtime/CoreUObject/Private/UObject/ScriptCore.cpp b/Engine/Source/Runtime/CoreUObject/Private/UObject/ScriptCore.cpp
--- a/Engine/Source/Runtime/CoreUObject/Private/UObject/ScriptCore.cpp
+++ b/Engine/Source/Runtime/CoreUObject/Private/UObject/ScriptCore.cpp
@@ -1239,15 +1239,15 @@ void UObject::execLetWeakObjPtr( FFrame& Stack, RESULT_DECL )
 			ObjectProperty = Cast<UObjectPropertyBase>(ArrayProp->Inner);
 		}
 	}
-
-	TWeakObjectPtr<UObject> NewValue = NULL;
+	
+	UObject* NewValue = NULL;
 	// evaluate the r-value for this expression into Value
 	Stack.Step( Stack.Object, &NewValue );
 
 	if (ObjAddr)
 	{
 		checkSlow(ObjectProperty);
-		ObjectProperty->SetObjectPropertyValue(ObjAddr, NewValue.Get());
+		ObjectProperty->SetObjectPropertyValue(ObjAddr, NewValue);
 	}
 }
 IMPLEMENT_VM_FUNCTION( EX_LetWeakObjPtr, execLetWeakObjPtr );
