#ue4 When LinkerLoader finds an exising object an the path, but with a different class - the import fails.

diff --git a/Engine/Source/Runtime/CoreUObject/Private/UObject/LinkerLoad.cpp b/Engine/Source/Runtime/CoreUObject/Private/UObject/LinkerLoad.cpp
--- a/Engine/Source/Runtime/CoreUObject/Private/UObject/LinkerLoad.cpp
+++ b/Engine/Source/Runtime/CoreUObject/Private/UObject/LinkerLoad.cpp
@@ -3006,6 +3006,8 @@ UObject* ULinkerLoad::CreateExport( int32 Index )
 		// Try to find existing object first in case we're a forced export to be able to reconcile. Also do it for the
 		// case of async loading as we cannot in-place replace objects.
 
+		UObject* ActualObjectWithTheName = StaticFindObjectFastInternal(NULL, ThisParent, Export.ObjectName, true);
+
 		if(	(FApp::IsGame() && !GIsEditor && !IsRunningCommandlet()) 
 		||	GIsAsyncLoading 
 		||	Export.bForcedExport 
@@ -3015,8 +3017,11 @@ UObject* ULinkerLoad::CreateExport( int32 Index )
 			// Find object after making sure it isn't already set. This would be bad as the code below NULLs it in a certain
 			// case, which if it had been set would cause a linker detach mismatch.
 			check( Export.Object == NULL );
-			Export.Object = StaticFindObjectFastInternal( LoadClass, ThisParent, Export.ObjectName, true, false, RF_NoFlags );
-		
+			if (ActualObjectWithTheName && (ActualObjectWithTheName->GetClass() == LoadClass))
+			{
+				Export.Object = ActualObjectWithTheName;
+			}
+
 			// Object is found in memory.
 			if( Export.Object )
 			{
@@ -3068,6 +3073,12 @@ UObject* ULinkerLoad::CreateExport( int32 Index )
 
 		// Create the export object, marking it with the appropriate flags to
 		// indicate that the object's data still needs to be loaded.
+		if (ActualObjectWithTheName && (ActualObjectWithTheName->GetClass() != LoadClass))
+		{
+			UE_LOG(LogLinker, Error, TEXT("Failed import: class '%s' name '%s' outer '%s'. There is another object (of '%s' class) at the path."),
+				*LoadClass->GetName(), *Export.ObjectName.ToString(), *ThisParent->GetName(), *ActualObjectWithTheName->GetClass()->GetName());
+			return NULL;
+		}
 
 		EObjectFlags ObjectLoadFlags = Export.ObjectFlags;
 		// if we are loading objects just to verify an object reference during script compilation,
