Prevent StaticLoadObject usage during PostLoad.

Added UE_CLOG(..., Fatal, ...) on GIsRoutingPostLoad in StaticLoadObject and refactored one usage of GIsRoutingPostLoad to use TGuardValue. Had to refactor MetaData loading procedures to make sure current lazy loading isn't actually fired during PostLoad.

#ttp 331738 - CORE: Calling LoadObject during PostLoad should not be allowed

diff --git a/Engine/Source/Editor/UnrealEd/Classes/ThumbnailRendering/SceneThumbnailInfoWithPrimitive.h b/Engine/Source/Editor/UnrealEd/Classes/ThumbnailRendering/SceneThumbnailInfoWithPrimitive.h
--- a/Engine/Source/Editor/UnrealEd/Classes/ThumbnailRendering/SceneThumbnailInfoWithPrimitive.h
+++ b/Engine/Source/Editor/UnrealEd/Classes/ThumbnailRendering/SceneThumbnailInfoWithPrimitive.h
@@ -20,6 +20,4 @@ class USceneThumbnailInfoWithPrimitive : public USceneThumbnailInfo
 	/** The custom mesh used when the primitive type is TPT_None */
 	UPROPERTY()
 	FStringAssetReference PreviewMesh;
-
-	virtual void PostLoad() OVERRIDE;
 };
diff --git a/Engine/Source/Editor/UnrealEd/Private/ThumbnailRendering/SceneThumbnailInfoWithPrimitive.cpp b/Engine/Source/Editor/UnrealEd/Private/ThumbnailRendering/SceneThumbnailInfoWithPrimitive.cpp
--- a/Engine/Source/Editor/UnrealEd/Private/ThumbnailRendering/SceneThumbnailInfoWithPrimitive.cpp
+++ b/Engine/Source/Editor/UnrealEd/Private/ThumbnailRendering/SceneThumbnailInfoWithPrimitive.cpp
@@ -11,16 +11,3 @@ USceneThumbnailInfoWithPrimitive::USceneThumbnailInfoWithPrimitive(const class F
 	OrbitZoom = -400.f;
 }
 
-void USceneThumbnailInfoWithPrimitive::PostLoad()
-{
-	Super::PostLoad();
-
-	if ( GIsEditor )
-	{
-		// Load the PreviewMesh with the thumbnail info so it will be ready for visualization.
-		if ( PreviewMesh.IsValid() )
-		{
-			LoadObject<UObject>(nullptr, *PreviewMesh.ToString());
-		}
-	}
-}
diff --git a/Engine/Source/Runtime/CoreUObject/Private/UObject/LinkerLoad.cpp b/Engine/Source/Runtime/CoreUObject/Private/UObject/LinkerLoad.cpp
--- a/Engine/Source/Runtime/CoreUObject/Private/UObject/LinkerLoad.cpp
+++ b/Engine/Source/Runtime/CoreUObject/Private/UObject/LinkerLoad.cpp
@@ -2421,6 +2421,51 @@ bool ULinkerLoad::VerifyImportInner(const int32 ImportIndex, FString& WarningSuf
 	return false;
 }
 
+UObject* ULinkerLoad::CreateExportAndPreload(int32 ExportIndex, bool bForcePreload /* = false */)
+{
+	UObject *Object = CreateExport(ExportIndex);
+	if (Object && (bForcePreload || (Cast<UClass>(Object) != NULL) || Object->IsTemplate() || (Cast<UObjectRedirector>(Object) != NULL)))
+	{
+		Preload(Object);
+	}
+
+	return Object;
+}
+
+int32 ULinkerLoad::LoadMetaDataFromExportMap(bool bForcePreload/* = false */)
+{
+	int32 MetaDataIndex = INDEX_NONE;
+
+	// Try to find MetaData and load it first as other objects can depend on it.
+	for (int32 ExportIndex = 0; ExportIndex < ExportMap.Num(); ++ExportIndex)
+	{
+		if (ExportMap[ExportIndex].ObjectName == NAME_PackageMetaData)
+		{
+			CreateExportAndPreload(ExportIndex, bForcePreload);
+			MetaDataIndex = ExportIndex;
+			break;
+		}
+	}
+
+	// If not found then try to use old name and rename.
+	if (MetaDataIndex == INDEX_NONE)
+	{
+		for (int32 ExportIndex = 0; ExportIndex < ExportMap.Num(); ++ExportIndex)
+		{
+			if (ExportMap[ExportIndex].ObjectName == *UMetaData::StaticClass()->GetName())
+			{
+				UObject* Object = CreateExportAndPreload(ExportIndex, bForcePreload);
+				Object->Rename(*FName(NAME_PackageMetaData).ToString(), NULL, REN_ForceNoResetLoaders);
+
+				MetaDataIndex = ExportIndex;
+				break;
+			}
+		}
+	}
+
+	return MetaDataIndex;
+}
+
 /**
  * Loads all objects in package.
  *
@@ -2437,13 +2482,22 @@ void ULinkerLoad::LoadAllObjects( bool bForcePreload )
 	bool bAllowedToShowStatusUpdate = (LoadFlags & ( LOAD_Quiet | LOAD_SeekFree ) ) == 0;
 	double StartTime = FPlatformTime::Seconds();
 
-	for( int32 i=0; i< ExportMap.Num(); i++ )
+	// MetaData object index in this package.
+	int32 MetaDataIndex = INDEX_NONE;
+
+	if(!FPlatformProperties::RequiresCookedData())
 	{
-		UObject* Object = CreateExport( i );
-		if( Object && (bForcePreload || (Cast<UClass>(Object) != NULL) || Object->IsTemplate() || (Cast<UObjectRedirector>(Object) != NULL)) )
+		MetaDataIndex = LoadMetaDataFromExportMap(bForcePreload);
+	}
+
+	for(int32 ExportIndex = 0; ExportIndex < ExportMap.Num(); ++ExportIndex)
+	{
+		if(ExportIndex == MetaDataIndex)
 		{
-			Preload( Object );
+			continue;
 		}
+
+		CreateExportAndPreload(ExportIndex, bForcePreload);
 	}
 
 	// Mark package as having been fully loaded.
diff --git a/Engine/Source/Runtime/CoreUObject/Private/UObject/MetaData.cpp b/Engine/Source/Runtime/CoreUObject/Private/UObject/MetaData.cpp
--- a/Engine/Source/Runtime/CoreUObject/Private/UObject/MetaData.cpp
+++ b/Engine/Source/Runtime/CoreUObject/Private/UObject/MetaData.cpp
@@ -61,7 +61,7 @@ FMetaDataUtilities::FMoveMetadataHelperContext::~FMoveMetadataHelperContext()
 	if (GIsEditor && (OldObject->GetOutermost() != OldPackage))
 	{
 		UMetaData *NewMetaData = OldObject->GetOutermost()->GetMetaData();
-		UMetaData *OldMetaData = OldPackage->GetMetaData(false);
+		UMetaData *OldMetaData = OldPackage->GetMetaData();
 		TMap< FName, FString> *OldObjectMetaData = OldMetaData->ObjectMetaDataMap.Find(OldObject);
 		if (OldObjectMetaData != NULL)
 		{
diff --git a/Engine/Source/Runtime/CoreUObject/Private/UObject/Package.cpp b/Engine/Source/Runtime/CoreUObject/Private/UObject/Package.cpp
--- a/Engine/Source/Runtime/CoreUObject/Private/UObject/Package.cpp
+++ b/Engine/Source/Runtime/CoreUObject/Private/UObject/Package.cpp
@@ -90,52 +90,27 @@ void UPackage::AddReferencedObjects(UObject* InThis, FReferenceCollector& Collec
 }
 
 /**
- * Get (after possibly creating) a metadata object for this package
- *
- * @param	bAllowLoadObject				Can load an object to find it's UMetaData if not currently loaded.
+ * Gets (after possibly creating) a metadata object for this package
  *
  * @return A valid UMetaData pointer for all objects in this package
  */
-UMetaData* UPackage::GetMetaData(const bool bAllowLoadObject)
+UMetaData* UPackage::GetMetaData()
 {
+	checkf(!FPlatformProperties::RequiresCookedData(), TEXT("MetaData is only allowed in the Editor."));
+
+	// If there is no MetaData, try to find it.
 	if (MetaData == NULL)
 	{
-		// first try to load it
-		if ( (PackageFlags&(PKG_Compiling|PKG_ContainsMap|PKG_CompiledIn)) == 0 && this != GetTransientPackage())
-		{
-			// Try to find before loading, meta data may already exists and we want to avoid calling LoadAllObjects again.
-			MetaData = Cast<UMetaData>(StaticFindObjectFast(UMetaData::StaticClass(), this, FName(NAME_PackageMetaData)));
-			if (MetaData == NULL && bAllowLoadObject)
-			{
-				MetaData = LoadObject<UMetaData>(this, *FName(NAME_PackageMetaData).ToString(), NULL, LOAD_NoWarn|LOAD_Quiet, NULL);
-				if (MetaData == NULL)
-				{
-					// Check old name, and rename if applicable
-					MetaData = LoadObject<UMetaData>(this, *UMetaData::StaticClass()->GetName(), NULL, LOAD_NoWarn|LOAD_Quiet, NULL);
-					if (MetaData)
-					{				
-						MetaData->Rename(*FName(NAME_PackageMetaData).ToString(), NULL, REN_ForceNoResetLoaders);
-					}
-				}
-			}
-		}
+		MetaData = Cast<UMetaData>(StaticFindObjectFast(UMetaData::StaticClass(), this, FName(NAME_PackageMetaData)));
 
-		// if it wasn't found, then create it
-		if (MetaData == NULL)
+		// If MetaData is NULL then it wasn't loaded by linker, so we have to create it.
+		if(MetaData == NULL)
 		{
-			// make a metadata object, but only allow it to be loaded in the editor
 			MetaData = ConstructObject<UMetaData>(UMetaData::StaticClass(), this, NAME_PackageMetaData, RF_Standalone);
 		}
-		check(MetaData);
 	}
 
-	if (MetaData && MetaData->HasAllFlags(RF_NeedLoad) && bAllowLoadObject)
-	{
-		if (auto Linker = MetaData->GetLinker())
-		{
-			Linker->Preload(MetaData);
-		}
-	}
+	check(MetaData);
 
 	return MetaData;
 }
diff --git a/Engine/Source/Runtime/CoreUObject/Private/UObject/UObjectGlobals.cpp b/Engine/Source/Runtime/CoreUObject/Private/UObject/UObjectGlobals.cpp
--- a/Engine/Source/Runtime/CoreUObject/Private/UObject/UObjectGlobals.cpp
+++ b/Engine/Source/Runtime/CoreUObject/Private/UObject/UObjectGlobals.cpp
@@ -451,6 +451,7 @@ UPackage* CreatePackage( UObject* InOuter, const TCHAR* PackageName )
 	{
 		UE_LOG(LogUObjectGlobals, Fatal, TEXT("%s"), TEXT("Attempted to create a package named 'None'") );
 	}
+
 	return Result;
 }
 
@@ -547,7 +548,16 @@ bool ResolveName( UObject*& InPackage, FString& InName, bool Create, bool Throw
 		{
 			if( Create )
 			{
-				InPackage = CreatePackage( InPackage, *PartialName );
+				if (!ScriptPackageName)
+				{
+					InPackage = LoadPackage(Cast<UPackage>(InPackage), *PartialName, 0);
+				}
+				if (!InPackage)
+				{
+					InPackage = CreatePackage(InPackage, *PartialName);
+				}
+
+				check(InPackage);
 			}
 			else
 			{
@@ -631,6 +641,8 @@ bool ParseObject( const TCHAR* Stream, const TCHAR* Match, UClass* Class, UObjec
  */
 UObject* StaticLoadObject(UClass* ObjectClass, UObject* InOuter, const TCHAR* InName, const TCHAR* Filename, uint32 LoadFlags, UPackageMap* Sandbox, bool bAllowObjectReconciliation )
 {
+	UE_CLOG(GIsRoutingPostLoad, LogUObjectGlobals, Fatal, TEXT("Calling StaticLoadObject during PostLoad is forbidden."));
+
 	SCOPE_CYCLE_COUNTER(STAT_LoadObject);
 	check(ObjectClass);
 	check(InName);
@@ -806,7 +818,6 @@ UPackage* LoadPackage( UPackage* InOuter, const TCHAR* InLongPackageName, uint32
 		{
 			Linker->LoadAllObjects();
 		}
-		
 #if WITH_EDITOR
 		// Add a LoadContext string to the endload function in the form of: "<FileToLoad> Package"
 		EndLoad( GIsEditor ? *FPaths::GetBaseFilename(FileToLoad) : NULL );
@@ -879,7 +890,8 @@ UPackage* LoadPackage( UPackage* InOuter, const TCHAR* InLongPackageName, uint32
 	}
 	// Mark package as loaded.
 	Result->SetFlags(RF_WasLoaded);
-	return Cast<UPackage>(Result);
+
+	return Result;
 }
 
 /**
@@ -1089,25 +1101,24 @@ void EndLoad(const TCHAR* LoadContext)
 					}
 				}
 
-				// set this so that we can perform certain operations in which are only safe once all objects have been de-serialized.
-				GIsRoutingPostLoad = true;
-
-				// Postload objects.
-				for( int32 i=0; i<ObjLoaded.Num(); i++ )
 				{
+					// set this so that we can perform certain operations in which are only safe once all objects have been de-serialized.
+					TGuardValue<bool> GuardIsRoutingPostLoad(GIsRoutingPostLoad, true);
+
+					// Postload objects.
+					for(int32 i = 0; i < ObjLoaded.Num(); i++)
+					{
 #if WITH_EDITOR
-					if ( bAllowStatusUpdate )
+						if(bAllowStatusUpdate)
 						{
-						UpdateObjectLoadingStatusMessage( LoadContext, StartTime );
-					}
+							UpdateObjectLoadingStatusMessage(LoadContext, StartTime);
+						}
 #endif
-					UObject* Obj = ObjLoaded[i];
-					check(Obj);
-					Obj->ConditionalPostLoad();
+						UObject* Obj = ObjLoaded[i];
+						check(Obj);
+						Obj->ConditionalPostLoad();
+					}
 				}
-
-				GIsRoutingPostLoad = false;
-
 #if WITH_EDITOR
 				// Send global notification for each object that was loaded.
 				// Useful for updating UI such as ContentBrowser's loaded status.
diff --git a/Engine/Source/Runtime/CoreUObject/Public/UObject/CoreObject.h b/Engine/Source/Runtime/CoreUObject/Public/UObject/CoreObject.h
--- a/Engine/Source/Runtime/CoreUObject/Public/UObject/CoreObject.h
+++ b/Engine/Source/Runtime/CoreUObject/Public/UObject/CoreObject.h
@@ -304,13 +304,13 @@ public:
 	// MetaData 
 
 	/**
-	 * Get (after possibly creating) a metadata object for this package
+	 * Gets (after possibly creating) a metadata object for this package
 	 *
 	 * @param	bAllowLoadObject				Can load an object to find it's UMetaData if not currently loaded.
 	 *
 	 * @return A valid UMetaData pointer for all objects in this package
 	 */
-	class UMetaData* GetMetaData(const bool bAllowLoadObject = true);
+	class UMetaData* GetMetaData();
 
 	/**
 	 * Save one specific object (along with any objects it references contained within the same Outer) into an Unreal package.
diff --git a/Engine/Source/Runtime/CoreUObject/Public/UObject/Linker.h b/Engine/Source/Runtime/CoreUObject/Public/UObject/Linker.h
--- a/Engine/Source/Runtime/CoreUObject/Public/UObject/Linker.h
+++ b/Engine/Source/Runtime/CoreUObject/Public/UObject/Linker.h
@@ -1520,6 +1520,29 @@ public:
 private:
 
 	UObject* CreateExport( int32 Index );
+
+	/**
+	 * Creates export and preload if requested.
+	 *
+	 * @param Index Index of the export in export map.
+	 * @param bForcePreload Whether to explicitly call Preload (serialize)
+	 *        right away instead of being called from EndLoad().
+	 *
+	 * @return Created object.
+	 */
+	UObject* CreateExportAndPreload(int32 ExportIndex, bool bForcePreload = false);
+
+	/** 
+	 * Looks for and loads meta data object from export map.
+	 *
+	 * @param bForcePreload Whether to explicitly call Preload (serialize)
+	 *        right away instead of being called from EndLoad().
+	 * 
+	 * @return If found returns index of meta data object in the export map,
+	 *         INDEX_NONE otherwise.
+	 */
+	int32 LoadMetaDataFromExportMap(bool bForcePreload = false);
+
 	UObject* CreateImport( int32 Index );
 	UObject* IndexToObject( FPackageIndex Index );
 
