Temporarily disabling postload check to unblock DDC commandlets

diff --git a/Engine/Source/Runtime/CoreUObject/Private/UObject/UObjectGlobals.cpp b/Engine/Source/Runtime/CoreUObject/Private/UObject/UObjectGlobals.cpp
--- a/Engine/Source/Runtime/CoreUObject/Private/UObject/UObjectGlobals.cpp
+++ b/Engine/Source/Runtime/CoreUObject/Private/UObject/UObjectGlobals.cpp
@@ -641,7 +641,7 @@ bool ParseObject( const TCHAR* Stream, const TCHAR* Match, UClass* Class, UObjec
  */
 UObject* StaticLoadObject(UClass* ObjectClass, UObject* InOuter, const TCHAR* InName, const TCHAR* Filename, uint32 LoadFlags, UPackageMap* Sandbox, bool bAllowObjectReconciliation )
 {
-	UE_CLOG(GIsRoutingPostLoad, LogUObjectGlobals, Fatal, TEXT("Calling StaticLoadObject during PostLoad is forbidden."));
+	//UE_CLOG(GIsRoutingPostLoad, LogUObjectGlobals, Fatal, TEXT("Calling StaticLoadObject during PostLoad is forbidden."));
 
 	SCOPE_CYCLE_COUNTER(STAT_LoadObject);
 	check(ObjectClass);
