UStructProperty::ImportText_Internal: Loop for closing quote will never start, so changed while into do-while block.
#ttp 332373 - UStructProperty::ImportText_Internal: Loop for closing quote will never start

diff --git a/Engine/Source/Runtime/CoreUObject/Private/UObject/PropertyStruct.cpp b/Engine/Source/Runtime/CoreUObject/Private/UObject/PropertyStruct.cpp
--- a/Engine/Source/Runtime/CoreUObject/Private/UObject/PropertyStruct.cpp
+++ b/Engine/Source/Runtime/CoreUObject/Private/UObject/PropertyStruct.cpp
@@ -327,10 +327,10 @@ const TCHAR* UStructProperty::ImportText_Internal( const TCHAR* InBuffer, void*
 				SkipWhitespace(Buffer);
 				if (*Buffer == TCHAR('\"'))
 				{
-					while (*Buffer && *Buffer != TCHAR('\"') && *Buffer != TCHAR('\n') && *Buffer != TCHAR('\r'))
+					do
 					{
 						Buffer++;
-					}
+					} while (*Buffer && *Buffer != TCHAR('\"') && *Buffer != TCHAR('\n') && *Buffer != TCHAR('\r'));
 
 					if (*Buffer != TCHAR('\"'))
 					{
@@ -455,4 +455,3 @@ IMPLEMENT_CORE_INTRINSIC_CLASS(UStructProperty, UProperty,
 		Class->EmitObjectReference( STRUCT_OFFSET( UStructProperty, Struct ) );
 	}
 );
-
