Making Obj->IsA(NULL) invalid

diff --git a/Engine/Source/Runtime/CoreUObject/Private/UObject/UObjectBaseUtility.cpp b/Engine/Source/Runtime/CoreUObject/Private/UObject/UObjectBaseUtility.cpp
--- a/Engine/Source/Runtime/CoreUObject/Private/UObject/UObjectBaseUtility.cpp
+++ b/Engine/Source/Runtime/CoreUObject/Private/UObject/UObjectBaseUtility.cpp
@@ -235,10 +235,8 @@ bool UObjectBaseUtility::RootPackageHasAnyFlags( uint32 CheckFlagMask ) const
  */
 bool UObjectBaseUtility::IsA( const UClass* SomeBase ) const
 {
-	if ( SomeBase==NULL )
-	{
-		return true;
-	}
+	UE_CLOG(!SomeBase, LogObj, Fatal, TEXT("IsA(NULL) cannot yield meaningful results"));
+
 	for ( UClass* TempClass=GetClass(); TempClass; TempClass=TempClass->GetSuperClass() )
 	{
 		if ( TempClass == SomeBase )
