Delegates: simplified script delegate API

Upgrade Notes:
The two separate function calls "MyScriptDelegate.SetObject(SomeObject); MyScriptDelegate.SetFunctionName(SomeFunctionName);" where combined into a single function call "MyScriptDelegate.BindUFunction(SomeObject, SomeFunctionName);"

diff --git a/Engine/Source/Editor/BlueprintGraph/Private/K2Node_ActorBoundEvent.cpp b/Engine/Source/Editor/BlueprintGraph/Private/K2Node_ActorBoundEvent.cpp
--- a/Engine/Source/Editor/BlueprintGraph/Private/K2Node_ActorBoundEvent.cpp
+++ b/Engine/Source/Editor/BlueprintGraph/Private/K2Node_ActorBoundEvent.cpp
@@ -64,8 +64,7 @@ void UK2Node_ActorBoundEvent::DestroyNode()
 			{
 				// Create a delegate of the correct signature to remove
 				FScriptDelegate Delegate;
-				Delegate.SetFunctionName(CustomFunctionName);
-				Delegate.SetObject(LSA);
+				Delegate.BindUFunction(LSA, CustomFunctionName);
 				
 				// Attempt to remove it from the target's MC delegate
 				FMulticastScriptDelegate* TargetDelegate = GetTargetDelegate();
diff --git a/Engine/Source/Editor/KismetCompiler/Private/KismetCompilerCppBackend.cpp b/Engine/Source/Editor/KismetCompiler/Private/KismetCompilerCppBackend.cpp
--- a/Engine/Source/Editor/KismetCompiler/Private/KismetCompilerCppBackend.cpp
+++ b/Engine/Source/Editor/KismetCompiler/Private/KismetCompilerCppBackend.cpp
@@ -509,10 +509,10 @@ void FKismetCppBackend::EmitBindDelegateStatement(FKismetFunctionContext& Functi
 
 	Emit(Body, 
 		*FString::Printf(
-			TEXT("\t\t\tif(FScriptDelegate* __Delegate = %s)\n\t\t\t{\n\t\t\t\t__Delegate->SetFunctionName(%s);\n\t\t\t\t__Delegate->SetObject(%s);\n\t\t\t|\n"),
+			TEXT("\t\t\tif(FScriptDelegate* __Delegate = %s)\n\t\t\t{\n\t\t\t\t__Delegate->BindUFunction(%s,%s);\n\t\t\t|\n"),
 			*Delegate, 
-			*NameTerm, 
-			*ObjectTerm
+			*ObjectTerm,
+			*NameTerm
 		)
 	);
 }
diff --git a/Engine/Source/Editor/PropertyEditor/Private/PropertyNode.cpp b/Engine/Source/Editor/PropertyEditor/Private/PropertyNode.cpp
--- a/Engine/Source/Editor/PropertyEditor/Private/PropertyNode.cpp
+++ b/Engine/Source/Editor/PropertyEditor/Private/PropertyNode.cpp
@@ -1144,9 +1144,9 @@ private:
 		if ( DelegateProp != NULL )
 		{
 			FScriptDelegate* DelegateValue = DelegateProp->GetPropertyValuePtr(PropertyValueAddress);
-			if (DelegateValue->GetObject() && DelegateValue->GetObject()->IsDefaultSubobject())
+			if (DelegateValue->GetUObject() && DelegateValue->GetUObject()->IsDefaultSubobject())
 			{
-				Components.AddUnique(DelegateValue->GetObject());
+				Components.AddUnique(DelegateValue->GetUObject());
 			}
 
 			bResult = true;
diff --git a/Engine/Source/Editor/UnrealEd/Private/Kismet2/BlueprintEditorUtils.cpp b/Engine/Source/Editor/UnrealEd/Private/Kismet2/BlueprintEditorUtils.cpp
--- a/Engine/Source/Editor/UnrealEd/Private/Kismet2/BlueprintEditorUtils.cpp
+++ b/Engine/Source/Editor/UnrealEd/Private/Kismet2/BlueprintEditorUtils.cpp
@@ -4694,8 +4694,7 @@ bool FBlueprintEditorUtils::FixLevelScriptActorBindings(ALevelScriptActor* Level
 
 					// Create the delegate, and add it if it doesn't already exist
 					FScriptDelegate Delegate;
-					Delegate.SetFunctionName(TargetFunction);
-					Delegate.SetObject(LevelScriptActor);
+					Delegate.BindUFunction(LevelScriptActor, TargetFunction);
 					TargetDelegate->AddUnique(Delegate);
 				}
 				else
diff --git a/Engine/Source/Runtime/Core/Public/UObject/ScriptDelegates.h b/Engine/Source/Runtime/Core/Public/UObject/ScriptDelegates.h
--- a/Engine/Source/Runtime/Core/Public/UObject/ScriptDelegates.h
+++ b/Engine/Source/Runtime/Core/Public/UObject/ScriptDelegates.h
@@ -19,7 +19,7 @@ public:
 
 	/** Constructor */
 	TScriptDelegate() 
-		: Object( NULL ),
+		: Object( nullptr ),
 		  FunctionName( NAME_None )
 	{ }
 
@@ -32,7 +32,7 @@ private:
 		{
 			if (UObject* ObjectPtr = Object.Get())
 			{
-				return ((UObjectTemplate*)ObjectPtr)->FindFunction(FunctionName) != NULL;
+				return ((UObjectTemplate*)ObjectPtr)->FindFunction(FunctionName) != nullptr;
 			}
 		}
 
@@ -41,6 +41,18 @@ private:
 
 public:
 
+	/**
+	 * Binds a UFunction to this delegate.
+	 *
+	 * @param InObject The object to call the function on.
+	 * @param InFunctionName The name of the function to call.
+	 */
+	void BindUFunction( class UObject* InObject, const FName& InFunctionName )
+	{
+		Object = InObject;
+		FunctionName = InFunctionName;
+	}
+
 	/** 
 	 * Checks to see if the user object bound to this delegate is still valid
 	 *
@@ -58,7 +70,7 @@ public:
 	 */
 	inline bool IsBoundToObject(void const* InUserObject) const
 	{
-		return InUserObject && (InUserObject == GetObject());
+		return InUserObject && (InUserObject == GetUObject());
 	}
 
 	/** 
@@ -76,7 +88,7 @@ public:
 	 */
 	void Unbind()
 	{
-		Object = NULL;
+		Object = nullptr;
 		FunctionName = NAME_None;
 	}
 
@@ -98,7 +110,7 @@ public:
 	{
 		if( IsBound() )
 		{
-			return ((UObjectTemplate*)GetObject())->GetPathName() + TEXT(".") + GetFunctionName().ToString();
+			return ((UObjectTemplate*)GetUObject())->GetPathName() + TEXT(".") + GetFunctionName().ToString();
 		}
 		return TEXT( "<Unbound>" );
 	}
@@ -127,32 +139,12 @@ public:
 		FunctionName = Other.FunctionName;
 	}
 
-	/**
-	 * Sets the object name.  Usually, you should never call this yourself.  Use BindDynamic() instead!
-	*
-	 * @param	InObject	New object
-	*/
-	void SetObject( class UObject* InObject )
-	{
-		Object = InObject;
-	}
-
-	/**
-	 * Sets the function name.  Usually, you should never call this yourself.  Use BindDynamic() instead!
-	 *
-	 * @param	InFunctionName	New function name
-	 */
-	void SetFunctionName( const FName InFunctionName )
-	{
-		FunctionName = InFunctionName;
-	}
-
 	/** 
 	 * Gets the object bound to this delegate
 	 *
 	 * @return	The object
 	 */
-	class UObject* GetObject()
+	class UObject* GetUObject()
 	{
 		// Downcast UObjectBase to UObject
 		return static_cast< UObject* >( Object.Get() );
@@ -163,7 +155,7 @@ public:
 	 *
 	 * @return	The object
 	 */
-	const class UObject* GetObject() const
+	const class UObject* GetUObject() const
 	{
 		// Downcast UObjectBase to UObject
 		return static_cast< const UObject* >( Object.Get() );
@@ -209,7 +201,7 @@ public:
 
 protected:
 
-	/** The object bound to this delegate, or NULL if no object is bound */
+	/** The object bound to this delegate, or nullptr if no object is bound */
 	TWeakPtr Object;
 
 	/** Name of the function to call on the bound object */
@@ -400,8 +392,8 @@ public:
 		TArray< UObject* > OutputList;
 		for( typename FInvocationList::TIterator CurDelegate( InvocationList ); CurDelegate; ++CurDelegate )
 		{
-			UObject* CurObject = CurDelegate->GetObject();
-			if( CurObject != NULL )
+			UObject* CurObject = CurDelegate->GetUObject();
+			if( CurObject != nullptr )
 			{
 				OutputList.Add( CurObject );
 			}
diff --git a/Engine/Source/Runtime/CoreUObject/Private/UObject/PropertyDelegate.cpp b/Engine/Source/Runtime/CoreUObject/Private/UObject/PropertyDelegate.cpp
--- a/Engine/Source/Runtime/CoreUObject/Private/UObject/PropertyDelegate.cpp
+++ b/Engine/Source/Runtime/CoreUObject/Private/UObject/PropertyDelegate.cpp
@@ -12,17 +12,20 @@ void UDelegateProperty::InstanceSubobjects(void* Data, void const* DefaultData,
 	for( int32 i=0; i<ArrayDim; i++ )
 	{
 		FScriptDelegate& DestDelegate = ((FScriptDelegate*)Data)[i];
-		UObject* CurrentValue = DestDelegate.GetObject();
-		if ( CurrentValue )
+		UObject* CurrentUObject = DestDelegate.GetUObject();
+
+		if (CurrentUObject)
 		{
 			UObject *Template = NULL;
+
 			if (DefaultData)
 			{
 				FScriptDelegate& DefaultDelegate = ((FScriptDelegate*)DefaultData)[i];
-				Template = DefaultDelegate.GetObject();
+				Template = DefaultDelegate.GetUObject();
 			}
-			UObject* NewValue = InstanceGraph->InstancePropertyValue(Template, CurrentValue, Owner, HasAnyPropertyFlags(CPF_Transient), false, true);
-			DestDelegate.SetObject(NewValue);
+
+			UObject* NewUObject = InstanceGraph->InstancePropertyValue(Template, CurrentUObject, Owner, HasAnyPropertyFlags(CPF_Transient), false, true);
+			DestDelegate.BindUFunction(NewUObject, DestDelegate.GetFunctionName());
 		}
 	}
 }
@@ -40,11 +43,11 @@ bool UDelegateProperty::Identical( const void* A, const void* B, uint32 PortFlag
 	}
 	else if ( DA->GetFunctionName() == DB->GetFunctionName() )
 	{
-		if ( DA->GetObject() == DB->GetObject() )
+		if ( DA->GetUObject() == DB->GetUObject() )
 		{
 			bResult = true;
 		}
-		else if	((DA->GetObject() == NULL || DB->GetObject() == NULL)
+		else if	((DA->GetUObject() == NULL || DB->GetUObject() == NULL)
 			&&	(PortFlags&PPF_DeltaComparison) != 0)
 		{
 			bResult = true;
@@ -82,7 +85,7 @@ void UDelegateProperty::ExportTextItem( FString& ValueStr, const void* PropertyV
 	check(ScriptDelegate != NULL);
 	bool bDelegateHasValue = ScriptDelegate->GetFunctionName() != NAME_None;
 	ValueStr += FString::Printf( TEXT("%s.%s"),
-		ScriptDelegate->GetObject() != NULL ? *ScriptDelegate->GetObject()->GetName() : TEXT("(null)"),
+		ScriptDelegate->GetUObject() != NULL ? *ScriptDelegate->GetUObject()->GetName() : TEXT("(null)"),
 		*ScriptDelegate->GetFunctionName().ToString() );
 }
 
diff --git a/Engine/Source/Runtime/CoreUObject/Private/UObject/PropertyHelper.h b/Engine/Source/Runtime/CoreUObject/Private/UObject/PropertyHelper.h
--- a/Engine/Source/Runtime/CoreUObject/Private/UObject/PropertyHelper.h
+++ b/Engine/Source/Runtime/CoreUObject/Private/UObject/PropertyHelper.h
@@ -269,8 +269,7 @@ namespace DelegatePropertyTools
 
 		//UE_LOG(LogProperty, Log, TEXT("... importing delegate FunctionName:'%s'(%s)   Object:'%s'(%s)"),Func != NULL ? *Func->GetName() : TEXT("NULL"), FuncName, Object != NULL ? *Object->GetFullName() : TEXT("NULL"), ObjName);
 	
-		Delegate.SetObject( Func ? Object : NULL );
-		Delegate.SetFunctionName( Func ? Func->GetFName() : NAME_None );
+		Delegate.BindUFunction( Func ? Object : NULL, Func ? Func->GetFName() : NAME_None );
 
 		return ( Func != NULL && Object != NULL ) ? Buffer : NULL;
 	}
diff --git a/Engine/Source/Runtime/CoreUObject/Private/UObject/PropertyMulticastDelegate.cpp b/Engine/Source/Runtime/CoreUObject/Private/UObject/PropertyMulticastDelegate.cpp
--- a/Engine/Source/Runtime/CoreUObject/Private/UObject/PropertyMulticastDelegate.cpp
+++ b/Engine/Source/Runtime/CoreUObject/Private/UObject/PropertyMulticastDelegate.cpp
@@ -22,24 +22,26 @@ void UMulticastDelegateProperty::InstanceSubobjects(void* Data, void const* Defa
 			for(; CurInvocation && DefaultInvocation; ++CurInvocation, ++DefaultInvocation )
 			{
 				FScriptDelegate& DestDelegateInvocation = *CurInvocation;
-				UObject* CurrentValue = DestDelegateInvocation.GetObject();
-				if (CurrentValue)
+				UObject* CurrentUObject = DestDelegateInvocation.GetUObject();
+
+				if (CurrentUObject)
 				{
 					FScriptDelegate& DefaultDelegateInvocation = *DefaultInvocation;
-					UObject *Template = DefaultDelegateInvocation.GetObject();
-					UObject* NewValue = InstanceGraph->InstancePropertyValue(Template, CurrentValue, Owner, HasAnyPropertyFlags(CPF_Transient), false, true);
-					DestDelegateInvocation.SetObject(NewValue);
+					UObject *Template = DefaultDelegateInvocation.GetUObject();
+					UObject* NewUObject = InstanceGraph->InstancePropertyValue(Template, CurrentUObject, Owner, HasAnyPropertyFlags(CPF_Transient), false, true);
+					DestDelegateInvocation.BindUFunction(NewUObject, DestDelegateInvocation.GetFunctionName());
 				}
 			}
 			// now finish up the ones for which there is no default
 			for(; CurInvocation; ++CurInvocation )
 			{
 				FScriptDelegate& DestDelegateInvocation = *CurInvocation;
-				UObject* CurrentValue = DestDelegateInvocation.GetObject();
-				if (CurrentValue)
+				UObject* CurrentUObject = DestDelegateInvocation.GetUObject();
+
+				if (CurrentUObject)
 				{
-					UObject* NewValue = InstanceGraph->InstancePropertyValue(NULL, CurrentValue, Owner, HasAnyPropertyFlags(CPF_Transient), false, true);
-					DestDelegateInvocation.SetObject(NewValue);
+					UObject* NewUObject = InstanceGraph->InstancePropertyValue(NULL, CurrentUObject, Owner, HasAnyPropertyFlags(CPF_Transient), false, true);
+					DestDelegateInvocation.BindUFunction(NewUObject, DestDelegateInvocation.GetFunctionName());
 				}
 			}
 		}
@@ -53,11 +55,12 @@ void UMulticastDelegateProperty::InstanceSubobjects(void* Data, void const* Defa
 			for( FMulticastScriptDelegate::FInvocationList::TIterator CurInvocation( DestDelegate.InvocationList ); CurInvocation; ++CurInvocation )
 			{
 				FScriptDelegate& DestDelegateInvocation = *CurInvocation;
-				UObject* CurrentValue = DestDelegateInvocation.GetObject();
-				if (CurrentValue)
+				UObject* CurrentUObject = DestDelegateInvocation.GetUObject();
+
+				if (CurrentUObject)
 				{
-					UObject* NewValue = InstanceGraph->InstancePropertyValue(NULL, CurrentValue, Owner, HasAnyPropertyFlags(CPF_Transient), false, true);
-					DestDelegateInvocation.SetObject(NewValue);
+					UObject* NewUObject = InstanceGraph->InstancePropertyValue(NULL, CurrentUObject, Owner, HasAnyPropertyFlags(CPF_Transient), false, true);
+					DestDelegateInvocation.BindUFunction(NewUObject, DestDelegateInvocation.GetFunctionName());
 				}
 			}
 		}
@@ -83,8 +86,8 @@ bool UMulticastDelegateProperty::Identical( const void* A, const void* B, uint32
 			const FScriptDelegate& InvocationA = DA->InvocationList[ CurInvocationIndex ];
 			const FScriptDelegate& InvocationB = DB->InvocationList[ CurInvocationIndex ];
 			
-			if( InvocationA.GetObject() != InvocationB.GetObject() ||
-				!( (PortFlags&PPF_DeltaComparison) != 0 && ( InvocationA.GetObject() == NULL || InvocationB.GetObject() == NULL ) ) )
+			if( InvocationA.GetUObject() != InvocationB.GetUObject() ||
+				!( (PortFlags&PPF_DeltaComparison) != 0 && ( InvocationA.GetUObject() == NULL || InvocationB.GetUObject() == NULL ) ) )
 			{
 				bResult = false;
 				break;
@@ -136,7 +139,7 @@ void UMulticastDelegateProperty::ExportTextItem( FString& ValueStr, const void*
 
 			bool bDelegateHasValue = CurInvocation->GetFunctionName() != NAME_None;
 			ValueStr += FString::Printf( TEXT("%s.%s"),
-				CurInvocation->GetObject() != NULL ? *CurInvocation->GetObject()->GetName() : TEXT("(null)"),
+				CurInvocation->GetUObject() != NULL ? *CurInvocation->GetUObject()->GetName() : TEXT("(null)"),
 				*CurInvocation->GetFunctionName().ToString() );
 		}
 	}
diff --git a/Engine/Source/Runtime/CoreUObject/Private/UObject/ScriptCore.cpp b/Engine/Source/Runtime/CoreUObject/Private/UObject/ScriptCore.cpp
--- a/Engine/Source/Runtime/CoreUObject/Private/UObject/ScriptCore.cpp
+++ b/Engine/Source/Runtime/CoreUObject/Private/UObject/ScriptCore.cpp
@@ -1310,8 +1310,7 @@ void UObject::execLetDelegate( FFrame& Stack, RESULT_DECL )
 
 	if (DelegateAddr != NULL)
 	{
-		DelegateAddr->SetFunctionName( Delegate.GetFunctionName() );
-		DelegateAddr->SetObject( Delegate.GetObject() );
+		DelegateAddr->BindUFunction( Delegate.GetUObject(), Delegate.GetFunctionName() );
 	}
 }
 IMPLEMENT_VM_FUNCTION( EX_LetDelegate, execLetDelegate );
@@ -1618,9 +1617,7 @@ IMPLEMENT_VM_FUNCTION( EX_ObjectConst, execObjectConst );
 void UObject::execInstanceDelegate( FFrame& Stack, RESULT_DECL )
 {
 	FName FunctionName = Stack.ReadName();
-
-	((FScriptDelegate*)Result)->SetFunctionName( FunctionName );
-	((FScriptDelegate*)Result)->SetObject( (FunctionName == NAME_None) ? NULL : this );
+	((FScriptDelegate*)Result)->BindUFunction( (FunctionName == NAME_None) ? NULL : this, FunctionName );
 }
 IMPLEMENT_VM_FUNCTION( EX_InstanceDelegate, execInstanceDelegate );
 
@@ -1640,8 +1637,7 @@ void UObject::execBindDelegate( FFrame& Stack, RESULT_DECL )
 
 	if (DelegateAddr)
 	{
-		DelegateAddr->SetFunctionName(FunctionName);
-		DelegateAddr->SetObject(ObjectForDelegate);
+		DelegateAddr->BindUFunction(ObjectForDelegate, FunctionName);
 	}
 }
 IMPLEMENT_VM_FUNCTION( EX_BindDelegate, execBindDelegate );
diff --git a/Engine/Source/Runtime/Engine/Classes/Components/InputComponent.h b/Engine/Source/Runtime/Engine/Classes/Components/InputComponent.h
--- a/Engine/Source/Runtime/Engine/Classes/Components/InputComponent.h
+++ b/Engine/Source/Runtime/Engine/Classes/Components/InputComponent.h
@@ -41,8 +41,7 @@ struct TInputUnifiedDelegate
 	inline void BindDelegate(UObject* Object, const FName FuncName)
 	{
 		FuncDelegate.Unbind();
-		FuncDynDelegate.SetObject(Object);
-		FuncDynDelegate.SetFunctionName(FuncName);
+		FuncDynDelegate.BindUFunction(Object, FuncName);
 	}
 
 	DelegateType & GetDelegateForManualSet()
diff --git a/Engine/Source/Runtime/Engine/Private/AI/Navigation/PathFollowingComponent.cpp b/Engine/Source/Runtime/Engine/Private/AI/Navigation/PathFollowingComponent.cpp
--- a/Engine/Source/Runtime/Engine/Private/AI/Navigation/PathFollowingComponent.cpp
+++ b/Engine/Source/Runtime/Engine/Private/AI/Navigation/PathFollowingComponent.cpp
@@ -1008,8 +1008,7 @@ bool UPathFollowingComponent::UpdateMovementComponent(bool bForce)
 		if (MyPawn)
 		{
 			FScriptDelegate Delegate;
-			Delegate.SetFunctionName(TEXT("OnActorBump"));
-			Delegate.SetObject(this);
+			Delegate.BindUFunction(this, "OnActorBump");
 			MyPawn->OnActorHit.AddUnique(Delegate);
 
 			SetMovementComponent(MyPawn->FindComponentByClass<UNavMovementComponent>());
diff --git a/Engine/Source/Runtime/Engine/Private/BlueprintGeneratedClass.cpp b/Engine/Source/Runtime/Engine/Private/BlueprintGeneratedClass.cpp
--- a/Engine/Source/Runtime/Engine/Private/BlueprintGeneratedClass.cpp
+++ b/Engine/Source/Runtime/Engine/Private/BlueprintGeneratedClass.cpp
@@ -281,8 +281,8 @@ void UBlueprintGeneratedClass::BindDynamicDelegates(AActor* InInstance)
 	{
 		if ( ensure(DynamicBindingObjects[Index] != NULL) )
 		{
-			DynamicBindingObjects[Index]->BindDynamicDelegates(InInstance);
-		}
+		DynamicBindingObjects[Index]->BindDynamicDelegates(InInstance);
+	}
 	}
 
 	// call on super class, if it's a BlueprintGeneratedClass
@@ -369,8 +369,7 @@ void UBlueprintGeneratedClass::CreateComponentsForActor(AActor* Actor) const
 			{
 				// Create delegate for all keys in this track
 				FScriptDelegate EventDelegate;
-				EventDelegate.SetObject(Actor);
-				EventDelegate.SetFunctionName(TimelineTemplate->GetEventTrackFunctionName(TrackIdx));
+				EventDelegate.BindUFunction(Actor, TimelineTemplate->GetEventTrackFunctionName(TrackIdx));
 
 				// Create an entry in Events for each key of this track
 				for (auto It(EventTrackTemplate->CurveKeys->FloatCurve.GetKeyIterator()); It; ++It)
@@ -412,14 +411,12 @@ void UBlueprintGeneratedClass::CreateComponentsForActor(AActor* Actor) const
 
 		// Set up delegate that gets called after all properties are updated
 		FScriptDelegate UpdateDelegate;
-		UpdateDelegate.SetObject(Actor);
-		UpdateDelegate.SetFunctionName(TimelineTemplate->GetUpdateFunctionName());
+		UpdateDelegate.BindUFunction(Actor, TimelineTemplate->GetUpdateFunctionName());
 		NewTimeline->SetTimelinePostUpdateFunc(FOnTimelineEvent(UpdateDelegate));
 
 		// Set up finished delegate that gets called after all properties are updated
 		FScriptDelegate FinishedDelegate;
-		FinishedDelegate.SetObject(Actor);
-		FinishedDelegate.SetFunctionName(TimelineTemplate->GetFinishedFunctionName());
+		FinishedDelegate.BindUFunction(Actor, TimelineTemplate->GetFinishedFunctionName());
 		NewTimeline->SetTimelineFinishedFunc(FOnTimelineEvent(FinishedDelegate));
 
 		NewTimeline->RegisterComponent();
diff --git a/Engine/Source/Runtime/Engine/Private/ComponentDelegateBinding.cpp b/Engine/Source/Runtime/Engine/Private/ComponentDelegateBinding.cpp
--- a/Engine/Source/Runtime/Engine/Private/ComponentDelegateBinding.cpp
+++ b/Engine/Source/Runtime/Engine/Private/ComponentDelegateBinding.cpp
@@ -33,8 +33,7 @@ void UComponentDelegateBinding::BindDynamicDelegates(AActor* InInstance) const
 					FMulticastScriptDelegate* TargetDelegate = DelegateProp->GetPropertyValuePtr_InContainer(Component);
 
 					FScriptDelegate Delegate;
-					Delegate.SetFunctionName(Binding.FunctionNameToBind);
-					Delegate.SetObject(InInstance);
+					Delegate.BindUFunction(InInstance, Binding.FunctionNameToBind);
 					TargetDelegate->AddUnique(Delegate);
 				}
 			}
diff --git a/Engine/Source/Runtime/Engine/Private/KismetSystemLibrary.cpp b/Engine/Source/Runtime/Engine/Private/KismetSystemLibrary.cpp
--- a/Engine/Source/Runtime/Engine/Private/KismetSystemLibrary.cpp
+++ b/Engine/Source/Runtime/Engine/Private/KismetSystemLibrary.cpp
@@ -236,8 +236,7 @@ void UKismetSystemLibrary::K2_SetTimer(UObject* Object, FString FunctionName, fl
 	}
 
 	FTimerDynamicDelegate Delegate;
-	Delegate.SetObject(Object);
-	Delegate.SetFunctionName(FunctionFName);
+	Delegate.BindUFunction(Object, FunctionFName);
 
 	if (Delegate.IsBound())
 	{
@@ -253,8 +252,7 @@ void UKismetSystemLibrary::K2_SetTimer(UObject* Object, FString FunctionName, fl
 void UKismetSystemLibrary::K2_ClearTimer(UObject* Object, FString FunctionName)
 {
 	FTimerDynamicDelegate Delegate;
-	Delegate.SetObject(Object);
-	Delegate.SetFunctionName(FName(*FunctionName));
+	Delegate.BindUFunction(Object, *FunctionName);
 
 	if (Delegate.IsBound())
 	{
@@ -274,8 +272,7 @@ void UKismetSystemLibrary::K2_ClearTimer(UObject* Object, FString FunctionName)
 void UKismetSystemLibrary::K2_PauseTimer(UObject* Object, FString FunctionName)
 {
 	FTimerDynamicDelegate Delegate;
-	Delegate.SetObject(Object);
-	Delegate.SetFunctionName(FName(*FunctionName));
+	Delegate.BindUFunction(Object, *FunctionName);
 
 	if (Delegate.IsBound())
 	{
@@ -292,8 +289,7 @@ void UKismetSystemLibrary::K2_PauseTimer(UObject* Object, FString FunctionName)
 void UKismetSystemLibrary::K2_UnPauseTimer(UObject* Object, FString FunctionName)
 {
 	FTimerDynamicDelegate Delegate;
-	Delegate.SetObject(Object);
-	Delegate.SetFunctionName(FName(*FunctionName));
+	Delegate.BindUFunction(Object, *FunctionName);
 
 	if (Delegate.IsBound())
 	{
@@ -310,8 +306,7 @@ void UKismetSystemLibrary::K2_UnPauseTimer(UObject* Object, FString FunctionName
 bool UKismetSystemLibrary::K2_IsTimerActive(UObject* Object, FString FunctionName)
 {
 	FTimerDynamicDelegate Delegate;
-	Delegate.SetObject(Object);
-	Delegate.SetFunctionName(FName(*FunctionName));
+	Delegate.BindUFunction(Object, *FunctionName);
 
 	if (Delegate.IsBound())
 	{
@@ -329,8 +324,7 @@ bool UKismetSystemLibrary::K2_IsTimerActive(UObject* Object, FString FunctionNam
 bool UKismetSystemLibrary::K2_IsTimerPaused(UObject* Object, FString FunctionName)
 {
 	FTimerDynamicDelegate Delegate;
-	Delegate.SetObject(Object);
-	Delegate.SetFunctionName(FName(*FunctionName));
+	Delegate.BindUFunction(Object, *FunctionName);
 
 	if (Delegate.IsBound())
 	{
@@ -348,8 +342,7 @@ bool UKismetSystemLibrary::K2_IsTimerPaused(UObject* Object, FString FunctionNam
 bool UKismetSystemLibrary::K2_TimerExists(UObject* Object, FString FunctionName)
 {
 	FTimerDynamicDelegate Delegate;
-	Delegate.SetObject(Object);
-	Delegate.SetFunctionName(FName(*FunctionName));
+	Delegate.BindUFunction(Object, *FunctionName);
 
 	if (Delegate.IsBound())
 	{
@@ -367,8 +360,7 @@ bool UKismetSystemLibrary::K2_TimerExists(UObject* Object, FString FunctionName)
 float UKismetSystemLibrary::K2_GetTimerElapsedTime(UObject* Object, FString FunctionName)
 {
 	FTimerDynamicDelegate Delegate;
-	Delegate.SetObject(Object);
-	Delegate.SetFunctionName(FName(*FunctionName));
+	Delegate.BindUFunction(Object, *FunctionName);
 
 	if (Delegate.IsBound())
 	{
@@ -386,8 +378,7 @@ float UKismetSystemLibrary::K2_GetTimerElapsedTime(UObject* Object, FString Func
 float UKismetSystemLibrary::K2_GetTimerRemainingTime(UObject* Object, FString FunctionName)
 {
 	FTimerDynamicDelegate Delegate;
-	Delegate.SetObject(Object);
-	Delegate.SetFunctionName(FName(*FunctionName));
+	Delegate.BindUFunction(Object, *FunctionName);
 
 	if (Delegate.IsBound())
 	{
