[332494] Fix for a serialization crash that occurs when there is a cyclic dependency between a Blueprint class and an implemented interface.

#change Bumped UObject version.
#change Modified UClass::Serialize() to skip over serialization of the Interfaces array until after the Link() step. This ensures that Interface class serialization will not interrupt serialization of classes that might potentially also implement the interface.
#add Added test case assets to QAGame.

diff --git a/Engine/Source/Runtime/Core/Public/UObject/ObjectVersion.h b/Engine/Source/Runtime/Core/Public/UObject/ObjectVersion.h
--- a/Engine/Source/Runtime/Core/Public/UObject/ObjectVersion.h
+++ b/Engine/Source/Runtime/Core/Public/UObject/ObjectVersion.h
@@ -534,6 +534,8 @@ enum EUnrealEngineObjectUE4Version
 	VER_UE4_FIX_BLUEPRINT_VARIABLE_FLAGS,
 	// Vehicles use Nm for Torque instead of cm and RPM instead of rad/s part two (missed conversion for some variables
 	VER_UE4_VEHICLES_UNIT_CHANGE2,
+	// Changed order of interface class serialization
+	VER_UE4_UCLASS_SERIALIZE_INTERFACES_AFTER_LINKING,
  
 	// -----<new versions can be added before this line>-------------------------------------------------
 	// - this needs to be the last line (see note below)
diff --git a/Engine/Source/Runtime/CoreUObject/Private/UObject/Class.cpp b/Engine/Source/Runtime/CoreUObject/Private/UObject/Class.cpp
--- a/Engine/Source/Runtime/CoreUObject/Private/UObject/Class.cpp
+++ b/Engine/Source/Runtime/CoreUObject/Private/UObject/Class.cpp
@@ -2763,7 +2763,33 @@ void UClass::Serialize( FArchive& Ar )
 		Ar << ComponentNameToDefaultObjectMapSerialized;
 	}
 
-	Ar << Interfaces;
+	int32 NumInterfaces = 0;
+	int64 InterfacesStart = 0L;
+	if(Ar.IsLoading())
+	{
+		// Always start with no interfaces
+		Interfaces.Empty();
+
+		// In older versions, interface classes were serialized before linking. In case of cyclic dependencies, we need to skip over the serialized array and defer the load until after Link() is called below.
+		if(Ar.UE4Ver() < VER_UE4_UCLASS_SERIALIZE_INTERFACES_AFTER_LINKING && !GIsDuplicatingClassForReinstancing)
+		{
+			// Get our current position
+			InterfacesStart = Ar.Tell();
+
+			// Load the length of the Interfaces array
+			Ar << NumInterfaces;
+
+			// Seek past the Interfaces array
+			struct FSerializedInterfaceReference
+			{
+				FPackageIndex Class;
+				int32 PointerOffset;
+				bool bImplementedByK2;
+			};
+			Ar.Seek(InterfacesStart + sizeof(NumInterfaces) + NumInterfaces * sizeof(FSerializedInterfaceReference));
+		}
+	}
+
 	Ar << ClassGeneratedBy;
 
 	if(Ar.IsLoading())
@@ -2774,6 +2800,38 @@ void UClass::Serialize( FArchive& Ar )
 		Link(Ar, true);
 	}
 
+	if(Ar.IsLoading())
+	{
+		// Save current position
+		int64 CurrentOffset = Ar.Tell();
+
+		// In older versions, we need to seek backwards to the start of the interfaces array
+		if(Ar.UE4Ver() < VER_UE4_UCLASS_SERIALIZE_INTERFACES_AFTER_LINKING && !GIsDuplicatingClassForReinstancing)
+		{
+			Ar.Seek(InterfacesStart);
+		}
+		
+		// Load serialized interface classes
+		TArray<FImplementedInterface> SerializedInterfaces;
+		Ar << SerializedInterfaces;
+
+		// Apply loaded interfaces only if we have not already set them (i.e. during compile-on-load)
+		if(Interfaces.Num() == 0 && SerializedInterfaces.Num() > 0)
+		{
+			Interfaces = SerializedInterfaces;
+		}
+
+		// In older versions, seek back to our current position after linking
+		if(Ar.UE4Ver() < VER_UE4_UCLASS_SERIALIZE_INTERFACES_AFTER_LINKING && !GIsDuplicatingClassForReinstancing)
+		{
+			Ar.Seek(CurrentOffset);
+		}
+	}
+	else
+	{
+		Ar << Interfaces;
+	}
+
 	if( Ar.UE4Ver() < VER_UE4_DONTSORTCATEGORIES_REMOVED )
 	{
 		TArray<FName> DeprecatedDontSortCategories;
