Fixed comment.

diff --git a/Engine/Source/Runtime/CoreUObject/Public/UObject/ObjectBase.h b/Engine/Source/Runtime/CoreUObject/Public/UObject/ObjectBase.h
--- a/Engine/Source/Runtime/CoreUObject/Public/UObject/ObjectBase.h
+++ b/Engine/Source/Runtime/CoreUObject/Public/UObject/ObjectBase.h
@@ -720,7 +720,7 @@ namespace UP
 		/// Property should always be reset to the default value during any type of duplication (copy/paste, binary duplication, etc.)
 		DuplicateTransient,
 
-		/// Property should always be reset to the default value during any type of duplication (copy/paste, binary duplication, etc.)
+		/// Property should always be reset to the default value unless it's being duplicated for a PIE session
 		NonPIETransient,
 
 		/// Value is copied out after function call. Only valid on function param declaration.
