UStruct::Link wont cause a crash when a property is changed while linking  (ttp#334298)

diff --git a/Engine/Source/Runtime/CoreUObject/Private/UObject/Class.cpp b/Engine/Source/Runtime/CoreUObject/Private/UObject/Class.cpp
--- a/Engine/Source/Runtime/CoreUObject/Private/UObject/Class.cpp
+++ b/Engine/Source/Runtime/CoreUObject/Private/UObject/Class.cpp
@@ -417,34 +417,54 @@ void UStruct::Link(FArchive& Ar, bool bRelinkExistingProperties)
 			Ar.Preload(Field);
 		}
 
-		PropertiesSize = 0;
-		MinAlignment = 1;
-
-		if (InheritanceSuper)
+		int32 LoopNum = 1;
+		for (int32 LoopIter = 0; LoopIter < LoopNum; LoopIter++)
 		{
-			PropertiesSize = InheritanceSuper->GetPropertiesSize();
-			MinAlignment = InheritanceSuper->GetMinAlignment();
-		}
+			PropertiesSize = 0;
+			MinAlignment = 1;
 
-		for (UField* Field = Children; Field; Field = Field->Next)
-		{
-			if (Field->GetOuter() != this)
+			if (InheritanceSuper)
 			{
-				break;
+				PropertiesSize = InheritanceSuper->GetPropertiesSize();
+				MinAlignment = InheritanceSuper->GetMinAlignment();
 			}
 
-			if (UProperty* Property = Cast<UProperty>(Field))
+			for (UField* Field = Children; Field; Field = Field->Next)
 			{
+				if (Field->GetOuter() != this)
+				{
+					break;
+				}
+
+				if (UProperty* Property = Cast<UProperty>(Field))
+				{
 #if !WITH_EDITORONLY_DATA
-				// If we don't have the editor, make sure we aren't trying to link properties that are editor only.
-				check( !Property->IsEditorOnlyProperty() );
+					// If we don't have the editor, make sure we aren't trying to link properties that are editor only.
+					check(!Property->IsEditorOnlyProperty());
 #endif // WITH_EDITORONLY_DATA
-				ensureMsgf(Property->GetOuter() == this, TEXT("Linking '%s'. Property '%s' has outer '%s'"),
-					*GetFullName(), *Property->GetName(), *Property->GetOuter()->GetFullName());
-				PropertiesSize = Property->Link(Ar);
-				MinAlignment = FMath::Max( MinAlignment, Property->GetMinAlignment() );
+					ensureMsgf(Property->GetOuter() == this, TEXT("Linking '%s'. Property '%s' has outer '%s'"),
+						*GetFullName(), *Property->GetName(), *Property->GetOuter()->GetFullName());
+
+					// Linking a property can cause a recompilation of the struct. 
+					// When the property was changed, the struct should be relinked again, to be sure, the PropertiesSize is actual.
+					const bool bPropertyIsTransient = Property->HasAllFlags(RF_Transient);
+					const FName PropertyName = Property->GetFName();
+
+					PropertiesSize = Property->Link(Ar);
+
+					if ((bPropertyIsTransient != Property->HasAllFlags(RF_Transient)) || (PropertyName != Property->GetFName()))
+					{
+						LoopNum++;
+						const int32 MaxLoopLimit = 64;
+						ensure(LoopNum < MaxLoopLimit);
+						break;
+					}
+
+					MinAlignment = FMath::Max(MinAlignment, Property->GetMinAlignment());
+				}
 			}
 		}
+
 		bool bHandledWithCppStructOps = false;
 		if (GetClass()->IsChildOf(UScriptStruct::StaticClass()))
 		{
