ScriptBytecode loading optimization in editor (bytecode serialization is skipped).

diff --git a/Engine/Source/Runtime/CoreUObject/Private/UObject/Class.cpp b/Engine/Source/Runtime/CoreUObject/Private/UObject/Class.cpp
--- a/Engine/Source/Runtime/CoreUObject/Private/UObject/Class.cpp
+++ b/Engine/Source/Runtime/CoreUObject/Private/UObject/Class.cpp
@@ -1203,7 +1203,8 @@ void UStruct::Serialize( FArchive& Ar )
 				int32 ScriptStart = Ar.Tell();
 				Ar.Serialize(TempScript.GetData(), ScriptStorageSize);
 
-				if ((Ar.UE4Ver() < VER_MIN_SCRIPTVM_UE4) || (Ar.LicenseeUE4Ver() < VER_MIN_SCRIPTVM_LICENSEEUE4))
+				const bool bSkipByteCodeSerialization = (0 != WITH_EDITOR);
+				if (bSkipByteCodeSerialization || (Ar.UE4Ver() < VER_MIN_SCRIPTVM_UE4) || (Ar.LicenseeUE4Ver() < VER_MIN_SCRIPTVM_LICENSEEUE4))
 				{
 					// Discard the bytecode as it's too old and might cause serialization errors
 					ScriptStorageSize = 0;
