ScriptStructs are initialized faster when STRUCT_ZeroConstructor flag is set.

diff --git a/Engine/Source/Runtime/CoreUObject/Private/UObject/Class.cpp b/Engine/Source/Runtime/CoreUObject/Private/UObject/Class.cpp
--- a/Engine/Source/Runtime/CoreUObject/Private/UObject/Class.cpp
+++ b/Engine/Source/Runtime/CoreUObject/Private/UObject/Class.cpp
@@ -2038,6 +2038,10 @@ void UScriptStruct::InitializeScriptStruct(void* InDest, int32 ArrayDim)
 
 	//@todo UE4 optimize
 	FMemory::Memzero(Dest, ArrayDim * Stride);
+	if (StructFlags & STRUCT_ZeroConstructor)
+	{
+		return;
+	}
 
 	int32 InitializedSize = 0;
 	UScriptStruct::ICppStructOps* TheCppStructOps = GetCppStructOps();
