Back out revision 124 from //depot/UE4/Engine/Source/Runtime/CoreUObject/Private/UObject/Class.cpp

diff --git a/Engine/Source/Runtime/CoreUObject/Private/UObject/Class.cpp b/Engine/Source/Runtime/CoreUObject/Private/UObject/Class.cpp
--- a/Engine/Source/Runtime/CoreUObject/Private/UObject/Class.cpp
+++ b/Engine/Source/Runtime/CoreUObject/Private/UObject/Class.cpp
@@ -2038,10 +2038,6 @@ void UScriptStruct::InitializeScriptStruct(void* InDest, int32 ArrayDim)
 
 	//@todo UE4 optimize
 	FMemory::Memzero(Dest, ArrayDim * Stride);
-	if (StructFlags & STRUCT_ZeroConstructor)
-	{
-		return;
-	}
 
 	int32 InitializedSize = 0;
 	UScriptStruct::ICppStructOps* TheCppStructOps = GetCppStructOps();
