Fix issue when saving empty cooked package crashes the coooker. Saving fails gracefully with log warning message.
#ttp 331348 - CORE: Cooking ShooterGame with -compressed crashes the cooker
#proj core
#branch UE4
#summary While cooking, when we save a package, if we find out the package is empty, it shouldn't be saved.

diff --git a/Engine/Source/Runtime/CoreUObject/Private/UObject/SavePackage.cpp b/Engine/Source/Runtime/CoreUObject/Private/UObject/SavePackage.cpp
--- a/Engine/Source/Runtime/CoreUObject/Private/UObject/SavePackage.cpp
+++ b/Engine/Source/Runtime/CoreUObject/Private/UObject/SavePackage.cpp
@@ -2605,9 +2605,23 @@ bool UPackage::SavePackage( UPackage* InOuter, UObject* Base, EObjectFlags TopLe
 							Obj->UnMark(EObjectMark(OBJECTMARK_TagExp|OBJECTMARK_TagImp));
 						}
 					}
+
+					// Obj->UnMark will delete object if it stripped all its tags.
+					// If all exportable objects are deleted, an attempt to save
+					// package will cause a crash.
+					GetObjectsWithAnyMarks(TagExpObjects, OBJECTMARK_TagExp);
+					if (TagExpObjects.Num() == 0)
+					{
+						if (!(SaveFlags & SAVE_NoError))
+						{
+							UE_LOG(LogSavePackage, Warning, TEXT("No exportable objects found in package. Package not saved."));
+						}
+
+						return false;
+					}
 				}
 
-				// Import objects & names.				
+				// Import objects & names.
 				{
 					TArray<UObject*> TagExpObjects;
 					GetObjectsWithAnyMarks(TagExpObjects, OBJECTMARK_TagExp);
