Calling MCDelegate won't cause a crash.
ttp333724 LIVE: CRASH: BP: Calling an Event Dispatcher from a non-placed blueprint variable

diff --git a/Engine/Source/Runtime/CoreUObject/Private/UObject/ScriptCore.cpp b/Engine/Source/Runtime/CoreUObject/Private/UObject/ScriptCore.cpp
--- a/Engine/Source/Runtime/CoreUObject/Private/UObject/ScriptCore.cpp
+++ b/Engine/Source/Runtime/CoreUObject/Private/UObject/ScriptCore.cpp
@@ -1461,7 +1461,6 @@ public:
 		Stack.MostRecentProperty = NULL;
 		Stack.Step( Stack.Object, NULL );
 		const FMulticastScriptDelegate* DelegateAddr = (FMulticastScriptDelegate*)Stack.MostRecentPropertyAddress;
-		check(NULL != DelegateAddr);
 
 		//Fill parameters
 		uint8* Parameters = (uint8*)FMemory_Alloca(SignatureFunction->ParmsSize);
@@ -1489,7 +1488,10 @@ public:
 		Stack.Code++;
 
 		//Process delegate
-		DelegateAddr->ProcessMulticastDelegate<UObject>(Parameters);
+		if (DelegateAddr)
+		{
+			DelegateAddr->ProcessMulticastDelegate<UObject>(Parameters);
+		}
 		
 		//Clean parameters
 		for (UProperty* Destruct = SignatureFunction->DestructorLink; Destruct; Destruct = Destruct->DestructorLinkNext)
