Back out revision 123 from //depot/UE4/Engine/Source/Runtime/CoreUObject/Private/UObject/Class.cpp

diff --git a/Engine/Source/Runtime/CoreUObject/Private/UObject/Class.cpp b/Engine/Source/Runtime/CoreUObject/Private/UObject/Class.cpp
--- a/Engine/Source/Runtime/CoreUObject/Private/UObject/Class.cpp
+++ b/Engine/Source/Runtime/CoreUObject/Private/UObject/Class.cpp
@@ -1203,8 +1203,7 @@ void UStruct::Serialize( FArchive& Ar )
 				int32 ScriptStart = Ar.Tell();
 				Ar.Serialize(TempScript.GetData(), ScriptStorageSize);
 
-				const bool bSkipByteCodeSerialization = (0 != WITH_EDITOR);
-				if (bSkipByteCodeSerialization || (Ar.UE4Ver() < VER_MIN_SCRIPTVM_UE4) || (Ar.LicenseeUE4Ver() < VER_MIN_SCRIPTVM_LICENSEEUE4))
+				if ((Ar.UE4Ver() < VER_MIN_SCRIPTVM_UE4) || (Ar.LicenseeUE4Ver() < VER_MIN_SCRIPTVM_LICENSEEUE4))
 				{
 					// Discard the bytecode as it's too old and might cause serialization errors
 					ScriptStorageSize = 0;
