'Force Delete' an UDStruct wont cause a crash. ttp#335026

diff --git a/Engine/Source/Runtime/CoreUObject/Private/UObject/PropertyStruct.cpp b/Engine/Source/Runtime/CoreUObject/Private/UObject/PropertyStruct.cpp
--- a/Engine/Source/Runtime/CoreUObject/Private/UObject/PropertyStruct.cpp
+++ b/Engine/Source/Runtime/CoreUObject/Private/UObject/PropertyStruct.cpp
@@ -37,7 +37,10 @@ void UStructProperty::LinkInternal(FArchive& Ar)
 
 	// Preload is required here in order to load the value of Struct->PropertiesSize
 	Ar.Preload(Struct);
-	checkSlow(Struct);
+	if ( !ensure(Struct) )
+	{
+		Struct = GetFallbackStruct();
+	}
 	Struct->RecursivelyPreload();
 	
 	ElementSize = Align(Struct->PropertiesSize, Struct->GetMinAlignment());
