Ensure, that VMBytecode was properly serialized.

diff --git a/Engine/Source/Runtime/CoreUObject/Private/UObject/Class.cpp b/Engine/Source/Runtime/CoreUObject/Private/UObject/Class.cpp
--- a/Engine/Source/Runtime/CoreUObject/Private/UObject/Class.cpp
+++ b/Engine/Source/Runtime/CoreUObject/Private/UObject/Class.cpp
@@ -1202,6 +1202,7 @@ void UStruct::Serialize( FArchive& Ar )
 				TempScript.AddUninitialized(ScriptStorageSize);
 				int32 ScriptStart = Ar.Tell();
 				Ar.Serialize(TempScript.GetData(), ScriptStorageSize);
+				const int32 ScriptEnd = Ar.Tell();
 
 				if ((Ar.UE4Ver() < VER_MIN_SCRIPTVM_UE4) || (Ar.LicenseeUE4Ver() < VER_MIN_SCRIPTVM_LICENSEEUE4))
 				{
@@ -1219,6 +1220,7 @@ void UStruct::Serialize( FArchive& Ar )
 					{	
 						SerializeExpr( iCode, Ar );
 					}
+					ensure(Ar.Tell() == ScriptEnd);
 				}
 				// and update the SHA (does nothing if not currently calculating SHA)
 				LinkerLoad->UpdateScriptSHAKey(TempScript);
