Adding code to handle benign cases of net guid's being reassigned

diff --git a/Engine/Source/Runtime/CoreUObject/Private/UObject/CoreNet.cpp b/Engine/Source/Runtime/CoreUObject/Private/UObject/CoreNet.cpp
--- a/Engine/Source/Runtime/CoreUObject/Private/UObject/CoreNet.cpp
+++ b/Engine/Source/Runtime/CoreUObject/Private/UObject/CoreNet.cpp
@@ -398,7 +398,26 @@ FNetworkGUID UPackageMap::AssignNetGUID(const UObject* Object, FNetworkGUID NewN
 		if ( ExistingCacheObjectPtr )
 		{
 			const UObject * OldObject = ExistingCacheObjectPtr->Object.Get();
-			UE_LOG( LogNetPackageMap, Warning, TEXT( "Reassigning NetGUID <%s> to %s (was assigned to object %s)" ), *NewNetworkGUID.ToString(), Object ? *Object->GetName() : TEXT( "NULL" ), OldObject ? *OldObject->GetName() : TEXT( "NULL" ) );
+
+			// If this net guid was found but the old object is NULL, this can happen due to:
+			//	1. Actor channel was closed locally (but we don't remove the net guid entry, since we can't know for sure if it will be referenced again)
+			//		a. Then when we re-create a channel, and assign this actor, we will find the old guid entry here
+			//	2. Dynamic object was locally GC'd, but then exported again from the server
+			//
+			// If this net guid was found and the objects match, we don't care. This can happen due to:
+			//	1. Same thing above can happen, but if we for some reason didn't destroy the actor/object we will see this case
+			//
+			// If the object pointers are different, this can be a problem, 
+			//	since this should only be possible if something gets out of sync during the net guid exchange code
+			if ( OldObject != NULL && OldObject != Object )
+			{
+				UE_LOG( LogNetPackageMap, Warning, TEXT( "Reassigning NetGUID <%s> to %s (was assigned to object %s)" ), *NewNetworkGUID.ToString(), Object ? *Object->GetPathName() : TEXT( "NULL" ), OldObject ? *OldObject->GetPathName() : TEXT( "NULL" ) );
+			}
+			else
+			{
+				UE_LOG( LogNetPackageMap, Verbose, TEXT( "Reassigning NetGUID <%s> to %s (was assigned to object %s)" ), *NewNetworkGUID.ToString(), Object ? *Object->GetPathName() : TEXT( "NULL" ), OldObject ? *OldObject->GetPathName() : TEXT( "NULL" ) );
+			}
+
 			Cache->NetGUIDLookup.Remove( ExistingCacheObjectPtr->Object );
 		}
 
@@ -406,7 +425,7 @@ FNetworkGUID UPackageMap::AssignNetGUID(const UObject* Object, FNetworkGUID NewN
 
 		if ( ExistingNetworkGUIDPtr )
 		{
-			UE_LOG( LogNetPackageMap, Warning, TEXT( "Changing NetGUID on object %s from <%s> to <%s>" ), Object ? *Object->GetName() : TEXT( "NULL" ), *ExistingNetworkGUIDPtr->ToString(), *NewNetworkGUID.ToString() );
+			UE_LOG( LogNetPackageMap, Warning, TEXT( "Changing NetGUID on object %s from <%s:%s> to <%s:%s>" ), Object ? *Object->GetName() : TEXT( "NULL" ), *ExistingNetworkGUIDPtr->ToString(), ExistingNetworkGUIDPtr->IsDynamic() ? TEXT("TRUE") : TEXT("FALSE"), *NewNetworkGUID.ToString(), NewNetworkGUID.IsDynamic() ? TEXT("TRUE") : TEXT("FALSE") );
 			Cache->ObjectLookup.Remove( *ExistingNetworkGUIDPtr );
 		}
 	}
