Log full name when net guids are changed + comment why this can happen

diff --git a/Engine/Source/Runtime/CoreUObject/Private/UObject/CoreNet.cpp b/Engine/Source/Runtime/CoreUObject/Private/UObject/CoreNet.cpp
--- a/Engine/Source/Runtime/CoreUObject/Private/UObject/CoreNet.cpp
+++ b/Engine/Source/Runtime/CoreUObject/Private/UObject/CoreNet.cpp
@@ -425,7 +425,9 @@ FNetworkGUID UPackageMap::AssignNetGUID(const UObject* Object, FNetworkGUID NewN
 
 		if ( ExistingNetworkGUIDPtr )
 		{
-			UE_LOG( LogNetPackageMap, Warning, TEXT( "Changing NetGUID on object %s from <%s:%s> to <%s:%s>" ), Object ? *Object->GetName() : TEXT( "NULL" ), *ExistingNetworkGUIDPtr->ToString(), ExistingNetworkGUIDPtr->IsDynamic() ? TEXT("TRUE") : TEXT("FALSE"), *NewNetworkGUID.ToString(), NewNetworkGUID.IsDynamic() ? TEXT("TRUE") : TEXT("FALSE") );
+			// This can happen when the server destroys a package/object and the client doesn't.
+			// The server will simply assign a new guid, and the client will assign to the old object that never got deleted locally
+			UE_LOG( LogNetPackageMap, Warning, TEXT( "Changing NetGUID on object %s from <%s:%s> to <%s:%s>" ), Object ? *Object->GetPathName() : TEXT( "NULL" ), *ExistingNetworkGUIDPtr->ToString(), ExistingNetworkGUIDPtr->IsDynamic() ? TEXT("TRUE") : TEXT("FALSE"), *NewNetworkGUID.ToString(), NewNetworkGUID.IsDynamic() ? TEXT("TRUE") : TEXT("FALSE") );
 			Cache->ObjectLookup.Remove( *ExistingNetworkGUIDPtr );
 		}
 	}
