NeedLoadInEditorGame flag added. VER_UE4_LOAD_FOR_EDITOR_GAME. OBJECTMARK_NotForEditorGame.
NeedLoadInEditorGame means, that the object must be always loaded in 'editor -game'. When it's false it doesn't mean, that it's never loaded in editor game.

diff --git a/Engine/Source/Runtime/Core/Public/UObject/ObjectVersion.h b/Engine/Source/Runtime/Core/Public/UObject/ObjectVersion.h
--- a/Engine/Source/Runtime/Core/Public/UObject/ObjectVersion.h
+++ b/Engine/Source/Runtime/Core/Public/UObject/ObjectVersion.h
@@ -542,6 +542,8 @@ enum EUnrealEngineObjectUE4Version
 	VER_UE4_FIX_MATERIAL_COORDS,
  	// Changed SpeedTree wind presets to v7
 	VER_UE4_SPEEDTREE_WIND_V7,
+	// NeedsLoadForEditorGame added
+	VER_UE4_LOAD_FOR_EDITOR_GAME,
 
 	// -----<new versions can be added before this line>-------------------------------------------------
 	// - this needs to be the last line (see note below)
diff --git a/Engine/Source/Runtime/CoreUObject/Private/UObject/MetaData.cpp b/Engine/Source/Runtime/CoreUObject/Private/UObject/MetaData.cpp
--- a/Engine/Source/Runtime/CoreUObject/Private/UObject/MetaData.cpp
+++ b/Engine/Source/Runtime/CoreUObject/Private/UObject/MetaData.cpp
@@ -387,6 +387,11 @@ bool UMetaData::NeedsLoadForServer() const
 	return false; 
 }
 
+bool UMetaData::NeedsLoadForEditorGame() const
+{
+	return true;
+}
+
 void UMetaData::InitializeRedirectMap()
 {
 	static bool bKeyRedirectMapInitialized = false;
diff --git a/Engine/Source/Runtime/CoreUObject/Private/UObject/ObjectResource.cpp b/Engine/Source/Runtime/CoreUObject/Private/UObject/ObjectResource.cpp
--- a/Engine/Source/Runtime/CoreUObject/Private/UObject/ObjectResource.cpp
+++ b/Engine/Source/Runtime/CoreUObject/Private/UObject/ObjectResource.cpp
@@ -43,6 +43,7 @@ FObjectExport::FObjectExport()
 ,	bNotForClient	( false														)
 ,	bNotForServer	( false														)
 ,	PackageGuid		( FGuid(0,0,0,0)											)
+,	bNotForEditorGame	(true													)
 {}
 
 FObjectExport::FObjectExport( UObject* InObject )
@@ -60,11 +61,13 @@ FObjectExport::FObjectExport( UObject* InObject )
 ,	bNotForServer	( false														)
 ,	PackageGuid		( FGuid(0,0,0,0)											)
 ,	PackageFlags	( 0															)
+,	bNotForEditorGame	(true													)
 {
 	if(Object)		
 	{
 		bNotForClient = Object->HasAnyMarks(OBJECTMARK_NotForClient);
 		bNotForServer = Object->HasAnyMarks(OBJECTMARK_NotForServer);
+		bNotForEditorGame = Object->HasAnyMarks(OBJECTMARK_NotForEditorGame);
 	}
 }
 
@@ -103,6 +106,11 @@ FArchive& operator<<( FArchive& Ar, FObjectExport& E )
 	Ar << E.PackageGuid;
 	Ar << E.PackageFlags;
 
+	if (Ar.UE4Ver() >= VER_UE4_LOAD_FOR_EDITOR_GAME)
+	{
+		Ar << E.bNotForEditorGame;
+	}
+
 	return Ar;
 }
 
diff --git a/Engine/Source/Runtime/CoreUObject/Private/UObject/SavePackage.cpp b/Engine/Source/Runtime/CoreUObject/Private/UObject/SavePackage.cpp
--- a/Engine/Source/Runtime/CoreUObject/Private/UObject/SavePackage.cpp
+++ b/Engine/Source/Runtime/CoreUObject/Private/UObject/SavePackage.cpp
@@ -384,6 +384,22 @@ FArchive& FArchiveSaveTagExports::operator<<( UObject*& Obj )
 				Search = Search->GetOuter();
 			} while (Search && !Obj->HasAllMarks(ObjectMarks));
 
+			{
+				bool bNeedsLoadForEditorGame = false;
+				for (UObject* OuterIt = Obj; OuterIt; OuterIt = OuterIt->GetOuter())
+				{
+					if (OuterIt->NeedsLoadForEditorGame())
+					{
+						bNeedsLoadForEditorGame = true;
+						break;
+					}
+				}
+				if (!bNeedsLoadForEditorGame)
+				{
+					Obj->Mark(OBJECTMARK_NotForEditorGame);
+				}
+			}
+
 			// skip these checks if the Template object is the CDO for an intrinsic class, as they will never have any load flags set.
 			if ( Template != NULL 
 			&& (!Template->GetClass()->HasAnyClassFlags(CLASS_Intrinsic) || !Template->HasAnyFlags(RF_ClassDefaultObject)) )
diff --git a/Engine/Source/Runtime/CoreUObject/Public/UObject/CoreObject.h b/Engine/Source/Runtime/CoreUObject/Public/UObject/CoreObject.h
--- a/Engine/Source/Runtime/CoreUObject/Public/UObject/CoreObject.h
+++ b/Engine/Source/Runtime/CoreUObject/Public/UObject/CoreObject.h
@@ -499,6 +499,7 @@ public:
 	virtual void Serialize(FArchive& Ar) OVERRIDE;
 	virtual bool NeedsLoadForClient() const OVERRIDE;
 	virtual bool NeedsLoadForServer() const OVERRIDE;
+	virtual bool NeedsLoadForEditorGame() const OVERRIDE;
 	virtual bool IsAsset() const OVERRIDE { return false; }
 	// End of UObject interface
 
diff --git a/Engine/Source/Runtime/CoreUObject/Public/UObject/Linker.h b/Engine/Source/Runtime/CoreUObject/Public/UObject/Linker.h
--- a/Engine/Source/Runtime/CoreUObject/Public/UObject/Linker.h
+++ b/Engine/Source/Runtime/CoreUObject/Public/UObject/Linker.h
@@ -257,6 +257,14 @@ struct FObjectExport : public FObjectResource
 	 */
 	uint32			PackageFlags;
 
+	/**
+	 * whether the export should be always loaded in editor game
+	 * False means that the object is necessary for editor game,
+	 * True doesn't means, that the object won't be loaded.
+	 * Serialized
+	 */
+	bool			bNotForEditorGame;
+
 	/**
 	 * Constructors
 	 */
@@ -1093,6 +1101,13 @@ class ULinker : public UObject, public FLinkerTables
 		{
 			return true;
 		}
+#if WITH_EDITOR
+		if (!Export.bNotForEditorGame)
+		{
+			ensure(true);
+			return false;
+		}
+#endif
 		if (FilterClientButNotServer && Export.bNotForServer) // "we are a dedicated server"
 		{
 			return true;
diff --git a/Engine/Source/Runtime/CoreUObject/Public/UObject/UObject.h b/Engine/Source/Runtime/CoreUObject/Public/UObject/UObject.h
--- a/Engine/Source/Runtime/CoreUObject/Public/UObject/UObject.h
+++ b/Engine/Source/Runtime/CoreUObject/Public/UObject/UObject.h
@@ -255,6 +255,16 @@ public:
 		return true; 
 	}
 
+	/**
+	 * Called during saving to determine the load flags to save with the object.
+	 *
+	 * @return	true if this object should always be loaded for editor game
+	 */
+	virtual bool NeedsLoadForEditorGame() const
+	{
+		return false;
+	}
+
 	/** 
 	 *	Determines if you can create an object from the supplied template in the current context (editor, client only, dedicated server, game/listen) 
 	 *	This calls NeedsLoadForClient & NeedsLoadForServer
diff --git a/Engine/Source/Runtime/CoreUObject/Public/UObject/UObjectMarks.h b/Engine/Source/Runtime/CoreUObject/Public/UObject/UObjectMarks.h
--- a/Engine/Source/Runtime/CoreUObject/Public/UObject/UObjectMarks.h
+++ b/Engine/Source/Runtime/CoreUObject/Public/UObject/UObjectMarks.h
@@ -47,6 +47,7 @@ enum EObjectMark
 	OBJECTMARK_TagExp				= 0x00000010,		// Temporary export tag in load/save.
 	OBJECTMARK_NotForClient			= 0x00000020,		// Temporary save tag for client load flag.
 	OBJECTMARK_NotForServer			= 0x00000040,		// Temporary save tag for server load flag.
+	OBJECTMARK_NotForEditorGame		= 0x00000080,		// Temporary save tag for editorgame load flag.
 
 	OBJECTMARK_ALLMARKS				= 0xffffffff,		// -1, all possible marks
 };
diff --git a/Engine/Source/Runtime/Engine/Classes/Engine/Blueprint.h b/Engine/Source/Runtime/Engine/Classes/Engine/Blueprint.h
--- a/Engine/Source/Runtime/Engine/Classes/Engine/Blueprint.h
+++ b/Engine/Source/Runtime/Engine/Classes/Engine/Blueprint.h
@@ -494,6 +494,7 @@ public:
 	virtual void TagSubobjects(EObjectFlags NewFlags) OVERRIDE;
 	virtual bool NeedsLoadForClient() const OVERRIDE;
 	virtual bool NeedsLoadForServer() const OVERRIDE;
+	virtual bool NeedsLoadForEditorGame() const OVERRIDE;
 	// End of UObject interface
 
 	/** Get the Blueprint object that generated the supplied class */
diff --git a/Engine/Source/Runtime/Engine/Private/Blueprint.cpp b/Engine/Source/Runtime/Engine/Private/Blueprint.cpp
--- a/Engine/Source/Runtime/Engine/Private/Blueprint.cpp
+++ b/Engine/Source/Runtime/Engine/Private/Blueprint.cpp
@@ -784,6 +784,10 @@ bool UBlueprint::NeedsLoadForServer() const
 	return !Helper.bDontLoad;
 }
 
+bool UBlueprint::NeedsLoadForEditorGame() const
+{
+	return true;
+}
 
 void UBlueprint::TagSubobjects(EObjectFlags NewFlags)
 {
