Changed "Graph is linked to object(s) in external map" dialog to have a single "OK" button.
#ttp 335736 - EDITOR: When saving a blueprint, received a "can't save" error popup message with "yes"/"no" buttons
#branch UE4
#proj Runtime.CoreUObject
#change Changed MessageDialog type to be EAppMsgType::Ok, and moved the code conditional on 'Yes' out of the condition, so it's always executed.

diff --git a/Engine/Source/Runtime/CoreUObject/Private/UObject/SavePackage.cpp b/Engine/Source/Runtime/CoreUObject/Private/UObject/SavePackage.cpp
--- a/Engine/Source/Runtime/CoreUObject/Private/UObject/SavePackage.cpp
+++ b/Engine/Source/Runtime/CoreUObject/Private/UObject/SavePackage.cpp
@@ -2826,18 +2826,16 @@ bool UPackage::SavePackage( UPackage* InOuter, UObject* Base, EObjectFlags TopLe
 					const FText Message = FText::Format( NSLOCTEXT("Core", "LinkedToObjectsInOtherMap_FindCulpritQ", "Can't save {FileName}: Graph is linked to object(s) in external map.\nExternal Object(s):\n{ObjectNames}  \nTry to find the chain of references to that object (may take some time)?"), Args );
 
 					FString CulpritString = TEXT( "Unknown" );
-					const bool bProceed = FMessageDialog::Open(EAppMsgType::YesNo, Message ) == EAppReturnType::Yes;
-					if ( bProceed )
+					FMessageDialog::Open(EAppMsgType::Ok, Message);
+
+					FindMostLikelyCulprit( IllegalObjectsInOtherMaps, MostLikelyCulprit, PropertyRef );
+					if( MostLikelyCulprit != NULL && PropertyRef != NULL )
 					{
-						FindMostLikelyCulprit( IllegalObjectsInOtherMaps, MostLikelyCulprit, PropertyRef );
-						if( MostLikelyCulprit != NULL && PropertyRef != NULL )
-						{
-							CulpritString = FString::Printf(TEXT("%s (%s)"), *MostLikelyCulprit->GetFullName(), *PropertyRef->GetName());
-						}
-						else if (MostLikelyCulprit != NULL)
-						{
-							CulpritString = FString::Printf(TEXT("%s (Unknown property)"), *MostLikelyCulprit->GetFullName());
-						}
+						CulpritString = FString::Printf(TEXT("%s (%s)"), *MostLikelyCulprit->GetFullName(), *PropertyRef->GetName());
+					}
+					else if (MostLikelyCulprit != NULL)
+					{
+						CulpritString = FString::Printf(TEXT("%s (Unknown property)"), *MostLikelyCulprit->GetFullName());
 					}
 
 					// Free the file handle and delete the temporary file
@@ -2885,14 +2883,12 @@ bool UPackage::SavePackage( UPackage* InOuter, UObject* Base, EObjectFlags TopLe
 					const FText Message = FText::Format( NSLOCTEXT("Core", "LinkedToObjectsInOtherMap_FindCulpritQ", "Can't save {FileName}: Graph is linked to object(s) in external map.\nExternal Object(s):\n{ObjectNames}  \nTry to find the chain of references to that object (may take some time)?"), Args );
 
 					FString CulpritString = TEXT( "Unknown" );
-					const bool bProceed = FMessageDialog::Open(EAppMsgType::YesNo, Message ) == EAppReturnType::Yes;
-					if ( bProceed )
-					{
-						FindMostLikelyCulprit( PrivateObjects, MostLikelyCulprit, PropertyRef );
-						CulpritString = FString::Printf(TEXT("%s (%s)"),
-							(MostLikelyCulprit != NULL) ? *MostLikelyCulprit->GetFullName() : TEXT("(unknown culprit)"),
-							(PropertyRef != NULL) ? *PropertyRef->GetName() : TEXT("unknown property ref"));
-					}
+					FMessageDialog::Open(EAppMsgType::Ok, Message);
+
+					FindMostLikelyCulprit( PrivateObjects, MostLikelyCulprit, PropertyRef );
+					CulpritString = FString::Printf(TEXT("%s (%s)"),
+						(MostLikelyCulprit != NULL) ? *MostLikelyCulprit->GetFullName() : TEXT("(unknown culprit)"),
+						(PropertyRef != NULL) ? *PropertyRef->GetName() : TEXT("unknown property ref"));
 
 					// free the file handle and delete the temporary file
 					Linker->Detach();
