Object Redirectors will no longer be saved during cooking.

diff --git a/Engine/Source/Runtime/CoreUObject/Private/UObject/ObjectRedirector.cpp b/Engine/Source/Runtime/CoreUObject/Private/UObject/ObjectRedirector.cpp
--- a/Engine/Source/Runtime/CoreUObject/Private/UObject/ObjectRedirector.cpp
+++ b/Engine/Source/Runtime/CoreUObject/Private/UObject/ObjectRedirector.cpp
@@ -37,6 +37,21 @@ void UObjectRedirector::Serialize( FArchive& Ar )
 	Ar << DestinationObject;
 }
 
+bool UObjectRedirector::NeedsLoadForClient() const
+{
+	return false;
+}
+
+bool UObjectRedirector::NeedsLoadForServer() const
+{
+	return false;
+}
+
+bool UObjectRedirector::NeedsLoadForEditorGame() const
+{
+	return true;
+}
+
 void UObjectRedirector::GetAssetRegistryTags(TArray<FAssetRegistryTag>& OutTags) const
 {
 	FString DestVal;
diff --git a/Engine/Source/Runtime/CoreUObject/Private/UObject/SavePackage.cpp b/Engine/Source/Runtime/CoreUObject/Private/UObject/SavePackage.cpp
--- a/Engine/Source/Runtime/CoreUObject/Private/UObject/SavePackage.cpp
+++ b/Engine/Source/Runtime/CoreUObject/Private/UObject/SavePackage.cpp
@@ -2633,7 +2633,7 @@ bool UPackage::SavePackage( UPackage* InOuter, UObject* Base, EObjectFlags TopLe
 					{
 						if (!(SaveFlags & SAVE_NoError))
 						{
-							UE_LOG(LogSavePackage, Warning, TEXT("No exportable objects found in package. Package not saved."));
+							UE_LOG(LogSavePackage, Display, TEXT("No exports found (or all exports are editor-only) for %s. Package will not be saved."), *BaseFilename);
 						}
 
 						return false;
diff --git a/Engine/Source/Runtime/CoreUObject/Public/UObject/ObjectRedirector.h b/Engine/Source/Runtime/CoreUObject/Public/UObject/ObjectRedirector.h
--- a/Engine/Source/Runtime/CoreUObject/Public/UObject/ObjectRedirector.h
+++ b/Engine/Source/Runtime/CoreUObject/Public/UObject/ObjectRedirector.h
@@ -21,6 +21,9 @@ class UObjectRedirector : public UObject
 	// UObject interface.
 	virtual void PreSave();
 	void Serialize( FArchive& Ar );
+	virtual bool NeedsLoadForClient() const OVERRIDE;
+	virtual bool NeedsLoadForServer() const OVERRIDE;
+	virtual bool NeedsLoadForEditorGame() const OVERRIDE;
 	virtual void GetAssetRegistryTags(TArray<FAssetRegistryTag>& OutTags) const OVERRIDE;
 
 	/**
