FStructureOnScope moved to core. ScriptStruct functions are const.

diff --git a/Engine/Source/Editor/BlueprintGraph/Private/K2Node_MakeStruct.cpp b/Engine/Source/Editor/BlueprintGraph/Private/K2Node_MakeStruct.cpp
--- a/Engine/Source/Editor/BlueprintGraph/Private/K2Node_MakeStruct.cpp
+++ b/Engine/Source/Editor/BlueprintGraph/Private/K2Node_MakeStruct.cpp
@@ -192,7 +192,7 @@ void UK2Node_MakeStruct::AllocateDefaultPins()
 		CreatePin(EGPD_Output, Schema->PC_Struct, TEXT(""), StructType, false, false, StructType->GetName());
 		
 		{
-			FStructureEditorUtils::FStructOnScope StructOnScope(Cast<UScriptStruct>(StructType));
+			FStructOnScope StructOnScope(Cast<UScriptStruct>(StructType));
 			FMakeStructPinManager OptionalPinManager(StructOnScope.GetStructMemory());
 			OptionalPinManager.RebuildPropertyList(ShowPinForProperties, StructType);
 			OptionalPinManager.CreateVisiblePins(ShowPinForProperties, StructType, EGPD_Input, this);
diff --git a/Engine/Source/Editor/UnrealEd/Public/Kismet2/StructureEditorUtils.h b/Engine/Source/Editor/UnrealEd/Public/Kismet2/StructureEditorUtils.h
--- a/Engine/Source/Editor/UnrealEd/Public/Kismet2/StructureEditorUtils.h
+++ b/Engine/Source/Editor/UnrealEd/Public/Kismet2/StructureEditorUtils.h
@@ -44,41 +44,6 @@ public:
 		}
 	};
 
-	struct FStructOnScope
-	{
-	private:
-		UScriptStruct* ScriptStruct;
-		uint8* SampleStructMemory;
-
-		FStructOnScope(const FStructOnScope&);
-		FStructOnScope& operator=(const FStructOnScope&);
-
-	public:
-		FStructOnScope(UScriptStruct* InScriptStruct) 
-			: ScriptStruct(InScriptStruct)
-			, SampleStructMemory(NULL)
-		{
-			if(ScriptStruct)
-			{
-				SampleStructMemory = (uint8*)FMemory::Malloc(ScriptStruct->GetStructureSize());
-				ScriptStruct->InitializeScriptStruct(SampleStructMemory);
-			}
-		}
-
-		uint8* GetStructMemory() { return SampleStructMemory; }
-		const uint8* GetStructMemory() const { return SampleStructMemory; }
-
-		~FStructOnScope()
-		{
-			if(ScriptStruct)
-			{
-				ScriptStruct->DestroyScriptStruct(SampleStructMemory);
-				FMemory::Free(SampleStructMemory);
-				SampleStructMemory = NULL;
-			}
-		}
-	};
-
 	//STRUCTURE
 	static UUserDefinedStruct* CreateUserDefinedStruct(UObject* InParent, FName Name, EObjectFlags Flags);
 
diff --git a/Engine/Source/Runtime/CoreUObject/Private/UObject/Class.cpp b/Engine/Source/Runtime/CoreUObject/Private/UObject/Class.cpp
--- a/Engine/Source/Runtime/CoreUObject/Private/UObject/Class.cpp
+++ b/Engine/Source/Runtime/CoreUObject/Private/UObject/Class.cpp
@@ -1964,7 +1964,7 @@ void UScriptStruct::Link(FArchive& Ar, bool bRelinkExistingProperties)
 	}
 }
 
-bool UScriptStruct::CompareScriptStruct( const void* A, const void* B, uint32 PortFlags )
+bool UScriptStruct::CompareScriptStruct(const void* A, const void* B, uint32 PortFlags) const
 {
 	check(A);
 	if (StructFlags & STRUCT_IdenticalNative)
@@ -1993,7 +1993,7 @@ bool UScriptStruct::CompareScriptStruct( const void* A, const void* B, uint32 Po
 }
 
 
-void UScriptStruct::CopyScriptStruct(void* InDest, void const* InSrc, int32 ArrayDim)
+void UScriptStruct::CopyScriptStruct(void* InDest, void const* InSrc, int32 ArrayDim) const
 {
 	uint8 *Dest = (uint8*)InDest;
 	check(Dest);
@@ -2030,7 +2030,7 @@ void UScriptStruct::CopyScriptStruct(void* InDest, void const* InSrc, int32 Arra
 }
 
 
-void UScriptStruct::InitializeScriptStruct(void* InDest, int32 ArrayDim)
+void UScriptStruct::InitializeScriptStruct(void* InDest, int32 ArrayDim) const
 {
 	uint8 *Dest = (uint8*)InDest;
 	check(Dest);
@@ -2077,7 +2077,7 @@ void UScriptStruct::InitializeScriptStruct(void* InDest, int32 ArrayDim)
 	}
 }
 
-void UScriptStruct::ClearScriptStruct(void* Dest, int32 ArrayDim)
+void UScriptStruct::ClearScriptStruct(void* Dest, int32 ArrayDim) const
 {
 	uint8 *Data = (uint8*)Dest;
 	int32 Stride = GetStructureSize();
@@ -2123,7 +2123,7 @@ void UScriptStruct::ClearScriptStruct(void* Dest, int32 ArrayDim)
 
 }
 
-void UScriptStruct::DestroyScriptStruct(void* Dest, int32 ArrayDim)
+void UScriptStruct::DestroyScriptStruct(void* Dest, int32 ArrayDim) const
 {
 	if (StructFlags & (STRUCT_IsPlainOldData | STRUCT_NoDestructor))
 	{
diff --git a/Engine/Source/Runtime/CoreUObject/Public/UObject/Class.h b/Engine/Source/Runtime/CoreUObject/Public/UObject/Class.h
--- a/Engine/Source/Runtime/CoreUObject/Public/UObject/Class.h
+++ b/Engine/Source/Runtime/CoreUObject/Public/UObject/Class.h
@@ -950,14 +950,14 @@ public:
 	/** Look for the CppStructOps and hook it up **/
 	COREUOBJECT_API void PrepareCppStructOps();
 
-	FORCEINLINE ICppStructOps* GetCppStructOps()
+	FORCEINLINE ICppStructOps* GetCppStructOps() const
 	{
 		check(bPrepareCppStructOpsCompleted);
 		return CppStructOps;
 	}
 
 	/** return true if these cpp ops are not for me, but rather this is an incomplete cpp ops from my base class **/
-	FORCEINLINE bool InheritedCppStructOps()
+	FORCEINLINE bool InheritedCppStructOps() const
 	{
 		check(bPrepareCppStructOpsCompleted);
 		return bCppStructOpsFromBaseClass;
@@ -974,7 +974,7 @@ public:
 	 * If it is native, it is assumed to have defaults because it has a constructor
 	 * @return true if this struct has defaults
 	**/
-	FORCEINLINE bool HasDefaults()
+	FORCEINLINE bool HasDefaults() const
 	{
 		return !!GetCppStructOps();
 	}
@@ -984,7 +984,7 @@ public:
 	 *
 	 * @param	Ar	Archive the struct is going to be serialized with later on
 	 */
-	bool ShouldSerializeAtomically( FArchive& Ar )
+	bool ShouldSerializeAtomically(FArchive& Ar) const
 	{
 		if( (StructFlags&STRUCT_Atomic) != 0)
 		{
@@ -1004,7 +1004,7 @@ public:
 	 * @param	PortFlags	Comparison flags
 	 * @return true if the structs are identical
 	 */
-	bool CompareScriptStruct( const void* A, const void* B, uint32 PortFlags );
+	bool CompareScriptStruct(const void* A, const void* B, uint32 PortFlags) const;
 
 	/**
 	 * Copy a struct over an existing struct
@@ -1014,7 +1014,7 @@ public:
 	 * @param	ArrayDim	Number of elements in the array
 	 * @param	Stride		Stride of the array, If this default (0), then we will pull the size from the struct
 	 */
-	COREUOBJECT_API void CopyScriptStruct(void* Dest, void const* Src, int32 ArrayDim = 1);
+	COREUOBJECT_API void CopyScriptStruct(void* Dest, void const* Src, int32 ArrayDim = 1) const;
 	/**
 	 * Initialize a struct over uninitialized memory. This may be done by calling the native constructor or individually initializing properties
 	 *
@@ -1022,7 +1022,7 @@ public:
 	 * @param	ArrayDim	Number of elements in the array
 	 * @param	Stride		Stride of the array, If this default (0), then we will pull the size from the struct
 	 */
-	COREUOBJECT_API void InitializeScriptStruct(void* Dest, int32 ArrayDim = 1);
+	COREUOBJECT_API void InitializeScriptStruct(void* Dest, int32 ArrayDim = 1) const;
 	/**
 	 * Reintialize a struct in memory. This may be done by calling the native destructor and then the constructor or individually reinitializing properties
 	 *
@@ -1030,7 +1030,7 @@ public:
 	 * @param	ArrayDim	Number of elements in the array
 	 * @param	Stride		Stride of the array, only relevant if there more than one element. If this default (0), then we will pull the size from the struct
 	 */
-	void ClearScriptStruct(void* Dest, int32 ArrayDim = 1);
+	void ClearScriptStruct(void* Dest, int32 ArrayDim = 1) const;
 	/**
 	 * Destroy a struct in memory. This may be done by calling the native destructor and then the constructor or individually reinitializing properties
 	 *
@@ -1038,11 +1038,52 @@ public:
 	 * @param	ArrayDim	Number of elements in the array
 	 * @param	Stride		Stride of the array. If this default (0), then we will pull the size from the struct
 	 */
-	COREUOBJECT_API void DestroyScriptStruct(void* Dest, int32 ArrayDim = 1);
+	COREUOBJECT_API void DestroyScriptStruct(void* Dest, int32 ArrayDim = 1) const;
 
 	virtual COREUOBJECT_API void RecursivelyPreload();
 };
 
+struct FStructOnScope
+{
+private:
+	TWeakObjectPtr<const UScriptStruct> ScriptStruct;
+	uint8* SampleStructMemory;
+
+	FStructOnScope(const FStructOnScope&);
+	FStructOnScope& operator=(const FStructOnScope&);
+
+public:
+	FStructOnScope(const UScriptStruct* InScriptStruct)
+		: ScriptStruct(InScriptStruct)
+		, SampleStructMemory(NULL)
+	{
+		if (ScriptStruct.IsValid())
+		{
+			SampleStructMemory = (uint8*)FMemory::Malloc(ScriptStruct->GetStructureSize());
+			ScriptStruct.Get()->InitializeScriptStruct(SampleStructMemory);
+		}
+	}
+
+	uint8* GetStructMemory() { return SampleStructMemory; }
+	const uint8* GetStructMemory() const { return SampleStructMemory; }
+	const UScriptStruct* GetStruct() { return ScriptStruct.Get(); }
+	bool IsValid() { return ScriptStruct.IsValid() && SampleStructMemory; }
+
+	~FStructOnScope()
+	{
+		if (ScriptStruct.IsValid() && SampleStructMemory)
+		{
+			ScriptStruct.Get()->DestroyScriptStruct(SampleStructMemory);
+			ScriptStruct = NULL;
+		}
+
+		if (SampleStructMemory)
+		{
+			FMemory::Free(SampleStructMemory);
+			SampleStructMemory = NULL;
+		}
+	}
+};
 
 /*-----------------------------------------------------------------------------
 	UFunction.
diff --git a/Engine/Source/Runtime/Engine/Private/UserDefinedStruct.cpp b/Engine/Source/Runtime/Engine/Private/UserDefinedStruct.cpp
--- a/Engine/Source/Runtime/Engine/Private/UserDefinedStruct.cpp
+++ b/Engine/Source/Runtime/Engine/Private/UserDefinedStruct.cpp
@@ -70,7 +70,7 @@ void UUserDefinedStruct::SerializeTaggedProperties(FArchive& Ar, uint8* Data, US
 		&& Ar.IsPersistent();
 
 	const bool bPrepareDefaultStruct = bUseNewDefaults || bLoadDefaultFirst;
-	FStructureEditorUtils::FStructOnScope StructDefaultMem(bPrepareDefaultStruct ? UDDefaultsStruct : NULL);
+	FStructOnScope StructDefaultMem(bPrepareDefaultStruct ? UDDefaultsStruct : NULL);
 	if (bPrepareDefaultStruct)
 	{
 		FStructureEditorUtils::Fill_MakeStructureDefaultValue(UDDefaultsStruct, StructDefaultMem.GetStructMemory());
