Duplicated objects now get PostLoad called on them properly by setting the RF_NeedsPostLoad flag and calling ConditionalPostLoad

diff --git a/Engine/Source/Runtime/CoreUObject/Private/Serialization/DuplicateDataWriter.cpp b/Engine/Source/Runtime/CoreUObject/Private/Serialization/DuplicateDataWriter.cpp
--- a/Engine/Source/Runtime/CoreUObject/Private/Serialization/DuplicateDataWriter.cpp
+++ b/Engine/Source/Runtime/CoreUObject/Private/Serialization/DuplicateDataWriter.cpp
@@ -67,6 +67,12 @@ FArchive& FDuplicateDataWriter::operator<<(FAssetPtr& AssetPtr)
  */
 UObject* FDuplicateDataWriter::AddDuplicate(UObject* SourceObject,UObject* DupObject)
 {
+	if ( DupObject && !DupObject->IsTemplate() )
+	{
+		// Make sure the duplicated object is prepared to postload
+		DupObject->SetFlags(RF_NeedPostLoad|RF_NeedPostLoadSubobjects);
+	}
+
 	// Check for an existing duplicate of the object; if found, use that one instead of creating a new one.
 	FDuplicatedObject Info = DuplicatedObjectAnnotation.GetAnnotation( SourceObject );
 	if ( Info.IsDefault() )
diff --git a/Engine/Source/Runtime/CoreUObject/Private/UObject/Obj.cpp b/Engine/Source/Runtime/CoreUObject/Private/UObject/Obj.cpp
--- a/Engine/Source/Runtime/CoreUObject/Private/UObject/Obj.cpp
+++ b/Engine/Source/Runtime/CoreUObject/Private/UObject/Obj.cpp
@@ -565,36 +565,28 @@ void UObject::ConditionalPostLoad()
 {
 	if( HasAnyFlags(RF_NeedPostLoad) )
 	{
-		if( GetLinker() )
-		{
 #if !(UE_BUILD_SHIPPING || UE_BUILD_TEST)
-			checkSlow(!DebugPostLoad.Contains(this));
-			DebugPostLoad.Add(this);
+		checkSlow(!DebugPostLoad.Contains(this));
+		DebugPostLoad.Add(this);
 #endif
-			ClearFlags( RF_NeedPostLoad );
+		ClearFlags( RF_NeedPostLoad );
 
-			UObject* ObjectArchetype = GetArchetype();
-			if ( ObjectArchetype != NULL )
-			{
-				//make sure our archetype executes ConditionalPostLoad first.
-				ObjectArchetype->ConditionalPostLoad();
-			}
+		UObject* ObjectArchetype = GetArchetype();
+		if ( ObjectArchetype != NULL )
+		{
+			//make sure our archetype executes ConditionalPostLoad first.
+			ObjectArchetype->ConditionalPostLoad();
+		}
 
-			ConditionalPostLoadSubobjects();
-			PostLoad();
+		ConditionalPostLoadSubobjects();
+		PostLoad();
 
 #if !(UE_BUILD_SHIPPING || UE_BUILD_TEST)
-			if( DebugPostLoad.Contains(this) )
-			{
-				UE_LOG(LogObj, Fatal, TEXT("%s failed to route PostLoad.  Please call Super::PostLoad() in your <className>::PostLoad() function. "), *GetFullName() );
-			}
-#endif
-		}
-		else
+		if( DebugPostLoad.Contains(this) )
 		{
-			UE_LOG(LogObj, Log, TEXT("%s did not route PostLoad, because it didn't have a linker"), *GetFullName());
-			ClearFlags( RF_NeedPostLoad|RF_NeedPostLoadSubobjects );
+			UE_LOG(LogObj, Fatal, TEXT("%s failed to route PostLoad.  Please call Super::PostLoad() in your <className>::PostLoad() function. "), *GetFullName() );
 		}
+#endif
 	}
 }
 
diff --git a/Engine/Source/Runtime/CoreUObject/Private/UObject/UObjectGlobals.cpp b/Engine/Source/Runtime/CoreUObject/Private/UObject/UObjectGlobals.cpp
--- a/Engine/Source/Runtime/CoreUObject/Private/UObject/UObjectGlobals.cpp
+++ b/Engine/Source/Runtime/CoreUObject/Private/UObject/UObjectGlobals.cpp
@@ -1508,8 +1508,7 @@ UObject* StaticDuplicateObjectEx( FObjectDuplicationParameters& Parameters )
 			if ( !DupObjectInfo.DuplicatedObject->IsTemplate() )
 			{
 				// Don't want to call PostLoad on class duplicated CDOs
-				DupObjectInfo.DuplicatedObject->PostLoad();
-				DupObjectInfo.DuplicatedObject->PostLoadSubobjects(NULL);
+				DupObjectInfo.DuplicatedObject->ConditionalPostLoad();
 			}
 			DupObjectInfo.DuplicatedObject->CheckDefaultSubobjects();
 		}
