FStructOnScope improvements

diff --git a/Engine/Source/Runtime/CoreUObject/Public/UObject/Class.h b/Engine/Source/Runtime/CoreUObject/Public/UObject/Class.h
--- a/Engine/Source/Runtime/CoreUObject/Public/UObject/Class.h
+++ b/Engine/Source/Runtime/CoreUObject/Public/UObject/Class.h
@@ -1093,11 +1093,14 @@ public:
 	}
 
 	uint8* GetStructMemory() { return SampleStructMemory; }
+
 	const uint8* GetStructMemory() const { return SampleStructMemory; }
-	const UScriptStruct* GetStruct() { return ScriptStruct.Get(); }
-	bool IsValid() { return ScriptStruct.IsValid() && SampleStructMemory; }
 
-	~FStructOnScope()
+	const UScriptStruct* GetStruct() const { return ScriptStruct.Get(); }
+
+	bool IsValid() const { return ScriptStruct.IsValid() && SampleStructMemory; }
+
+	void Destroy()
 	{
 		if (ScriptStruct.IsValid() && SampleStructMemory)
 		{
@@ -1111,6 +1114,11 @@ public:
 			SampleStructMemory = NULL;
 		}
 	}
+
+	~FStructOnScope()
+	{
+		Destroy();
+	}
 };
 
 /*-----------------------------------------------------------------------------
