TSubobjectPtr documentation

#docs More detailed TSubobjectPtr comments

diff --git a/Engine/Source/Runtime/CoreUObject/Public/UObject/UObjectGlobals.h b/Engine/Source/Runtime/CoreUObject/Public/UObject/UObjectGlobals.h
--- a/Engine/Source/Runtime/CoreUObject/Public/UObject/UObjectGlobals.h
+++ b/Engine/Source/Runtime/CoreUObject/Public/UObject/UObjectGlobals.h
@@ -575,9 +575,29 @@ class TSubobjectPtrConstructor
 };
 
 /**
- * TSubobjectPtr - Subobject pointer.
- * Prevents anything C++ from overwriting the subobject pointer.
- * It can only be assigned to with PCIP.CreateDefaultSubobject (via TSubobjectPtrConstructor).
+ * TSubobjectPtr - Sub-object smart pointer, owns a reference to instanced objects (sub-objects / components).
+ * Prevents anything in C++ from overwriting the sub-object pointer.
+ * Can (and should) be declared as a UPROPERTY():
+ *
+ *   UPROPERTY()
+ *   TSubobjectPtr<UActorComponent> MyComponent;
+ *
+ * It can only be assigned to with PCIP.CreateDefaultSubobject (via TSubobjectPtrConstructor) in the owning object's constructor.
+ *
+ *   MyComponent = PCIP.CreateDefaultSubobject<UPathFollowingComponent>(this, TEXT("PathFollowingComponent"));
+ *
+ * Initialized with InvalidPtrValue by default because it always needs to be initialized (either with NULL or a pointer to sub-object in the constructor
+ * of the owner object).
+ * Usually used with Actor components to specify that the actor-derived class is the owner of the component and prevent other derived classes from
+ * overwriting it in any other way than through PCIP object.
+ *
+ * Implements a structure dereference operator for convenience
+ *
+ *   MyComponent->AttachTo(Owner);
+ *
+ * Can be reset to NULL, although this functionality is mostly for internal use:
+ *
+ *   MyComponent.Reset();
  */
 template <class SubobjectType>
 class TSubobjectPtr : public FSubobjectPtr
@@ -609,7 +629,7 @@ public:
 	{
 		checkAtCompileTime((CanConvertPointerFromTo<DerivedSubobjectType, SubobjectType>::Result), Subobject_Pointers_Must_Be_Compatible);
 	}
-	/** Gets the subobject pointer. */
+	/** Gets the sub-object pointer. */
 	FORCEINLINE SubobjectType* Get() const
 	{
 		return (SubobjectType*)Object;
@@ -622,12 +642,12 @@ public:
 		Set(Other.Object);
 		return *this;
 	}
-	/** Gets the subobject pointer. */
+	/** Gets the sub-object pointer. */
 	FORCEINLINE SubobjectType* operator->() const
 	{
 		return (SubobjectType*)Object;
 	}
-	/** Gets the subobject pointer. */
+	/** Gets the sub-object pointer. */
 	FORCEINLINE operator SubobjectType*() const
 	{
 		return (SubobjectType*)Object;
