BPGC bytecode serialization is skipped in Editor.

diff --git a/Engine/Config/BaseEditor.ini b/Engine/Config/BaseEditor.ini
--- a/Engine/Config/BaseEditor.ini
+++ b/Engine/Config/BaseEditor.ini
@@ -269,3 +269,6 @@ bAllowClassAndBlueprintPinMatching=false
 bReplaceBlueprintWithClass=false
 bDontLoadBlueprintOutsideEditor=false
 bBlueprintIsNotBlueprintType=false
+
+[StructSerialization]
+SkipByteCodeSerialization=true
\ No newline at end of file
diff --git a/Engine/Source/Runtime/CoreUObject/Private/UObject/Class.cpp b/Engine/Source/Runtime/CoreUObject/Private/UObject/Class.cpp
--- a/Engine/Source/Runtime/CoreUObject/Private/UObject/Class.cpp
+++ b/Engine/Source/Runtime/CoreUObject/Private/UObject/Class.cpp
@@ -1207,7 +1207,20 @@ void UStruct::Serialize( FArchive& Ar )
 				Ar.Serialize(TempScript.GetData(), ScriptStorageSize);
 				const int32 ScriptEnd = Ar.Tell();
 
-				if ((Ar.UE4Ver() < VER_MIN_SCRIPTVM_UE4) || (Ar.LicenseeUE4Ver() < VER_MIN_SCRIPTVM_LICENSEEUE4))
+				bool bSkipByteCodeSerialization = false;
+#if WITH_EDITOR
+				struct FSkipByteCodeSerializationHelper
+				{
+					bool bValue;
+					FSkipByteCodeSerializationHelper() : bValue(false)
+					{
+						GConfig->GetBool(TEXT("StructSerialization"), TEXT("SkipByteCodeSerialization"), bValue, GEditorIni);
+					}
+				};
+				static FSkipByteCodeSerializationHelper SkipByteCodeHelper;
+				bSkipByteCodeSerialization = SkipByteCodeHelper.bValue;
+#endif // WITH_EDITOR
+				if (bSkipByteCodeSerialization || (Ar.UE4Ver() < VER_MIN_SCRIPTVM_UE4) || (Ar.LicenseeUE4Ver() < VER_MIN_SCRIPTVM_LICENSEEUE4))
 				{
 					// Discard the bytecode as it's too old and might cause serialization errors
 					ScriptStorageSize = 0;
