Package map cleanup

* We only need one FClassNetCache, so moved this into a single manager on UNetDriver
* Moving towards making UPackageMap abstract

diff --git a/Engine/Source/Runtime/CoreUObject/Private/UObject/CoreNet.cpp b/Engine/Source/Runtime/CoreUObject/Private/UObject/CoreNet.cpp
--- a/Engine/Source/Runtime/CoreUObject/Private/UObject/CoreNet.cpp
+++ b/Engine/Source/Runtime/CoreUObject/Private/UObject/CoreNet.cpp
@@ -58,6 +58,52 @@ FClassNetCache::FClassNetCache( UClass* InClass )
 : Class( InClass )
 {}
 
+FClassNetCache * FClassNetCacheMgr::GetClassNetCache( UClass* Class )
+{
+	FClassNetCache * Result = ClassFieldIndices.FindRef( Class );
+
+	if ( !Result )
+	{
+		Result					= ClassFieldIndices.Add( Class, new FClassNetCache( Class ) );
+		Result->Super			= NULL;
+		Result->FieldsBase		= 0;
+
+		if ( Class->GetSuperClass() )
+		{
+			Result->Super		= GetClassNetCache(Class->GetSuperClass());
+			Result->FieldsBase	= Result->Super->GetMaxIndex();
+		}
+
+		Result->Fields.Empty( Class->NetFields.Num() );
+
+		for( int32 i = 0; i < Class->NetFields.Num(); i++ )
+		{
+			// Add sandboxed items to net cache.  
+			UField * Field = Class->NetFields[i];
+			int32 ThisIndex	= Result->GetMaxIndex();
+			new ( Result->Fields )FFieldNetCache( Field, ThisIndex );
+		}
+
+		Result->Fields.Shrink();
+
+		for ( TArray< FFieldNetCache >::TIterator It( Result->Fields ); It; ++It )
+		{
+			Result->FieldMap.Add( It->Field, &*It );
+		}
+	}
+	return Result;
+}
+
+void FClassNetCacheMgr::ClearClassNetCache()
+{
+	for ( auto It = ClassFieldIndices.CreateIterator(); It; ++It)
+	{
+		delete It.Value();
+	}
+
+	ClassFieldIndices.Empty();
+}
+
 /*-----------------------------------------------------------------------------
 	UPackageMap implementation.
 -----------------------------------------------------------------------------*/
@@ -132,47 +178,6 @@ UObject * UPackageMap::GetObjectFromNetGUID( const FNetworkGUID & NetGUID )
 	return NULL;
 }
 
-FClassNetCache* UPackageMap::GetClassNetCache( UClass* Class )
-{
-	FClassNetCache* Result = ClassFieldIndices.FindRef(Class);
-	if( !Result )
-	{
-		Result                       = ClassFieldIndices.Add( Class, new FClassNetCache(Class) );
-		Result->Super                = NULL;
-		Result->FieldsBase           = 0;
-		if( Class->GetSuperClass() )
-		{
-			Result->Super		         = GetClassNetCache(Class->GetSuperClass());
-			Result->FieldsBase           = Result->Super->GetMaxIndex();
-		}
-
-		Result->Fields.Empty( Class->NetFields.Num() );
-		for( int32 i=0; i<Class->NetFields.Num(); i++ )
-		{
-			// Add sandboxed items to net cache.  
-			UField* Field = Class->NetFields[i];
-			int32 ThisIndex	= Result->GetMaxIndex();
-			new(Result->Fields)FFieldNetCache( Field, ThisIndex );
-		}
-
-		Result->Fields.Shrink();
-		for( TArray<FFieldNetCache>::TIterator It(Result->Fields); It; ++It )
-		{
-			Result->FieldMap.Add( It->Field, &*It );
-		}
-	}
-	return Result;
-}
-
-void UPackageMap::ClearClassNetCache()
-{
-	for( auto It = ClassFieldIndices.CreateIterator(); It; ++It)
-	{
-		delete It.Value();
-	}
-	ClassFieldIndices.Empty();
-}
-
 void UPackageMap::Serialize( FArchive& Ar )
 {
 	Super::Serialize( Ar );
@@ -186,7 +191,6 @@ void UPackageMap::AddReferencedObjects(UObject* InThis, FReferenceCollector& Col
 
 void UPackageMap::FinishDestroy()
 {
-	ClearClassNetCache();
 	Super::FinishDestroy();
 }
 
diff --git a/Engine/Source/Runtime/CoreUObject/Public/UObject/CoreNet.h b/Engine/Source/Runtime/CoreUObject/Public/UObject/CoreNet.h
--- a/Engine/Source/Runtime/CoreUObject/Public/UObject/CoreNet.h
+++ b/Engine/Source/Runtime/CoreUObject/Public/UObject/CoreNet.h
@@ -13,7 +13,7 @@ DECLARE_DELEGATE_RetVal_OneParam( bool, FNetObjectIsDynamic, const UObject*);
 //
 // Information about a field.
 //
-class FFieldNetCache
+class COREUOBJECT_API FFieldNetCache
 {
 public:
 	UField* Field;
@@ -28,9 +28,9 @@ public:
 //
 // Information about a class, cached for network coordination.
 //
-class FClassNetCache
+class COREUOBJECT_API FClassNetCache
 {
-	friend class UPackageMap;
+	friend class FClassNetCacheMgr;
 public:
 	FClassNetCache();
 	FClassNetCache( UClass* Class );
@@ -62,11 +62,22 @@ public:
 		return NULL;
 	}
 private:
-	int32 FieldsBase;
-	FClassNetCache* Super;
-	UClass* Class;
-	TArray<FFieldNetCache> Fields;
-	TMap<UObject*,FFieldNetCache*> FieldMap;
+	int32								FieldsBase;
+	FClassNetCache *					Super;
+	TWeakObjectPtr< UClass >			Class;
+	TArray< FFieldNetCache >			Fields;
+	TMap< UObject *, FFieldNetCache * > FieldMap;
+};
+
+class COREUOBJECT_API FClassNetCacheMgr
+{
+public:
+	/** get the cached field to index mappings for the given class */
+	FClassNetCache *	GetClassNetCache( UClass * Class );
+	void				ClearClassNetCache();
+
+private:
+	TMap< TWeakObjectPtr< UClass >, FClassNetCache * > ClassFieldIndices;
 };
 
 //
@@ -121,10 +132,6 @@ class COREUOBJECT_API UPackageMap : public UObject
 	// @todo document
 	virtual bool SerializeName( FArchive& Ar, FName& Name );
 
-	/** get the cached field to index mappings for the given class */
-	FClassNetCache *	GetClassNetCache(UClass* Class);
-	void				ClearClassNetCache();
-
 	virtual void ResetPackageMap();
 
 	virtual void SetLocked(bool L) { }
@@ -164,11 +171,6 @@ class COREUOBJECT_API UPackageMap : public UObject
 
 protected:
 
-	virtual bool IsNetGUIDAuthority() { return true; }
-
-	// @todo document
-	TMap<TWeakObjectPtr<UClass>, FClassNetCache*> ClassFieldIndices;
-
 	bool			bSuppressLogs;
 	bool			bShouldSerializeUnAckedObjects;
 	bool			bSerializedUnAckedObject;
diff --git a/Engine/Source/Runtime/Engine/Classes/Engine/NetDriver.h b/Engine/Source/Runtime/Engine/Classes/Engine/NetDriver.h
--- a/Engine/Source/Runtime/Engine/Classes/Engine/NetDriver.h
+++ b/Engine/Source/Runtime/Engine/Classes/Engine/NetDriver.h
@@ -166,7 +166,9 @@ public:
 	class UWorld* World;
 
 	/** @todo document */
-	TSharedPtr< class FNetGUIDCache > GuidCache;
+	TSharedPtr< class FNetGUIDCache >		GuidCache;
+
+	TSharedPtr< class FClassNetCacheMgr >	NetCache;
 
 	/** The loaded UClass of the net connection type to use */
 	UPROPERTY()
diff --git a/Engine/Source/Runtime/Engine/Classes/Engine/PackageMapClient.h b/Engine/Source/Runtime/Engine/Classes/Engine/PackageMapClient.h
--- a/Engine/Source/Runtime/Engine/Classes/Engine/PackageMapClient.h
+++ b/Engine/Source/Runtime/Engine/Classes/Engine/PackageMapClient.h
@@ -143,7 +143,7 @@ protected:
 
 	bool	ShouldSendFullPath(const UObject* Object, const FNetworkGUID &NetGUID);
 	
-	virtual bool IsNetGUIDAuthority() OVERRIDE;
+	bool	IsNetGUIDAuthority();
 
 	bool Locked;
 	bool IsExportingNetGUIDBunch;
diff --git a/Engine/Source/Runtime/Engine/Private/DataChannel.cpp b/Engine/Source/Runtime/Engine/Private/DataChannel.cpp
--- a/Engine/Source/Runtime/Engine/Private/DataChannel.cpp
+++ b/Engine/Source/Runtime/Engine/Private/DataChannel.cpp
@@ -1344,7 +1344,7 @@ void UActorChannel::SetChannelActor( AActor* InActor )
 	// Set stuff.
 	Actor                      = InActor;
 	ActorClass                 = Actor->GetClass();
-	FClassNetCache* ClassCache = Connection->PackageMap->GetClassNetCache( ActorClass );
+	FClassNetCache* ClassCache = Connection->Driver->NetCache->GetClassNetCache( ActorClass );
 
 	UE_LOG(LogNetTraffic, VeryVerbose, TEXT("SetChannelActor[%d]: Actor: %s ActorClass: %s. %s"), ChIndex, Actor ? *Actor->GetPathName() : TEXT("NULL"), ActorClass ? *ActorClass->GetName() : TEXT("NULL"), *GetName() );
 
@@ -1872,7 +1872,7 @@ void UActorChannel::EndContentBlock( UObject *Obj, FOutBunch &Bunch, FClassNetCa
 	check(Obj);
 	if (!ClassCache)
 	{
-		ClassCache = Connection->PackageMap->GetClassNetCache( Obj->GetClass() );
+		ClassCache = Connection->Driver->NetCache->GetClassNetCache( Obj->GetClass() );
 	}
 
 	// Write max int to signify done
diff --git a/Engine/Source/Runtime/Engine/Private/DataReplication.cpp b/Engine/Source/Runtime/Engine/Private/DataReplication.cpp
--- a/Engine/Source/Runtime/Engine/Private/DataReplication.cpp
+++ b/Engine/Source/Runtime/Engine/Private/DataReplication.cpp
@@ -361,7 +361,7 @@ bool FObjectReplicator::ReceivedBunch( FInBunch &Bunch, const FReplicationFlags
 
 	const bool bIsServer = ( OwningChannel->Connection->Driver->ServerConnection == NULL );
 
-	FClassNetCache * ClassCache = PackageMap->GetClassNetCache( ObjectClass );
+	FClassNetCache * ClassCache = OwningChannel->Connection->Driver->NetCache->GetClassNetCache( ObjectClass );
 
 	if ( ClassCache == NULL )
 	{
@@ -977,7 +977,7 @@ bool FObjectReplicator::ReplicateProperties( FOutBunch & Bunch, FReplicationFlag
 
 	if ( WroteImportantData )
 	{
-		OwningChannel->EndContentBlock( Object, Bunch, OwningChannelConnection->PackageMap->GetClassNetCache( ObjectClass ) );
+		OwningChannel->EndContentBlock( Object, Bunch, OwningChannelConnection->Driver->NetCache->GetClassNetCache( ObjectClass ) );
 	}
 
 	return WroteImportantData;
diff --git a/Engine/Source/Runtime/Engine/Private/NetConnection.cpp b/Engine/Source/Runtime/Engine/Private/NetConnection.cpp
--- a/Engine/Source/Runtime/Engine/Private/NetConnection.cpp
+++ b/Engine/Source/Runtime/Engine/Private/NetConnection.cpp
@@ -1654,11 +1654,6 @@ bool UNetConnection::ActorIsAvailableOnClient(const AActor* ThisActor)
 void UNetConnection::ResetGameWorldState()
 {
 	//Clear out references and do whatever else so that nothing holds onto references that it doesn't need to.
-	if (PackageMap)
-	{
-		PackageMap->ClearClassNetCache();	// Clear the cache net: it will recreate itself after seamless travel
-	}
-
 	DestroyedStartupOrDormantActors.Empty();
 	RecentlyDormantActors.Empty();
 	DormantActors.Empty();
diff --git a/Engine/Source/Runtime/Engine/Private/NetworkDriver.cpp b/Engine/Source/Runtime/Engine/Private/NetworkDriver.cpp
--- a/Engine/Source/Runtime/Engine/Private/NetworkDriver.cpp
+++ b/Engine/Source/Runtime/Engine/Private/NetworkDriver.cpp
@@ -103,6 +103,7 @@ void UNetDriver::PostInitProperties()
 		RemoteRoleProperty	= FindObjectChecked<UProperty>( AActor::StaticClass(), TEXT("RemoteRole") );
 
 		GuidCache			= TSharedPtr< FNetGUIDCache >( new FNetGUIDCache( this ) );
+		NetCache			= TSharedPtr< FClassNetCacheMgr >( new FClassNetCacheMgr() );
 
 		ProfileStats		= FParse::Param(FCommandLine::Get(),TEXT("profilestats"));
 	}
@@ -616,7 +617,7 @@ void UNetDriver::InternalProcessRemoteFunction
 	UObject * TargetObj = SubObject ? SubObject : Actor;
 
 	// Make sure this function exists for both parties.
-	FClassNetCache* ClassCache = Connection->PackageMap->GetClassNetCache( TargetObj->GetClass() );
+	FClassNetCache* ClassCache = NetCache->GetClassNetCache( TargetObj->GetClass() );
 	if (!ClassCache)
 	{
 		DEBUG_REMOTEFUNCTION(TEXT("ClassNetCache empty, not calling %s::%s"), *Actor->GetName(), *Function->GetName());
@@ -1145,8 +1146,6 @@ bool UNetDriver::HandleNetDisconnectCommand( const TCHAR* Cmd, FOutputDevice& Ar
 bool UNetDriver::HandleNetDumpServerRPCCommand( const TCHAR* Cmd, FOutputDevice& Ar )
 {
 #if WITH_SERVER_CODE
-	UPackageMap * PackageMap = new( this )UPackageMap( FPostConstructInitializeProperties() );
-
 	for ( TObjectIterator<UClass> ClassIt; ClassIt; ++ClassIt )
 	{
 		bool bHasNetFields = false;
@@ -1175,7 +1174,7 @@ bool UNetDriver::HandleNetDumpServerRPCCommand( const TCHAR* Cmd, FOutputDevice&
 
 			if ( Function != NULL && Function->FunctionFlags & FUNC_NetServer )
 			{
-				FClassNetCache * ClassCache = PackageMap->GetClassNetCache( *ClassIt );
+				FClassNetCache * ClassCache = NetCache->GetClassNetCache( *ClassIt );
 
 				FFieldNetCache * FieldCache = ClassCache->GetFromField( Function );
 
@@ -2889,6 +2888,12 @@ void UNetDriver::SetWorld(class UWorld* InWorld)
 void UNetDriver::ResetGameWorldState()
 {
 	DestroyedStartupOrDormantActors.Empty();
+
+	if ( NetCache.IsValid() )
+	{
+		NetCache->ClearClassNetCache();	// Clear the cache net: it will recreate itself after seamless travel
+	}
+
 	if (ServerConnection)
 	{
 		ServerConnection->ResetGameWorldState();
diff --git a/Engine/Source/Runtime/Engine/Private/RepLayout.cpp b/Engine/Source/Runtime/Engine/Private/RepLayout.cpp
--- a/Engine/Source/Runtime/Engine/Private/RepLayout.cpp
+++ b/Engine/Source/Runtime/Engine/Private/RepLayout.cpp
@@ -947,7 +947,7 @@ void FRepLayout::WritePropertyHeader(
 	UNetConnection * Connection = OwningChannel->Connection;
 
 	// Get class network info cache.
-	FClassNetCache * ClassCache = Connection->PackageMap->GetClassNetCache( ObjectClass );
+	FClassNetCache * ClassCache = Connection->Driver->NetCache->GetClassNetCache( ObjectClass );
 	check( ClassCache );
 
 	if ( !bContentBlockWritten )
