More package map cleanup

diff --git a/Engine/Source/Runtime/CoreUObject/Private/UObject/CoreNet.cpp b/Engine/Source/Runtime/CoreUObject/Private/UObject/CoreNet.cpp
--- a/Engine/Source/Runtime/CoreUObject/Private/UObject/CoreNet.cpp
+++ b/Engine/Source/Runtime/CoreUObject/Private/UObject/CoreNet.cpp
@@ -153,59 +153,6 @@ bool UPackageMap::SerializeName(FArchive& Ar, FName& Name)
 	return true;
 }
 
-bool UPackageMap::SerializeObject( FArchive& Ar, UClass* Class, UObject*& Object, FNetworkGUID *OutNetGUID )
-{
-	// For now this is unsupported, but we could implement this serialization in a way that is compatible with UPackageMapClient,
-	// it would just always serialize objects as <NetGUID, FullPath>.
-	UE_LOG(LogCoreNet, Fatal,TEXT("Unexpected UPackageMap::SerializeObject"));
-	return true;
-}
-
-bool UPackageMap::WriteObject( FArchive& Ar, UObject* Outer, FNetworkGUID NetGUID, FString ObjName )
-{
-	UE_LOG(LogCoreNet, Fatal,TEXT("Unexpected UPackageMap::WriteObject"));
-	return true;
-}
-
-void UPackageMap::LogDebugInfo( FOutputDevice & Ar)
-{
-	UE_LOG(LogCoreNet, Fatal,TEXT("Unexpected UPackageMap::LogDebugInfo"));
-}
-
-UObject * UPackageMap::GetObjectFromNetGUID( const FNetworkGUID & NetGUID )
-{
-	UE_LOG(LogCoreNet, Fatal,TEXT("Unexpected UPackageMap::GetObjectFromNetGUID"));
-	return NULL;
-}
-
-void UPackageMap::Serialize( FArchive& Ar )
-{
-	Super::Serialize( Ar );
-}
-
-void UPackageMap::AddReferencedObjects(UObject* InThis, FReferenceCollector& Collector)
-{
-	UPackageMap * This = CastChecked<UPackageMap>( InThis );
-	return Super::AddReferencedObjects( InThis, Collector );
-}
-
-void UPackageMap::FinishDestroy()
-{
-	Super::FinishDestroy();
-}
-
-void UPackageMap::ResetPackageMap()
-{
-	UE_LOG(LogCoreNet, Fatal,TEXT("Unexpected UPackageMap::ResetPackageMap"));
-}
-
-void UPackageMap::PostInitProperties()
-{
-	Super::PostInitProperties();
-	bShouldSerializeUnAckedObjects = true;
-	bSerializedUnAckedObject = false;
-}
-
 IMPLEMENT_CORE_INTRINSIC_CLASS(UPackageMap, UObject,
 	{
 	}
diff --git a/Engine/Source/Runtime/CoreUObject/Public/UObject/CoreNet.h b/Engine/Source/Runtime/CoreUObject/Public/UObject/CoreNet.h
--- a/Engine/Source/Runtime/CoreUObject/Public/UObject/CoreNet.h
+++ b/Engine/Source/Runtime/CoreUObject/Public/UObject/CoreNet.h
@@ -105,81 +105,57 @@ public:
 	friend COREUOBJECT_API FArchive& operator<<( FArchive& Ar, FPackageInfo& I );
 };
 
-/** maximum index that may ever be used to serialize an object over the network
- * this affects the number of bits used to send object references
- */
-#define MAX_OBJECT_INDEX (uint32(1) << 31)
-
 //
 // Maps objects and names to and from indices for network communication.
 //
 class COREUOBJECT_API UPackageMap : public UObject
 {
-	DECLARE_CLASS_INTRINSIC(UPackageMap,UObject,CLASS_Transient|0,CoreUObject);
-
-	// Begin UObject interface.
-	virtual void	PostInitProperties();
-	virtual void	Serialize( FArchive& Ar ) OVERRIDE;
-	virtual void	FinishDestroy() OVERRIDE;
-	static void		AddReferencedObjects(UObject* InThis, FReferenceCollector& Collector);
-	// End UObject interface.
-	
-	virtual bool WriteObject( FArchive& Ar, UObject* Outer, FNetworkGUID NetGUID, FString ObjName );
+	DECLARE_CLASS_INTRINSIC( UPackageMap, UObject, CLASS_Transient | CLASS_Abstract | 0, CoreUObject );
 
-	// @todo document
-	virtual bool SerializeObject( FArchive& Ar, UClass* Class, UObject*& Obj, FNetworkGUID *OutNetGUID = NULL );
+	virtual bool		WriteObject( FArchive & Ar, UObject * Outer, FNetworkGUID NetGUID, FString ObjName ) { return false; }
 
 	// @todo document
-	virtual bool SerializeName( FArchive& Ar, FName& Name );
+	virtual bool		SerializeObject( FArchive& Ar, UClass* Class, UObject*& Obj, FNetworkGUID *OutNetGUID = NULL ) { return false; }
 
-	virtual void ResetPackageMap();
-
-	virtual void SetLocked(bool L) { }
+	// @todo document
+	virtual bool		SerializeName( FArchive& Ar, FName& Name );
 
-	// Variables.
+	virtual void		ResetPackageMap() { }
 
-	virtual UObject * ResolvePathAndAssignNetGUID( FNetworkGUID & InOutNetGUID, const FString & PathName, const FString & FilenameOverride, UObject * ObjOuter, const bool bNoLoad = false ) { return NULL; }
+	virtual void		SetLocked( bool L ) { }
 
-	void ResetUnAckedObject() { bSerializedUnAckedObject = false; }
-	bool SerializedUnAckedObject() { return bSerializedUnAckedObject; }
+	virtual UObject *	ResolvePathAndAssignNetGUID( FNetworkGUID & InOutNetGUID, const FString & PathName, const FString & FilenameOverride, UObject * ObjOuter, const bool bNoLoad = false ) { return NULL; }
 
-	void ResetHasSerializedCDO() { bSerializedCDO = false; }
-	bool HasSerializedCDO() const { return bSerializedCDO; }
-	
-	virtual bool SerializeNewActor(FArchive& Ar, class UActorChannel *Channel, class AActor*& Actor) { return false; }
+	virtual bool		SerializeNewActor(FArchive & Ar, class UActorChannel * Channel, class AActor *& Actor) { return false; }
 
-	virtual bool NetGUIDHasBeenAckd(FNetworkGUID NetGUID) { return false; }
-	virtual void ReceivedNak( const int32 NakPacketId ) { }
-	virtual void ReceivedAck( const int32 AckPacketId ) { }
-	virtual void NotifyBunchCommit( const int32 OutPacketId, const TArray< FNetworkGUID > & ExportNetGUIDs ) { }
+	virtual void		ReceivedNak( const int32 NakPacketId ) { }
+	virtual void		ReceivedAck( const int32 AckPacketId ) { }
+	virtual void		NotifyBunchCommit( const int32 OutPacketId, const TArray< FNetworkGUID > & ExportNetGUIDs ) { }
 
-	virtual void GetNetGUIDStats(int32 &AckCount, int32 &UnAckCount, int32 &PendingCount) { }
+	virtual void		GetNetGUIDStats(int32 & AckCount, int32 & UnAckCount, int32 & PendingCount) { }
 
-	virtual void NotifyStreamingLevelUnload(UObject* UnloadedLevel) { }
+	virtual void		NotifyStreamingLevelUnload( UObject * UnloadedLevel ) { }
 
-	virtual bool PrintExportBatch() { return false; }
+	virtual bool		PrintExportBatch() { return false; }
 
-	void	SetDebugContextString(FString Str) { DebugContextString = Str; }
-	void	ClearDebugContextString() { DebugContextString.Empty(); }
+	void				SetDebugContextString( const FString & Str ) { DebugContextString = Str; }
+	void				ClearDebugContextString() { DebugContextString.Empty(); }
 
-	void			ResetLoadedUnmappedObject() { bLoadedUnmappedObject = false; }
-	bool			GetLoadedUnmappedObject() const { return bLoadedUnmappedObject; }
-	FNetworkGUID	GetLastUnmappedNetGUID() const { return LastUnmappedNetGUID; }
+	void				ResetLoadedUnmappedObject() { bLoadedUnmappedObject = false; }
+	bool				GetLoadedUnmappedObject() const { return bLoadedUnmappedObject; }
+	FNetworkGUID		GetLastUnmappedNetGUID() const { return LastUnmappedNetGUID; }
 
-	virtual void		LogDebugInfo( FOutputDevice & Ar);
-	virtual UObject *	GetObjectFromNetGUID( const FNetworkGUID & NetGUID );
+	virtual void		LogDebugInfo( FOutputDevice & Ar) { }
+	virtual UObject *	GetObjectFromNetGUID( const FNetworkGUID & NetGUID ) { return NULL; }
 
 protected:
 
-	bool			bSuppressLogs;
-	bool			bShouldSerializeUnAckedObjects;
-	bool			bSerializedUnAckedObject;
-	bool			bSerializedCDO;
+	bool				bSuppressLogs;
 
-	bool			bLoadedUnmappedObject;
-	FNetworkGUID	LastUnmappedNetGUID;
+	bool				bLoadedUnmappedObject;
+	FNetworkGUID		LastUnmappedNetGUID;
 
-	FString			DebugContextString;
+	FString				DebugContextString;
 };
 
 /** Represents a range of PacketIDs, inclusive */
diff --git a/Engine/Source/Runtime/Engine/Classes/Engine/PackageMapClient.h b/Engine/Source/Runtime/Engine/Classes/Engine/PackageMapClient.h
--- a/Engine/Source/Runtime/Engine/Classes/Engine/PackageMapClient.h
+++ b/Engine/Source/Runtime/Engine/Classes/Engine/PackageMapClient.h
@@ -9,9 +9,6 @@
  *	Once ACK'd, objects are just serialized as <NetGUID>. The result is higher bandwidth usage upfront for new clients, but minimal bandwidth once
  *	things gets going.
  *
- *	A further optimization is enabled by default with the SerializeOuter variable. When true, references will actualy be serialized via:
- *	<NetGUID, <(Outer *), Object Name> >. Where Outer * is a reference to the UObject's Outer.
- *
  *	The main advantages from this are:
  *		-Flexibility. No precomputed net indices or large package lists need to be exchanged for UObject serialization.
  *		-Cross version communication. The name is all that is needed to exchange references.
@@ -77,7 +74,6 @@ class UPackageMapClient : public UPackageMap
 	,	Connection(InConnection)
 	{
 		GuidCache = InNetGUIDCache;
-		SerializeOuter = true;
 		Locked = false;
 		ExportNetGUIDCount = 0;
 		IsExportingNetGUIDBunch = false;
@@ -105,7 +101,7 @@ class UPackageMapClient : public UPackageMap
 
 	// UPackageMapClient Connection specific methods
 
-	virtual bool NetGUIDHasBeenAckd(FNetworkGUID NetGUID) OVERRIDE;
+	virtual bool NetGUIDHasBeenAckd(FNetworkGUID NetGUID);
 
 	virtual void ReceivedNak( const int32 NakPacketId ) OVERRIDE;
 	virtual void ReceivedAck( const int32 AckPacketId ) OVERRIDE;
@@ -159,10 +155,6 @@ protected:
 
 	// -----------------
 
-	
-	// Enables optimization to send <outer reference + Name> instead of <full path name> when serializing unackd objects
-	bool							SerializeOuter;
-
 	// Bunches of NetGUID/path tables to send with the current content bunch
 	TArray<FOutBunch* >				ExportBunches;
 	FOutBunch *						CurrentExportBunch;
diff --git a/Engine/Source/Runtime/Engine/Private/PackageMapClient.cpp b/Engine/Source/Runtime/Engine/Private/PackageMapClient.cpp
--- a/Engine/Source/Runtime/Engine/Private/PackageMapClient.cpp
+++ b/Engine/Source/Runtime/Engine/Private/PackageMapClient.cpp
@@ -59,7 +59,6 @@ bool UPackageMapClient::SerializeObject( FArchive& Ar, UClass* Class, UObject*&
 			return SerializeObject( Ar, Class, NullObj, OutNetGUID);
 		}
 
-		bSerializedUnAckedObject = false;
 		FNetworkGUID NetGUID = GuidCache->GetOrAssignNetGUID( Object );
 
 		// Write out NetGUID to caller if necessary
@@ -81,14 +80,6 @@ bool UPackageMapClient::SerializeObject( FArchive& Ar, UClass* Class, UObject*&
 			}
 		}
 
-#if 0
-		if ( Object != NULL && Object->HasAnyFlags( RF_ClassDefaultObject ) )
-		{
-			UE_LOG( LogNetPackageMap, Warning, TEXT( "UPackageMapClient::SerializeObject: Sent CDO: %s, %s" ), *Object->GetFullName(), *Class->GetFullName() );
-			bSerializedCDO = true;
-		}
-#endif
-
 		return true;
 	}
 	else if (Ar.IsLoading())
@@ -130,14 +121,6 @@ bool UPackageMapClient::SerializeObject( FArchive& Ar, UClass* Class, UObject*&
 
 		UE_CLOG(!bSuppressLogs, LogNetPackageMap, Log, TEXT("UPackageMapClient::SerializeObject Serialized Object %s as <%s>"), Object ? *Object->GetPathName() : TEXT("NULL"), *NetGUID.ToString() );
 
-#if 0
-		if ( Object != NULL && Object->HasAnyFlags( RF_ClassDefaultObject ) )
-		{
-			UE_LOG( LogNetPackageMap, Warning, TEXT( "UPackageMapClient::SerializeObject: Received CDO: %s, %s" ), *Object->GetFullName(), *Class->GetFullName() );
-			bSerializedCDO = true;
-		}
-#endif
-
 		// reference is mapped if it was not NULL (or was explicitly null)
 		return (Object != NULL || !NetGUID.IsValid());
 	}
diff --git a/Engine/Source/Runtime/Engine/Private/RepLayout.cpp b/Engine/Source/Runtime/Engine/Private/RepLayout.cpp
--- a/Engine/Source/Runtime/Engine/Private/RepLayout.cpp
+++ b/Engine/Source/Runtime/Engine/Private/RepLayout.cpp
@@ -903,8 +903,7 @@ uint16 FRepLayout::SendProperties_r(
 			const int32 NumStartBits = WriterState.Writer.GetNumBits();
 			
 			// This property changed, so send it
-			WriterState.Writer.PackageMap->ResetUnAckedObject();	// Set this to false so floats, ints, etc don't trigger it
-			bool bMapped = Cmd.Property->NetSerializeItem( WriterState.Writer, WriterState.Writer.PackageMap, (void*)( Data + Cmd.Offset ) );
+			Cmd.Property->NetSerializeItem( WriterState.Writer, WriterState.Writer.PackageMap, (void*)( Data + Cmd.Offset ) );
 
 			const int32 NumEndBits = WriterState.Writer.GetNumBits();
 
@@ -2300,8 +2299,6 @@ void FRepLayout::SerializeProperties_r(
 			continue;
 		}
 
-		Map->ResetHasSerializedCDO();
-		
 #if !(UE_BUILD_SHIPPING || UE_BUILD_TEST)
 		if (CVarDoReplicationContextString->GetInt() > 0)
 		{
@@ -2314,12 +2311,6 @@ void FRepLayout::SerializeProperties_r(
 			bHasUnmapped = true;
 		}
 
-		if ( Map->HasSerializedCDO() && Ar.IsLoading() )
-		{
-			UE_LOG( LogNetTraffic, Warning, TEXT( "CDO serialized. Parameter %s." ), *Cmd.Property->GetName() );
-			Ar.SetError();
-		}
-
 #if !(UE_BUILD_SHIPPING || UE_BUILD_TEST)
 		if (CVarDoReplicationContextString->GetInt() > 0)
 		{
