EditorOnlyBP: Function IsInBlueprint needs o check if there is any blueprint. ttp#337444

When CLASS_CompiledFromBlueprint is set and the ClassGeneratedBy is set - the function should return true.
But with EditorOnlyBP enabled the CLASS_CompiledFromBlueprint is set and the ClassGeneratedBy is null. Then the function should return false.
I didnÆt want to add any config setting for EditorOnlyBP outside the editor.ini. Hopefully the function will be removed, once the EditorOnlyBP becomes non-optional.

diff --git a/Engine/Source/Runtime/CoreUObject/Private/Blueprint/BlueprintSupport.cpp b/Engine/Source/Runtime/CoreUObject/Private/Blueprint/BlueprintSupport.cpp
--- a/Engine/Source/Runtime/CoreUObject/Private/Blueprint/BlueprintSupport.cpp
+++ b/Engine/Source/Runtime/CoreUObject/Private/Blueprint/BlueprintSupport.cpp
@@ -111,11 +111,16 @@ bool UObject::IsInBlueprint() const
  	while (TestObject)
  	{
  		const UClass *ClassObject = Cast<const UClass>(TestObject);
- 		if (ClassObject && ClassObject->HasAnyClassFlags(CLASS_CompiledFromBlueprint))
+		if (ClassObject 
+			&& ClassObject->HasAnyClassFlags(CLASS_CompiledFromBlueprint) 
+			&& ClassObject->ClassGeneratedBy)
  		{
  			return true;
  		}
- 		else if (TestObject->HasAnyFlags(RF_ClassDefaultObject) && TestObject->GetClass() && TestObject->GetClass()->HasAnyClassFlags(CLASS_CompiledFromBlueprint))
+		else if (TestObject->HasAnyFlags(RF_ClassDefaultObject) 
+			&& TestObject->GetClass() 
+			&& TestObject->GetClass()->HasAnyClassFlags(CLASS_CompiledFromBlueprint) 
+			&& TestObject->GetClass()->ClassGeneratedBy)
  		{
  			return true;
  		}
