Back out revision 57 from //depot/UE4/Engine/Source/Runtime/CoreUObject/Private/UObject/CoreNet.cpp

diff --git a/Engine/Source/Runtime/CoreUObject/Private/UObject/CoreNet.cpp b/Engine/Source/Runtime/CoreUObject/Private/UObject/CoreNet.cpp
--- a/Engine/Source/Runtime/CoreUObject/Private/UObject/CoreNet.cpp
+++ b/Engine/Source/Runtime/CoreUObject/Private/UObject/CoreNet.cpp
@@ -62,32 +62,31 @@ FClassNetCache * FClassNetCacheMgr::GetClassNetCache( UClass* Class )
 {
 	FClassNetCache * Result = ClassFieldIndices.FindRef( Class );
 
-
 	if ( !Result )
-{
-		Result                       = ClassFieldIndices.Add( Class, new FClassNetCache(Class) );
-		Result->Super                = NULL;
-		Result->FieldsBase           = 0;
+	{
+		Result					= ClassFieldIndices.Add( Class, new FClassNetCache( Class ) );
+		Result->Super			= NULL;
+		Result->FieldsBase		= 0;
 
-		if( Class->GetSuperClass() )
+		if ( Class->GetSuperClass() )
 		{
-			Result->Super		         = GetClassNetCache(Class->GetSuperClass());
-			Result->FieldsBase           = Result->Super->GetMaxIndex();
+			Result->Super		= GetClassNetCache(Class->GetSuperClass());
+			Result->FieldsBase	= Result->Super->GetMaxIndex();
 		}
 
 		Result->Fields.Empty( Class->NetFields.Num() );
 
-		for( int32 i=0; i<Class->NetFields.Num(); i++ )
+		for( int32 i = 0; i < Class->NetFields.Num(); i++ )
 		{
 			// Add sandboxed items to net cache.  
-			UField* Field = Class->NetFields[i];
-				int32 ThisIndex      = Result->GetMaxIndex();
-				new(Result->Fields)FFieldNetCache( Field, ThisIndex );
-			}
+			UField * Field = Class->NetFields[i];
+			int32 ThisIndex	= Result->GetMaxIndex();
+			new ( Result->Fields )FFieldNetCache( Field, ThisIndex );
+		}
 
 		Result->Fields.Shrink();
 
-		for( TArray<FFieldNetCache>::TIterator It(Result->Fields); It; ++It )
+		for ( TArray< FFieldNetCache >::TIterator It( Result->Fields ); It; ++It )
 		{
 			Result->FieldMap.Add( It->Field, &*It );
 		}
@@ -97,7 +96,7 @@ FClassNetCache * FClassNetCacheMgr::GetClassNetCache( UClass* Class )
 
 void FClassNetCacheMgr::ClearClassNetCache()
 {
-	for( auto It = ClassFieldIndices.CreateIterator(); It; ++It)
+	for ( auto It = ClassFieldIndices.CreateIterator(); It; ++It)
 	{
 		delete It.Value();
 	}
@@ -116,13 +115,15 @@ bool UPackageMap::SerializeName(FArchive& Ar, FName& Name)
 		uint8 bHardcoded = 0;
 		Ar.SerializeBits(&bHardcoded, 1);
 		if (bHardcoded)
-{
+		{
 			// replicated by hardcoded index
 			uint32 NameIndex;
 			Ar.SerializeInt(NameIndex, MAX_NETWORKED_HARDCODED_NAME + 1);
 			Name = EName(NameIndex);
 			// hardcoded names never have a Number
-}
+		}
+		else
+		{
 			// replicated by string
 			FString InString;
 			int32 InNumber;
@@ -148,6 +149,7 @@ bool UPackageMap::SerializeName(FArchive& Ar, FName& Name)
 			int32 OutNumber = Name.GetNumber();
 			Ar << OutString << OutNumber;
 		}
+	}
 	return true;
 }
 
