RF_Native uobjects are net addressable

diff --git a/Engine/Source/Runtime/CoreUObject/Private/UObject/Obj.cpp b/Engine/Source/Runtime/CoreUObject/Private/UObject/Obj.cpp
--- a/Engine/Source/Runtime/CoreUObject/Private/UObject/Obj.cpp
+++ b/Engine/Source/Runtime/CoreUObject/Private/UObject/Obj.cpp
@@ -3776,7 +3776,7 @@ void UObject::PostNetReceive()
 /** IsNameStableForNetworking means an object can be referred to its path name (relative to outer) over the network */
 bool UObject::IsNameStableForNetworking() const
 {
-	return IsDefaultSubobject() || HasAnyFlags( RF_WasLoaded | RF_DefaultSubObject );
+	return IsDefaultSubobject() || HasAnyFlags( RF_WasLoaded | RF_DefaultSubObject | RF_Native );
 }
 
 /** IsFullNameStableForNetworking means an object can be referred to its full path name over the network */
