Stats - Removed unused code

diff --git a/Engine/Source/Runtime/CoreUObject/Public/UObject/UObjectBaseUtility.h b/Engine/Source/Runtime/CoreUObject/Public/UObject/UObjectBaseUtility.h
--- a/Engine/Source/Runtime/CoreUObject/Public/UObject/UObjectBaseUtility.h
+++ b/Engine/Source/Runtime/CoreUObject/Public/UObject/UObjectBaseUtility.h
@@ -494,11 +494,3 @@ struct FScopeCycleCounterUObject
 	}
 };
 #endif
-
-#if STATS
-
-#define SCOPED_UOBJECT_CYCLE_COUNTER(Object)	\
-	FScopeCycleCounter CycleCount_##Stat(Stat);
-#else
-
-#endif
\ No newline at end of file
