GC TokenIndex tracking feature.
Now when token stream referenced object is in invalid state engine will try to print the reference (i.e. field or property) name in non-shipping and non-test builds.

diff --git a/Engine/Source/Runtime/CoreUObject/Private/UObject/Class.cpp b/Engine/Source/Runtime/CoreUObject/Private/UObject/Class.cpp
--- a/Engine/Source/Runtime/CoreUObject/Private/UObject/Class.cpp
+++ b/Engine/Source/Runtime/CoreUObject/Private/UObject/Class.cpp
@@ -347,7 +347,7 @@ void UField::RemoveMetaData(const FName& Key)
 #endif
 IMPLEMENT_CORE_INTRINSIC_CLASS(UField, UObject,
 	{
-		Class->EmitObjectReference( STRUCT_OFFSET( UField, Next ) );
+		Class->EmitObjectReference(STRUCT_OFFSET(UField, Next), TEXT("Next"));
 	}
 );
 
@@ -1420,13 +1420,13 @@ void UStruct::InstanceSubobjectTemplates( void* Data, void const* DefaultData, U
 IMPLEMENT_CORE_INTRINSIC_CLASS(UStruct, UField,
 	{
 		Class->ClassAddReferencedObjects = &UStruct::AddReferencedObjects;
-		Class->EmitObjectReference( STRUCT_OFFSET( UStruct, SuperStruct ) );
-		Class->EmitObjectReference( STRUCT_OFFSET( UStruct, Children ) );
+		Class->EmitObjectReference(STRUCT_OFFSET(UStruct, SuperStruct), TEXT("SuperStruct"));
+		Class->EmitObjectReference(STRUCT_OFFSET(UStruct, Children), TEXT("Children"));
 
 		// Note: None of the *Link members need to be emitted, as they only contain properties
 		// that are in the Children chain or SuperStruct->Children chains.
 
-		Class->EmitObjectArrayReference( STRUCT_OFFSET( UStruct, ScriptObjectReferences ) );
+		Class->EmitObjectArrayReference(STRUCT_OFFSET(UStruct, ScriptObjectReferences), TEXT("ScriptObjectReferences"));
 	}
 );
 
@@ -3535,13 +3535,23 @@ bool UClass::IsClassGroupName(const TCHAR* InGroupName) const
 
 #endif // WITH_EDITOR || HACK_HEADER_GENERATOR
 
+void UClass::PrependStreamWithSuperClass(UClass& SuperClass)
+{
+	ReferenceTokenStream.PrependStream(SuperClass.ReferenceTokenStream);
+
+#if !(UE_BUILD_TEST || UE_BUILD_SHIPPING)
+	DebugTokenMap.PrependWithSuperClass(SuperClass);
+#endif
+}
+
 IMPLEMENT_CORE_INTRINSIC_CLASS(UClass, UStruct,
 	{
 		Class->ClassAddReferencedObjects = &UClass::AddReferencedObjects;
-		Class->EmitObjectReference( STRUCT_OFFSET( UClass, ClassDefaultObject ) );
-		Class->EmitObjectReference( STRUCT_OFFSET( UClass, ClassWithin ) );
-		Class->EmitObjectReference( STRUCT_OFFSET( UClass, ClassGeneratedBy ) );
-		Class->EmitObjectArrayReference( STRUCT_OFFSET( UClass, NetFields ) );
+
+		Class->EmitObjectReference(STRUCT_OFFSET(UClass, ClassDefaultObject), TEXT("ClassDefaultObject"));
+		Class->EmitObjectReference(STRUCT_OFFSET(UClass, ClassWithin), TEXT("ClassWithin"));
+		Class->EmitObjectReference(STRUCT_OFFSET(UClass, ClassGeneratedBy), TEXT("ClassGeneratedBy"));
+		Class->EmitObjectArrayReference(STRUCT_OFFSET(UClass, NetFields), TEXT("NetFields"));
 	}
 );
 
diff --git a/Engine/Source/Runtime/CoreUObject/Private/UObject/GarbageCollection.cpp b/Engine/Source/Runtime/CoreUObject/Private/UObject/GarbageCollection.cpp
--- a/Engine/Source/Runtime/CoreUObject/Private/UObject/GarbageCollection.cpp
+++ b/Engine/Source/Runtime/CoreUObject/Private/UObject/GarbageCollection.cpp
@@ -227,10 +227,23 @@ static FORCEINLINE void HandleTokenStreamObjectReference( TArray<UObject*>& Obje
 #if !(UE_BUILD_TEST || UE_BUILD_SHIPPING)
 	if (Object && !Object->IsValidLowLevelFast())
 	{
-		UE_LOG(LogGarbage, Fatal, TEXT("Invalid object in GC: 0x%016llx, ReferencingObject: %s, TokenIndex: %d"), 
-			(int64)(PTRINT)Object, 
+		FString TokenDebugInfo;
+		if (UClass *Class = (ReferencingObject ? ReferencingObject->GetClass() : nullptr))
+		{
+			auto& TokenInfo = Class->DebugTokenMap.GetTokenInfo(TokenIndex);
+			TokenDebugInfo = FString::Printf(TEXT("ReferencingObjectClass: %s, Name: %s, Offset: %d"),
+				*Class->GetFullName(), *TokenInfo.Name.GetPlainNameString(), TokenInfo.Offset);
+		}
+		else
+		{
+			// This means this objects is most likely being referenced by AddReferencedObjects
+			TokenDebugInfo = FString::Printf(TEXT("Native Reference, TokenIndex: %d"), TokenIndex);
+		}
+
+		UE_LOG(LogGarbage, Fatal, TEXT("Invalid object in GC: 0x%016llx, ReferencingObject: %s, %s"),
+			(int64)(PTRINT)Object,
 			ReferencingObject ? *ReferencingObject->GetFullName() : TEXT("NULL"),
-			TokenIndex);
+			*TokenDebugInfo);
 	}
 #endif
 	HandleObjectReference(ObjectsToSerialize, ReferencingObject, Object, bAllowReferenceElimination);
@@ -605,7 +618,7 @@ public:
 						UObject**	ObjectPtr	= (UObject**)(StackEntryData + REFERENCE_INFO.Offset);
 						UObject*&	Object		= *ObjectPtr;
 						TokenReturnCount		= REFERENCE_INFO.ReturnCount;
-						HandleTokenStreamObjectReference( NewObjectsToSerialize, CurrentObject, Object, TokenStreamIndex, true );
+						HandleTokenStreamObjectReference(NewObjectsToSerialize, CurrentObject, Object, ReferenceTokenStreamIndex, true);
 					}
 					else if( REFERENCE_INFO.Type == GCRT_ArrayObject )
 					{
@@ -615,7 +628,7 @@ public:
 						for( int32 ObjectIndex=0; ObjectIndex<ObjectArray.Num(); ObjectIndex++ )
 						{
 							UObject*& Object = ObjectArray[ObjectIndex];
-							HandleTokenStreamObjectReference( NewObjectsToSerialize, CurrentObject, Object, TokenStreamIndex, true );
+							HandleTokenStreamObjectReference(NewObjectsToSerialize, CurrentObject, Object, ReferenceTokenStreamIndex, true);
 						}
 					}
 					else if( REFERENCE_INFO.Type == GCRT_ArrayStruct )
@@ -650,7 +663,7 @@ public:
 						UObject**	ObjectPtr	= (UObject**)(StackEntryData + REFERENCE_INFO.Offset);
 						UObject*&	Object		= *ObjectPtr;
 						TokenReturnCount		= REFERENCE_INFO.ReturnCount;
-						HandleTokenStreamObjectReference( NewObjectsToSerialize, CurrentObject, Object, TokenStreamIndex, false );
+						HandleTokenStreamObjectReference(NewObjectsToSerialize, CurrentObject, Object, ReferenceTokenStreamIndex, false);
 					}
 					else if( REFERENCE_INFO.Type == GCRT_FixedArray )
 					{
@@ -1363,17 +1376,18 @@ struct FGCReferenceFixedArrayTokenHelper
 	 * @param InOffset					offset into object/ struct
 	 * @param InCount					array count
 	 * @param InStride					array type stride (e.g. sizeof(struct) or sizeof(UObject*))
+	 * @param InProperty                the property this array represents
 	 */
-	FGCReferenceFixedArrayTokenHelper( FGCReferenceTokenStream* InReferenceTokenStream, int32 InOffset, int32 InCount, int32 InStride )
-	:	ReferenceTokenStream( InReferenceTokenStream )
-	,	Count( InCount )
+	FGCReferenceFixedArrayTokenHelper(UClass& OwnerClass, int32 InOffset, int32 InCount, int32 InStride, const UProperty& InProperty)
+		: ReferenceTokenStream(&OwnerClass.ReferenceTokenStream)
+	,	Count(InCount)
 	{
 		if( InCount > 1 )
 		{
-			FGCReferenceInfo FixedArrayReference( GCRT_FixedArray, InOffset );
-			ReferenceTokenStream->EmitReferenceInfo( FixedArrayReference );
-			ReferenceTokenStream->EmitStride( InStride );
-			ReferenceTokenStream->EmitCount( InCount );
+			OwnerClass.EmitObjectReference(InOffset, InProperty.GetFName(), GCRT_FixedArray);
+
+			OwnerClass.ReferenceTokenStream.EmitStride(InStride);
+			OwnerClass.ReferenceTokenStream.EmitCount(InCount);
 		}
 	}
 
@@ -1398,56 +1412,53 @@ private:
  * Emits tokens used by realtime garbage collection code to passed in ReferenceTokenStream. The offset emitted is relative
  * to the passed in BaseOffset which is used by e.g. arrays of structs.
  */
-void UProperty::EmitReferenceInfo( FGCReferenceTokenStream* ReferenceTokenStream, int32 BaseOffset )
+void UProperty::EmitReferenceInfo(UClass& OwnerClass, int32 BaseOffset)
 {
 }
 
 /**
- * Emits tokens used by realtime garbage collection code to passed in ReferenceTokenStream. The offset emitted is relative
+ * Emits tokens used by realtime garbage collection code to passed in OwnerClass' ReferenceTokenStream. The offset emitted is relative
  * to the passed in BaseOffset which is used by e.g. arrays of structs.
  */
-void UObjectProperty::EmitReferenceInfo( FGCReferenceTokenStream* ReferenceTokenStream, int32 BaseOffset )
+void UObjectProperty::EmitReferenceInfo(UClass& OwnerClass, int32 BaseOffset)
 {
-	FGCReferenceFixedArrayTokenHelper FixedArrayHelper( ReferenceTokenStream, BaseOffset + GetOffset_ForGC(), ArrayDim, sizeof(UObject*) );
-	FGCReferenceInfo ObjectReference( GCRT_Object, BaseOffset + GetOffset_ForGC() );
-	ReferenceTokenStream->EmitReferenceInfo( ObjectReference );
+	FGCReferenceFixedArrayTokenHelper FixedArrayHelper(OwnerClass, BaseOffset + GetOffset_ForGC(), ArrayDim, sizeof(UObject*), *this);
+	OwnerClass.EmitObjectReference(BaseOffset + GetOffset_ForGC(), GetFName(), GCRT_Object);
 }
 
 /**
- * Emits tokens used by realtime garbage collection code to passed in ReferenceTokenStream. The offset emitted is relative
+ * Emits tokens used by realtime garbage collection code to passed in OwnerClass' ReferenceTokenStream. The offset emitted is relative
  * to the passed in BaseOffset which is used by e.g. arrays of structs.
  */
-void UArrayProperty::EmitReferenceInfo( FGCReferenceTokenStream* ReferenceTokenStream, int32 BaseOffset )
+void UArrayProperty::EmitReferenceInfo(UClass& OwnerClass, int32 BaseOffset)
 {
 	if( Inner->ContainsObjectReference() )
 	{
 		if( Inner->IsA(UStructProperty::StaticClass()) )
 		{
-			FGCReferenceInfo ReferenceInfo( GCRT_ArrayStruct, BaseOffset + GetOffset_ForGC() );
-			ReferenceTokenStream->EmitReferenceInfo( ReferenceInfo );
-			ReferenceTokenStream->EmitStride( Inner->ElementSize );
-			const uint32 SkipIndexIndex = ReferenceTokenStream->EmitSkipIndexPlaceholder();
-			Inner->EmitReferenceInfo( ReferenceTokenStream, 0 );
-			const uint32 SkipIndex = ReferenceTokenStream->EmitReturn();
-			ReferenceTokenStream->UpdateSkipIndexPlaceholder( SkipIndexIndex, SkipIndex );
+			OwnerClass.EmitObjectReference(BaseOffset + GetOffset_ForGC(), GetFName(), GCRT_ArrayStruct);
+
+			OwnerClass.ReferenceTokenStream.EmitStride(Inner->ElementSize);
+			const uint32 SkipIndexIndex = OwnerClass.ReferenceTokenStream.EmitSkipIndexPlaceholder();
+			Inner->EmitReferenceInfo(OwnerClass, 0);
+			const uint32 SkipIndex = OwnerClass.ReferenceTokenStream.EmitReturn();
+			OwnerClass.ReferenceTokenStream.UpdateSkipIndexPlaceholder(SkipIndexIndex, SkipIndex);
 		}
 		else if( Inner->IsA(UObjectProperty::StaticClass()) )
 		{
-			FGCReferenceInfo ReferenceInfo( GCRT_ArrayObject, BaseOffset + GetOffset_ForGC() );
-			ReferenceTokenStream->EmitReferenceInfo( ReferenceInfo );
+			OwnerClass.EmitObjectReference(BaseOffset + GetOffset_ForGC(), GetFName(), GCRT_ArrayObject);
 		}
 		else if( Inner->IsA(UInterfaceProperty::StaticClass()) )
 		{
-			FGCReferenceInfo ReferenceInfo( GCRT_ArrayStruct, BaseOffset + GetOffset_ForGC() );
-			ReferenceTokenStream->EmitReferenceInfo( ReferenceInfo );
-			ReferenceTokenStream->EmitStride( Inner->ElementSize );
-			const uint32 SkipIndexIndex = ReferenceTokenStream->EmitSkipIndexPlaceholder();
+			OwnerClass.EmitObjectReference(BaseOffset + GetOffset_ForGC(), GetFName(), GCRT_ArrayStruct);
 
-			FGCReferenceInfo InnerReferenceInfo( GCRT_Object, 0 );
-			ReferenceTokenStream->EmitReferenceInfo( InnerReferenceInfo );
+			OwnerClass.ReferenceTokenStream.EmitStride(Inner->ElementSize);
+			const uint32 SkipIndexIndex = OwnerClass.ReferenceTokenStream.EmitSkipIndexPlaceholder();
 
-			const uint32 SkipIndex = ReferenceTokenStream->EmitReturn();
-			ReferenceTokenStream->UpdateSkipIndexPlaceholder( SkipIndexIndex, SkipIndex );
+			OwnerClass.EmitObjectReference(0, GetFName(), GCRT_Object);
+
+			const uint32 SkipIndex = OwnerClass.ReferenceTokenStream.EmitReturn();
+			OwnerClass.ReferenceTokenStream.UpdateSkipIndexPlaceholder(SkipIndexIndex, SkipIndex);
 		}
 		else
 		{
@@ -1461,31 +1472,30 @@ void UArrayProperty::EmitReferenceInfo( FGCReferenceTokenStream* ReferenceTokenS
  * Emits tokens used by realtime garbage collection code to passed in ReferenceTokenStream. The offset emitted is relative
  * to the passed in BaseOffset which is used by e.g. arrays of structs.
  */
-void UStructProperty::EmitReferenceInfo( FGCReferenceTokenStream* ReferenceTokenStream, int32 BaseOffset )
+void UStructProperty::EmitReferenceInfo(UClass& OwnerClass, int32 BaseOffset)
 {
 	if (Struct->StructFlags & STRUCT_AddStructReferencedObjects)
 	{
 		UScriptStruct::ICppStructOps* CppStructOps = Struct->GetCppStructOps();
 		check(CppStructOps); // else should not have STRUCT_AddStructReferencedObjects
 		check(!Struct->InheritedCppStructOps()); // else should not have STRUCT_AddStructReferencedObjects
-		FGCReferenceFixedArrayTokenHelper FixedArrayHelper( ReferenceTokenStream, BaseOffset + GetOffset_ForGC(), ArrayDim, ElementSize );
+		FGCReferenceFixedArrayTokenHelper FixedArrayHelper(OwnerClass, BaseOffset + GetOffset_ForGC(), ArrayDim, ElementSize, *this);
 
-		FGCReferenceInfo ReferenceInfo( GCRT_AddStructReferencedObjects, BaseOffset + GetOffset_ForGC());
-		ReferenceTokenStream->EmitReferenceInfo( ReferenceInfo );
+		OwnerClass.EmitObjectReference(BaseOffset + GetOffset_ForGC(), GetFName(), GCRT_AddStructReferencedObjects);
 
 		void *FunctionPtr = (void*)CppStructOps->AddStructReferencedObjects();
-		ReferenceTokenStream->EmitPointer(FunctionPtr);
+		OwnerClass.ReferenceTokenStream.EmitPointer(FunctionPtr);
 		return;
 	}
 	check(Struct);
 	if( ContainsObjectReference() )
 	{
-		FGCReferenceFixedArrayTokenHelper FixedArrayHelper( ReferenceTokenStream, BaseOffset + GetOffset_ForGC(), ArrayDim, ElementSize );
+		FGCReferenceFixedArrayTokenHelper FixedArrayHelper(OwnerClass, BaseOffset + GetOffset_ForGC(), ArrayDim, ElementSize, *this);
 
 		UProperty* Property = Struct->PropertyLink;
 		while( Property )
 		{
-			Property->EmitReferenceInfo( ReferenceTokenStream, BaseOffset + GetOffset_ForGC() );
+			Property->EmitReferenceInfo(OwnerClass, BaseOffset + GetOffset_ForGC());
 			Property = Property->PropertyLinkNext;
 		}
 	}
@@ -1495,54 +1505,34 @@ void UStructProperty::EmitReferenceInfo( FGCReferenceTokenStream* ReferenceToken
  * Emits tokens used by realtime garbage collection code to passed in ReferenceTokenStream. The offset emitted is relative
  * to the passed in BaseOffset which is used by e.g. arrays of structs.
  */
-void UInterfaceProperty::EmitReferenceInfo( FGCReferenceTokenStream* ReferenceTokenStream, int32 BaseOffset )
+void UInterfaceProperty::EmitReferenceInfo(UClass& OwnerClass, int32 BaseOffset)
 {
-	FGCReferenceFixedArrayTokenHelper FixedArrayHelper( ReferenceTokenStream, BaseOffset + GetOffset_ForGC(), ArrayDim, sizeof(FScriptInterface) );
+	FGCReferenceFixedArrayTokenHelper FixedArrayHelper(OwnerClass, BaseOffset + GetOffset_ForGC(), ArrayDim, sizeof(FScriptInterface), *this);
 
-	FGCReferenceInfo ObjectReference( GCRT_Object, BaseOffset + GetOffset_ForGC() );
-	ReferenceTokenStream->EmitReferenceInfo( ObjectReference );
+	OwnerClass.EmitObjectReference(BaseOffset + GetOffset_ForGC(), GetFName(), GCRT_Object);
 }
 
-/**
- * Realtime garbage collection helper function used to emit token containing information about a 
- * direct UObject reference at the passed in offset.
- *
- * @param Offset	offset into object at which object reference is stored
- */
-void UClass::EmitObjectReference( int32 Offset, EGCReferenceType Kind )
+void UClass::EmitObjectReference(int32 Offset, const FName& DebugName, EGCReferenceType Kind)
 {
-	FGCReferenceInfo ObjectReference( Kind, Offset );
-	ReferenceTokenStream.EmitReferenceInfo( ObjectReference );
+	FGCReferenceInfo ObjectReference(Kind, Offset);
+	int32 TokenIndex = ReferenceTokenStream.EmitReferenceInfo(ObjectReference);
+
+#if !(UE_BUILD_TEST || UE_BUILD_SHIPPING)
+	DebugTokenMap.MapToken(DebugName, Offset, TokenIndex);
+#endif
 }
 
-/**
- * Realtime garbage collection helper function used to emit token containing information about a 
- * an array of UObject references at the passed in offset. Handles both TArray and TTransArray.
- *
- * @param Offset	offset into object at which array of objects is stored
- */
-void UClass::EmitObjectArrayReference( int32 Offset )
+void UClass::EmitObjectArrayReference(int32 Offset, const FName& DebugName)
 {
-	check( HasAnyClassFlags( CLASS_Intrinsic ) );
-	FGCReferenceInfo ObjectReference( GCRT_ArrayObject, Offset );
-	ReferenceTokenStream.EmitReferenceInfo( ObjectReference );
+	check(HasAnyClassFlags(CLASS_Intrinsic));
+	EmitObjectReference(Offset, DebugName, GCRT_ArrayObject);
 }
 
-/**
- * Realtime garbage collection helper function used to indicate an array of structs at the passed in 
- * offset.
- *
- * @param Offset	offset into object at which array of structs is stored
- * @param Stride	size/ stride of struct
- * @return	index into token stream at which later on index to next token after the array is stored
- *			which is used to skip over empty dynamic arrays
- */
-uint32 UClass::EmitStructArrayBegin( int32 Offset, int32 Stride )
+uint32 UClass::EmitStructArrayBegin(int32 Offset, const FName& DebugName, int32 Stride)
 {
-	check( HasAnyClassFlags( CLASS_Intrinsic ) );
-	FGCReferenceInfo ReferenceInfo( GCRT_ArrayStruct, Offset );
-	ReferenceTokenStream.EmitReferenceInfo( ReferenceInfo );
-	ReferenceTokenStream.EmitStride( Stride );
+	check(HasAnyClassFlags(CLASS_Intrinsic));
+	EmitObjectReference(Offset, DebugName, GCRT_ArrayStruct);
+	ReferenceTokenStream.EmitStride(Stride);
 	const uint32 SkipIndexIndex = ReferenceTokenStream.EmitSkipIndexPlaceholder();
 	return SkipIndexIndex;
 }
@@ -1561,21 +1551,12 @@ void UClass::EmitStructArrayEnd( uint32 SkipIndexIndex )
 	ReferenceTokenStream.UpdateSkipIndexPlaceholder( SkipIndexIndex, SkipIndex );
 }
 
-/**
- * Realtime garbage collection helper function used to indicate the beginning of a fixed array.
- * All tokens issues between Begin and End will be replayed Count times.
- *
- * @param Offset	offset at which fixed array starts
- * @param Stride	Stride of array element, e.g. sizeof(struct) or sizeof(UObject*)
- * @param Count		fixed array count
- */
-void UClass::EmitFixedArrayBegin( int32 Offset, int32 Stride, int32 Count )
+void UClass::EmitFixedArrayBegin(int32 Offset, const FName& DebugName, int32 Stride, int32 Count)
 {
-	check( HasAnyClassFlags( CLASS_Intrinsic ) );
-	FGCReferenceInfo FixedArrayReference( GCRT_FixedArray, Offset );
-	ReferenceTokenStream.EmitReferenceInfo( FixedArrayReference );
-	ReferenceTokenStream.EmitStride( Stride );
-	ReferenceTokenStream.EmitCount( Count );
+	check(HasAnyClassFlags(CLASS_Intrinsic));
+	EmitObjectReference(Offset, DebugName, GCRT_FixedArray);
+	ReferenceTokenStream.EmitStride(Stride);
+	ReferenceTokenStream.EmitCount(Count);
 }
 
 /**
@@ -1600,7 +1581,7 @@ void UClass::AssembleReferenceTokenStream()
 		for( TFieldIterator<UProperty> It(this,EFieldIteratorFlags::ExcludeSuper); It; ++It)
 		{
 			UProperty* Property = *It;
-			Property->EmitReferenceInfo( &ReferenceTokenStream, 0 );
+			Property->EmitReferenceInfo(*this, 0);
 		}
 
 		if (GetSuperClass())
@@ -1610,7 +1591,7 @@ void UClass::AssembleReferenceTokenStream()
 			if (!GetSuperClass()->ReferenceTokenStream.IsEmpty())
 			{
 				// Prepend super's stream. This automatically handles removing the EOS token.
-				ReferenceTokenStream.PrependStream( GetSuperClass()->ReferenceTokenStream );
+				PrependStreamWithSuperClass(*GetSuperClass());
 			}
 		}
 		else
@@ -1633,8 +1614,8 @@ void UClass::AssembleReferenceTokenStream()
 		}
 
 		// Emit end of stream token.
-		FGCReferenceInfo EndOfStream( GCRT_EndOfStream, 0 );
-		ReferenceTokenStream.EmitReferenceInfo( EndOfStream );
+		const FName EOSDebugName("EOS");
+		EmitObjectReference(0, EOSDebugName, GCRT_EndOfStream);
 
 		// Shrink reference token stream to proper size.
 		ReferenceTokenStream.Shrink();
@@ -1721,14 +1702,9 @@ void FGCReferenceTokenStream::ReplaceOrAddAddReferencedObjectsCall(void (*AddRef
 	EmitPointer((const void*)AddReferencedObjectsPtr);
 }
 
-/**
- * Emit reference info
- *
- * @param ReferenceInfo	reference info to emit
- */
-void FGCReferenceTokenStream::EmitReferenceInfo( FGCReferenceInfo ReferenceInfo )
+int32 FGCReferenceTokenStream::EmitReferenceInfo(FGCReferenceInfo ReferenceInfo)
 {
-	Tokens.Add( ReferenceInfo );
+	return Tokens.Add(ReferenceInfo);
 }
 
 /**
@@ -1804,3 +1780,51 @@ uint32 FGCReferenceTokenStream::EmitReturn()
 	Tokens.Last() = ReferenceInfo;
 	return Tokens.Num();
 }
+
+#if !(UE_BUILD_TEST || UE_BUILD_SHIPPING)
+
+void FGCDebugReferenceTokenMap::MapToken(const FName& DebugName, int32 Offset, int32 TokenIndex)
+{
+	if(TokenMap.Num() <= TokenIndex)
+	{
+		TokenMap.AddZeroed(TokenIndex - TokenMap.Num() + 1);
+
+		TokenMap[TokenIndex] = { Offset, DebugName };
+	}
+	else
+	{
+		// Token already mapped.
+		checkNoEntry();
+	}
+}
+
+void FGCDebugReferenceTokenMap::PrependWithSuperClass(const UClass& SuperClass)
+{
+	if (SuperClass.ReferenceTokenStream.Size() == 0)
+	{
+		return;
+	}
+
+	// Check if token stream is already ended with end-of-stream token. If so then something's wrong.
+	checkSlow(TokenMap.Num() == 0 || TokenMap.Last().Name != "EOS");
+
+	int32 OldTokenNumber = TokenMap.Num();
+	int32 NewTokenOffset = SuperClass.ReferenceTokenStream.Size() - 1;
+	TokenMap.AddZeroed(NewTokenOffset);
+
+	for(int32 OldTokenIndex = OldTokenNumber - 1; OldTokenIndex >= 0; --OldTokenIndex)
+	{
+		TokenMap[OldTokenIndex + NewTokenOffset] = TokenMap[OldTokenIndex];
+	}
+
+	for(int32 NewTokenIndex = 0; NewTokenIndex < NewTokenOffset; ++NewTokenIndex)
+	{
+		TokenMap[NewTokenIndex] = SuperClass.DebugTokenMap.GetTokenInfo(NewTokenIndex);
+	}
+}
+
+const FTokenInfo& FGCDebugReferenceTokenMap::GetTokenInfo(int32 TokenIndex) const
+{
+	return TokenMap[TokenIndex];
+}
+#endif
diff --git a/Engine/Source/Runtime/CoreUObject/Private/UObject/Linker.cpp b/Engine/Source/Runtime/CoreUObject/Private/UObject/Linker.cpp
--- a/Engine/Source/Runtime/CoreUObject/Private/UObject/Linker.cpp
+++ b/Engine/Source/Runtime/CoreUObject/Private/UObject/Linker.cpp
@@ -746,15 +746,12 @@ void ResetLoadersForSave(UObject* InOuter, const TCHAR *Filename)
 	}
 }
 
-
-
-
 IMPLEMENT_CORE_INTRINSIC_CLASS(ULinker, UObject,
 	{
 		Class->ClassAddReferencedObjects = &ULinker::AddReferencedObjects;
-		Class->EmitObjectReference( STRUCT_OFFSET( ULinker, LinkerRoot ) );
-		const uint32 SkipIndexIndex = Class->EmitStructArrayBegin( STRUCT_OFFSET( ULinker, ImportMap ), sizeof(FObjectImport) );
-		Class->EmitObjectReference( STRUCT_OFFSET( FObjectImport, SourceLinker ) );
+		Class->EmitObjectReference(STRUCT_OFFSET(ULinker, LinkerRoot), TEXT("LinkerRoot"));
+		const uint32 SkipIndexIndex = Class->EmitStructArrayBegin(STRUCT_OFFSET(ULinker, ImportMap), TEXT("ImportMap"), sizeof(FObjectImport));
+		Class->EmitObjectReference(STRUCT_OFFSET(FObjectImport, SourceLinker), TEXT("SourceLinker"));
 		Class->EmitStructArrayEnd( SkipIndexIndex );
 	}
 );
diff --git a/Engine/Source/Runtime/CoreUObject/Private/UObject/ObjectRedirector.cpp b/Engine/Source/Runtime/CoreUObject/Private/UObject/ObjectRedirector.cpp
--- a/Engine/Source/Runtime/CoreUObject/Private/UObject/ObjectRedirector.cpp
+++ b/Engine/Source/Runtime/CoreUObject/Private/UObject/ObjectRedirector.cpp
@@ -95,7 +95,7 @@ bool UObjectRedirector::GetNativePropertyValues( TMap<FString,FString>& out_Prop
 
 IMPLEMENT_CORE_INTRINSIC_CLASS(UObjectRedirector, UObject,
 	{
-		Class->EmitObjectReference( STRUCT_OFFSET( UObjectRedirector, DestinationObject ) );
+		Class->EmitObjectReference(STRUCT_OFFSET(UObjectRedirector, DestinationObject), TEXT("DestinationObject"));
 	}
 );
 
diff --git a/Engine/Source/Runtime/CoreUObject/Private/UObject/Package.cpp b/Engine/Source/Runtime/CoreUObject/Private/UObject/Package.cpp
--- a/Engine/Source/Runtime/CoreUObject/Private/UObject/Package.cpp
+++ b/Engine/Source/Runtime/CoreUObject/Private/UObject/Package.cpp
@@ -183,6 +183,6 @@ bool UPackage::IsFullyLoaded()
 IMPLEMENT_CORE_INTRINSIC_CLASS(UPackage, UObject,
 	{
 		Class->ClassAddReferencedObjects = &UPackage::AddReferencedObjects;
-		Class->EmitObjectReference(STRUCT_OFFSET(UPackage, MetaData));
+		Class->EmitObjectReference(STRUCT_OFFSET(UPackage, MetaData), TEXT("MetaData"));
 	}
 );
diff --git a/Engine/Source/Runtime/CoreUObject/Private/UObject/PropertyArray.cpp b/Engine/Source/Runtime/CoreUObject/Private/UObject/PropertyArray.cpp
--- a/Engine/Source/Runtime/CoreUObject/Private/UObject/PropertyArray.cpp
+++ b/Engine/Source/Runtime/CoreUObject/Private/UObject/PropertyArray.cpp
@@ -349,7 +349,7 @@ bool UArrayProperty::SameType(const UProperty* Other) const
 
 IMPLEMENT_CORE_INTRINSIC_CLASS(UArrayProperty, UProperty,
 	{
-		Class->EmitObjectReference( STRUCT_OFFSET( UArrayProperty, Inner ) );
+		Class->EmitObjectReference(STRUCT_OFFSET(UArrayProperty, Inner), TEXT("Inner"));
 
 		// Ensure that TArray and FScriptArray are interchangeable, as FScriptArray will be used to access a native array property
 		// from script that is declared as a TArray in C++.
diff --git a/Engine/Source/Runtime/CoreUObject/Private/UObject/PropertyAssetClass.cpp b/Engine/Source/Runtime/CoreUObject/Private/UObject/PropertyAssetClass.cpp
--- a/Engine/Source/Runtime/CoreUObject/Private/UObject/PropertyAssetClass.cpp
+++ b/Engine/Source/Runtime/CoreUObject/Private/UObject/PropertyAssetClass.cpp
@@ -48,6 +48,6 @@ bool UAssetClassProperty::SameType(const UProperty* Other) const
 
 IMPLEMENT_CORE_INTRINSIC_CLASS(UAssetClassProperty, UAssetObjectProperty,
 	{
-		Class->EmitObjectReference( STRUCT_OFFSET( UAssetClassProperty, MetaClass ) );
+		Class->EmitObjectReference(STRUCT_OFFSET(UAssetClassProperty, MetaClass), TEXT("MetaClass"));
 	}
 );
diff --git a/Engine/Source/Runtime/CoreUObject/Private/UObject/PropertyBaseObject.cpp b/Engine/Source/Runtime/CoreUObject/Private/UObject/PropertyBaseObject.cpp
--- a/Engine/Source/Runtime/CoreUObject/Private/UObject/PropertyBaseObject.cpp
+++ b/Engine/Source/Runtime/CoreUObject/Private/UObject/PropertyBaseObject.cpp
@@ -459,6 +459,6 @@ bool UObjectPropertyBase::SameType(const UProperty* Other) const
 
 IMPLEMENT_CORE_INTRINSIC_CLASS(UObjectPropertyBase, UProperty,
 	{
-		Class->EmitObjectReference( STRUCT_OFFSET( UObjectProperty, PropertyClass ) );
+		Class->EmitObjectReference(STRUCT_OFFSET(UObjectProperty, PropertyClass), TEXT("PropertyClass"));
 	}
 );
diff --git a/Engine/Source/Runtime/CoreUObject/Private/UObject/PropertyByte.cpp b/Engine/Source/Runtime/CoreUObject/Private/UObject/PropertyByte.cpp
--- a/Engine/Source/Runtime/CoreUObject/Private/UObject/PropertyByte.cpp
+++ b/Engine/Source/Runtime/CoreUObject/Private/UObject/PropertyByte.cpp
@@ -138,7 +138,7 @@ const TCHAR* UByteProperty::ImportText_Internal( const TCHAR* InBuffer, void* Da
 
 IMPLEMENT_CORE_INTRINSIC_CLASS(UByteProperty, UNumericProperty,
 	{
-		Class->EmitObjectReference( STRUCT_OFFSET( UByteProperty, Enum ) );
+		Class->EmitObjectReference(STRUCT_OFFSET(UByteProperty, Enum), TEXT("Enum"));
 	}
 );
 
diff --git a/Engine/Source/Runtime/CoreUObject/Private/UObject/PropertyClass.cpp b/Engine/Source/Runtime/CoreUObject/Private/UObject/PropertyClass.cpp
--- a/Engine/Source/Runtime/CoreUObject/Private/UObject/PropertyClass.cpp
+++ b/Engine/Source/Runtime/CoreUObject/Private/UObject/PropertyClass.cpp
@@ -119,7 +119,7 @@ void UClassProperty::CheckValidObject(void* Value) const
 
 IMPLEMENT_CORE_INTRINSIC_CLASS(UClassProperty, UObjectProperty,
 	{
-		Class->EmitObjectReference( STRUCT_OFFSET( UClassProperty, MetaClass ) );
+		Class->EmitObjectReference(STRUCT_OFFSET(UClassProperty, MetaClass), TEXT("MetaClass"));
 	}
 );
 
diff --git a/Engine/Source/Runtime/CoreUObject/Private/UObject/PropertyDelegate.cpp b/Engine/Source/Runtime/CoreUObject/Private/UObject/PropertyDelegate.cpp
--- a/Engine/Source/Runtime/CoreUObject/Private/UObject/PropertyDelegate.cpp
+++ b/Engine/Source/Runtime/CoreUObject/Private/UObject/PropertyDelegate.cpp
@@ -108,6 +108,6 @@ bool UDelegateProperty::SameType(const UProperty* Other) const
 
 IMPLEMENT_CORE_INTRINSIC_CLASS(UDelegateProperty, UProperty,
 	{
-		Class->EmitObjectReference( STRUCT_OFFSET( UDelegateProperty, SignatureFunction ) );
+		Class->EmitObjectReference(STRUCT_OFFSET(UDelegateProperty, SignatureFunction), TEXT("SignatureFunction"));
 	}
 );
diff --git a/Engine/Source/Runtime/CoreUObject/Private/UObject/PropertyInterface.cpp b/Engine/Source/Runtime/CoreUObject/Private/UObject/PropertyInterface.cpp
--- a/Engine/Source/Runtime/CoreUObject/Private/UObject/PropertyInterface.cpp
+++ b/Engine/Source/Runtime/CoreUObject/Private/UObject/PropertyInterface.cpp
@@ -207,6 +207,6 @@ bool UInterfaceProperty::SameType(const UProperty* Other) const
 
 IMPLEMENT_CORE_INTRINSIC_CLASS(UInterfaceProperty, UProperty,
 	{
-		Class->EmitObjectReference( STRUCT_OFFSET( UInterfaceProperty, InterfaceClass ) );
+		Class->EmitObjectReference(STRUCT_OFFSET(UInterfaceProperty, InterfaceClass), TEXT("InterfaceClass"));
 	}
 );
diff --git a/Engine/Source/Runtime/CoreUObject/Private/UObject/PropertyMulticastDelegate.cpp b/Engine/Source/Runtime/CoreUObject/Private/UObject/PropertyMulticastDelegate.cpp
--- a/Engine/Source/Runtime/CoreUObject/Private/UObject/PropertyMulticastDelegate.cpp
+++ b/Engine/Source/Runtime/CoreUObject/Private/UObject/PropertyMulticastDelegate.cpp
@@ -258,7 +258,7 @@ bool UMulticastDelegateProperty::SameType(const UProperty* Other) const
 
 IMPLEMENT_CORE_INTRINSIC_CLASS(UMulticastDelegateProperty, UProperty,
 	{
-		Class->EmitObjectReference( STRUCT_OFFSET( UMulticastDelegateProperty, SignatureFunction ) );
+		Class->EmitObjectReference(STRUCT_OFFSET(UMulticastDelegateProperty, SignatureFunction), TEXT("SignatureFunction"));
 	}
 );
 
diff --git a/Engine/Source/Runtime/CoreUObject/Private/UObject/PropertyStruct.cpp b/Engine/Source/Runtime/CoreUObject/Private/UObject/PropertyStruct.cpp
--- a/Engine/Source/Runtime/CoreUObject/Private/UObject/PropertyStruct.cpp
+++ b/Engine/Source/Runtime/CoreUObject/Private/UObject/PropertyStruct.cpp
@@ -455,6 +455,6 @@ bool UStructProperty::SameType(const UProperty* Other) const
 
 IMPLEMENT_CORE_INTRINSIC_CLASS(UStructProperty, UProperty,
 	{
-		Class->EmitObjectReference( STRUCT_OFFSET( UStructProperty, Struct ) );
+		Class->EmitObjectReference(STRUCT_OFFSET(UStructProperty, Struct), TEXT("Struct"));
 	}
 );
diff --git a/Engine/Source/Runtime/CoreUObject/Private/UObject/UObjectBase.cpp b/Engine/Source/Runtime/CoreUObject/Private/UObject/UObjectBase.cpp
--- a/Engine/Source/Runtime/CoreUObject/Private/UObject/UObjectBase.cpp
+++ b/Engine/Source/Runtime/CoreUObject/Private/UObject/UObjectBase.cpp
@@ -337,8 +337,8 @@ bool UObjectBase::IsValidLowLevelFast(bool bRecursive /*= true*/) const
 
 void UObjectBase::EmitBaseReferences(UClass *RootClass)
 {
-	RootClass->EmitObjectReference( STRUCT_OFFSET( UObjectBase, Class ) );
-	RootClass->EmitObjectReference( STRUCT_OFFSET( UObjectBase, Outer ), GCRT_PersistentObject);
+	RootClass->EmitObjectReference(STRUCT_OFFSET(UObjectBase, Class), TEXT("Class"));
+	RootClass->EmitObjectReference(STRUCT_OFFSET(UObjectBase, Outer), TEXT("Outer"), GCRT_PersistentObject);
 }
 
 /** Enqueue the registration for this object. */
diff --git a/Engine/Source/Runtime/CoreUObject/Public/UObject/Class.h b/Engine/Source/Runtime/CoreUObject/Public/UObject/Class.h
--- a/Engine/Source/Runtime/CoreUObject/Public/UObject/Class.h
+++ b/Engine/Source/Runtime/CoreUObject/Public/UObject/Class.h
@@ -1773,9 +1773,21 @@ public:
 	 **/
 	TArray<FImplementedInterface> Interfaces;
 
+	/**
+	 * Prepends reference token stream with super class's stream.
+	 *
+	 * @param SuperClass Super class to prepend stream with.
+	 */
+	void PrependStreamWithSuperClass(UClass& SuperClass);
+
 	/** Reference token stream used by realtime garbage collector, finalized in AssembleReferenceTokenStream */
 	FGCReferenceTokenStream ReferenceTokenStream;
 
+#if !(UE_BUILD_TEST || UE_BUILD_SHIPPING)
+	/* TokenIndex map to look-up token stream index origin. */
+	FGCDebugReferenceTokenMap DebugTokenMap;
+#endif
+
 	/** This class's native functions. */
 	TArray<FNativeFunctionLookup> NativeFunctionLookupTable;
 
@@ -1993,28 +2005,32 @@ public:
 	 * Realtime garbage collection helper function used to emit token containing information about a 
 	 * direct UObject reference at the passed in offset.
 	 *
-	 * @param Offset	offset into object at which object reference is stored
+	 * @param Offset	Offset into object at which object reference is stored.
+	 * @param DebugName	DebugName for this objects token. Only used in non-shipping builds.
+	 * @param Kind		Optional parameter the describe the type of the reference.
 	 */
-	void EmitObjectReference(int32 Offset, EGCReferenceType Kind = GCRT_Object);
+	void EmitObjectReference(int32 Offset, const FName& DebugName, EGCReferenceType Kind = GCRT_Object);
 
 	/**
 	 * Realtime garbage collection helper function used to emit token containing information about a 
 	 * an array of UObject references at the passed in offset. Handles both TArray and TTransArray.
 	 *
-	 * @param Offset	offset into object at which array of objects is stored
+	 * @param Offset	Offset into object at which array of objects is stored.
+	 * @param DebugName	DebugName for this objects token. Only used in non-shipping builds.
 	 */
-	void EmitObjectArrayReference(int32 Offset);
+	void EmitObjectArrayReference(int32 Offset, const FName& DebugName);
 
 	/**
 	 * Realtime garbage collection helper function used to indicate an array of structs at the passed in 
 	 * offset.
 	 *
-	 * @param Offset	offset into object at which array of structs is stored
-	 * @param Stride	size/ stride of struct
-	 * @return	index into token stream at which later on index to next token after the array is stored
+	 * @param Offset	Offset into object at which array of structs is stored
+	 * @param DebugName	DebugName for this objects token. Only used in non-shipping builds.
+	 * @param Stride	Size/stride of struct
+	 * @return	Index into token stream at which later on index to next token after the array is stored
 	 *			which is used to skip over empty dynamic arrays
 	 */
-	uint32 EmitStructArrayBegin(int32 Offset, int32 Stride);
+	uint32 EmitStructArrayBegin(int32 Offset, const FName& DebugName, int32 Stride);
 
 	/**
 	 * Realtime garbage collection helper function used to indicate the end of an array of structs. The
@@ -2029,11 +2045,12 @@ public:
 	 * Realtime garbage collection helper function used to indicate the beginning of a fixed array.
 	 * All tokens issues between Begin and End will be replayed Count times.
 	 *
-	 * @param Offset	offset at which fixed array starts
-	 * @param Stride	Stride of array element, e.g. sizeof(struct) or sizeof(UObject*)
-	 * @param Count		fixed array count
+	 * @param Offset	Offset at which fixed array starts.
+	 * @param DebugName	DebugName for this objects token. Only used in non-shipping builds.
+	 * @param Stride	Stride of array element, e.g. sizeof(struct) or sizeof(UObject*).
+	 * @param Count		Fixed array count.
 	 */
-	void EmitFixedArrayBegin(int32 Offset, int32 Stride, int32 Count);
+	void EmitFixedArrayBegin(int32 Offset, const FName& DebugName, int32 Stride, int32 Count);
 	
 	/**
 	 * Realtime garbage collection helper function used to indicated the end of a fixed array.
diff --git a/Engine/Source/Runtime/CoreUObject/Public/UObject/GarbageCollection.h b/Engine/Source/Runtime/CoreUObject/Public/UObject/GarbageCollection.h
--- a/Engine/Source/Runtime/CoreUObject/Public/UObject/GarbageCollection.h
+++ b/Engine/Source/Runtime/CoreUObject/Public/UObject/GarbageCollection.h
@@ -132,6 +132,55 @@ struct FGCSkipInfo
 	};
 };
 
+#if !(UE_BUILD_TEST || UE_BUILD_SHIPPING)
+/**
+ * Stores debug info about the token.
+ */
+struct FTokenInfo
+{
+	/* Token offset. */
+	int32 Offset;
+	/* Token debug name. */
+	FName Name;
+};
+
+/**
+ * Stores info about GC token stream debug map.
+ */
+class FGCDebugReferenceTokenMap
+{
+public:
+	/**
+	 * Maps token.
+	 *
+	 * @param DebugName Field name.
+	 * @param Offset Field offset.
+	 * @param TokenIndex Token index.
+	 */
+	void MapToken(const FName& DebugName, int32 Offset, int32 TokenIndex);
+
+	/**
+	 * Prepends this token map with its super class token map.
+	 *
+	 * @param SuperClass Super class to obtain its token map.
+	 */
+	void PrependWithSuperClass(const UClass& SuperClass);
+
+	/**
+	 * Returns pointer to token info object at given index.
+	 *
+	 * @param TokenIndex Index of the token.
+	 *
+	 * @returns FTokenInfo object.
+	 */
+	const FTokenInfo& GetTokenInfo(int32 TokenIndex) const;
+
+private:
+	/* Token map. */
+	TArray<FTokenInfo> TokenMap;
+};
+#endif
+
 /**
  * Reference token stream class. Used for creating and parsing stream of object references.
  */
@@ -160,6 +209,16 @@ struct FGCReferenceTokenStream
 		Tokens.Empty();
 	}
 
+	/**
+	 * Returns the size ofthe reference token stream.
+	 *
+	 * @returns Size of the stream.
+	 */
+	int32 Size() const
+	{
+		return Tokens.Num();
+	}
+
 	/**
 	 * return true if this is empty
 	 */
@@ -176,11 +235,13 @@ struct FGCReferenceTokenStream
 	void PrependStream( const FGCReferenceTokenStream& Other );
 
 	/**
-	 * Emit reference info
+	 * Emit reference info.
+	 *
+	 * @param ReferenceInfo	reference info to emit.
 	 *
-	 * @param ReferenceInfo	reference info to emit
+	 * @return Index of the reference info in the token stream.
 	 */
-	void EmitReferenceInfo( FGCReferenceInfo ReferenceInfo );
+	int32 EmitReferenceInfo(FGCReferenceInfo ReferenceInfo);
 
 	/**
 	 * Emit placeholder for aray skip index, updated in UpdateSkipIndexPlaceholder
diff --git a/Engine/Source/Runtime/CoreUObject/Public/UObject/UnrealType.h b/Engine/Source/Runtime/CoreUObject/Public/UObject/UnrealType.h
--- a/Engine/Source/Runtime/CoreUObject/Public/UObject/UnrealType.h
+++ b/Engine/Source/Runtime/CoreUObject/Public/UObject/UnrealType.h
@@ -715,7 +715,7 @@ public:
 	 * Emits tokens used by realtime garbage collection code to passed in ReferenceTokenStream. The offset emitted is relative
 	 * to the passed in BaseOffset which is used by e.g. arrays of structs.
 	 */
-	virtual void EmitReferenceInfo( FGCReferenceTokenStream* ReferenceTokenStream, int32 BaseOffset );
+	virtual void EmitReferenceInfo(UClass& OwnerClass, int32 BaseOffset);
 
 	FORCEINLINE int32 GetSize() const
 	{
@@ -1766,7 +1766,7 @@ class COREUOBJECT_API UObjectProperty : public TUObjectPropertyBase<UObject*>
 
 	// UProperty interface
 	virtual void SerializeItem( FArchive& Ar, void* Value, int32 MaxReadBytes, void const* Defaults ) const OVERRIDE;
-	virtual void EmitReferenceInfo( FGCReferenceTokenStream* ReferenceTokenStream, int32 BaseOffset ) OVERRIDE;
+	virtual void EmitReferenceInfo(UClass& OwnerClass, int32 BaseOffset) OVERRIDE;
 	// End of UProperty interface
 	// UObjectPropertyBase interface
 	virtual UObject* GetObjectPropertyValue(const void* PropertyValueAddress) const OVERRIDE
@@ -2011,7 +2011,7 @@ public:
 
 	// UObject interface
 	virtual void Serialize( FArchive& Ar ) OVERRIDE;
-	virtual void EmitReferenceInfo( FGCReferenceTokenStream* ReferenceTokenStream, int32 BaseOffset ) OVERRIDE;
+	virtual void EmitReferenceInfo(UClass& OwnerClass, int32 BaseOffset) OVERRIDE;
 	// End of UObject interface
 };
 
@@ -2126,7 +2126,7 @@ public:
 	virtual void InstanceSubobjects( void* Data, void const* DefaultData, UObject* Owner, struct FObjectInstancingGraph* InstanceGraph ) OVERRIDE;
 	virtual bool ContainsObjectReference() const OVERRIDE;
 	virtual bool ContainsWeakObjectReference() const OVERRIDE;
-	virtual void EmitReferenceInfo( FGCReferenceTokenStream* ReferenceTokenStream, int32 BaseOffset ) OVERRIDE;
+	virtual void EmitReferenceInfo(UClass& OwnerClass, int32 BaseOffset) OVERRIDE;
 	virtual bool SameType(const UProperty* Other) const OVERRIDE;
 	// End of UProperty interface
 };
@@ -2484,7 +2484,7 @@ public:
 	virtual int32 GetMinAlignment() const OVERRIDE;
 	virtual bool ContainsObjectReference() const OVERRIDE;
 	virtual bool ContainsWeakObjectReference() const OVERRIDE;
-	virtual void EmitReferenceInfo( FGCReferenceTokenStream* ReferenceTokenStream, int32 BaseOffset ) OVERRIDE;
+	virtual void EmitReferenceInfo(UClass& OwnerClass, int32 BaseOffset) OVERRIDE;
 	virtual bool SameType(const UProperty* Other) const OVERRIDE;
 	// End of UProperty interface
 
diff --git a/Engine/Source/Runtime/Engine/Private/Model.cpp b/Engine/Source/Runtime/Engine/Private/Model.cpp
--- a/Engine/Source/Runtime/Engine/Private/Model.cpp
+++ b/Engine/Source/Runtime/Engine/Private/Model.cpp
@@ -423,10 +423,10 @@ SIZE_T UModel::GetResourceSize(EResourceSizeMode::Type Mode)
 IMPLEMENT_INTRINSIC_CLASS(UModel, ENGINE_API, UObject, CORE_API, 
 	{
 		Class->ClassAddReferencedObjects = &UModel::AddReferencedObjects;
-		Class->EmitObjectReference( STRUCT_OFFSET( UModel, Polys ) );
-		const uint32 SkipIndexIndex = Class->EmitStructArrayBegin( STRUCT_OFFSET( UModel, Surfs ), sizeof(FBspSurf) );
-		Class->EmitObjectReference( STRUCT_OFFSET( FBspSurf, Material ) );
-		Class->EmitObjectReference( STRUCT_OFFSET( FBspSurf, Actor ) );
+		Class->EmitObjectReference(STRUCT_OFFSET(UModel, Polys), TEXT("Polys"));
+		const uint32 SkipIndexIndex = Class->EmitStructArrayBegin(STRUCT_OFFSET(UModel, Surfs), TEXT("Surfs"), sizeof(FBspSurf));
+		Class->EmitObjectReference(STRUCT_OFFSET(FBspSurf, Material), TEXT("Material"));
+		Class->EmitObjectReference(STRUCT_OFFSET(FBspSurf, Actor), TEXT("Actor"));
 		Class->EmitStructArrayEnd( SkipIndexIndex );
 	}
 );
@@ -436,9 +436,9 @@ IMPLEMENT_INTRINSIC_CLASS(UModel, ENGINE_API, UObject, CORE_API,
 IMPLEMENT_INTRINSIC_CLASS(UModel, ENGINE_API, UObject, CORE_API, 
 	{
 		Class->ClassAddReferencedObjects = &UModel::AddReferencedObjects;
-		const uint32 SkipIndexIndex = Class->EmitStructArrayBegin( STRUCT_OFFSET( UModel, Surfs ), sizeof(FBspSurf) );
-		Class->EmitObjectReference( STRUCT_OFFSET( FBspSurf, Material ) );
-		Class->EmitObjectReference( STRUCT_OFFSET( FBspSurf, Actor ) );
+		const uint32 SkipIndexIndex = Class->EmitStructArrayBegin(STRUCT_OFFSET(UModel, Surfs), TEXT("Surfs"), sizeof(FBspSurf));
+		Class->EmitObjectReference(STRUCT_OFFSET(FBspSurf, Material), TEXT("Material"));
+		Class->EmitObjectReference(STRUCT_OFFSET(FBspSurf, Actor), TEXT("Actor"));
 		Class->EmitStructArrayEnd( SkipIndexIndex );
 	}
 );
