CIS fix for Mac.

diff --git a/Engine/Source/Runtime/CoreUObject/Private/UObject/GarbageCollection.cpp b/Engine/Source/Runtime/CoreUObject/Private/UObject/GarbageCollection.cpp
--- a/Engine/Source/Runtime/CoreUObject/Private/UObject/GarbageCollection.cpp
+++ b/Engine/Source/Runtime/CoreUObject/Private/UObject/GarbageCollection.cpp
@@ -1789,7 +1789,10 @@ void FGCDebugReferenceTokenMap::MapToken(const FName& DebugName, int32 Offset, i
 	{
 		TokenMap.AddZeroed(TokenIndex - TokenMap.Num() + 1);
 
-		TokenMap[TokenIndex] = { Offset, DebugName };
+		auto& TokenInfo = TokenMap[TokenIndex];
+
+		TokenInfo.Offset = Offset;
+		TokenInfo.Name = DebugName;
 	}
 	else
 	{
