Cosmetic changes for GC token stream debug output.

diff --git a/Engine/Source/Runtime/CoreUObject/Private/UObject/GarbageCollection.cpp b/Engine/Source/Runtime/CoreUObject/Private/UObject/GarbageCollection.cpp
--- a/Engine/Source/Runtime/CoreUObject/Private/UObject/GarbageCollection.cpp
+++ b/Engine/Source/Runtime/CoreUObject/Private/UObject/GarbageCollection.cpp
@@ -231,19 +231,19 @@ static FORCEINLINE void HandleTokenStreamObjectReference( TArray<UObject*>& Obje
 		if (UClass *Class = (ReferencingObject ? ReferencingObject->GetClass() : nullptr))
 		{
 			auto& TokenInfo = Class->DebugTokenMap.GetTokenInfo(TokenIndex);
-			TokenDebugInfo = FString::Printf(TEXT("ReferencingObjectClass: %s, Name: %s, Offset: %d"),
+			TokenDebugInfo = FString::Printf(TEXT("ReferencingObjectClass: %s, Property Name: %s, Offset: %d"),
 				*Class->GetFullName(), *TokenInfo.Name.GetPlainNameString(), TokenInfo.Offset);
 		}
 		else
 		{
 			// This means this objects is most likely being referenced by AddReferencedObjects
-			TokenDebugInfo = FString::Printf(TEXT("Native Reference, TokenIndex: %d"), TokenIndex);
+			TokenDebugInfo = TEXT("Native Reference");
 		}
 
-		UE_LOG(LogGarbage, Fatal, TEXT("Invalid object in GC: 0x%016llx, ReferencingObject: %s, %s"),
+		UE_LOG(LogGarbage, Fatal, TEXT("Invalid object in GC: 0x%016llx, ReferencingObject: %s, %s, TokenIndex: %d"),
 			(int64)(PTRINT)Object,
 			ReferencingObject ? *ReferencingObject->GetFullName() : TEXT("NULL"),
-			*TokenDebugInfo);
+			*TokenDebugInfo, TokenIndex);
 	}
 #endif
 	HandleObjectReference(ObjectsToSerialize, ReferencingObject, Object, bAllowReferenceElimination);
