UE4 Refactoring. Changed OVERRIDE and FINAL macros to keywords override and final.

diff --git a/Engine/Source/Runtime/Core/Public/Linux/LinuxApplication.h b/Engine/Source/Runtime/Core/Public/Linux/LinuxApplication.h
--- a/Engine/Source/Runtime/Core/Public/Linux/LinuxApplication.h
+++ b/Engine/Source/Runtime/Core/Public/Linux/LinuxApplication.h
@@ -22,36 +22,36 @@ public:
 public:	
 	virtual ~FLinuxApplication();
 	
-	virtual void DestroyApplication() OVERRIDE;
+	virtual void DestroyApplication() override;
 
 public:
-	virtual void SetMessageHandler( const TSharedRef< class FGenericApplicationMessageHandler >& InMessageHandler ) OVERRIDE;
+	virtual void SetMessageHandler( const TSharedRef< class FGenericApplicationMessageHandler >& InMessageHandler ) override;
 
-	virtual void PollGameDeviceState( const float TimeDelta ) OVERRIDE;
+	virtual void PollGameDeviceState( const float TimeDelta ) override;
 
-	virtual void PumpMessages( const float TimeDelta ) OVERRIDE;
+	virtual void PumpMessages( const float TimeDelta ) override;
 
-	virtual void ProcessDeferredEvents( const float TimeDelta ) OVERRIDE;
+	virtual void ProcessDeferredEvents( const float TimeDelta ) override;
 
-	virtual TSharedRef< FGenericWindow > MakeWindow() OVERRIDE;
+	virtual TSharedRef< FGenericWindow > MakeWindow() override;
 
-	virtual void InitializeWindow( const TSharedRef< FGenericWindow >& Window, const TSharedRef< FGenericWindowDefinition >& InDefinition, const TSharedPtr< FGenericWindow >& InParent, const bool bShowImmediately ) OVERRIDE;
+	virtual void InitializeWindow( const TSharedRef< FGenericWindow >& Window, const TSharedRef< FGenericWindowDefinition >& InDefinition, const TSharedPtr< FGenericWindow >& InParent, const bool bShowImmediately ) override;
 
-	virtual void SetCapture( const TSharedPtr< FGenericWindow >& InWindow ) OVERRIDE;
+	virtual void SetCapture( const TSharedPtr< FGenericWindow >& InWindow ) override;
 
-	virtual void* GetCapture( void ) const OVERRIDE;
+	virtual void* GetCapture( void ) const override;
 
-	virtual void SetHighPrecisionMouseMode( const bool Enable, const TSharedPtr< FGenericWindow >& InWindow ) OVERRIDE;
+	virtual void SetHighPrecisionMouseMode( const bool Enable, const TSharedPtr< FGenericWindow >& InWindow ) override;
 
-	virtual bool IsUsingHighPrecisionMouseMode() const OVERRIDE { return bUsingHighPrecisionMouseInput; }
+	virtual bool IsUsingHighPrecisionMouseMode() const override { return bUsingHighPrecisionMouseInput; }
 	//	
-	virtual FModifierKeysState GetModifierKeys() const OVERRIDE;
+	virtual FModifierKeysState GetModifierKeys() const override;
 	//	
-	virtual FPlatformRect GetWorkArea( const FPlatformRect& CurrentWindow ) const OVERRIDE;
+	virtual FPlatformRect GetWorkArea( const FPlatformRect& CurrentWindow ) const override;
 	//	X
-	virtual bool TryCalculatePopupWindowPosition( const FPlatformRect& InAnchor, const FVector2D& InSize, const EPopUpOrientation::Type Orientation, /*OUT*/ FVector2D* const CalculatedPopUpPosition ) const OVERRIDE;
+	virtual bool TryCalculatePopupWindowPosition( const FPlatformRect& InAnchor, const FVector2D& InSize, const EPopUpOrientation::Type Orientation, /*OUT*/ FVector2D* const CalculatedPopUpPosition ) const override;
 	//	
-	virtual void GetDisplayMetrics( FDisplayMetrics& OutDisplayMetrics ) const OVERRIDE;
+	virtual void GetDisplayMetrics( FDisplayMetrics& OutDisplayMetrics ) const override;
 
 	void AddPendingEvent( SDL_Event event );
 
diff --git a/Engine/Source/Runtime/Core/Public/Linux/LinuxCursor.h b/Engine/Source/Runtime/Core/Public/Linux/LinuxCursor.h
--- a/Engine/Source/Runtime/Core/Public/Linux/LinuxCursor.h
+++ b/Engine/Source/Runtime/Core/Public/Linux/LinuxCursor.h
@@ -16,17 +16,17 @@ public:
 
 	virtual ~FLinuxCursor();
 
-	virtual FVector2D GetPosition() const OVERRIDE;
+	virtual FVector2D GetPosition() const override;
 
-	virtual void SetPosition( const int32 X, const int32 Y ) OVERRIDE;
+	virtual void SetPosition( const int32 X, const int32 Y ) override;
 
-	virtual void SetType( const EMouseCursor::Type InNewCursor ) OVERRIDE;
+	virtual void SetType( const EMouseCursor::Type InNewCursor ) override;
 
-	virtual void GetSize( int32& Width, int32& Height ) const OVERRIDE;
+	virtual void GetSize( int32& Width, int32& Height ) const override;
 
-	virtual void Show( bool bShow ) OVERRIDE;
+	virtual void Show( bool bShow ) override;
 
-	virtual void Lock( const RECT* const Bounds ) OVERRIDE;
+	virtual void Lock( const RECT* const Bounds ) override;
 
 public:
 
diff --git a/Engine/Source/Runtime/Core/Public/Linux/LinuxPlatformFile.h b/Engine/Source/Runtime/Core/Public/Linux/LinuxPlatformFile.h
--- a/Engine/Source/Runtime/Core/Public/Linux/LinuxPlatformFile.h
+++ b/Engine/Source/Runtime/Core/Public/Linux/LinuxPlatformFile.h
@@ -15,24 +15,24 @@ protected:
 	virtual FString NormalizeFilename(const TCHAR* Filename);
 	virtual FString NormalizeDirectory(const TCHAR* Directory);
 public:
-	virtual bool FileExists(const TCHAR* Filename) OVERRIDE;
-	virtual int64 FileSize(const TCHAR* Filename) OVERRIDE;
-	virtual bool DeleteFile(const TCHAR* Filename) OVERRIDE;
-	virtual bool IsReadOnly(const TCHAR* Filename) OVERRIDE;
-	virtual bool MoveFile(const TCHAR* To, const TCHAR* From) OVERRIDE;
-	virtual bool SetReadOnly(const TCHAR* Filename, bool bNewReadOnlyValue) OVERRIDE;
+	virtual bool FileExists(const TCHAR* Filename) override;
+	virtual int64 FileSize(const TCHAR* Filename) override;
+	virtual bool DeleteFile(const TCHAR* Filename) override;
+	virtual bool IsReadOnly(const TCHAR* Filename) override;
+	virtual bool MoveFile(const TCHAR* To, const TCHAR* From) override;
+	virtual bool SetReadOnly(const TCHAR* Filename, bool bNewReadOnlyValue) override;
 
 
-	virtual FDateTime GetTimeStamp(const TCHAR* Filename) OVERRIDE;
+	virtual FDateTime GetTimeStamp(const TCHAR* Filename) override;
 
-	virtual void SetTimeStamp(const TCHAR* Filename, const FDateTime DateTime) OVERRIDE;
+	virtual void SetTimeStamp(const TCHAR* Filename, const FDateTime DateTime) override;
 
-	virtual FDateTime GetAccessTimeStamp(const TCHAR* Filename) OVERRIDE;
+	virtual FDateTime GetAccessTimeStamp(const TCHAR* Filename) override;
 
-	virtual IFileHandle* OpenRead(const TCHAR* Filename) OVERRIDE;
-	virtual IFileHandle* OpenWrite(const TCHAR* Filename, bool bAppend = false, bool bAllowRead = false) OVERRIDE;
-	virtual bool DirectoryExists(const TCHAR* Directory) OVERRIDE;
-	virtual bool CreateDirectory(const TCHAR* Directory) OVERRIDE;
-	virtual bool DeleteDirectory(const TCHAR* Directory) OVERRIDE;
+	virtual IFileHandle* OpenRead(const TCHAR* Filename) override;
+	virtual IFileHandle* OpenWrite(const TCHAR* Filename, bool bAppend = false, bool bAllowRead = false) override;
+	virtual bool DirectoryExists(const TCHAR* Directory) override;
+	virtual bool CreateDirectory(const TCHAR* Directory) override;
+	virtual bool DeleteDirectory(const TCHAR* Directory) override;
 	bool IterateDirectory(const TCHAR* Directory, FDirectoryVisitor& Visitor);
 };
diff --git a/Engine/Source/Runtime/Core/Public/Linux/LinuxWindow.h b/Engine/Source/Runtime/Core/Public/Linux/LinuxWindow.h
--- a/Engine/Source/Runtime/Core/Public/Linux/LinuxWindow.h
+++ b/Engine/Source/Runtime/Core/Public/Linux/LinuxWindow.h
@@ -43,59 +43,59 @@ public:
 	void AdjustWindowRegion( int32 Width, int32 Height );
 
 public:
-	virtual void ReshapeWindow( int32 X, int32 Y, int32 Width, int32 Height ) OVERRIDE;
+	virtual void ReshapeWindow( int32 X, int32 Y, int32 Width, int32 Height ) override;
 
-	virtual void* GetOSWindowHandle() const  OVERRIDE { return HWnd; }
+	virtual void* GetOSWindowHandle() const  override { return HWnd; }
 
 	//	not finished
-	virtual bool GetFullScreenInfo( int32& X, int32& Y, int32& Width, int32& Height ) const OVERRIDE;
+	virtual bool GetFullScreenInfo( int32& X, int32& Y, int32& Width, int32& Height ) const override;
 
-	virtual void MoveWindowTo ( int32 X, int32 Y ) OVERRIDE;
+	virtual void MoveWindowTo ( int32 X, int32 Y ) override;
 
-	virtual void BringToFront( bool bForce = false ) OVERRIDE;
+	virtual void BringToFront( bool bForce = false ) override;
 
-	virtual void Destroy() OVERRIDE;
+	virtual void Destroy() override;
 
-	virtual void Minimize() OVERRIDE;
+	virtual void Minimize() override;
 
-	virtual void Maximize() OVERRIDE;
+	virtual void Maximize() override;
 
-	virtual void Restore() OVERRIDE;
+	virtual void Restore() override;
 
-	virtual void Show() OVERRIDE;
+	virtual void Show() override;
 
-	virtual void Hide() OVERRIDE;
+	virtual void Hide() override;
 
-	virtual void SetWindowMode( EWindowMode::Type NewWindowMode ) OVERRIDE;
+	virtual void SetWindowMode( EWindowMode::Type NewWindowMode ) override;
 
-	virtual EWindowMode::Type GetWindowMode() const OVERRIDE { return WindowMode; } 
+	virtual EWindowMode::Type GetWindowMode() const override { return WindowMode; } 
 
-	virtual bool IsMaximized() const OVERRIDE;
+	virtual bool IsMaximized() const override;
 
-	virtual bool IsVisible() const OVERRIDE;
+	virtual bool IsVisible() const override;
 
-	virtual bool GetRestoredDimensions(int32& X, int32& Y, int32& Width, int32& Height) OVERRIDE;
+	virtual bool GetRestoredDimensions(int32& X, int32& Y, int32& Width, int32& Height) override;
 
-	virtual void SetWindowFocus() OVERRIDE;
+	virtual void SetWindowFocus() override;
 
-	virtual void SetOpacity( const float InOpacity ) OVERRIDE;
+	virtual void SetOpacity( const float InOpacity ) override;
 
-	virtual void Enable( bool bEnable ) OVERRIDE;
+	virtual void Enable( bool bEnable ) override;
 
-	virtual bool IsPointInWindow( int32 X, int32 Y ) const OVERRIDE;
+	virtual bool IsPointInWindow( int32 X, int32 Y ) const override;
 
-	virtual int32 GetWindowBorderSize() const OVERRIDE;
+	virtual int32 GetWindowBorderSize() const override;
 
-	virtual bool IsForegroundWindow() const OVERRIDE;
+	virtual bool IsForegroundWindow() const override;
 	/**
 	 * Sets the window text - usually the title but can also be text content for things like controls
 	 *
 	 * @param Text	The window's title or content text
 	 */
-	virtual void SetText(const TCHAR* const Text) OVERRIDE;
+	virtual void SetText(const TCHAR* const Text) override;
 
 	/** @return	Gives the native window a chance to adjust our stored window size before we cache it off */
-	virtual void AdjustCachedSize( FVector2D& Size ) const OVERRIDE;
+	virtual void AdjustCachedSize( FVector2D& Size ) const override;
 
 
 #if 0
@@ -104,25 +104,25 @@ public:
 	 */
 
 	/** IUnknown : attempt to get a reference to this instance as a pointer  */
-	HRESULT STDCALL QueryInterface( REFIID iid, void ** ppvObject ) OVERRIDE;
+	HRESULT STDCALL QueryInterface( REFIID iid, void ** ppvObject ) override;
 
 	/** IUnknown : Invoked by OLE when a new reference to this instance */
-	ULONG STDCALL AddRef( void ) OVERRIDE;
+	ULONG STDCALL AddRef( void ) override;
 
 	/** IUnknown : Invoked by OLE when a reference to this instance is released */
-	ULONG STDCALL Release( void ) OVERRIDE;
+	ULONG STDCALL Release( void ) override;
 
 	/** IDropTarget : Invoked by OLE when someone drag data our window initially */
-	virtual HRESULT STDCALL DragEnter( __RPC__in_opt IDataObject *DataObjectPointer, ::DWORD KeyState, POINTL CursorPosition, __RPC__inout ::DWORD *CursorEffect) OVERRIDE;
+	virtual HRESULT STDCALL DragEnter( __RPC__in_opt IDataObject *DataObjectPointer, ::DWORD KeyState, POINTL CursorPosition, __RPC__inout ::DWORD *CursorEffect) override;
 
 	/** IDropTarget : Invoked by OLE when someone drag data over our window */
-	virtual HRESULT STDCALL DragOver( ::DWORD KeyState, POINTL CursorPosition, __RPC__inout ::DWORD *CursorEffect) OVERRIDE;
+	virtual HRESULT STDCALL DragOver( ::DWORD KeyState, POINTL CursorPosition, __RPC__inout ::DWORD *CursorEffect) override;
 
 	/** IDropTarget : Invoked by OLE when someone drag data and it exits our window or hits ESC to cancel the drag and drop. */
-	virtual HRESULT STDCALL DragLeave( void ) OVERRIDE;
+	virtual HRESULT STDCALL DragLeave( void ) override;
 
 	/** IDropTarget : Invoked by OLE when someone drag data our window initially */
-	virtual HRESULT STDCALL Drop( __RPC__in_opt IDataObject *DataObjectPointer, ::DWORD KeyState, POINTL CursorPosition, __RPC__inout ::DWORD *CursorEffect) OVERRIDE;
+	virtual HRESULT STDCALL Drop( __RPC__in_opt IDataObject *DataObjectPointer, ::DWORD KeyState, POINTL CursorPosition, __RPC__inout ::DWORD *CursorEffect) override;
 #endif
 
 private:
diff --git a/Engine/Source/Runtime/Core/Public/Mac/MacApplication.h b/Engine/Source/Runtime/Core/Public/Mac/MacApplication.h
--- a/Engine/Source/Runtime/Core/Public/Mac/MacApplication.h
+++ b/Engine/Source/Runtime/Core/Public/Mac/MacApplication.h
@@ -25,40 +25,40 @@ public:
 
 public:
 
-	virtual void SetMessageHandler( const TSharedRef< class FGenericApplicationMessageHandler >& InMessageHandler ) OVERRIDE;
+	virtual void SetMessageHandler( const TSharedRef< class FGenericApplicationMessageHandler >& InMessageHandler ) override;
 
-	virtual void PollGameDeviceState( const float TimeDelta ) OVERRIDE;
+	virtual void PollGameDeviceState( const float TimeDelta ) override;
 
-	virtual void PumpMessages( const float TimeDelta ) OVERRIDE;
+	virtual void PumpMessages( const float TimeDelta ) override;
 
-	virtual void ProcessDeferredEvents( const float TimeDelta ) OVERRIDE;
+	virtual void ProcessDeferredEvents( const float TimeDelta ) override;
 
-	virtual TSharedRef< FGenericWindow > MakeWindow() OVERRIDE;
+	virtual TSharedRef< FGenericWindow > MakeWindow() override;
 
-	virtual void InitializeWindow( const TSharedRef< FGenericWindow >& Window, const TSharedRef< FGenericWindowDefinition >& InDefinition, const TSharedPtr< FGenericWindow >& InParent, const bool bShowImmediately ) OVERRIDE;
+	virtual void InitializeWindow( const TSharedRef< FGenericWindow >& Window, const TSharedRef< FGenericWindowDefinition >& InDefinition, const TSharedPtr< FGenericWindow >& InParent, const bool bShowImmediately ) override;
 
-	virtual void SetCapture( const TSharedPtr< FGenericWindow >& InWindow ) OVERRIDE;
+	virtual void SetCapture( const TSharedPtr< FGenericWindow >& InWindow ) override;
 
-	virtual void* GetCapture( void ) const OVERRIDE;
+	virtual void* GetCapture( void ) const override;
 
-	virtual void SetHighPrecisionMouseMode( const bool Enable, const TSharedPtr< FGenericWindow >& InWindow ) OVERRIDE;
+	virtual void SetHighPrecisionMouseMode( const bool Enable, const TSharedPtr< FGenericWindow >& InWindow ) override;
 
-	virtual bool IsUsingHighPrecisionMouseMode() const OVERRIDE { return bUsingHighPrecisionMouseInput; }
+	virtual bool IsUsingHighPrecisionMouseMode() const override { return bUsingHighPrecisionMouseInput; }
 
-	virtual bool IsUsingTrackpad() const OVERRIDE { return bUsingTrackpad; }
+	virtual bool IsUsingTrackpad() const override { return bUsingTrackpad; }
 
-	virtual FModifierKeysState GetModifierKeys() const OVERRIDE;
+	virtual FModifierKeysState GetModifierKeys() const override;
 
-	virtual FPlatformRect GetWorkArea( const FPlatformRect& CurrentWindow ) const OVERRIDE;
+	virtual FPlatformRect GetWorkArea( const FPlatformRect& CurrentWindow ) const override;
 
-	virtual void GetDisplayMetrics( FDisplayMetrics& OutDisplayMetrics ) const OVERRIDE;
+	virtual void GetDisplayMetrics( FDisplayMetrics& OutDisplayMetrics ) const override;
 
-	virtual EWindowTitleAlignment::Type GetWindowTitleAlignment() const OVERRIDE
+	virtual EWindowTitleAlignment::Type GetWindowTitleAlignment() const override
 	{
 		return EWindowTitleAlignment::Center;
 	}
 
-	virtual ITextInputMethodSystem *GetTextInputMethodSystem() OVERRIDE
+	virtual ITextInputMethodSystem *GetTextInputMethodSystem() override
 	{
 		return TextInputMethodSystem.Get();
 	}
@@ -78,7 +78,7 @@ public:
 	void UseMouseCaptureWindow( bool bUseMouseCaptureWindow );
 
 #if WITH_EDITOR
-    virtual void SendAnalytics(IAnalyticsProvider* Provider) OVERRIDE;
+    virtual void SendAnalytics(IAnalyticsProvider* Provider) override;
 
 	void SetIsUsingTrackpad(bool bInUsingTrackpad)
 	{
diff --git a/Engine/Source/Runtime/Core/Public/Mac/MacPlatformRunnableThread.h b/Engine/Source/Runtime/Core/Public/Mac/MacPlatformRunnableThread.h
--- a/Engine/Source/Runtime/Core/Public/Mac/MacPlatformRunnableThread.h
+++ b/Engine/Source/Runtime/Core/Public/Mac/MacPlatformRunnableThread.h
@@ -27,7 +27,7 @@ private:
 		pthread_setname_np(TCHAR_TO_ANSI(*ThreadName));
 	}
 
-	virtual int GetDefaultStackSize() OVERRIDE
+	virtual int GetDefaultStackSize() override
 	{
 		// default is 512 KB, we need more
 		return 1024 * 1024;
diff --git a/Engine/Source/Runtime/Core/Public/Mac/MacTextInputMethodSystem.h b/Engine/Source/Runtime/Core/Public/Mac/MacTextInputMethodSystem.h
--- a/Engine/Source/Runtime/Core/Public/Mac/MacTextInputMethodSystem.h
+++ b/Engine/Source/Runtime/Core/Public/Mac/MacTextInputMethodSystem.h
@@ -100,10 +100,10 @@ public:
 	void Terminate();
 	
 	// ITextInputMethodSystem Interface Begin
-	virtual TSharedPtr<ITextInputMethodChangeNotifier> RegisterContext(const TSharedRef<ITextInputMethodContext>& Context) OVERRIDE;
-	virtual void UnregisterContext(const TSharedRef<ITextInputMethodContext>& Context) OVERRIDE;
-	virtual void ActivateContext(const TSharedRef<ITextInputMethodContext>& Context) OVERRIDE;
-	virtual void DeactivateContext(const TSharedRef<ITextInputMethodContext>& Context) OVERRIDE;
+	virtual TSharedPtr<ITextInputMethodChangeNotifier> RegisterContext(const TSharedRef<ITextInputMethodContext>& Context) override;
+	virtual void UnregisterContext(const TSharedRef<ITextInputMethodContext>& Context) override;
+	virtual void ActivateContext(const TSharedRef<ITextInputMethodContext>& Context) override;
+	virtual void DeactivateContext(const TSharedRef<ITextInputMethodContext>& Context) override;
 	// ITextInputMethodSystem Interface End
 	
 private:
diff --git a/Engine/Source/Runtime/Core/Public/Mac/MacWindow.h b/Engine/Source/Runtime/Core/Public/Mac/MacWindow.h
--- a/Engine/Source/Runtime/Core/Public/Mac/MacWindow.h
+++ b/Engine/Source/Runtime/Core/Public/Mac/MacWindow.h
@@ -141,55 +141,55 @@ public:
 
 public:
 
-	virtual void ReshapeWindow( int32 X, int32 Y, int32 Width, int32 Height ) OVERRIDE;
+	virtual void ReshapeWindow( int32 X, int32 Y, int32 Width, int32 Height ) override;
 
-	virtual bool GetFullScreenInfo( int32& X, int32& Y, int32& Width, int32& Height ) const OVERRIDE;
+	virtual bool GetFullScreenInfo( int32& X, int32& Y, int32& Width, int32& Height ) const override;
 
-	virtual void MoveWindowTo ( int32 X, int32 Y ) OVERRIDE;
+	virtual void MoveWindowTo ( int32 X, int32 Y ) override;
 
-	virtual void BringToFront( bool bForce = false ) OVERRIDE;
+	virtual void BringToFront( bool bForce = false ) override;
 
-	virtual void Destroy() OVERRIDE;
+	virtual void Destroy() override;
 
-	virtual void Minimize() OVERRIDE;
+	virtual void Minimize() override;
 
-	virtual void Maximize() OVERRIDE;
+	virtual void Maximize() override;
 
-	virtual void Restore() OVERRIDE;
+	virtual void Restore() override;
 
-	virtual void Show() OVERRIDE;
+	virtual void Show() override;
 
-	virtual void Hide() OVERRIDE;
+	virtual void Hide() override;
 
-	virtual void SetWindowMode( EWindowMode::Type NewWindowMode ) OVERRIDE;
+	virtual void SetWindowMode( EWindowMode::Type NewWindowMode ) override;
 
-	virtual EWindowMode::Type GetWindowMode() const OVERRIDE { return WindowMode; } 
+	virtual EWindowMode::Type GetWindowMode() const override { return WindowMode; } 
 
-	virtual bool IsMaximized() const OVERRIDE;
+	virtual bool IsMaximized() const override;
 
-	virtual bool IsMinimized() const OVERRIDE;
+	virtual bool IsMinimized() const override;
 
-	virtual bool IsVisible() const OVERRIDE;
+	virtual bool IsVisible() const override;
 
-	virtual bool GetRestoredDimensions(int32& X, int32& Y, int32& Width, int32& Height) OVERRIDE;
+	virtual bool GetRestoredDimensions(int32& X, int32& Y, int32& Width, int32& Height) override;
 
-	virtual void SetWindowFocus() OVERRIDE;
+	virtual void SetWindowFocus() override;
 
-	virtual void SetOpacity( const float InOpacity ) OVERRIDE;
+	virtual void SetOpacity( const float InOpacity ) override;
 
-	virtual void Enable( bool bEnable ) OVERRIDE;
+	virtual void Enable( bool bEnable ) override;
 
-	virtual bool IsPointInWindow( int32 X, int32 Y ) const OVERRIDE;
+	virtual bool IsPointInWindow( int32 X, int32 Y ) const override;
 
-	virtual int32 GetWindowBorderSize() const OVERRIDE;
+	virtual int32 GetWindowBorderSize() const override;
 
-	virtual void* GetOSWindowHandle() const  OVERRIDE { return WindowHandle; }
+	virtual void* GetOSWindowHandle() const  override { return WindowHandle; }
 
-	virtual bool IsForegroundWindow() const OVERRIDE;
+	virtual bool IsForegroundWindow() const override;
 
-	virtual void SetText(const TCHAR* const Text) OVERRIDE;
+	virtual void SetText(const TCHAR* const Text) override;
 
-	virtual void AdjustCachedSize( FVector2D& Size ) const OVERRIDE;
+	virtual void AdjustCachedSize( FVector2D& Size ) const override;
 
 	/**
 	 * Sets the window text - usually the title but can also be text content for things like controls
diff --git a/Engine/Source/Runtime/Core/Public/Misc/AutomationTest.h b/Engine/Source/Runtime/Core/Public/Misc/AutomationTest.h
--- a/Engine/Source/Runtime/Core/Public/Misc/AutomationTest.h
+++ b/Engine/Source/Runtime/Core/Public/Misc/AutomationTest.h
@@ -562,7 +562,7 @@ private:
 		 * @param	V		String to serialize within the context
 		 * @param	Event	Event associated with the string
 		 */
-		virtual void Serialize( const TCHAR* V, ELogVerbosity::Type Verbosity, const class FName& Category ) OVERRIDE;
+		virtual void Serialize( const TCHAR* V, ELogVerbosity::Type Verbosity, const class FName& Category ) override;
 
 		/**
 		 * Set the automation test associated with the feedback context. The automation test is where all warnings, errors, etc.
@@ -1001,7 +1001,7 @@ class CommandName : public IAutomationLatentCommand \
 	public: \
 	virtual ~CommandName() \
 		{} \
-		virtual bool Update() OVERRIDE; \
+		virtual bool Update() override; \
 }
 
 #define DEFINE_LATENT_AUTOMATION_COMMAND_ONE_PARAMETER(CommandName,ParamType,ParamName)	\
@@ -1013,7 +1013,7 @@ class CommandName : public IAutomationLatentCommand \
 		{} \
 		virtual ~CommandName() \
 		{} \
-		virtual bool Update() OVERRIDE; \
+		virtual bool Update() override; \
 	private: \
 	ParamType ParamName; \
 }
@@ -1024,7 +1024,7 @@ class ENGINE_API CommandName : public IAutomationLatentCommand \
 	public: \
 	virtual ~CommandName() \
 		{} \
-		virtual bool Update() OVERRIDE; \
+		virtual bool Update() override; \
 }
 
 #define DEFINE_ENGINE_LATENT_AUTOMATION_COMMAND_ONE_PARAMETER(CommandName,ParamType,ParamName)	\
@@ -1036,7 +1036,7 @@ class ENGINE_API CommandName : public IAutomationLatentCommand \
 		{} \
 		virtual ~CommandName() \
 		{} \
-		virtual bool Update() OVERRIDE; \
+		virtual bool Update() override; \
 	private: \
 	ParamType ParamName; \
 }
@@ -1055,7 +1055,7 @@ public: \
 	F##ClassDeclaration(int32 InRoleIndex) : RoleIndex(InRoleIndex) {} \
 	virtual ~F##ClassDeclaration() {} \
 	virtual uint32 GetRoleIndex() const { return RoleIndex; } \
-	virtual void Run() OVERRIDE 
+	virtual void Run() override 
 
 //close the class and add to the framework
 #define END_NETWORK_AUTOMATION_COMMAND(ClassDeclaration,InRoleIndex) }; \
@@ -1103,7 +1103,7 @@ public: \
 		virtual bool IsStressTest() const { return true; } \
 		virtual uint32 GetRequiredDeviceNum() const { return 1; } \
 	protected: \
-		virtual void GetTests(TArray<FString>& OutBeautifiedNames, TArray <FString>& OutTestCommands) const OVERRIDE; \
+		virtual void GetTests(TArray<FString>& OutBeautifiedNames, TArray <FString>& OutTestCommands) const override; \
 		virtual bool RunTest(const FString& Parameters); \
 		virtual FString GetBeautifiedTestName() const { return PrettyName; } \
 	};
diff --git a/Engine/Source/Runtime/Core/Public/Misc/DelegateFilter.h b/Engine/Source/Runtime/Core/Public/Misc/DelegateFilter.h
--- a/Engine/Source/Runtime/Core/Public/Misc/DelegateFilter.h
+++ b/Engine/Source/Runtime/Core/Public/Misc/DelegateFilter.h
@@ -23,7 +23,7 @@ public:
 	}
 
 	DECLARE_DERIVED_EVENT( TDelegateFilter, IFilter<ItemType>::FChangedEvent, FChangedEvent );
-	virtual FChangedEvent& OnChanged() OVERRIDE { return ChangedEvent; }
+	virtual FChangedEvent& OnChanged() override { return ChangedEvent; }
 
 	/** 
 	 * Returns whether the specified Item passes the Filter's restrictions 
@@ -31,7 +31,7 @@ public:
 	 *	@param	InItem	The Item is check 
 	 *	@return			Whether the specified Item passed the filter
 	 */
-	virtual bool PassesFilter( ItemType InItem ) const OVERRIDE
+	virtual bool PassesFilter( ItemType InItem ) const override
 	{
 		return Predicate.Execute( InItem );
 	}
diff --git a/Engine/Source/Runtime/Core/Public/Misc/TextFilter.h b/Engine/Source/Runtime/Core/Public/Misc/TextFilter.h
--- a/Engine/Source/Runtime/Core/Public/Misc/TextFilter.h
+++ b/Engine/Source/Runtime/Core/Public/Misc/TextFilter.h
@@ -26,7 +26,7 @@ public:
 	}
 
 	DECLARE_DERIVED_EVENT(TTextFilter, IFilter<ItemType>::FChangedEvent, FChangedEvent);
-	virtual FChangedEvent& OnChanged() OVERRIDE { return ChangedEvent; }
+	virtual FChangedEvent& OnChanged() override { return ChangedEvent; }
 
 	/** 
 	 * Returns whether the specified Item passes the Filter's text restrictions 
@@ -34,7 +34,7 @@ public:
 	 *	@param	InItem	The Item to check 
 	 *	@return			Whether the specified Item passed the filter
 	 */
-	virtual bool PassesFilter( ItemType InItem ) const OVERRIDE
+	virtual bool PassesFilter( ItemType InItem ) const override
 	{
 		if( FilterSearchTerms.Num() == 0 )
 		{
diff --git a/Engine/Source/Runtime/Core/Public/Misc/TokenizedMessage.h b/Engine/Source/Runtime/Core/Public/Misc/TokenizedMessage.h
--- a/Engine/Source/Runtime/Core/Public/Misc/TokenizedMessage.h
+++ b/Engine/Source/Runtime/Core/Public/Misc/TokenizedMessage.h
@@ -248,7 +248,7 @@ public:
 	}
 
 	/** Begin IMessageToken interface */
-	virtual EMessageToken::Type GetType() const OVERRIDE
+	virtual EMessageToken::Type GetType() const override
 	{
 		return EMessageToken::Text;
 	}
@@ -273,7 +273,7 @@ public:
 	}
 
 	/** Begin IMessageToken interface */
-	virtual EMessageToken::Type GetType() const OVERRIDE
+	virtual EMessageToken::Type GetType() const override
 	{
 		return EMessageToken::Image;
 	}
@@ -308,7 +308,7 @@ public:
 	}
 
 	/** Begin IMessageToken interface */
-	virtual EMessageToken::Type GetType() const OVERRIDE
+	virtual EMessageToken::Type GetType() const override
 	{
 		return EMessageToken::Severity;
 	}
@@ -343,7 +343,7 @@ public:
 	}
 
 	/** Begin IMessageToken interface */
-	virtual EMessageToken::Type GetType() const OVERRIDE
+	virtual EMessageToken::Type GetType() const override
 	{
 		return EMessageToken::URL;
 	}
@@ -393,7 +393,7 @@ public:
 	}
 
 	/** Begin IMessageToken interface */
-	virtual EMessageToken::Type GetType() const OVERRIDE
+	virtual EMessageToken::Type GetType() const override
 	{
 		return EMessageToken::AssetName;
 	}
@@ -442,7 +442,7 @@ public:
 	}
 
 	/** Begin IMessageToken interface */
-	virtual EMessageToken::Type GetType() const OVERRIDE
+	virtual EMessageToken::Type GetType() const override
 	{
 		return EMessageToken::Documentation;
 	}
diff --git a/Engine/Source/Runtime/Core/Public/Modules/ModuleManager.h b/Engine/Source/Runtime/Core/Public/Modules/ModuleManager.h
--- a/Engine/Source/Runtime/Core/Public/Modules/ModuleManager.h
+++ b/Engine/Source/Runtime/Core/Public/Modules/ModuleManager.h
@@ -609,7 +609,7 @@ private:
 	static bool CheckModuleCompatibility(const TCHAR *Filename);
 
 	// Begin FExec interface.
-	virtual bool Exec( UWorld* Inworld, const TCHAR* Cmd, FOutputDevice& Ar ) OVERRIDE;
+	virtual bool Exec( UWorld* Inworld, const TCHAR* Cmd, FOutputDevice& Ar ) override;
 	// End of FExec interface.
 
 public:
@@ -765,7 +765,7 @@ class FDefaultGameModuleImpl
 	 *
 	 * @return True for "gameplay modules", or false for engine code modules, plugins, etc.
 	 */
-	virtual bool IsGameModule() const OVERRIDE
+	virtual bool IsGameModule() const override
 	{
 		return true;
 	}
diff --git a/Engine/Source/Runtime/Core/Public/ProfilingDebugging/MallocProfiler.h b/Engine/Source/Runtime/Core/Public/ProfilingDebugging/MallocProfiler.h
--- a/Engine/Source/Runtime/Core/Public/ProfilingDebugging/MallocProfiler.h
+++ b/Engine/Source/Runtime/Core/Public/ProfilingDebugging/MallocProfiler.h
@@ -420,7 +420,7 @@ public:
 	 * @param Alignment		The alignment of a hypothetical allocation request
 	 * @return				Returns the usable size that the allocation request would return. In other words you can ask for this greater amount without using any more actual memory.
 	 */
-	virtual SIZE_T QuantizeSize( SIZE_T Size, uint32 Alignment ) OVERRIDE
+	virtual SIZE_T QuantizeSize( SIZE_T Size, uint32 Alignment ) override
 	{
 		return UsedMalloc->QuantizeSize(Size,Alignment); 
 	}
@@ -428,7 +428,7 @@ public:
 	/** 
 	 * Malloc
 	 */
-	virtual void* Malloc( SIZE_T Size, uint32 Alignment ) OVERRIDE
+	virtual void* Malloc( SIZE_T Size, uint32 Alignment ) override
 	{
 		FScopeLock Lock( &CriticalSection );
 		void* Ptr = UsedMalloc->Malloc( Size, Alignment );
@@ -440,7 +440,7 @@ public:
 	/** 
 	 * Realloc
 	 */
-	virtual void* Realloc( void* OldPtr, SIZE_T NewSize, uint32 Alignment ) OVERRIDE
+	virtual void* Realloc( void* OldPtr, SIZE_T NewSize, uint32 Alignment ) override
 	{
 		FScopeLock Lock( &CriticalSection );
 		void* NewPtr = UsedMalloc->Realloc( OldPtr, NewSize, Alignment );
@@ -452,7 +452,7 @@ public:
 	/** 
 	 * Free
 	 */
-	virtual void Free( void* Ptr ) OVERRIDE
+	virtual void Free( void* Ptr ) override
 	{
 		FScopeLock Lock( &CriticalSection );
 		UsedMalloc->Free( Ptr );
@@ -464,27 +464,27 @@ public:
 	 * Returns if the allocator is guaranteed to be thread-safe and therefore
 	 * doesn't need a unnecessary thread-safety wrapper around it.
 	 */
-	virtual bool IsInternallyThreadSafe() const OVERRIDE
+	virtual bool IsInternallyThreadSafe() const override
 	{ 
 		return true; 
 	}
 
 	/** Called once per frame, gathers and sets all memory allocator statistics into the corresponding stats. */
-	virtual void UpdateStats() OVERRIDE
+	virtual void UpdateStats() override
 	{
 		FScopeLock Lock( &CriticalSection );
 		UsedMalloc->UpdateStats();
 	}
 
 	/** Writes allocator stats from the last update into the specified destination. */
-	virtual void GetAllocatorStats( FGenericMemoryStats& out_Stats ) OVERRIDE
+	virtual void GetAllocatorStats( FGenericMemoryStats& out_Stats ) override
 	{
 		FScopeLock Lock( &CriticalSection );
 		UsedMalloc->GetAllocatorStats( out_Stats );
 	}
 
 	/** Dumps allocator stats to an output device. */
-	virtual void DumpAllocatorStats( class FOutputDevice& Ar ) OVERRIDE
+	virtual void DumpAllocatorStats( class FOutputDevice& Ar ) override
 	{
 		FScopeLock Lock( &CriticalSection );
 		UsedMalloc->DumpAllocatorStats( Ar );
@@ -493,7 +493,7 @@ public:
 	/**
 	 * Validates the allocator's heap
 	 */
-	virtual bool ValidateHeap() OVERRIDE
+	virtual bool ValidateHeap() override
 	{
 		FScopeLock Lock( &CriticalSection );
 		return( UsedMalloc->ValidateHeap() );
@@ -506,7 +506,7 @@ public:
 	* @param SizeOut - If possible, this value is set to the size of the passed in pointer
 	* @return true if succeeded
 	*/
-	virtual bool GetAllocationSize(void *Original, SIZE_T &SizeOut) OVERRIDE
+	virtual bool GetAllocationSize(void *Original, SIZE_T &SizeOut) override
 	{
 		FScopeLock Lock( &CriticalSection );
 		return UsedMalloc->GetAllocationSize(Original,SizeOut);
@@ -514,7 +514,7 @@ public:
 
 	
 	// Begin Exec Interface
-	virtual bool Exec( UWorld* InWorld, const TCHAR* Cmd, FOutputDevice& Ar ) OVERRIDE;
+	virtual bool Exec( UWorld* InWorld, const TCHAR* Cmd, FOutputDevice& Ar ) override;
 	// End Exec Interface
 
 	/** 
@@ -525,7 +525,7 @@ public:
 	bool HandleSnapshotMemoryCommand( const TCHAR* Cmd, FOutputDevice& Ar );
 	bool HandleSnapshotMemoryFrameCommand( const TCHAR* Cmd, FOutputDevice& Ar );
 
-	virtual const TCHAR * GetDescriptiveName() OVERRIDE
+	virtual const TCHAR * GetDescriptiveName() override
 	{ 
 		FScopeLock ScopeLock( &CriticalSection );
 		check(UsedMalloc);
diff --git a/Engine/Source/Runtime/Core/Public/Serialization/Json/JsonWriter.h b/Engine/Source/Runtime/Core/Public/Serialization/Json/JsonWriter.h
--- a/Engine/Source/Runtime/Core/Public/Serialization/Json/JsonWriter.h
+++ b/Engine/Source/Runtime/Core/Public/Serialization/Json/JsonWriter.h
@@ -391,7 +391,7 @@ public:
 		delete this->Stream;
 	}
 
-	virtual bool Close() OVERRIDE
+	virtual bool Close() override
 	{
 		FString Out;
 		for (int32 i = 0; i < Bytes.Num(); i+=sizeof(TCHAR))
diff --git a/Engine/Source/Runtime/Core/Public/Stats/StatsFile.h b/Engine/Source/Runtime/Core/Public/Stats/StatsFile.h
--- a/Engine/Source/Runtime/Core/Public/Stats/StatsFile.h
+++ b/Engine/Source/Runtime/Core/Public/Stats/StatsFile.h
@@ -463,7 +463,7 @@ protected:
 	*	Grabs a frame from the local FStatsThreadState and adds it to the output.
 	*	Called from the stats thread, but the data is saved using the the FAsyncWriteWorker.
 	*/
-	virtual void WriteFrame( int64 TargetFrame, bool bNeedFullMetadata = false ) OVERRIDE;
+	virtual void WriteFrame( int64 TargetFrame, bool bNeedFullMetadata = false ) override;
 
 	/** Write a stat packed into the specified archive. */
 	void WriteStatPacket( FArchive& Ar, /*const*/ FStatPacket& StatPacket )
diff --git a/Engine/Source/Runtime/Core/Public/WinRT/WinRTApplication.h b/Engine/Source/Runtime/Core/Public/WinRT/WinRTApplication.h
--- a/Engine/Source/Runtime/Core/Public/WinRT/WinRTApplication.h
+++ b/Engine/Source/Runtime/Core/Public/WinRT/WinRTApplication.h
@@ -8,19 +8,19 @@ class FWinRTCursor : public ICursor
 {
 public:
 
-	virtual FVector2D GetPosition() const OVERRIDE;
+	virtual FVector2D GetPosition() const override;
 
-	virtual void SetPosition( const int32 X, const int32 Y ) OVERRIDE;
+	virtual void SetPosition( const int32 X, const int32 Y ) override;
 
 	void UpdatePosition( const FVector2D& NewPosition );
 
-	virtual void SetType( const EMouseCursor::Type InNewCursor ) OVERRIDE;
+	virtual void SetType( const EMouseCursor::Type InNewCursor ) override;
 
-	virtual void GetSize( int32& Width, int32& Height ) const OVERRIDE;
+	virtual void GetSize( int32& Width, int32& Height ) const override;
 
-	virtual void Show( bool bShow ) OVERRIDE;
+	virtual void Show( bool bShow ) override;
 
-	virtual void Lock( const RECT* const Bounds ) OVERRIDE;
+	virtual void Lock( const RECT* const Bounds ) override;
 
 
 private:
@@ -41,15 +41,15 @@ public:
 
 	virtual ~FWinRTApplication() {}
 
-	virtual void PollGameDeviceState( const float TimeDelta ) OVERRIDE;
+	virtual void PollGameDeviceState( const float TimeDelta ) override;
 
-	virtual FPlatformRect GetWorkArea( const FPlatformRect& CurrentWindow ) const OVERRIDE;
+	virtual FPlatformRect GetWorkArea( const FPlatformRect& CurrentWindow ) const override;
 
-	virtual void GetDisplayMetrics( FDisplayMetrics& OutDisplayMetrics ) const OVERRIDE;
+	virtual void GetDisplayMetrics( FDisplayMetrics& OutDisplayMetrics ) const override;
 
 	TSharedRef< class FGenericApplicationMessageHandler > GetMessageHandler() const;
 
-	void SetMessageHandler( const TSharedRef< FGenericApplicationMessageHandler >& InMessageHandler ) OVERRIDE;
+	void SetMessageHandler( const TSharedRef< FGenericApplicationMessageHandler >& InMessageHandler ) override;
 
 	TSharedRef< class FWinRTCursor > GetCursor() const;
 
diff --git a/Engine/Source/Runtime/Core/Public/Windows/WindowsApplication.h b/Engine/Source/Runtime/Core/Public/Windows/WindowsApplication.h
--- a/Engine/Source/Runtime/Core/Public/Windows/WindowsApplication.h
+++ b/Engine/Source/Runtime/Core/Public/Windows/WindowsApplication.h
@@ -156,53 +156,53 @@ public:
 
 	virtual ~FWindowsApplication();
 	
-	virtual void DestroyApplication() OVERRIDE;
+	virtual void DestroyApplication() override;
 
 
 public:
 
-	virtual void SetMessageHandler( const TSharedRef< class FGenericApplicationMessageHandler >& InMessageHandler ) OVERRIDE;
+	virtual void SetMessageHandler( const TSharedRef< class FGenericApplicationMessageHandler >& InMessageHandler ) override;
 
-	virtual void PollGameDeviceState( const float TimeDelta ) OVERRIDE;
+	virtual void PollGameDeviceState( const float TimeDelta ) override;
 
-	virtual void PumpMessages( const float TimeDelta ) OVERRIDE;
+	virtual void PumpMessages( const float TimeDelta ) override;
 
-	virtual void ProcessDeferredEvents( const float TimeDelta ) OVERRIDE;
+	virtual void ProcessDeferredEvents( const float TimeDelta ) override;
 
-	virtual TSharedRef< FGenericWindow > MakeWindow() OVERRIDE;
+	virtual TSharedRef< FGenericWindow > MakeWindow() override;
 
-	virtual void InitializeWindow( const TSharedRef< FGenericWindow >& Window, const TSharedRef< FGenericWindowDefinition >& InDefinition, const TSharedPtr< FGenericWindow >& InParent, const bool bShowImmediately ) OVERRIDE;
+	virtual void InitializeWindow( const TSharedRef< FGenericWindow >& Window, const TSharedRef< FGenericWindowDefinition >& InDefinition, const TSharedPtr< FGenericWindow >& InParent, const bool bShowImmediately ) override;
 
-	virtual void SetCapture( const TSharedPtr< FGenericWindow >& InWindow ) OVERRIDE;
+	virtual void SetCapture( const TSharedPtr< FGenericWindow >& InWindow ) override;
 
-	virtual void* GetCapture( void ) const OVERRIDE;
+	virtual void* GetCapture( void ) const override;
 
-	virtual void SetHighPrecisionMouseMode( const bool Enable, const TSharedPtr< FGenericWindow >& InWindow ) OVERRIDE;
+	virtual void SetHighPrecisionMouseMode( const bool Enable, const TSharedPtr< FGenericWindow >& InWindow ) override;
 
-	virtual bool IsUsingHighPrecisionMouseMode() const OVERRIDE { return bUsingHighPrecisionMouseInput; }
+	virtual bool IsUsingHighPrecisionMouseMode() const override { return bUsingHighPrecisionMouseInput; }
 
-	virtual FModifierKeysState GetModifierKeys() const OVERRIDE;
+	virtual FModifierKeysState GetModifierKeys() const override;
 
-	virtual FPlatformRect GetWorkArea( const FPlatformRect& CurrentWindow ) const OVERRIDE;
+	virtual FPlatformRect GetWorkArea( const FPlatformRect& CurrentWindow ) const override;
 
-	virtual bool TryCalculatePopupWindowPosition( const FPlatformRect& InAnchor, const FVector2D& InSize, const EPopUpOrientation::Type Orientation, /*OUT*/ FVector2D* const CalculatedPopUpPosition ) const OVERRIDE;
+	virtual bool TryCalculatePopupWindowPosition( const FPlatformRect& InAnchor, const FVector2D& InSize, const EPopUpOrientation::Type Orientation, /*OUT*/ FVector2D* const CalculatedPopUpPosition ) const override;
 
-	virtual void GetDisplayMetrics( FDisplayMetrics& OutDisplayMetrics ) const OVERRIDE;
+	virtual void GetDisplayMetrics( FDisplayMetrics& OutDisplayMetrics ) const override;
 
-	virtual void GetInitialDisplayMetrics( FDisplayMetrics& OutDisplayMetrics ) const OVERRIDE;
+	virtual void GetInitialDisplayMetrics( FDisplayMetrics& OutDisplayMetrics ) const override;
 
-	virtual EWindowTitleAlignment::Type GetWindowTitleAlignment() const OVERRIDE;
+	virtual EWindowTitleAlignment::Type GetWindowTitleAlignment() const override;
 
-	virtual IForceFeedbackSystem *GetForceFeedbackSystem() OVERRIDE 
+	virtual IForceFeedbackSystem *GetForceFeedbackSystem() override 
 	{
 		return this; 
 	}
 
-	virtual void SetChannelValue (int32 ControllerId, FForceFeedbackChannelType ChannelType, float Value) OVERRIDE;
+	virtual void SetChannelValue (int32 ControllerId, FForceFeedbackChannelType ChannelType, float Value) override;
 
-	virtual void SetChannelValues (int32 ControllerId, const FForceFeedbackValues &Values) OVERRIDE;
+	virtual void SetChannelValues (int32 ControllerId, const FForceFeedbackValues &Values) override;
 
-	virtual ITextInputMethodSystem *GetTextInputMethodSystem() OVERRIDE
+	virtual ITextInputMethodSystem *GetTextInputMethodSystem() override
 	{
 		return TextInputMethodSystem.Get();
 	}
diff --git a/Engine/Source/Runtime/Core/Public/Windows/WindowsPlatformNamedPipe.h b/Engine/Source/Runtime/Core/Public/Windows/WindowsPlatformNamedPipe.h
--- a/Engine/Source/Runtime/Core/Public/Windows/WindowsPlatformNamedPipe.h
+++ b/Engine/Source/Runtime/Core/Public/Windows/WindowsPlatformNamedPipe.h
@@ -20,34 +20,34 @@ public:
 public:
 
 	// Create a named pipe, using overlapped IO if bAsync=1
-	virtual bool Create(const FString& PipeName, bool bAsServer, bool bAsync) OVERRIDE;
+	virtual bool Create(const FString& PipeName, bool bAsServer, bool bAsync) override;
 
 	// Destroy the pipe
-	virtual bool Destroy() OVERRIDE;
+	virtual bool Destroy() override;
 
 	// Open a connection from a client
-	virtual bool OpenConnection() OVERRIDE;
+	virtual bool OpenConnection() override;
 
 	// Blocks if there's an IO operation in progress until it's done or errors out
-	virtual bool BlockForAsyncIO() OVERRIDE;
+	virtual bool BlockForAsyncIO() override;
 
 	// Lets the user know if the pipe is ready to send or receive data
-	virtual bool IsReadyForRW() const OVERRIDE;
+	virtual bool IsReadyForRW() const override;
 
 	// Updates status of async state of the current pipe
-	virtual bool UpdateAsyncStatus() OVERRIDE;
+	virtual bool UpdateAsyncStatus() override;
 
 	// Writes a buffer out
-	virtual bool WriteBytes(int32 NumBytes, const void* Data) OVERRIDE;
+	virtual bool WriteBytes(int32 NumBytes, const void* Data) override;
 
 	// Reads a buffer in
-	virtual bool ReadBytes(int32 NumBytes, void* OutData) OVERRIDE;
+	virtual bool ReadBytes(int32 NumBytes, void* OutData) override;
 
 	// Returns true if the pipe has been created and hasn't been destroyed
-	virtual bool IsCreated() const OVERRIDE;
+	virtual bool IsCreated() const override;
 
 	// Return true if the pipe has had any communication error
-	virtual bool HasFailed() const OVERRIDE;
+	virtual bool HasFailed() const override;
 
 
 private:
diff --git a/Engine/Source/Runtime/Core/Public/Windows/WindowsTextInputMethodSystem.h b/Engine/Source/Runtime/Core/Public/Windows/WindowsTextInputMethodSystem.h
--- a/Engine/Source/Runtime/Core/Public/Windows/WindowsTextInputMethodSystem.h
+++ b/Engine/Source/Runtime/Core/Public/Windows/WindowsTextInputMethodSystem.h
@@ -19,17 +19,17 @@ public:
 	{}
 
 	// IUnknown Interface Begin
-	STDMETHODIMP QueryInterface(REFIID riid, void **ppvObj) OVERRIDE;
-	STDMETHODIMP_(ULONG) AddRef() OVERRIDE;
-	STDMETHODIMP_(ULONG) Release() OVERRIDE;
+	STDMETHODIMP QueryInterface(REFIID riid, void **ppvObj) override;
+	STDMETHODIMP_(ULONG) AddRef() override;
+	STDMETHODIMP_(ULONG) Release() override;
 	// IUnknown Interface End
 
 	// ITfInputProcessorProfileActivationSink Interface Begin
-	STDMETHODIMP OnActivated(DWORD dwProfileType, LANGID langid, REFCLSID clsid, REFGUID catid, REFGUID guidProfile, HKL hkl, DWORD dwFlags) OVERRIDE;
+	STDMETHODIMP OnActivated(DWORD dwProfileType, LANGID langid, REFCLSID clsid, REFGUID catid, REFGUID guidProfile, HKL hkl, DWORD dwFlags) override;
 	// ITfInputProcessorProfileActivationSink Interface End
 
 	// ITfActiveLanguageProfileNotifySink Interface Begin
-	STDMETHODIMP OnActivated(REFCLSID clsid, REFGUID guidProfile, BOOL fActivated) OVERRIDE;
+	STDMETHODIMP OnActivated(REFCLSID clsid, REFGUID guidProfile, BOOL fActivated) override;
 	// ITfActiveLanguageProfileNotifySink Interface End
 
 public:
@@ -52,10 +52,10 @@ public:
 	void Terminate();
 
 	// ITextInputMethodSystem Interface Begin
-	virtual TSharedPtr<ITextInputMethodChangeNotifier> RegisterContext(const TSharedRef<ITextInputMethodContext>& Context) OVERRIDE;
-	virtual void UnregisterContext(const TSharedRef<ITextInputMethodContext>& Context) OVERRIDE;
-	virtual void ActivateContext(const TSharedRef<ITextInputMethodContext>& Context) OVERRIDE;
-	virtual void DeactivateContext(const TSharedRef<ITextInputMethodContext>& Context) OVERRIDE;
+	virtual TSharedPtr<ITextInputMethodChangeNotifier> RegisterContext(const TSharedRef<ITextInputMethodContext>& Context) override;
+	virtual void UnregisterContext(const TSharedRef<ITextInputMethodContext>& Context) override;
+	virtual void ActivateContext(const TSharedRef<ITextInputMethodContext>& Context) override;
+	virtual void DeactivateContext(const TSharedRef<ITextInputMethodContext>& Context) override;
 	// ITextInputMethodSystem Interface End
 
 	int32 ProcessMessage(HWND hwnd, uint32 msg, WPARAM wParam, LPARAM lParam);
diff --git a/Engine/Source/Runtime/Core/Public/Windows/WindowsWindow.h b/Engine/Source/Runtime/Core/Public/Windows/WindowsWindow.h
--- a/Engine/Source/Runtime/Core/Public/Windows/WindowsWindow.h
+++ b/Engine/Source/Runtime/Core/Public/Windows/WindowsWindow.h
@@ -38,80 +38,80 @@ public:
 
 public:
 
-	virtual void ReshapeWindow( int32 X, int32 Y, int32 Width, int32 Height ) OVERRIDE;
+	virtual void ReshapeWindow( int32 X, int32 Y, int32 Width, int32 Height ) override;
 
-	virtual bool GetFullScreenInfo( int32& X, int32& Y, int32& Width, int32& Height ) const OVERRIDE;
+	virtual bool GetFullScreenInfo( int32& X, int32& Y, int32& Width, int32& Height ) const override;
 
-	virtual void MoveWindowTo ( int32 X, int32 Y ) OVERRIDE;
+	virtual void MoveWindowTo ( int32 X, int32 Y ) override;
 
-	virtual void BringToFront( bool bForce = false ) OVERRIDE;
+	virtual void BringToFront( bool bForce = false ) override;
 
-	virtual void HACK_ForceToFront() OVERRIDE;
+	virtual void HACK_ForceToFront() override;
 
-	virtual void Destroy() OVERRIDE;
+	virtual void Destroy() override;
 
-	virtual void Minimize() OVERRIDE;
+	virtual void Minimize() override;
 
-	virtual void Maximize() OVERRIDE;
+	virtual void Maximize() override;
 
-	virtual void Restore() OVERRIDE;
+	virtual void Restore() override;
 
-	virtual void Show() OVERRIDE;
+	virtual void Show() override;
 
-	virtual void Hide() OVERRIDE;
+	virtual void Hide() override;
 
-	virtual void SetWindowMode( EWindowMode::Type NewWindowMode ) OVERRIDE;
+	virtual void SetWindowMode( EWindowMode::Type NewWindowMode ) override;
 
-	virtual EWindowMode::Type GetWindowMode() const OVERRIDE { return WindowMode; } 
+	virtual EWindowMode::Type GetWindowMode() const override { return WindowMode; } 
 
-	virtual bool IsMaximized() const OVERRIDE;
+	virtual bool IsMaximized() const override;
 
-	virtual bool IsMinimized() const OVERRIDE;
+	virtual bool IsMinimized() const override;
 
-	virtual bool IsVisible() const OVERRIDE;
+	virtual bool IsVisible() const override;
 
-	virtual bool GetRestoredDimensions(int32& X, int32& Y, int32& Width, int32& Height) OVERRIDE;
+	virtual bool GetRestoredDimensions(int32& X, int32& Y, int32& Width, int32& Height) override;
 
-	virtual void SetWindowFocus() OVERRIDE;
+	virtual void SetWindowFocus() override;
 
-	virtual void SetOpacity( const float InOpacity ) OVERRIDE;
+	virtual void SetOpacity( const float InOpacity ) override;
 
-	virtual void Enable( bool bEnable ) OVERRIDE;
+	virtual void Enable( bool bEnable ) override;
 
-	virtual bool IsPointInWindow( int32 X, int32 Y ) const OVERRIDE;
+	virtual bool IsPointInWindow( int32 X, int32 Y ) const override;
 
-	virtual int32 GetWindowBorderSize() const OVERRIDE;
+	virtual int32 GetWindowBorderSize() const override;
 
-	virtual void* GetOSWindowHandle() const  OVERRIDE { return HWnd; }
+	virtual void* GetOSWindowHandle() const  override { return HWnd; }
 
-	virtual bool IsForegroundWindow() const OVERRIDE;
+	virtual bool IsForegroundWindow() const override;
 
-	virtual void SetText(const TCHAR* const Text) OVERRIDE;
+	virtual void SetText(const TCHAR* const Text) override;
 
 	/**
 	 * Slate Win32Windows implement the IDropTarget and IUnknown OLE interfaces.
 	 */
 
 	/** IUnknown : attempt to get a reference to this instance as a pointer  */
-	HRESULT STDCALL QueryInterface( REFIID iid, void ** ppvObject ) OVERRIDE;
+	HRESULT STDCALL QueryInterface( REFIID iid, void ** ppvObject ) override;
 
 	/** IUnknown : Invoked by OLE when a new reference to this instance */
-	ULONG STDCALL AddRef( void ) OVERRIDE;
+	ULONG STDCALL AddRef( void ) override;
 
 	/** IUnknown : Invoked by OLE when a reference to this instance is released */
-	ULONG STDCALL Release( void ) OVERRIDE;
+	ULONG STDCALL Release( void ) override;
 
 	/** IDropTarget : Invoked by OLE when someone drag data our window initially */
-	virtual HRESULT STDCALL DragEnter( __RPC__in_opt IDataObject *DataObjectPointer, ::DWORD KeyState, POINTL CursorPosition, __RPC__inout ::DWORD *CursorEffect) OVERRIDE;
+	virtual HRESULT STDCALL DragEnter( __RPC__in_opt IDataObject *DataObjectPointer, ::DWORD KeyState, POINTL CursorPosition, __RPC__inout ::DWORD *CursorEffect) override;
 
 	/** IDropTarget : Invoked by OLE when someone drag data over our window */
-	virtual HRESULT STDCALL DragOver( ::DWORD KeyState, POINTL CursorPosition, __RPC__inout ::DWORD *CursorEffect) OVERRIDE;
+	virtual HRESULT STDCALL DragOver( ::DWORD KeyState, POINTL CursorPosition, __RPC__inout ::DWORD *CursorEffect) override;
 
 	/** IDropTarget : Invoked by OLE when someone drag data and it exits our window or hits ESC to cancel the drag and drop. */
-	virtual HRESULT STDCALL DragLeave( void ) OVERRIDE;
+	virtual HRESULT STDCALL DragLeave( void ) override;
 
 	/** IDropTarget : Invoked by OLE when someone drag data our window initially */
-	virtual HRESULT STDCALL Drop( __RPC__in_opt IDataObject *DataObjectPointer, ::DWORD KeyState, POINTL CursorPosition, __RPC__inout ::DWORD *CursorEffect) OVERRIDE;
+	virtual HRESULT STDCALL Drop( __RPC__in_opt IDataObject *DataObjectPointer, ::DWORD KeyState, POINTL CursorPosition, __RPC__inout ::DWORD *CursorEffect) override;
 
 	/**
 	 * Sets the window text - usually the title but can also be text content for things like controls
@@ -124,7 +124,7 @@ public:
 	void AdjustWindowRegion( int32 Width, int32 Height );
 
 	/** @return	Gives the native window a chance to adjust our stored window size before we cache it off */
-	virtual void AdjustCachedSize( FVector2D& Size ) const OVERRIDE;
+	virtual void AdjustCachedSize( FVector2D& Size ) const override;
 
 	/** Called when our parent window is minimized (which will in turn cause us to become minimized) */
 	void OnParentWindowMinimized();
diff --git a/Engine/Source/Runtime/CoreUObject/Private/UObject/CoreNative.cpp b/Engine/Source/Runtime/CoreUObject/Private/UObject/CoreNative.cpp
--- a/Engine/Source/Runtime/CoreUObject/Private/UObject/CoreNative.cpp
+++ b/Engine/Source/Runtime/CoreUObject/Private/UObject/CoreNative.cpp
@@ -9,7 +9,7 @@
 class FCoreUObjectModule : public FDefaultModuleImpl
 {
 public:
-	virtual void StartupModule() OVERRIDE
+	virtual void StartupModule() override
 	{
 		void InitUObject();
 		FCoreDelegates::OnInit.AddStatic(InitUObject);
diff --git a/Engine/Source/Runtime/CoreUObject/Private/UObject/GarbageCollection.cpp b/Engine/Source/Runtime/CoreUObject/Private/UObject/GarbageCollection.cpp
--- a/Engine/Source/Runtime/CoreUObject/Private/UObject/GarbageCollection.cpp
+++ b/Engine/Source/Runtime/CoreUObject/Private/UObject/GarbageCollection.cpp
@@ -262,7 +262,7 @@ public:
 	{
 	}
 
-	virtual void HandleObjectReference( UObject*& Object, const UObject* ReferencingObject, const UObject* ReferencingProperty ) OVERRIDE
+	virtual void HandleObjectReference( UObject*& Object, const UObject* ReferencingObject, const UObject* ReferencingProperty ) override
 	{
 #if !(UE_BUILD_TEST || UE_BUILD_SHIPPING)
 		if (Object && !Object->IsValidLowLevelFast())
@@ -275,15 +275,15 @@ public:
 #endif
 		::HandleObjectReference( ObjectArray, const_cast<UObject*>(ReferencingObject), Object, bAllowEliminatingReferences );
 	}
-	virtual bool IsIgnoringArchetypeRef() const OVERRIDE
+	virtual bool IsIgnoringArchetypeRef() const override
 	{
 		return false;
 	}
-	virtual bool IsIgnoringTransient() const OVERRIDE
+	virtual bool IsIgnoringTransient() const override
 	{
 		return false;
 	}
-	virtual void AllowEliminatingReferences( bool bAllow ) OVERRIDE
+	virtual void AllowEliminatingReferences( bool bAllow ) override
 	{
 		bAllowEliminatingReferences = bAllow;
 	}
diff --git a/Engine/Source/Runtime/CoreUObject/Private/UObject/Obj.cpp b/Engine/Source/Runtime/CoreUObject/Private/UObject/Obj.cpp
--- a/Engine/Source/Runtime/CoreUObject/Private/UObject/Obj.cpp
+++ b/Engine/Source/Runtime/CoreUObject/Private/UObject/Obj.cpp
@@ -923,7 +923,7 @@ public:
 	}
 
 	// Begin FReferenceCollector interface.
-	virtual void HandleObjectReference(UObject*& InObject, const UObject* InReferencingObject, const UObject* InReferencingProperty) OVERRIDE
+	virtual void HandleObjectReference(UObject*& InObject, const UObject* InReferencingObject, const UObject* InReferencingProperty) override
 	{
 		// Only care about unique default subobjects that are outside of the referencing object's outer chain.
 		// Also ignore references to subobjects if they share the same Outer.
@@ -936,8 +936,8 @@ public:
 			ObjectArray.Add(InObject);
 		}
 	}
-	virtual bool IsIgnoringArchetypeRef() const OVERRIDE { return true; }
-	virtual bool IsIgnoringTransient() const OVERRIDE { return true; }
+	virtual bool IsIgnoringArchetypeRef() const override { return true; }
+	virtual bool IsIgnoringTransient() const override { return true; }
 	// End FReferenceCollector interface.
 
 protected:
diff --git a/Engine/Source/Runtime/CoreUObject/Private/UObject/UObjectGlobals.cpp b/Engine/Source/Runtime/CoreUObject/Private/UObject/UObjectGlobals.cpp
--- a/Engine/Source/Runtime/CoreUObject/Private/UObject/UObjectGlobals.cpp
+++ b/Engine/Source/Runtime/CoreUObject/Private/UObject/UObjectGlobals.cpp
@@ -2258,11 +2258,11 @@ public:
 	}
 
 	// FReferenceCollector interface
-	virtual bool IsIgnoringArchetypeRef() const OVERRIDE
+	virtual bool IsIgnoringArchetypeRef() const override
 	{
 		return false;
 	}
-	virtual bool IsIgnoringTransient() const OVERRIDE
+	virtual bool IsIgnoringTransient() const override
 	{
 		return false;
 	}
@@ -2372,7 +2372,7 @@ private:
 		ObjectsToSerialize.Add( Object );
 	}
 
-	virtual void HandleObjectReference( UObject*& InObject, const UObject* InReferencingObject, const UObject* InReferencingProperty ) OVERRIDE
+	virtual void HandleObjectReference( UObject*& InObject, const UObject* InReferencingObject, const UObject* InReferencingProperty ) override
 	{
 		checkSlow( !InObject || InObject->IsValidLowLevel() );		
 		if( InObject )
diff --git a/Engine/Source/Runtime/CoreUObject/Public/Misc/TextBuffer.h b/Engine/Source/Runtime/CoreUObject/Public/Misc/TextBuffer.h
--- a/Engine/Source/Runtime/CoreUObject/Public/Misc/TextBuffer.h
+++ b/Engine/Source/Runtime/CoreUObject/Public/Misc/TextBuffer.h
@@ -48,7 +48,7 @@ public:
 
 	virtual void Serialize (FArchive& Ar);
 
-	virtual void Serialize (const TCHAR* Data, ELogVerbosity::Type Verbosity, const class FName& Category ) OVERRIDE;
+	virtual void Serialize (const TCHAR* Data, ELogVerbosity::Type Verbosity, const class FName& Category ) override;
 
 
 private:
diff --git a/Engine/Source/Runtime/CoreUObject/Public/Misc/UObjectToken.h b/Engine/Source/Runtime/CoreUObject/Public/Misc/UObjectToken.h
--- a/Engine/Source/Runtime/CoreUObject/Public/Misc/UObjectToken.h
+++ b/Engine/Source/Runtime/CoreUObject/Public/Misc/UObjectToken.h
@@ -19,12 +19,12 @@ public:
 	}
 
 	/** Begin IMessageToken interface */
-	virtual EMessageToken::Type GetType() const OVERRIDE
+	virtual EMessageToken::Type GetType() const override
 	{
 		return EMessageToken::Object;
 	}
 
-	virtual const FOnMessageTokenActivated& GetOnMessageTokenActivated() const OVERRIDE;
+	virtual const FOnMessageTokenActivated& GetOnMessageTokenActivated() const override;
 	/** End IMessageToken interface */
 
 	/** Get the object referenced by this token */
diff --git a/Engine/Source/Runtime/CoreUObject/Public/Serialization/ArchiveUObject.h b/Engine/Source/Runtime/CoreUObject/Public/Serialization/ArchiveUObject.h
--- a/Engine/Source/Runtime/CoreUObject/Public/Serialization/ArchiveUObject.h
+++ b/Engine/Source/Runtime/CoreUObject/Public/Serialization/ArchiveUObject.h
@@ -95,11 +95,11 @@ public:
 	}
 
 	// FArchive interface
-	COREUOBJECT_API virtual FArchive& operator<<( class FName& N ) OVERRIDE;
-	COREUOBJECT_API virtual FArchive& operator<<( class UObject*& Res ) OVERRIDE;
+	COREUOBJECT_API virtual FArchive& operator<<( class FName& N ) override;
+	COREUOBJECT_API virtual FArchive& operator<<( class UObject*& Res ) override;
 	COREUOBJECT_API virtual FArchive& operator<<( FLazyObjectPtr& LazyObjectPtr );
 	COREUOBJECT_API virtual FArchive& operator<<( FAssetPtr& AssetPtr );
-	COREUOBJECT_API virtual FString GetArchiveName() const OVERRIDE;
+	COREUOBJECT_API virtual FString GetArchiveName() const override;
 	// End of FArchive interface
 
 protected:
@@ -127,11 +127,11 @@ public:
 	}
 
 	// FArchive interface
-	COREUOBJECT_API virtual FArchive& operator<<( class FName& N ) OVERRIDE;
-	COREUOBJECT_API virtual FArchive& operator<<( class UObject*& Res ) OVERRIDE;
+	COREUOBJECT_API virtual FArchive& operator<<( class FName& N ) override;
+	COREUOBJECT_API virtual FArchive& operator<<( class UObject*& Res ) override;
 	COREUOBJECT_API virtual FArchive& operator<<( FLazyObjectPtr& LazyObjectPtr );
 	COREUOBJECT_API virtual FArchive& operator<<( FAssetPtr& AssetPtr );
-	COREUOBJECT_API virtual FString GetArchiveName() const OVERRIDE;
+	COREUOBJECT_API virtual FString GetArchiveName() const override;
 	// End of FArchive interface
 
 protected:
diff --git a/Engine/Source/Runtime/CoreUObject/Public/Serialization/ArchiveUObjectBase.h b/Engine/Source/Runtime/CoreUObject/Public/Serialization/ArchiveUObjectBase.h
--- a/Engine/Source/Runtime/CoreUObject/Public/Serialization/ArchiveUObjectBase.h
+++ b/Engine/Source/Runtime/CoreUObject/Public/Serialization/ArchiveUObjectBase.h
@@ -14,7 +14,7 @@ class COREUOBJECT_API FArchiveUObject : public FArchive
 public:
 
 	// Begin FArchive Interface
-	virtual FArchive& operator<< (class FLazyObjectPtr& Value) OVERRIDE;
-	virtual FArchive& operator<< (class FAssetPtr& Value) OVERRIDE;
+	virtual FArchive& operator<< (class FLazyObjectPtr& Value) override;
+	virtual FArchive& operator<< (class FAssetPtr& Value) override;
 	// End FArchive Interface
 };
diff --git a/Engine/Source/Runtime/CoreUObject/Public/UObject/Class.h b/Engine/Source/Runtime/CoreUObject/Public/UObject/Class.h
--- a/Engine/Source/Runtime/CoreUObject/Public/UObject/Class.h
+++ b/Engine/Source/Runtime/CoreUObject/Public/UObject/Class.h
@@ -62,8 +62,8 @@ class COREUOBJECT_API UField : public UObject
 	UField(EStaticConstructor, EObjectFlags InFlags);
 
 	// UObject interface.
-	virtual void Serialize( FArchive& Ar ) OVERRIDE;
-	virtual void PostLoad() OVERRIDE;
+	virtual void Serialize( FArchive& Ar ) override;
+	virtual void PostLoad() override;
 
 	// UField interface.
 	virtual void AddCppProperty( UProperty* Property );
@@ -259,13 +259,13 @@ public:
 	explicit UStruct(const class FPostConstructInitializeProperties& PCIP, UStruct* InSuperStruct, SIZE_T ParamsSize = 0, SIZE_T Alignment = 0 );
 
 	// UObject interface.
-	virtual void Serialize(FArchive& Ar) OVERRIDE;
-	virtual void FinishDestroy() OVERRIDE;
-	virtual void RegisterDependencies() OVERRIDE;
+	virtual void Serialize(FArchive& Ar) override;
+	virtual void FinishDestroy() override;
+	virtual void RegisterDependencies() override;
 	static void AddReferencedObjects(UObject* InThis, FReferenceCollector& Collector);
 
 	// UField interface.
-	virtual void AddCppProperty(UProperty* Property) OVERRIDE;
+	virtual void AddCppProperty(UProperty* Property) override;
 
 	/**
 	 * Creates new copies of components
@@ -302,7 +302,7 @@ public:
 	virtual void SerializeTaggedProperties( FArchive& Ar, uint8* Data, UStruct* DefaultsStruct, uint8* Defaults ) const;
 
 	virtual EExprToken SerializeExpr(int32& iCode, FArchive& Ar);
-	virtual void TagSubobjects(EObjectFlags NewFlags) OVERRIDE;
+	virtual void TagSubobjects(EObjectFlags NewFlags) override;
 
 	/**
 	 * Returns the struct/ class prefix used for the C++ declaration of this struct/ class.
@@ -803,124 +803,124 @@ public:
 			: ICppStructOps(sizeof(CPPSTRUCT), ALIGNOF(CPPSTRUCT))
 		{
 		}
-		virtual bool HasNoopConstructor() OVERRIDE
+		virtual bool HasNoopConstructor() override
 		{
 			return TTraits::WithNoInitConstructor;
 		}		
-		virtual bool HasZeroConstructor() OVERRIDE
+		virtual bool HasZeroConstructor() override
 		{
 			return TTraits::WithZeroConstructor;
 		}
-		virtual void Construct(void *Dest) OVERRIDE
+		virtual void Construct(void *Dest) override
 		{
 			check(!TTraits::WithZeroConstructor); // don't call this if we have indicated it is not necessary
 			// that could have been an if statement, but we might as well force optimization above the virtual call
 			// could also not attempt to call the constructor for types where this is not possible, but I didn't do that here
 			ConstructWithNoInitOrNot<CPPSTRUCT>(Dest);
 		}
-		virtual bool HasDestructor() OVERRIDE
+		virtual bool HasDestructor() override
 		{
 			return !(TTraits::WithNoDestructor || TIsPODType<CPPSTRUCT>::Value);
 		}
-		virtual void Destruct(void *Dest) OVERRIDE
+		virtual void Destruct(void *Dest) override
 		{
 			check(!(TTraits::WithNoDestructor || TIsPODType<CPPSTRUCT>::Value)); // don't call this if we have indicated it is not necessary
 			// that could have been an if statement, but we might as well force optimization above the virtual call
 			// could also not attempt to call the destructor for types where this is not possible, but I didn't do that here
 			((CPPSTRUCT*)Dest)->~CPPSTRUCT();
 		}
-		virtual bool HasSerializer() OVERRIDE
+		virtual bool HasSerializer() override
 		{
 			return TTraits::WithSerializer;
 		}
-		virtual bool Serialize(FArchive& Ar, void *Data) OVERRIDE
+		virtual bool Serialize(FArchive& Ar, void *Data) override
 		{
 			check(TTraits::WithSerializer); // don't call this if we have indicated it is not necessary
 			return SerializeOrNot(Ar, (CPPSTRUCT*)Data);
 		}
-		virtual bool HasPostSerialize() OVERRIDE
+		virtual bool HasPostSerialize() override
 		{
 			return TTraits::WithPostSerialize;
 		}
-		virtual void PostSerialize(const FArchive& Ar, void *Data) OVERRIDE
+		virtual void PostSerialize(const FArchive& Ar, void *Data) override
 		{
 			check(TTraits::WithPostSerialize); // don't call this if we have indicated it is not necessary
 			PostSerializeOrNot(Ar, (CPPSTRUCT*)Data);
 		}
-		virtual bool HasNetSerializer() OVERRIDE
+		virtual bool HasNetSerializer() override
 		{
 			return TTraits::WithNetSerializer;
 		}
-		virtual bool HasNetDeltaSerializer() OVERRIDE
+		virtual bool HasNetDeltaSerializer() override
 		{
 			return TTraits::WithNetDeltaSerializer;
 		}
-		virtual bool NetSerialize(FArchive& Ar, class UPackageMap* Map, bool& bOutSuccess, void *Data) OVERRIDE
+		virtual bool NetSerialize(FArchive& Ar, class UPackageMap* Map, bool& bOutSuccess, void *Data) override
 		{
 			return NetSerializeOrNot(Ar, Map, bOutSuccess, (CPPSTRUCT*)Data);
 		}
-		virtual bool NetDeltaSerialize(FNetDeltaSerializeInfo & DeltaParms, void *Data) OVERRIDE
+		virtual bool NetDeltaSerialize(FNetDeltaSerializeInfo & DeltaParms, void *Data) override
 		{
 			return NetDeltaSerializeOrNot(DeltaParms, (CPPSTRUCT*)Data);
 		}
-		virtual bool IsPlainOldData() OVERRIDE
+		virtual bool IsPlainOldData() override
 		{
 			return TIsPODType<CPPSTRUCT>::Value;
 		}
-		virtual bool HasCopy() OVERRIDE
+		virtual bool HasCopy() override
 		{
 			return TTraits::WithCopy;
 		}
-		virtual bool Copy(void* Dest, void const* Src, int32 ArrayDim) OVERRIDE
+		virtual bool Copy(void* Dest, void const* Src, int32 ArrayDim) override
 		{
 			return CopyOrNot((CPPSTRUCT*)Dest, (CPPSTRUCT const*)Src, ArrayDim);
 		}
-		virtual bool HasIdentical() OVERRIDE
+		virtual bool HasIdentical() override
 		{
 			return TTraits::WithIdentical || TTraits::WithIdenticalViaEquality;
 		}
-		virtual bool Identical(const void* A, const void* B, uint32 PortFlags, bool& bOutResult) OVERRIDE
+		virtual bool Identical(const void* A, const void* B, uint32 PortFlags, bool& bOutResult) override
 		{
 			check((TTraits::WithIdentical || TTraits::WithIdenticalViaEquality)); // don't call this if we have indicated it is not necessary
 			return IdenticalOrNot((const CPPSTRUCT*)A, (const CPPSTRUCT*)B, PortFlags, bOutResult);
 		}
-		virtual bool HasExportTextItem() OVERRIDE
+		virtual bool HasExportTextItem() override
 		{
 			return TTraits::WithExportTextItem;
 		}
-		virtual bool ExportTextItem(FString& ValueStr, const void* PropertyValue, const void* DefaultValue, class UObject* Parent, int32 PortFlags, class UObject* ExportRootScope) OVERRIDE
+		virtual bool ExportTextItem(FString& ValueStr, const void* PropertyValue, const void* DefaultValue, class UObject* Parent, int32 PortFlags, class UObject* ExportRootScope) override
 		{
 			check(TTraits::WithExportTextItem); // don't call this if we have indicated it is not necessary
 			return ExportTextItemOrNot(ValueStr, (const CPPSTRUCT*)PropertyValue, (const CPPSTRUCT*)DefaultValue, Parent, PortFlags, ExportRootScope);
 		}
-		virtual bool HasImportTextItem() OVERRIDE
+		virtual bool HasImportTextItem() override
 		{
 			return TTraits::WithImportTextItem;
 		}
-		virtual bool ImportTextItem(const TCHAR*& Buffer, void* Data, int32 PortFlags, class UObject* OwnerObject, FOutputDevice* ErrorText) OVERRIDE
+		virtual bool ImportTextItem(const TCHAR*& Buffer, void* Data, int32 PortFlags, class UObject* OwnerObject, FOutputDevice* ErrorText) override
 		{
 			check(TTraits::WithImportTextItem); // don't call this if we have indicated it is not necessary
 			return ImportTextItemOrNot(Buffer, (CPPSTRUCT*)Data, PortFlags, OwnerObject, ErrorText);
 		}
-		virtual bool HasAddStructReferencedObjects() OVERRIDE
+		virtual bool HasAddStructReferencedObjects() override
 		{
 			return TTraits::WithAddStructReferencedObjects;
 		}
-		virtual TPointerToAddStructReferencedObjects AddStructReferencedObjects() OVERRIDE
+		virtual TPointerToAddStructReferencedObjects AddStructReferencedObjects() override
 		{
 			check(TTraits::WithAddStructReferencedObjects); // don't call this if we have indicated it is not necessary
 			return &AddStructReferencedObjectsOrNot<CPPSTRUCT>;
 		}
-		virtual bool HasSerializeFromMismatchedTag() OVERRIDE
+		virtual bool HasSerializeFromMismatchedTag() override
 		{
 			return TTraits::WithSerializeFromMismatchedTag;
 		}
-		virtual bool SerializeFromMismatchedTag(struct FPropertyTag const& Tag, FArchive& Ar, void *Data) OVERRIDE
+		virtual bool SerializeFromMismatchedTag(struct FPropertyTag const& Tag, FArchive& Ar, void *Data) override
 		{
 			check(TTraits::WithSerializeFromMismatchedTag); // don't call this if we have indicated it is not allowed
 			return SerializeFromMismatchedTagOrNot(Tag, Ar, (CPPSTRUCT*)Data);
 		}
-		virtual bool HasMessageHandling() OVERRIDE
+		virtual bool HasMessageHandling() override
 		{
 			return TTraits::WithMessageHandling;
 		}
@@ -962,11 +962,11 @@ private:
 public:
 
 	// UObject Interface
-	virtual COREUOBJECT_API void Serialize( FArchive& Ar ) OVERRIDE;
-	virtual COREUOBJECT_API void PostLoad() OVERRIDE;
+	virtual COREUOBJECT_API void Serialize( FArchive& Ar ) override;
+	virtual COREUOBJECT_API void PostLoad() override;
 
 	// UStruct interface.
-	virtual COREUOBJECT_API void Link(FArchive& Ar, bool bRelinkExistingProperties) OVERRIDE;
+	virtual COREUOBJECT_API void Link(FArchive& Ar, bool bRelinkExistingProperties) override;
 	// End of UStruct interface.
 
 	/** Stash a CppStructOps for future use 
@@ -1188,14 +1188,14 @@ public:
 	void InitializeDerivedMembers();
 
 	// UObject interface.
-	virtual void Serialize( FArchive& Ar ) OVERRIDE;
+	virtual void Serialize( FArchive& Ar ) override;
 
 	// UField interface.
-	virtual void Bind() OVERRIDE;
+	virtual void Bind() override;
 
 	// UStruct interface.
-	virtual UStruct* GetInheritanceSuper() const OVERRIDE { return NULL;}
-	virtual void Link(FArchive& Ar, bool bRelinkExistingProperties) OVERRIDE;
+	virtual UStruct* GetInheritanceSuper() const override { return NULL;}
+	virtual void Link(FArchive& Ar, bool bRelinkExistingProperties) override;
 
 	// UFunction interface.
 	UFunction* GetSuperFunction() const
@@ -1296,7 +1296,7 @@ protected:
 
 public:
 	// UObject interface.
-	COREUOBJECT_API virtual void Serialize(FArchive& Ar) OVERRIDE;
+	COREUOBJECT_API virtual void Serialize(FArchive& Ar) override;
 	// End of UObject interface.
 
 	COREUOBJECT_API ~UEnum();
@@ -1725,7 +1725,7 @@ public:
 	// List of network relevant fields (properties and functions)
 	TArray<UField*> NetFields;
 
-	virtual bool IsNameStableForNetworking() const OVERRIDE { return true; }		// For now, assume all classes have stable net names
+	virtual bool IsNameStableForNetworking() const override { return true; }		// For now, assume all classes have stable net names
 
 #if WITH_EDITOR || HACK_HEADER_GENERATOR 
 	// Editor only properties
@@ -1850,27 +1850,27 @@ public:
 	UFunction* FindFunctionByName(FName InName, EIncludeSuperFlag::Type IncludeSuper = EIncludeSuperFlag::IncludeSuper) const;
 
 	// UObject interface.
-	virtual void Serialize(FArchive& Ar) OVERRIDE;
-	virtual void PostLoad() OVERRIDE;
-	virtual void FinishDestroy() OVERRIDE;
-	virtual void DeferredRegister(UClass *UClassStaticClass,const TCHAR* PackageName,const TCHAR* Name) OVERRIDE;
-	virtual bool Rename(const TCHAR* NewName = NULL, UObject* NewOuter = NULL, ERenameFlags Flags = REN_None) OVERRIDE;
-	virtual void TagSubobjects(EObjectFlags NewFlags) OVERRIDE;
-	virtual void PostInitProperties() OVERRIDE;
+	virtual void Serialize(FArchive& Ar) override;
+	virtual void PostLoad() override;
+	virtual void FinishDestroy() override;
+	virtual void DeferredRegister(UClass *UClassStaticClass,const TCHAR* PackageName,const TCHAR* Name) override;
+	virtual bool Rename(const TCHAR* NewName = NULL, UObject* NewOuter = NULL, ERenameFlags Flags = REN_None) override;
+	virtual void TagSubobjects(EObjectFlags NewFlags) override;
+	virtual void PostInitProperties() override;
 	static void AddReferencedObjects(UObject* InThis, FReferenceCollector& Collector);
-	virtual FRestoreForUObjectOverwrite* GetRestoreForUObjectOverwrite() OVERRIDE;
-	virtual FString GetDesc() OVERRIDE;
-	virtual bool IsAsset() const OVERRIDE { return false; }
+	virtual FRestoreForUObjectOverwrite* GetRestoreForUObjectOverwrite() override;
+	virtual FString GetDesc() override;
+	virtual bool IsAsset() const override { return false; }
 	// End of UObject interface.
 
 	// UField interface.
-	virtual void Bind() OVERRIDE;
-	virtual const TCHAR* GetPrefixCPP() const OVERRIDE;
+	virtual void Bind() override;
+	virtual const TCHAR* GetPrefixCPP() const override;
 	// End of UField interface.
 
 	// UStruct interface.
-	virtual void Link(FArchive& Ar, bool bRelinkExistingProperties) OVERRIDE;
-	virtual void SetSuperStruct(UStruct* NewSuperStruct) OVERRIDE;
+	virtual void Link(FArchive& Ar, bool bRelinkExistingProperties) override;
+	virtual void SetSuperStruct(UStruct* NewSuperStruct) override;
 	// End of UStruct interface.
 	
 	/**
diff --git a/Engine/Source/Runtime/CoreUObject/Public/UObject/CoreObject.h b/Engine/Source/Runtime/CoreUObject/Public/UObject/CoreObject.h
--- a/Engine/Source/Runtime/CoreUObject/Public/UObject/CoreObject.h
+++ b/Engine/Source/Runtime/CoreUObject/Public/UObject/CoreObject.h
@@ -64,7 +64,7 @@ private:
 	FName ForcedExportBasePackageName;
 public:
 
-	virtual bool IsNameStableForNetworking() const OVERRIDE { return true; }		// For now, assume all packages have stable net names
+	virtual bool IsNameStableForNetworking() const override { return true; }		// For now, assume all packages have stable net names
 
 	/** Package flags, serialized.*/
 	uint32	PackageFlags;
@@ -100,7 +100,7 @@ public:
 	static void AddReferencedObjects(UObject* InThis, FReferenceCollector& Collector);
 
 	/** Packages are never assets */
-	virtual bool IsAsset() const OVERRIDE { return false; }
+	virtual bool IsAsset() const override { return false; }
 
 	// UPackage interface.
 
@@ -498,11 +498,11 @@ public:
 	void RemoveMetaDataOutsidePackage();
 
 	// UObject interface
-	virtual void Serialize(FArchive& Ar) OVERRIDE;
-	virtual bool NeedsLoadForClient() const OVERRIDE;
-	virtual bool NeedsLoadForServer() const OVERRIDE;
-	virtual bool NeedsLoadForEditorGame() const OVERRIDE;
-	virtual bool IsAsset() const OVERRIDE { return false; }
+	virtual void Serialize(FArchive& Ar) override;
+	virtual bool NeedsLoadForClient() const override;
+	virtual bool NeedsLoadForServer() const override;
+	virtual bool NeedsLoadForEditorGame() const override;
+	virtual bool IsAsset() const override { return false; }
 	// End of UObject interface
 
 #if HACK_HEADER_GENERATOR
diff --git a/Engine/Source/Runtime/CoreUObject/Public/UObject/ObjectRedirector.h b/Engine/Source/Runtime/CoreUObject/Public/UObject/ObjectRedirector.h
--- a/Engine/Source/Runtime/CoreUObject/Public/UObject/ObjectRedirector.h
+++ b/Engine/Source/Runtime/CoreUObject/Public/UObject/ObjectRedirector.h
@@ -21,10 +21,10 @@ class UObjectRedirector : public UObject
 	// UObject interface.
 	virtual void PreSave();
 	void Serialize( FArchive& Ar );
-	virtual bool NeedsLoadForClient() const OVERRIDE;
-	virtual bool NeedsLoadForServer() const OVERRIDE;
-	virtual bool NeedsLoadForEditorGame() const OVERRIDE;
-	virtual void GetAssetRegistryTags(TArray<FAssetRegistryTag>& OutTags) const OVERRIDE;
+	virtual bool NeedsLoadForClient() const override;
+	virtual bool NeedsLoadForServer() const override;
+	virtual bool NeedsLoadForEditorGame() const override;
+	virtual void GetAssetRegistryTags(TArray<FAssetRegistryTag>& OutTags) const override;
 
 	/**
 	 * Callback for retrieving a textual representation of natively serialized properties.  Child classes should implement this method if they wish
diff --git a/Engine/Source/Runtime/CoreUObject/Public/UObject/ReferenceChainSearch.h b/Engine/Source/Runtime/CoreUObject/Public/UObject/ReferenceChainSearch.h
--- a/Engine/Source/Runtime/CoreUObject/Public/UObject/ReferenceChainSearch.h
+++ b/Engine/Source/Runtime/CoreUObject/Public/UObject/ReferenceChainSearch.h
@@ -81,10 +81,10 @@ public:
 		{
 		}
 
-		virtual void HandleObjectReference( UObject*& InObject, const UObject* ReferencingObject, const UObject* ReferencingProperty ) OVERRIDE;
+		virtual void HandleObjectReference( UObject*& InObject, const UObject* ReferencingObject, const UObject* ReferencingProperty ) override;
 
-		virtual bool IsIgnoringArchetypeRef() const OVERRIDE	{ return false; }
-		virtual bool IsIgnoringTransient() const OVERRIDE		{ return false; }
+		virtual bool IsIgnoringArchetypeRef() const override	{ return false; }
+		virtual bool IsIgnoringTransient() const override		{ return false; }
 
 		TArray<FReferenceChainLink> References;
 	};
diff --git a/Engine/Source/Runtime/CoreUObject/Public/UObject/Stack.h b/Engine/Source/Runtime/CoreUObject/Public/UObject/Stack.h
--- a/Engine/Source/Runtime/CoreUObject/Public/UObject/Stack.h
+++ b/Engine/Source/Runtime/CoreUObject/Public/UObject/Stack.h
@@ -118,7 +118,7 @@ public:
 	template<class TProperty, typename TNativeType>
 	FORCEINLINE_DEBUGGABLE TNativeType& StepCompiledInRef(void*const TemporaryBuffer);
 
-	COREUOBJECT_API virtual void Serialize( const TCHAR* V, ELogVerbosity::Type Verbosity, const class FName& Category ) OVERRIDE;
+	COREUOBJECT_API virtual void Serialize( const TCHAR* V, ELogVerbosity::Type Verbosity, const class FName& Category ) override;
 	
 	COREUOBJECT_API static void KismetExecutionMessage(const TCHAR* Message, ELogVerbosity::Type Verbosity);
 
diff --git a/Engine/Source/Runtime/CoreUObject/Public/UObject/UObjectGlobals.h b/Engine/Source/Runtime/CoreUObject/Public/UObject/UObjectGlobals.h
--- a/Engine/Source/Runtime/CoreUObject/Public/UObject/UObjectGlobals.h
+++ b/Engine/Source/Runtime/CoreUObject/Public/UObject/UObjectGlobals.h
@@ -1393,9 +1393,9 @@ public:
 	virtual void FindReferences(UObject* Object);	
 
 	// FReferenceCollector interface.
-	virtual void HandleObjectReference(UObject*& Object, const UObject* ReferencingObject, const UObject* ReferencingProperty) OVERRIDE;
-	virtual bool IsIgnoringArchetypeRef() const OVERRIDE { return bShouldIgnoreArchetype; }
-	virtual bool IsIgnoringTransient() const OVERRIDE { return bShouldIgnoreTransient; }
+	virtual void HandleObjectReference(UObject*& Object, const UObject* ReferencingObject, const UObject* ReferencingProperty) override;
+	virtual bool IsIgnoringArchetypeRef() const override { return bShouldIgnoreArchetype; }
+	virtual bool IsIgnoringTransient() const override { return bShouldIgnoreTransient; }
 
 protected:
 
diff --git a/Engine/Source/Runtime/CoreUObject/Public/UObject/UTextProperty.h b/Engine/Source/Runtime/CoreUObject/Public/UObject/UTextProperty.h
--- a/Engine/Source/Runtime/CoreUObject/Public/UObject/UTextProperty.h
+++ b/Engine/Source/Runtime/CoreUObject/Public/UObject/UTextProperty.h
@@ -21,9 +21,9 @@ public:
 	}
 
 	// UProperty interface
-	virtual bool Identical( const void* A, const void* B, uint32 PortFlags ) const OVERRIDE;
-	virtual void SerializeItem( FArchive& Ar, void* Value, int32 MaxReadBytes, void const* Defaults ) const OVERRIDE;
-	virtual void ExportTextItem( FString& ValueStr, const void* PropertyValue, const void* DefaultValue, UObject* Parent, int32 PortFlags, UObject* ExportRootScope ) const OVERRIDE;
-	virtual const TCHAR* ImportText_Internal( const TCHAR* Buffer, void* Data, int32 PortFlags, UObject* OwnerObject, FOutputDevice* ErrorText ) const OVERRIDE;
+	virtual bool Identical( const void* A, const void* B, uint32 PortFlags ) const override;
+	virtual void SerializeItem( FArchive& Ar, void* Value, int32 MaxReadBytes, void const* Defaults ) const override;
+	virtual void ExportTextItem( FString& ValueStr, const void* PropertyValue, const void* DefaultValue, UObject* Parent, int32 PortFlags, UObject* ExportRootScope ) const override;
+	virtual const TCHAR* ImportText_Internal( const TCHAR* Buffer, void* Data, int32 PortFlags, UObject* OwnerObject, FOutputDevice* ErrorText ) const override;
 	// End of UProperty interface
 };
\ No newline at end of file
