UE4 Refactoring. Changed OVERRIDE and FINAL macros to keywords override and final.

diff --git a/Engine/Source/Runtime/CoreUObject/Public/UObject/UnrealType.h b/Engine/Source/Runtime/CoreUObject/Public/UObject/UnrealType.h
--- a/Engine/Source/Runtime/CoreUObject/Public/UObject/UnrealType.h
+++ b/Engine/Source/Runtime/CoreUObject/Public/UObject/UnrealType.h
@@ -160,7 +160,7 @@ public:
 	UProperty( const class FPostConstructInitializeProperties& PCIP, ECppProperty, int32 InOffset, uint64 InFlags );
 
 	// UObject interface
-	virtual void Serialize( FArchive& Ar ) OVERRIDE;
+	virtual void Serialize( FArchive& Ar ) override;
 	// End of UObject interface
 
 	/** parses and imports a text definition of a single property's value (if array, may be an individual element)
@@ -946,11 +946,11 @@ public:
 	}
 
 	// UHT interface
-	virtual FString GetCPPType( FString* ExtendedTypeText=NULL, uint32 CPPExportFlags=0 ) const OVERRIDE
+	virtual FString GetCPPType( FString* ExtendedTypeText=NULL, uint32 CPPExportFlags=0 ) const override
 	{
 		return FString(TTypeFundamentals::GetTypeName());
 	}
-	virtual bool PassCPPArgsByRef() const OVERRIDE
+	virtual bool PassCPPArgsByRef() const override
 	{
 		// non-pod data is passed by reference
 		return !TIsPODType<TCppType>::Value;
@@ -958,32 +958,32 @@ public:
 	// End of UHT interface
 
 	// UProperty interface.
-	virtual int32 GetMinAlignment() const OVERRIDE
+	virtual int32 GetMinAlignment() const override
 	{
 		return TTypeFundamentals::CPPAlignment;
 	}
-	virtual void LinkInternal(FArchive& Ar) OVERRIDE
+	virtual void LinkInternal(FArchive& Ar) override
 	{
 		SetElementSize();
 		this->PropertyFlags |= TTypeFundamentals::GetComputedFlagsPropertyFlags();
 
 	}
-	virtual void CopyValuesInternal( void* Dest, void const* Src, int32 Count ) const OVERRIDE
+	virtual void CopyValuesInternal( void* Dest, void const* Src, int32 Count ) const override
 	{
 		for (int32 Index = 0; Index < Count; Index++)
 		{
 			TTypeFundamentals::GetPropertyValuePtr(Dest)[Index] = TTypeFundamentals::GetPropertyValuePtr(Src)[Index];
 		}
 	}
-	virtual void ClearValueInternal( void* Data ) const OVERRIDE
+	virtual void ClearValueInternal( void* Data ) const override
 	{
 		TTypeFundamentals::SetPropertyValue(Data, TTypeFundamentals::GetDefaultPropertyValue());
 	}
-	virtual void InitializeValueInternal( void* Dest ) const OVERRIDE
+	virtual void InitializeValueInternal( void* Dest ) const override
 	{
 		TTypeFundamentals::IntializePropertyValue(Dest);
 	}
-	virtual void DestroyValueInternal( void* Dest ) const OVERRIDE
+	virtual void DestroyValueInternal( void* Dest ) const override
 	{
 		for( int32 i = 0; i < this->ArrayDim; ++i )
 		{
@@ -1044,11 +1044,11 @@ public:
 	}
 
 	// UProperty interface.
-	virtual bool Identical( const void* A, const void* B, uint32 PortFlags=0 ) const OVERRIDE
+	virtual bool Identical( const void* A, const void* B, uint32 PortFlags=0 ) const override
 	{
 		return TTypeFundamentals::GetPropertyValue(A) == TTypeFundamentals::GetOptionalPropertyValue(B);
 	}
-	virtual void SerializeItem( FArchive& Ar, void* Value, int32 MaxReadBytes, void const* Defaults ) const OVERRIDE
+	virtual void SerializeItem( FArchive& Ar, void* Value, int32 MaxReadBytes, void const* Defaults ) const override
 	{
 		Ar << *TTypeFundamentals::GetPropertyValuePtr(Value);
 	}
@@ -1065,9 +1065,9 @@ class COREUOBJECT_API UNumericProperty : public UProperty
 	{}
 
 	// UProperty interface.
-	virtual const TCHAR* ImportText_Internal( const TCHAR* Buffer, void* Data, int32 PortFlags, UObject* Parent, FOutputDevice* ErrorText ) const OVERRIDE;
+	virtual const TCHAR* ImportText_Internal( const TCHAR* Buffer, void* Data, int32 PortFlags, UObject* Parent, FOutputDevice* ErrorText ) const override;
 
-	virtual void ExportTextItem( FString& ValueStr, const void* PropertyValue, const void* DefaultValue, UObject* Parent, int32 PortFlags, UObject* ExportRootScope ) const OVERRIDE;
+	virtual void ExportTextItem( FString& ValueStr, const void* PropertyValue, const void* DefaultValue, UObject* Parent, int32 PortFlags, UObject* ExportRootScope ) const override;
 	// End of UProperty interface
 
 	// UNumericProperty interface.
@@ -1203,48 +1203,48 @@ public:
 
 	// UNumericProperty interface.
 
-	virtual bool IsFloatingPoint() const OVERRIDE
+	virtual bool IsFloatingPoint() const override
 	{
 		return TIsFloatType<TCppType>::Value;
 	}
-	virtual bool IsInteger() const OVERRIDE
+	virtual bool IsInteger() const override
 	{
 		return TIsIntegralType<TCppType>::Value;
 	}
-	virtual void SetIntPropertyValue(void* Data, uint64 Value) const OVERRIDE
+	virtual void SetIntPropertyValue(void* Data, uint64 Value) const override
 	{
 		check(TIsIntegralType<TCppType>::Value);
 		TTypeFundamentals::SetPropertyValue(Data, Value);
 	}
-	virtual void SetIntPropertyValue(void* Data, int64 Value) const OVERRIDE
+	virtual void SetIntPropertyValue(void* Data, int64 Value) const override
 	{
 		check(TIsIntegralType<TCppType>::Value);
 		TTypeFundamentals::SetPropertyValue(Data, Value);
 	}
-	virtual void SetFloatingPointPropertyValue(void* Data, double Value) const OVERRIDE
+	virtual void SetFloatingPointPropertyValue(void* Data, double Value) const override
 	{
 		check(TIsFloatType<TCppType>::Value);
 		TTypeFundamentals::SetPropertyValue(Data, Value);
 	}
-	virtual void SetNumericPropertyValueFromString(void * Data, TCHAR const* Value) const OVERRIDE
+	virtual void SetNumericPropertyValueFromString(void * Data, TCHAR const* Value) const override
 	{
 		TTypeFundamentals::FromString(Data, Value);
 	}
-	virtual FString GetNumericPropertyValueToString(void const* Data) const OVERRIDE
+	virtual FString GetNumericPropertyValueToString(void const* Data) const override
 	{
 		return TTypeFundamentals::ToString(Data);
 	}
-	virtual int64 GetSignedIntPropertyValue(void const* Data) const OVERRIDE
+	virtual int64 GetSignedIntPropertyValue(void const* Data) const override
 	{
 		check(TIsIntegralType<TCppType>::Value);
 		return TTypeFundamentals::GetPropertyValue(Data);
 	}
-	virtual uint64 GetUnsignedIntPropertyValue(void const* Data) const OVERRIDE
+	virtual uint64 GetUnsignedIntPropertyValue(void const* Data) const override
 	{
 		check(TIsIntegralType<TCppType>::Value);
 		return TTypeFundamentals::GetPropertyValue(Data);
 	}
-	virtual double GetFloatingPointPropertyValue(void const* Data) const OVERRIDE
+	virtual double GetFloatingPointPropertyValue(void const* Data) const override
 	{
 		check(TIsFloatType<TCppType>::Value);
 		return TTypeFundamentals::GetPropertyValue(Data);
@@ -1274,23 +1274,23 @@ class COREUOBJECT_API UByteProperty : public TProperty_Numeric<uint8>
 	}
 
 	// UObject interface.
-	virtual void Serialize( FArchive& Ar ) OVERRIDE;
+	virtual void Serialize( FArchive& Ar ) override;
 	static void AddReferencedObjects(UObject* InThis, FReferenceCollector& Collector);
 	// End of UObject interface
 
 	// UHT interface
-	virtual FString GetCPPType( FString* ExtendedTypeText=NULL, uint32 CPPExportFlags=0 ) const OVERRIDE;
+	virtual FString GetCPPType( FString* ExtendedTypeText=NULL, uint32 CPPExportFlags=0 ) const override;
 	// End of UHT interface
 
 	// UProperty interface.
-	virtual void SerializeItem( FArchive& Ar, void* Value, int32 MaxReadBytes, void const* Defaults ) const OVERRIDE;
-	virtual bool NetSerializeItem( FArchive& Ar, UPackageMap* Map, void* Data, TArray<uint8> * MetaData = NULL ) const OVERRIDE;
-	virtual void ExportTextItem( FString& ValueStr, const void* PropertyValue, const void* DefaultValue, UObject* Parent, int32 PortFlags, UObject* ExportRootScope ) const OVERRIDE;
-	virtual const TCHAR* ImportText_Internal( const TCHAR* Buffer, void* Data, int32 PortFlags, UObject* Parent, FOutputDevice* ErrorText ) const OVERRIDE;
+	virtual void SerializeItem( FArchive& Ar, void* Value, int32 MaxReadBytes, void const* Defaults ) const override;
+	virtual bool NetSerializeItem( FArchive& Ar, UPackageMap* Map, void* Data, TArray<uint8> * MetaData = NULL ) const override;
+	virtual void ExportTextItem( FString& ValueStr, const void* PropertyValue, const void* DefaultValue, UObject* Parent, int32 PortFlags, UObject* ExportRootScope ) const override;
+	virtual const TCHAR* ImportText_Internal( const TCHAR* Buffer, void* Data, int32 PortFlags, UObject* Parent, FOutputDevice* ErrorText ) const override;
 	// End of UProperty interface
 
 	// UNumericProperty interface.
-	virtual UEnum* GetIntPropertyEnum() const OVERRIDE
+	virtual UEnum* GetIntPropertyEnum() const override
 	{
 		return Enum;
 	}
@@ -1496,25 +1496,25 @@ public:
 	UBoolProperty( const class FPostConstructInitializeProperties& PCIP, ECppProperty, int32 InOffset, uint64 InFlags, uint32 InBitMask, uint32 InElementSize, bool bIsNativeBool );
 
 	// UObject interface.
-	virtual void Serialize( FArchive& Ar ) OVERRIDE;
+	virtual void Serialize( FArchive& Ar ) override;
 	// End of UObject interface
 
 	// UHT interface
-	virtual FString GetCPPType( FString* ExtendedTypeText, uint32 CPPExportFlags ) const OVERRIDE;
-	virtual FString GetCPPMacroType( FString& ExtendedTypeText ) const OVERRIDE;
+	virtual FString GetCPPType( FString* ExtendedTypeText, uint32 CPPExportFlags ) const override;
+	virtual FString GetCPPMacroType( FString& ExtendedTypeText ) const override;
 	// End of UHT interface
 
 	// UProperty interface.
-	virtual void LinkInternal(FArchive& Ar) OVERRIDE;
-	virtual bool Identical( const void* A, const void* B, uint32 PortFlags ) const OVERRIDE;
-	virtual void SerializeItem( FArchive& Ar, void* Value, int32 MaxReadBytes, void const* Defaults ) const OVERRIDE;
-	virtual bool NetSerializeItem( FArchive& Ar, UPackageMap* Map, void* Data, TArray<uint8> * MetaData = NULL ) const OVERRIDE;
-	virtual void ExportTextItem( FString& ValueStr, const void* PropertyValue, const void* DefaultValue, UObject* Parent, int32 PortFlags, UObject* ExportRootScope ) const OVERRIDE;
-	virtual const TCHAR* ImportText_Internal( const TCHAR* Buffer, void* Data, int32 PortFlags, UObject* Parent, FOutputDevice* ErrorText ) const OVERRIDE;
-	virtual void CopyValuesInternal( void* Dest, void const* Src, int32 Count ) const OVERRIDE;
-	virtual void ClearValueInternal( void* Data ) const OVERRIDE;
-	virtual void InitializeValueInternal( void* Dest ) const OVERRIDE;
-	virtual int32 GetMinAlignment() const OVERRIDE;
+	virtual void LinkInternal(FArchive& Ar) override;
+	virtual bool Identical( const void* A, const void* B, uint32 PortFlags ) const override;
+	virtual void SerializeItem( FArchive& Ar, void* Value, int32 MaxReadBytes, void const* Defaults ) const override;
+	virtual bool NetSerializeItem( FArchive& Ar, UPackageMap* Map, void* Data, TArray<uint8> * MetaData = NULL ) const override;
+	virtual void ExportTextItem( FString& ValueStr, const void* PropertyValue, const void* DefaultValue, UObject* Parent, int32 PortFlags, UObject* ExportRootScope ) const override;
+	virtual const TCHAR* ImportText_Internal( const TCHAR* Buffer, void* Data, int32 PortFlags, UObject* Parent, FOutputDevice* ErrorText ) const override;
+	virtual void CopyValuesInternal( void* Dest, void const* Src, int32 Count ) const override;
+	virtual void ClearValueInternal( void* Data ) const override;
+	virtual void InitializeValueInternal( void* Dest ) const override;
+	virtual int32 GetMinAlignment() const override;
 	// End of UProperty interface
 
 	// Emulate the CPP type API, see TPropertyTypeFundamentals
@@ -1592,18 +1592,18 @@ class COREUOBJECT_API UObjectPropertyBase : public UProperty
 	{}
 
 	// UObject interface
-	virtual void Serialize( FArchive& Ar ) OVERRIDE;
+	virtual void Serialize( FArchive& Ar ) override;
 	static void AddReferencedObjects(UObject* InThis, FReferenceCollector& Collector);
 	// End of UObject interface
 
 	// UProperty interface
-	virtual bool Identical( const void* A, const void* B, uint32 PortFlags ) const OVERRIDE;
-	virtual bool NetSerializeItem( FArchive& Ar, UPackageMap* Map, void* Data, TArray<uint8> * MetaData = NULL ) const OVERRIDE;
-	virtual void ExportTextItem( FString& ValueStr, const void* PropertyValue, const void* DefaultValue, UObject* Parent, int32 PortFlags, UObject* ExportRootScope ) const OVERRIDE;
-	virtual const TCHAR* ImportText_Internal( const TCHAR* Buffer, void* Data, int32 PortFlags, UObject* OwnerObject, FOutputDevice* ErrorText ) const OVERRIDE;
-	virtual FName GetID() const OVERRIDE;
-	virtual void InstanceSubobjects( void* Data, void const* DefaultData, UObject* Owner, struct FObjectInstancingGraph* InstanceGraph ) OVERRIDE;
-	virtual bool SameType(const UProperty* Other) const OVERRIDE;
+	virtual bool Identical( const void* A, const void* B, uint32 PortFlags ) const override;
+	virtual bool NetSerializeItem( FArchive& Ar, UPackageMap* Map, void* Data, TArray<uint8> * MetaData = NULL ) const override;
+	virtual void ExportTextItem( FString& ValueStr, const void* PropertyValue, const void* DefaultValue, UObject* Parent, int32 PortFlags, UObject* ExportRootScope ) const override;
+	virtual const TCHAR* ImportText_Internal( const TCHAR* Buffer, void* Data, int32 PortFlags, UObject* OwnerObject, FOutputDevice* ErrorText ) const override;
+	virtual FName GetID() const override;
+	virtual void InstanceSubobjects( void* Data, void const* DefaultData, UObject* Owner, struct FObjectInstancingGraph* InstanceGraph ) override;
+	virtual bool SameType(const UProperty* Other) const override;
 	/**
 	 * Copy the value for a single element of this property. To the script VM.
 	 * 
@@ -1735,11 +1735,11 @@ public:
 	}
 
 	// UProperty interface.
-	virtual bool ContainsObjectReference() const OVERRIDE
+	virtual bool ContainsObjectReference() const override
 	{
 		return !TIsWeakPointerType<InTCppType>::Value;
 	}
-	virtual bool ContainsWeakObjectReference() const OVERRIDE
+	virtual bool ContainsWeakObjectReference() const override
 	{
 		return TIsWeakPointerType<InTCppType>::Value;
 	}
@@ -1760,20 +1760,20 @@ class COREUOBJECT_API UObjectProperty : public TUObjectPropertyBase<UObject*>
 	}
 
 	// UHT interface
-	virtual FString GetCPPMacroType( FString& ExtendedTypeText ) const  OVERRIDE;
-	virtual FString GetCPPType( FString* ExtendedTypeText, uint32 CPPExportFlags ) const OVERRIDE;
+	virtual FString GetCPPMacroType( FString& ExtendedTypeText ) const  override;
+	virtual FString GetCPPType( FString* ExtendedTypeText, uint32 CPPExportFlags ) const override;
 	// End of UHT interface
 
 	// UProperty interface
-	virtual void SerializeItem( FArchive& Ar, void* Value, int32 MaxReadBytes, void const* Defaults ) const OVERRIDE;
-	virtual void EmitReferenceInfo(UClass& OwnerClass, int32 BaseOffset) OVERRIDE;
+	virtual void SerializeItem( FArchive& Ar, void* Value, int32 MaxReadBytes, void const* Defaults ) const override;
+	virtual void EmitReferenceInfo(UClass& OwnerClass, int32 BaseOffset) override;
 	// End of UProperty interface
 	// UObjectPropertyBase interface
-	virtual UObject* GetObjectPropertyValue(const void* PropertyValueAddress) const OVERRIDE
+	virtual UObject* GetObjectPropertyValue(const void* PropertyValueAddress) const override
 	{
 		return GetPropertyValue(PropertyValueAddress);
 	}
-	virtual void SetObjectPropertyValue(void* PropertyValueAddress, UObject* Value) const OVERRIDE
+	virtual void SetObjectPropertyValue(void* PropertyValueAddress, UObject* Value) const override
 	{
 		SetPropertyValue(PropertyValueAddress, Value);
 	}
@@ -1793,20 +1793,20 @@ class COREUOBJECT_API UWeakObjectProperty : public TUObjectPropertyBase<FWeakObj
 	}
 	
 	// UHT interface
-	virtual FString GetCPPMacroType( FString& ExtendedTypeText ) const  OVERRIDE;
-	virtual FString GetCPPType( FString* ExtendedTypeText, uint32 CPPExportFlags ) const OVERRIDE;
+	virtual FString GetCPPMacroType( FString& ExtendedTypeText ) const  override;
+	virtual FString GetCPPType( FString* ExtendedTypeText, uint32 CPPExportFlags ) const override;
 	// End of UHT interface
 
 	// UProperty interface
-	virtual void SerializeItem( FArchive& Ar, void* Value, int32 MaxReadBytes, void const* Defaults ) const OVERRIDE;
+	virtual void SerializeItem( FArchive& Ar, void* Value, int32 MaxReadBytes, void const* Defaults ) const override;
 	// End of UProperty interface
 
 	// UObjectProperty interface
-	virtual UObject* GetObjectPropertyValue(const void* PropertyValueAddress) const OVERRIDE
+	virtual UObject* GetObjectPropertyValue(const void* PropertyValueAddress) const override
 	{
 		return GetPropertyValue(PropertyValueAddress).Get();
 	}
-	virtual void SetObjectPropertyValue(void* PropertyValueAddress, UObject* Value) const OVERRIDE
+	virtual void SetObjectPropertyValue(void* PropertyValueAddress, UObject* Value) const override
 	{
 		SetPropertyValue(PropertyValueAddress, TCppType(Value));
 	}
@@ -1826,26 +1826,26 @@ class COREUOBJECT_API ULazyObjectProperty : public TUObjectPropertyBase<FLazyObj
 	}
 
 	// UHT interface
-	virtual FString GetCPPMacroType( FString& ExtendedTypeText ) const  OVERRIDE;
-	virtual FString GetCPPType( FString* ExtendedTypeText, uint32 CPPExportFlags ) const OVERRIDE;
+	virtual FString GetCPPMacroType( FString& ExtendedTypeText ) const  override;
+	virtual FString GetCPPType( FString* ExtendedTypeText, uint32 CPPExportFlags ) const override;
 	// End of UHT interface
 
 	// UProperty interface
-	virtual FName GetID() const OVERRIDE;
-	virtual bool Identical( const void* A, const void* B, uint32 PortFlags ) const OVERRIDE;
-	virtual void SerializeItem( FArchive& Ar, void* Value, int32 MaxReadBytes, void const* Defaults ) const OVERRIDE;
+	virtual FName GetID() const override;
+	virtual bool Identical( const void* A, const void* B, uint32 PortFlags ) const override;
+	virtual void SerializeItem( FArchive& Ar, void* Value, int32 MaxReadBytes, void const* Defaults ) const override;
 	// End of UProperty interface
 
 	// UObjectProperty interface
-	virtual UObject* GetObjectPropertyValue(const void* PropertyValueAddress) const OVERRIDE
+	virtual UObject* GetObjectPropertyValue(const void* PropertyValueAddress) const override
 	{
 		return GetPropertyValue(PropertyValueAddress).Get();
 	}
-	virtual void SetObjectPropertyValue(void* PropertyValueAddress, UObject* Value) const OVERRIDE
+	virtual void SetObjectPropertyValue(void* PropertyValueAddress, UObject* Value) const override
 	{
 		SetPropertyValue(PropertyValueAddress, TCppType(Value));
 	}
-	virtual bool AllowCrossLevel() const OVERRIDE
+	virtual bool AllowCrossLevel() const override
 	{
 		return true;
 	}
@@ -1864,28 +1864,28 @@ class COREUOBJECT_API UAssetObjectProperty : public TUObjectPropertyBase<FAssetP
 	{}
 
 	// UHT interface
-	virtual FString GetCPPMacroType( FString& ExtendedTypeText ) const  OVERRIDE;
-	virtual FString GetCPPType( FString* ExtendedTypeText, uint32 CPPExportFlags ) const OVERRIDE;
+	virtual FString GetCPPMacroType( FString& ExtendedTypeText ) const  override;
+	virtual FString GetCPPType( FString* ExtendedTypeText, uint32 CPPExportFlags ) const override;
 	// End of UHT interface
 
 	// UProperty interface
-	virtual FName GetID() const OVERRIDE;
-	virtual bool Identical( const void* A, const void* B, uint32 PortFlags ) const OVERRIDE;
-	virtual void SerializeItem( FArchive& Ar, void* Value, int32 MaxReadBytes, void const* Defaults ) const OVERRIDE;
-	virtual void ExportTextItem( FString& ValueStr, const void* PropertyValue, const void* DefaultValue, UObject* Parent, int32 PortFlags, UObject* ExportRootScope ) const OVERRIDE;
-	virtual const TCHAR* ImportText_Internal( const TCHAR* Buffer, void* Data, int32 PortFlags, UObject* OwnerObject, FOutputDevice* ErrorText ) const OVERRIDE;
+	virtual FName GetID() const override;
+	virtual bool Identical( const void* A, const void* B, uint32 PortFlags ) const override;
+	virtual void SerializeItem( FArchive& Ar, void* Value, int32 MaxReadBytes, void const* Defaults ) const override;
+	virtual void ExportTextItem( FString& ValueStr, const void* PropertyValue, const void* DefaultValue, UObject* Parent, int32 PortFlags, UObject* ExportRootScope ) const override;
+	virtual const TCHAR* ImportText_Internal( const TCHAR* Buffer, void* Data, int32 PortFlags, UObject* OwnerObject, FOutputDevice* ErrorText ) const override;
 	// End of UProperty interface
 
 	// UObjectProperty interface
-	virtual UObject* GetObjectPropertyValue(const void* PropertyValueAddress) const OVERRIDE
+	virtual UObject* GetObjectPropertyValue(const void* PropertyValueAddress) const override
 	{
 		return GetPropertyValue(PropertyValueAddress).Get();
 	}
-	virtual void SetObjectPropertyValue(void* PropertyValueAddress, UObject* Value) const OVERRIDE
+	virtual void SetObjectPropertyValue(void* PropertyValueAddress, UObject* Value) const override
 	{
 		SetPropertyValue(PropertyValueAddress, TCppType(Value));
 	}
-	virtual bool AllowCrossLevel() const OVERRIDE
+	virtual bool AllowCrossLevel() const override
 	{
 		return true;
 	}
@@ -1913,22 +1913,22 @@ public:
 	}
 
 	// UObject interface
-	virtual void Serialize( FArchive& Ar ) OVERRIDE;
+	virtual void Serialize( FArchive& Ar ) override;
 	static void AddReferencedObjects(UObject* InThis, FReferenceCollector& Collector);
 	// End of UObject interface
 
 	// UHT interface
-	virtual FString GetCPPMacroType( FString& ExtendedTypeText ) const  OVERRIDE;
-	virtual FString GetCPPType( FString* ExtendedTypeText, uint32 CPPExportFlags ) const OVERRIDE;
+	virtual FString GetCPPMacroType( FString& ExtendedTypeText ) const  override;
+	virtual FString GetCPPType( FString* ExtendedTypeText, uint32 CPPExportFlags ) const override;
 	// End of UHT interface
 
 	// UProperty interface
-	virtual const TCHAR* ImportText_Internal( const TCHAR* Buffer, void* Data, int32 PortFlags, UObject* OwnerObject, FOutputDevice* ErrorText ) const OVERRIDE;
-	virtual bool SameType(const UProperty* Other) const OVERRIDE;
+	virtual const TCHAR* ImportText_Internal( const TCHAR* Buffer, void* Data, int32 PortFlags, UObject* OwnerObject, FOutputDevice* ErrorText ) const override;
+	virtual bool SameType(const UProperty* Other) const override;
 	// End of UProperty interface
 
 protected:
-	virtual void CheckValidObject(void* Value) const OVERRIDE;
+	virtual void CheckValidObject(void* Value) const override;
 };
 
 /*-----------------------------------------------------------------------------
@@ -1951,17 +1951,17 @@ public:
 	{}
 
 	// UHT interface
-	virtual FString GetCPPMacroType( FString& ExtendedTypeText ) const  OVERRIDE;
-	virtual FString GetCPPType( FString* ExtendedTypeText, uint32 CPPExportFlags ) const OVERRIDE;
+	virtual FString GetCPPMacroType( FString& ExtendedTypeText ) const  override;
+	virtual FString GetCPPType( FString* ExtendedTypeText, uint32 CPPExportFlags ) const override;
 	// End of UHT interface
 
 	// UObject interface
-	virtual void Serialize( FArchive& Ar ) OVERRIDE;
+	virtual void Serialize( FArchive& Ar ) override;
 	static void AddReferencedObjects(UObject* InThis, FReferenceCollector& Collector);
 	// End of UObject interface
 
 	// UProperty interface
-	virtual bool SameType(const UProperty* Other) const OVERRIDE;
+	virtual bool SameType(const UProperty* Other) const override;
 	// End of UProperty interface
 };
 
@@ -1994,24 +1994,24 @@ public:
 	}
 
 	// UHT interface
-	virtual FString GetCPPMacroType( FString& ExtendedTypeText ) const  OVERRIDE;
-	virtual FString GetCPPType( FString* ExtendedTypeText, uint32 CPPExportFlags ) const OVERRIDE;
+	virtual FString GetCPPMacroType( FString& ExtendedTypeText ) const  override;
+	virtual FString GetCPPType( FString* ExtendedTypeText, uint32 CPPExportFlags ) const override;
 	// End of UHT interface
 
 	// UProperty interface
-	virtual void LinkInternal(FArchive& Ar) OVERRIDE;
-	virtual bool Identical( const void* A, const void* B, uint32 PortFlags ) const OVERRIDE;
-	virtual void SerializeItem( FArchive& Ar, void* Value, int32 MaxReadBytes, void const* Defaults ) const OVERRIDE;
-	virtual bool NetSerializeItem( FArchive& Ar, UPackageMap* Map, void* Data, TArray<uint8> * MetaData = NULL ) const OVERRIDE;
-	virtual void ExportTextItem( FString& ValueStr, const void* PropertyValue, const void* DefaultValue, UObject* Parent, int32 PortFlags, UObject* ExportRootScope ) const OVERRIDE;
-	virtual const TCHAR* ImportText_Internal( const TCHAR* Buffer, void* Data, int32 PortFlags, UObject* OwnerObject, FOutputDevice* ErrorText ) const OVERRIDE;
-	virtual bool ContainsObjectReference() const OVERRIDE;
-	virtual bool SameType(const UProperty* Other) const OVERRIDE;
+	virtual void LinkInternal(FArchive& Ar) override;
+	virtual bool Identical( const void* A, const void* B, uint32 PortFlags ) const override;
+	virtual void SerializeItem( FArchive& Ar, void* Value, int32 MaxReadBytes, void const* Defaults ) const override;
+	virtual bool NetSerializeItem( FArchive& Ar, UPackageMap* Map, void* Data, TArray<uint8> * MetaData = NULL ) const override;
+	virtual void ExportTextItem( FString& ValueStr, const void* PropertyValue, const void* DefaultValue, UObject* Parent, int32 PortFlags, UObject* ExportRootScope ) const override;
+	virtual const TCHAR* ImportText_Internal( const TCHAR* Buffer, void* Data, int32 PortFlags, UObject* OwnerObject, FOutputDevice* ErrorText ) const override;
+	virtual bool ContainsObjectReference() const override;
+	virtual bool SameType(const UProperty* Other) const override;
 	// End of UProperty interface
 
 	// UObject interface
-	virtual void Serialize( FArchive& Ar ) OVERRIDE;
-	virtual void EmitReferenceInfo(UClass& OwnerClass, int32 BaseOffset) OVERRIDE;
+	virtual void Serialize( FArchive& Ar ) override;
+	virtual void EmitReferenceInfo(UClass& OwnerClass, int32 BaseOffset) override;
 	// End of UObject interface
 };
 
@@ -2039,8 +2039,8 @@ public:
 	}
 
 	// UProperty interface
-	virtual void ExportTextItem( FString& ValueStr, const void* PropertyValue, const void* DefaultValue, UObject* Parent, int32 PortFlags, UObject* ExportRootScope ) const OVERRIDE;
-	virtual const TCHAR* ImportText_Internal( const TCHAR* Buffer, void* Data, int32 PortFlags, UObject* OwnerObject, FOutputDevice* ErrorText ) const OVERRIDE;
+	virtual void ExportTextItem( FString& ValueStr, const void* PropertyValue, const void* DefaultValue, UObject* Parent, int32 PortFlags, UObject* ExportRootScope ) const override;
+	virtual const TCHAR* ImportText_Internal( const TCHAR* Buffer, void* Data, int32 PortFlags, UObject* OwnerObject, FOutputDevice* ErrorText ) const override;
 	// End of UProperty interface
 };
 
@@ -2068,8 +2068,8 @@ public:
 	}
 
 	// UProperty interface
-	virtual void ExportTextItem( FString& ValueStr, const void* PropertyValue, const void* DefaultValue, UObject* Parent, int32 PortFlags, UObject* ExportRootScope ) const OVERRIDE;
-	virtual const TCHAR* ImportText_Internal( const TCHAR* Buffer, void* Data, int32 PortFlags, UObject* OwnerObject, FOutputDevice* ErrorText ) const OVERRIDE;
+	virtual void ExportTextItem( FString& ValueStr, const void* PropertyValue, const void* DefaultValue, UObject* Parent, int32 PortFlags, UObject* ExportRootScope ) const override;
+	virtual const TCHAR* ImportText_Internal( const TCHAR* Buffer, void* Data, int32 PortFlags, UObject* OwnerObject, FOutputDevice* ErrorText ) const override;
 	// End of UProperty interface
 };
 
@@ -2101,33 +2101,33 @@ public:
 	}
 
 	// UObject interface
-	virtual void Serialize( FArchive& Ar ) OVERRIDE;
+	virtual void Serialize( FArchive& Ar ) override;
 	static void AddReferencedObjects(UObject* InThis, FReferenceCollector& Collector);
 	// End of UObject interface
 
 	// UField interface
-	virtual void AddCppProperty( UProperty* Property ) OVERRIDE;
+	virtual void AddCppProperty( UProperty* Property ) override;
 	// End of UField interface
 
 	// UProperty interface
-	virtual FString GetCPPMacroType( FString& ExtendedTypeText ) const  OVERRIDE;
-	virtual FString GetCPPType( FString* ExtendedTypeText, uint32 CPPExportFlags ) const OVERRIDE;
-	virtual void LinkInternal(FArchive& Ar) OVERRIDE;
-	virtual bool Identical( const void* A, const void* B, uint32 PortFlags ) const OVERRIDE;
-	virtual void SerializeItem( FArchive& Ar, void* Value, int32 MaxReadBytes, void const* Defaults ) const OVERRIDE;
-	virtual bool NetSerializeItem( FArchive& Ar, UPackageMap* Map, void* Data, TArray<uint8> * MetaData = NULL ) const OVERRIDE;
-	virtual void ExportTextItem( FString& ValueStr, const void* PropertyValue, const void* DefaultValue, UObject* Parent, int32 PortFlags, UObject* ExportRootScope ) const OVERRIDE;
-	virtual const TCHAR* ImportText_Internal( const TCHAR* Buffer, void* Data, int32 PortFlags, UObject* OwnerObject, FOutputDevice* ErrorText ) const OVERRIDE;
-	virtual void CopyValuesInternal( void* Dest, void const* Src, int32 Count  ) const OVERRIDE;
-	virtual void ClearValueInternal( void* Data ) const OVERRIDE;
-	virtual void DestroyValueInternal( void* Dest ) const OVERRIDE;
-	virtual bool IsLocalized() const OVERRIDE;
-	virtual bool PassCPPArgsByRef() const OVERRIDE;
-	virtual void InstanceSubobjects( void* Data, void const* DefaultData, UObject* Owner, struct FObjectInstancingGraph* InstanceGraph ) OVERRIDE;
-	virtual bool ContainsObjectReference() const OVERRIDE;
-	virtual bool ContainsWeakObjectReference() const OVERRIDE;
-	virtual void EmitReferenceInfo(UClass& OwnerClass, int32 BaseOffset) OVERRIDE;
-	virtual bool SameType(const UProperty* Other) const OVERRIDE;
+	virtual FString GetCPPMacroType( FString& ExtendedTypeText ) const  override;
+	virtual FString GetCPPType( FString* ExtendedTypeText, uint32 CPPExportFlags ) const override;
+	virtual void LinkInternal(FArchive& Ar) override;
+	virtual bool Identical( const void* A, const void* B, uint32 PortFlags ) const override;
+	virtual void SerializeItem( FArchive& Ar, void* Value, int32 MaxReadBytes, void const* Defaults ) const override;
+	virtual bool NetSerializeItem( FArchive& Ar, UPackageMap* Map, void* Data, TArray<uint8> * MetaData = NULL ) const override;
+	virtual void ExportTextItem( FString& ValueStr, const void* PropertyValue, const void* DefaultValue, UObject* Parent, int32 PortFlags, UObject* ExportRootScope ) const override;
+	virtual const TCHAR* ImportText_Internal( const TCHAR* Buffer, void* Data, int32 PortFlags, UObject* OwnerObject, FOutputDevice* ErrorText ) const override;
+	virtual void CopyValuesInternal( void* Dest, void const* Src, int32 Count  ) const override;
+	virtual void ClearValueInternal( void* Data ) const override;
+	virtual void DestroyValueInternal( void* Dest ) const override;
+	virtual bool IsLocalized() const override;
+	virtual bool PassCPPArgsByRef() const override;
+	virtual void InstanceSubobjects( void* Data, void const* DefaultData, UObject* Owner, struct FObjectInstancingGraph* InstanceGraph ) override;
+	virtual bool ContainsObjectReference() const override;
+	virtual bool ContainsWeakObjectReference() const override;
+	virtual void EmitReferenceInfo(UClass& OwnerClass, int32 BaseOffset) override;
+	virtual bool SameType(const UProperty* Other) const override;
 	// End of UProperty interface
 };
 
@@ -2462,30 +2462,30 @@ public:
 	UStructProperty( const class FPostConstructInitializeProperties& PCIP, ECppProperty, int32 InOffset, uint64 InFlags, UScriptStruct* InStruct );
 
 	// UObject interface
-	virtual void Serialize( FArchive& Ar ) OVERRIDE;
+	virtual void Serialize( FArchive& Ar ) override;
 	static void AddReferencedObjects(UObject* InThis, FReferenceCollector& Collector);
 	// End of UObject interface
 
 	// UProperty interface
-	virtual FString GetCPPMacroType( FString& ExtendedTypeText ) const  OVERRIDE;
-	virtual FString GetCPPType( FString* ExtendedTypeText, uint32 CPPExportFlags ) const OVERRIDE;
-	virtual void LinkInternal(FArchive& Ar) OVERRIDE;
-	virtual bool Identical( const void* A, const void* B, uint32 PortFlags ) const OVERRIDE;
-	virtual void SerializeItem( FArchive& Ar, void* Value, int32 MaxReadBytes, void const* Defaults ) const OVERRIDE;
-	virtual bool NetSerializeItem( FArchive& Ar, UPackageMap* Map, void* Data, TArray<uint8> * MetaData = NULL ) const OVERRIDE;
-	virtual void ExportTextItem( FString& ValueStr, const void* PropertyValue, const void* DefaultValue, UObject* Parent, int32 PortFlags, UObject* ExportRootScope ) const OVERRIDE;
-	virtual const TCHAR* ImportText_Internal( const TCHAR* Buffer, void* Data, int32 PortFlags, UObject* OwnerObject, FOutputDevice* ErrorText ) const OVERRIDE;
-	virtual void CopyValuesInternal( void* Dest, void const* Src, int32 Count  ) const OVERRIDE;
-	virtual void ClearValueInternal( void* Data ) const OVERRIDE;
-	virtual void DestroyValueInternal( void* Dest ) const OVERRIDE;
-	virtual void InitializeValueInternal( void* Dest ) const OVERRIDE;
-	virtual bool IsLocalized() const OVERRIDE;
-	virtual void InstanceSubobjects( void* Data, void const* DefaultData, UObject* Owner, struct FObjectInstancingGraph* InstanceGraph ) OVERRIDE;
-	virtual int32 GetMinAlignment() const OVERRIDE;
-	virtual bool ContainsObjectReference() const OVERRIDE;
-	virtual bool ContainsWeakObjectReference() const OVERRIDE;
-	virtual void EmitReferenceInfo(UClass& OwnerClass, int32 BaseOffset) OVERRIDE;
-	virtual bool SameType(const UProperty* Other) const OVERRIDE;
+	virtual FString GetCPPMacroType( FString& ExtendedTypeText ) const  override;
+	virtual FString GetCPPType( FString* ExtendedTypeText, uint32 CPPExportFlags ) const override;
+	virtual void LinkInternal(FArchive& Ar) override;
+	virtual bool Identical( const void* A, const void* B, uint32 PortFlags ) const override;
+	virtual void SerializeItem( FArchive& Ar, void* Value, int32 MaxReadBytes, void const* Defaults ) const override;
+	virtual bool NetSerializeItem( FArchive& Ar, UPackageMap* Map, void* Data, TArray<uint8> * MetaData = NULL ) const override;
+	virtual void ExportTextItem( FString& ValueStr, const void* PropertyValue, const void* DefaultValue, UObject* Parent, int32 PortFlags, UObject* ExportRootScope ) const override;
+	virtual const TCHAR* ImportText_Internal( const TCHAR* Buffer, void* Data, int32 PortFlags, UObject* OwnerObject, FOutputDevice* ErrorText ) const override;
+	virtual void CopyValuesInternal( void* Dest, void const* Src, int32 Count  ) const override;
+	virtual void ClearValueInternal( void* Data ) const override;
+	virtual void DestroyValueInternal( void* Dest ) const override;
+	virtual void InitializeValueInternal( void* Dest ) const override;
+	virtual bool IsLocalized() const override;
+	virtual void InstanceSubobjects( void* Data, void const* DefaultData, UObject* Owner, struct FObjectInstancingGraph* InstanceGraph ) override;
+	virtual int32 GetMinAlignment() const override;
+	virtual bool ContainsObjectReference() const override;
+	virtual bool ContainsWeakObjectReference() const override;
+	virtual void EmitReferenceInfo(UClass& OwnerClass, int32 BaseOffset) override;
+	virtual bool SameType(const UProperty* Other) const override;
 	// End of UProperty interface
 
 	bool UseNativeSerialization() const;
@@ -2536,19 +2536,19 @@ public:
 	}
 
 	// UObject interface
-	virtual void Serialize( FArchive& Ar ) OVERRIDE;
+	virtual void Serialize( FArchive& Ar ) override;
 	// End of UObject interface
 
 	// UProperty interface
-	virtual FString GetCPPType( FString* ExtendedTypeText, uint32 CPPExportFlags ) const OVERRIDE;
-	virtual bool Identical( const void* A, const void* B, uint32 PortFlags ) const OVERRIDE;
-	virtual void SerializeItem( FArchive& Ar, void* Value, int32 MaxReadBytes, void const* Defaults ) const OVERRIDE;
-	virtual bool NetSerializeItem( FArchive& Ar, UPackageMap* Map, void* Data, TArray<uint8> * MetaData = NULL ) const OVERRIDE;
-	virtual void ExportTextItem( FString& ValueStr, const void* PropertyValue, const void* DefaultValue, UObject* Parent, int32 PortFlags, UObject* ExportRootScope ) const OVERRIDE;
-	virtual const TCHAR* ImportText_Internal( const TCHAR* Buffer, void* Data, int32 PortFlags, UObject* OwnerObject, FOutputDevice* ErrorText ) const OVERRIDE;
-	virtual bool ContainsWeakObjectReference() const OVERRIDE;
-	virtual void InstanceSubobjects( void* Data, void const* DefaultData, UObject* Owner, struct FObjectInstancingGraph* InstanceGraph ) OVERRIDE;
-	virtual bool SameType(const UProperty* Other) const OVERRIDE;
+	virtual FString GetCPPType( FString* ExtendedTypeText, uint32 CPPExportFlags ) const override;
+	virtual bool Identical( const void* A, const void* B, uint32 PortFlags ) const override;
+	virtual void SerializeItem( FArchive& Ar, void* Value, int32 MaxReadBytes, void const* Defaults ) const override;
+	virtual bool NetSerializeItem( FArchive& Ar, UPackageMap* Map, void* Data, TArray<uint8> * MetaData = NULL ) const override;
+	virtual void ExportTextItem( FString& ValueStr, const void* PropertyValue, const void* DefaultValue, UObject* Parent, int32 PortFlags, UObject* ExportRootScope ) const override;
+	virtual const TCHAR* ImportText_Internal( const TCHAR* Buffer, void* Data, int32 PortFlags, UObject* OwnerObject, FOutputDevice* ErrorText ) const override;
+	virtual bool ContainsWeakObjectReference() const override;
+	virtual void InstanceSubobjects( void* Data, void const* DefaultData, UObject* Owner, struct FObjectInstancingGraph* InstanceGraph ) override;
+	virtual bool SameType(const UProperty* Other) const override;
 	// End of UProperty interface
 };
 
@@ -2581,19 +2581,19 @@ public:
 	}
 
 	// UObject interface
-	virtual void Serialize( FArchive& Ar ) OVERRIDE;
+	virtual void Serialize( FArchive& Ar ) override;
 	// End of UObject interface
 
 	// UProperty interface
-	virtual FString GetCPPType( FString* ExtendedTypeText, uint32 CPPExportFlags ) const OVERRIDE;
-	virtual bool Identical( const void* A, const void* B, uint32 PortFlags ) const OVERRIDE;
-	virtual void SerializeItem( FArchive& Ar, void* Value, int32 MaxReadBytes, void const* Defaults ) const OVERRIDE;
-	virtual bool NetSerializeItem( FArchive& Ar, UPackageMap* Map, void* Data, TArray<uint8> * MetaData = NULL ) const OVERRIDE;
-	virtual void ExportTextItem( FString& ValueStr, const void* PropertyValue, const void* DefaultValue, UObject* Parent, int32 PortFlags, UObject* ExportRootScope ) const OVERRIDE;
-	virtual const TCHAR* ImportText_Internal( const TCHAR* Buffer, void* Data, int32 PortFlags, UObject* OwnerObject, FOutputDevice* ErrorText ) const OVERRIDE;
-	virtual bool ContainsWeakObjectReference() const OVERRIDE;
-	virtual void InstanceSubobjects( void* Data, void const* DefaultData, UObject* Owner, struct FObjectInstancingGraph* InstanceGraph ) OVERRIDE;
-	virtual bool SameType(const UProperty* Other) const OVERRIDE;
+	virtual FString GetCPPType( FString* ExtendedTypeText, uint32 CPPExportFlags ) const override;
+	virtual bool Identical( const void* A, const void* B, uint32 PortFlags ) const override;
+	virtual void SerializeItem( FArchive& Ar, void* Value, int32 MaxReadBytes, void const* Defaults ) const override;
+	virtual bool NetSerializeItem( FArchive& Ar, UPackageMap* Map, void* Data, TArray<uint8> * MetaData = NULL ) const override;
+	virtual void ExportTextItem( FString& ValueStr, const void* PropertyValue, const void* DefaultValue, UObject* Parent, int32 PortFlags, UObject* ExportRootScope ) const override;
+	virtual const TCHAR* ImportText_Internal( const TCHAR* Buffer, void* Data, int32 PortFlags, UObject* OwnerObject, FOutputDevice* ErrorText ) const override;
+	virtual bool ContainsWeakObjectReference() const override;
+	virtual void InstanceSubobjects( void* Data, void const* DefaultData, UObject* Owner, struct FObjectInstancingGraph* InstanceGraph ) override;
+	virtual bool SameType(const UProperty* Other) const override;
 	// End of UProperty interface
 
 protected:
diff --git a/Engine/Source/Runtime/Engine/Classes/AI/Navigation/NavAreas/NavArea.h b/Engine/Source/Runtime/Engine/Classes/AI/Navigation/NavAreas/NavArea.h
--- a/Engine/Source/Runtime/Engine/Classes/AI/Navigation/NavAreas/NavArea.h
+++ b/Engine/Source/Runtime/Engine/Classes/AI/Navigation/NavAreas/NavArea.h
@@ -65,8 +65,8 @@ class ENGINE_API UNavArea : public UObject
 	};
 #endif
 
-	virtual void PostInitProperties() OVERRIDE;
-	virtual void FinishDestroy() OVERRIDE;
+	virtual void PostInitProperties() override;
+	virtual void FinishDestroy() override;
 
 	FORCEINLINE uint16 GetAreaFlags() const { return AreaFlags; }
 	FORCEINLINE bool HasFlags(uint16 InFlags) const { return (InFlags & AreaFlags) != 0; }
diff --git a/Engine/Source/Runtime/Engine/Classes/AI/Navigation/NavAreas/NavAreaMeta_SwitchByAgent.h b/Engine/Source/Runtime/Engine/Classes/AI/Navigation/NavAreas/NavAreaMeta_SwitchByAgent.h
--- a/Engine/Source/Runtime/Engine/Classes/AI/Navigation/NavAreas/NavAreaMeta_SwitchByAgent.h
+++ b/Engine/Source/Runtime/Engine/Classes/AI/Navigation/NavAreas/NavAreaMeta_SwitchByAgent.h
@@ -58,10 +58,10 @@ class ENGINE_API UNavAreaMeta_SwitchByAgent : public UNavAreaMeta
 	UPROPERTY(EditAnywhere, Category=AgentTypes)
 	TSubclassOf<UNavArea> Agent15Area;
 
-	virtual TSubclassOf<UNavArea> PickAreaClass(const class AActor* Actor, const FNavAgentProperties& NavAgent) const OVERRIDE;
+	virtual TSubclassOf<UNavArea> PickAreaClass(const class AActor* Actor, const FNavAgentProperties& NavAgent) const override;
 
 #if WITH_EDITOR
 	/** setup AgentXArea properties */
-	virtual void UpdateAgentConfig() OVERRIDE;
+	virtual void UpdateAgentConfig() override;
 #endif
 };
diff --git a/Engine/Source/Runtime/Engine/Classes/AI/Navigation/NavCollision.h b/Engine/Source/Runtime/Engine/Classes/AI/Navigation/NavCollision.h
--- a/Engine/Source/Runtime/Engine/Classes/AI/Navigation/NavCollision.h
+++ b/Engine/Source/Runtime/Engine/Classes/AI/Navigation/NavCollision.h
@@ -76,9 +76,9 @@ class ENGINE_API UNavCollision : public UObject
 	FFormatContainer CookedFormatData;
 
 	// Begin UObject interface.
-	virtual void Serialize(FArchive& Ar) OVERRIDE;
-	virtual void PostLoad() OVERRIDE;
-	virtual SIZE_T GetResourceSize(EResourceSizeMode::Type Mode) OVERRIDE;
+	virtual void Serialize(FArchive& Ar) override;
+	virtual void PostLoad() override;
+	virtual SIZE_T GetResourceSize(EResourceSizeMode::Type Mode) override;
 	// End UObject interface.
 
 	FGuid GetGuid() const;
diff --git a/Engine/Source/Runtime/Engine/Classes/AI/Navigation/NavFilters/RecastFilter_UseDefaultArea.h b/Engine/Source/Runtime/Engine/Classes/AI/Navigation/NavFilters/RecastFilter_UseDefaultArea.h
--- a/Engine/Source/Runtime/Engine/Classes/AI/Navigation/NavFilters/RecastFilter_UseDefaultArea.h
+++ b/Engine/Source/Runtime/Engine/Classes/AI/Navigation/NavFilters/RecastFilter_UseDefaultArea.h
@@ -9,5 +9,5 @@ class URecastFilter_UseDefaultArea : public UNavigationQueryFilter
 {
 	GENERATED_UCLASS_BODY()
 
-	virtual void InitializeFilter(const ANavigationData* NavData, FNavigationQueryFilter* Filter) const OVERRIDE;
+	virtual void InitializeFilter(const ANavigationData* NavData, FNavigationQueryFilter* Filter) const override;
 };
diff --git a/Engine/Source/Runtime/Engine/Classes/AI/Navigation/NavLinkCustomComponent.h b/Engine/Source/Runtime/Engine/Classes/AI/Navigation/NavLinkCustomComponent.h
--- a/Engine/Source/Runtime/Engine/Classes/AI/Navigation/NavLinkCustomComponent.h
+++ b/Engine/Source/Runtime/Engine/Classes/AI/Navigation/NavLinkCustomComponent.h
@@ -24,19 +24,19 @@ class ENGINE_API UNavLinkCustomComponent : public UNavRelevantComponent, public
 	DECLARE_DELEGATE_TwoParams(FBroadcastFilter, UNavLinkCustomComponent* /*ThisComp*/, TArray<class UNavigationComponent*>& /*NotifyList*/);
 
 	// BEGIN INavLinkCustomInterface
-	virtual void GetLinkData(FVector& LeftPt, FVector& RightPt, ENavLinkDirection::Type& Direction) const OVERRIDE;
-	virtual TSubclassOf<UNavArea> GetLinkAreaClass() const OVERRIDE;
-	virtual uint32 GetLinkId() const OVERRIDE;
-	virtual bool IsLinkPathfindingAllowed(const UObject* Querier) const OVERRIDE;
-	virtual bool OnLinkMoveStarted(class UPathFollowingComponent* PathComp, const FVector& DestPoint) OVERRIDE;
-	virtual void OnLinkMoveFinished(class UPathFollowingComponent* PathComp) OVERRIDE;
+	virtual void GetLinkData(FVector& LeftPt, FVector& RightPt, ENavLinkDirection::Type& Direction) const override;
+	virtual TSubclassOf<UNavArea> GetLinkAreaClass() const override;
+	virtual uint32 GetLinkId() const override;
+	virtual bool IsLinkPathfindingAllowed(const UObject* Querier) const override;
+	virtual bool OnLinkMoveStarted(class UPathFollowingComponent* PathComp, const FVector& DestPoint) override;
+	virtual void OnLinkMoveFinished(class UPathFollowingComponent* PathComp) override;
 	// END INavLinkCustomInterface
 
 	// BEGIN UNavRelevantComponent
-	virtual void OnRegister() OVERRIDE;
-	virtual void OnOwnerRegistered() OVERRIDE;
-	virtual void OnOwnerUnregistered() OVERRIDE;
-	virtual void OnApplyModifiers(struct FCompositeNavModifier& Modifiers) OVERRIDE;
+	virtual void OnRegister() override;
+	virtual void OnOwnerRegistered() override;
+	virtual void OnOwnerUnregistered() override;
+	virtual void OnApplyModifiers(struct FCompositeNavModifier& Modifiers) override;
 	// END UNavRelevantComponent
 
 	/** set basic link data: end points and direction */
diff --git a/Engine/Source/Runtime/Engine/Classes/AI/Navigation/NavLinkProxy.h b/Engine/Source/Runtime/Engine/Classes/AI/Navigation/NavLinkProxy.h
--- a/Engine/Source/Runtime/Engine/Classes/AI/Navigation/NavLinkProxy.h
+++ b/Engine/Source/Runtime/Engine/Classes/AI/Navigation/NavLinkProxy.h
@@ -38,21 +38,21 @@ class ENGINE_API ANavLinkProxy : public AActor, public INavLinkHostInterface, pu
 #endif // WITH_EDITORONLY_DATA
 
 	// BEGIN INavRelevantActorInterface
-	virtual bool UpdateNavigationRelevancy() OVERRIDE;
-	virtual bool GetNavigationRelevantData(struct FNavigationRelevantData& Data) const OVERRIDE;
+	virtual bool UpdateNavigationRelevancy() override;
+	virtual bool GetNavigationRelevantData(struct FNavigationRelevantData& Data) const override;
 	// END INavRelevantActorInterface
 
 	// BEGIN INavLinkHostInterface
-	virtual bool GetNavigationLinksClasses(TArray<TSubclassOf<class UNavLinkDefinition> >& OutClasses) const OVERRIDE;
-	virtual bool GetNavigationLinksArray(TArray<FNavigationLink>& OutLink, TArray<FNavigationSegmentLink>& OutSegments) const OVERRIDE;
+	virtual bool GetNavigationLinksClasses(TArray<TSubclassOf<class UNavLinkDefinition> >& OutClasses) const override;
+	virtual bool GetNavigationLinksArray(TArray<FNavigationLink>& OutLink, TArray<FNavigationSegmentLink>& OutSegments) const override;
 	// END INavLinkHostInterface
 
 #if WITH_EDITOR
-	virtual void PostEditChangeProperty(FPropertyChangedEvent& PropertyChangedEvent) OVERRIDE;
+	virtual void PostEditChangeProperty(FPropertyChangedEvent& PropertyChangedEvent) override;
 #endif // WITH_EDITOR
 
-	virtual void PostInitProperties() OVERRIDE;
-	virtual FBox GetComponentsBoundingBox(bool bNonColliding = false) const OVERRIDE;
+	virtual void PostInitProperties() override;
+	virtual FBox GetComponentsBoundingBox(bool bNonColliding = false) const override;
 
 	//////////////////////////////////////////////////////////////////////////
 	// Blueprint interface for smart links
diff --git a/Engine/Source/Runtime/Engine/Classes/AI/Navigation/NavLinkRenderingComponent.h b/Engine/Source/Runtime/Engine/Classes/AI/Navigation/NavLinkRenderingComponent.h
--- a/Engine/Source/Runtime/Engine/Classes/AI/Navigation/NavLinkRenderingComponent.h
+++ b/Engine/Source/Runtime/Engine/Classes/AI/Navigation/NavLinkRenderingComponent.h
@@ -9,14 +9,14 @@ class ENGINE_API UNavLinkRenderingComponent : public UPrimitiveComponent
 	GENERATED_UCLASS_BODY()
 		
 	// Begin UPrimitiveComponent Interface
-	virtual FPrimitiveSceneProxy* CreateSceneProxy() OVERRIDE;
+	virtual FPrimitiveSceneProxy* CreateSceneProxy() override;
 
 	/** Should recreate proxy one very update */
-	virtual bool ShouldRecreateProxyOnUpdateTransform() const OVERRIDE { return true; }
+	virtual bool ShouldRecreateProxyOnUpdateTransform() const override { return true; }
 	// End UPrimitiveComponent Interface
 
 	// Begin USceneComponent Interface
-	virtual FBoxSphereBounds CalcBounds(const FTransform &LocalToWorld) const OVERRIDE;
+	virtual FBoxSphereBounds CalcBounds(const FTransform &LocalToWorld) const override;
 	// End USceneComponent Interface
 };
 
diff --git a/Engine/Source/Runtime/Engine/Classes/AI/Navigation/NavMeshBoundsVolume.h b/Engine/Source/Runtime/Engine/Classes/AI/Navigation/NavMeshBoundsVolume.h
--- a/Engine/Source/Runtime/Engine/Classes/AI/Navigation/NavMeshBoundsVolume.h
+++ b/Engine/Source/Runtime/Engine/Classes/AI/Navigation/NavMeshBoundsVolume.h
@@ -15,12 +15,12 @@ class ANavMeshBoundsVolume : public AVolume
 	GENERATED_UCLASS_BODY()
 
 	// Begin AActor Interface
-	virtual void PostInitializeComponents() OVERRIDE;
+	virtual void PostInitializeComponents() override;
 #if WITH_EDITOR
-	virtual void ReregisterAllComponents() OVERRIDE;
+	virtual void ReregisterAllComponents() override;
 	// End AActor Interface
 	// Begin UObject Interface
-	virtual void PostEditChangeProperty( struct FPropertyChangedEvent& PropertyChangedEvent) OVERRIDE;
+	virtual void PostEditChangeProperty( struct FPropertyChangedEvent& PropertyChangedEvent) override;
 	// End UObject Interface
 
 #endif // WITH_EDITOR
diff --git a/Engine/Source/Runtime/Engine/Classes/AI/Navigation/NavMeshRenderingComponent.h b/Engine/Source/Runtime/Engine/Classes/AI/Navigation/NavMeshRenderingComponent.h
--- a/Engine/Source/Runtime/Engine/Classes/AI/Navigation/NavMeshRenderingComponent.h
+++ b/Engine/Source/Runtime/Engine/Classes/AI/Navigation/NavMeshRenderingComponent.h
@@ -9,16 +9,16 @@ class ENGINE_API UNavMeshRenderingComponent : public UPrimitiveComponent
 	GENERATED_UCLASS_BODY()
 
 	// Begin UPrimitiveComponent Interface
-	virtual FPrimitiveSceneProxy* CreateSceneProxy() OVERRIDE;
+	virtual FPrimitiveSceneProxy* CreateSceneProxy() override;
 	// End UPrimitiveComponent Interface
 
 	// Begin UActorComponent Interface
-	virtual void CreateRenderState_Concurrent() OVERRIDE;
-	virtual void DestroyRenderState_Concurrent() OVERRIDE;
+	virtual void CreateRenderState_Concurrent() override;
+	virtual void DestroyRenderState_Concurrent() override;
 	// End UActorComponent Interface
 
 	// Begin USceneComponent Interface
-	virtual FBoxSphereBounds CalcBounds(const FTransform &LocalToWorld) const OVERRIDE;
+	virtual FBoxSphereBounds CalcBounds(const FTransform &LocalToWorld) const override;
 	// End USceneComponent Interface
 
 	void GatherData(struct FNavMeshSceneProxyData*) const;
diff --git a/Engine/Source/Runtime/Engine/Classes/AI/Navigation/NavModifierVolume.h b/Engine/Source/Runtime/Engine/Classes/AI/Navigation/NavModifierVolume.h
--- a/Engine/Source/Runtime/Engine/Classes/AI/Navigation/NavModifierVolume.h
+++ b/Engine/Source/Runtime/Engine/Classes/AI/Navigation/NavModifierVolume.h
@@ -15,6 +15,6 @@ class ANavModifierVolume : public AVolume, public INavRelevantActorInterface
 	UPROPERTY(EditAnywhere, BlueprintReadWrite, Category=Default)
 	TSubclassOf<class UNavArea> AreaClass;
 
-	virtual bool GetNavigationRelevantData(struct FNavigationRelevantData& Data) const OVERRIDE;
-	virtual bool UpdateNavigationRelevancy() OVERRIDE { SetNavigationRelevancy(true); return true; }
+	virtual bool GetNavigationRelevantData(struct FNavigationRelevantData& Data) const override;
+	virtual bool UpdateNavigationRelevancy() override { SetNavigationRelevancy(true); return true; }
 };
diff --git a/Engine/Source/Runtime/Engine/Classes/AI/Navigation/NavRelevantComponent.h b/Engine/Source/Runtime/Engine/Classes/AI/Navigation/NavRelevantComponent.h
--- a/Engine/Source/Runtime/Engine/Classes/AI/Navigation/NavRelevantComponent.h
+++ b/Engine/Source/Runtime/Engine/Classes/AI/Navigation/NavRelevantComponent.h
@@ -12,8 +12,8 @@ class ENGINE_API UNavRelevantComponent : public UActorComponent
 	virtual void OnOwnerRegistered();
 	virtual void OnOwnerUnregistered();
 
-	virtual void OnRegister() OVERRIDE;
-	virtual void OnUnregister() OVERRIDE;
+	virtual void OnRegister() override;
+	virtual void OnUnregister() override;
 
 	virtual void OnApplyModifiers(struct FCompositeNavModifier& Modifiers);
 
diff --git a/Engine/Source/Runtime/Engine/Classes/AI/Navigation/NavTestRenderingComponent.h b/Engine/Source/Runtime/Engine/Classes/AI/Navigation/NavTestRenderingComponent.h
--- a/Engine/Source/Runtime/Engine/Classes/AI/Navigation/NavTestRenderingComponent.h
+++ b/Engine/Source/Runtime/Engine/Classes/AI/Navigation/NavTestRenderingComponent.h
@@ -8,8 +8,8 @@ class UNavTestRenderingComponent: public UPrimitiveComponent
 {
 	GENERATED_UCLASS_BODY()
 
-	virtual FPrimitiveSceneProxy* CreateSceneProxy() OVERRIDE;
-	virtual FBoxSphereBounds CalcBounds(const FTransform &LocalToWorld) const OVERRIDE;
-	virtual void CreateRenderState_Concurrent() OVERRIDE;
-	virtual void DestroyRenderState_Concurrent() OVERRIDE;
+	virtual FPrimitiveSceneProxy* CreateSceneProxy() override;
+	virtual FBoxSphereBounds CalcBounds(const FTransform &LocalToWorld) const override;
+	virtual void CreateRenderState_Concurrent() override;
+	virtual void DestroyRenderState_Concurrent() override;
 };
diff --git a/Engine/Source/Runtime/Engine/Classes/AI/Navigation/NavigationData.h b/Engine/Source/Runtime/Engine/Classes/AI/Navigation/NavigationData.h
--- a/Engine/Source/Runtime/Engine/Classes/AI/Navigation/NavigationData.h
+++ b/Engine/Source/Runtime/Engine/Classes/AI/Navigation/NavigationData.h
@@ -59,13 +59,13 @@ class ENGINE_API ANavigationData : public AActor
 	virtual ~ANavigationData();
 
 	// Begin UObject/AActor Interface
-	virtual void PostInitProperties() OVERRIDE;
-	virtual void PostInitializeComponents() OVERRIDE;
-	virtual void PostLoad() OVERRIDE;
+	virtual void PostInitProperties() override;
+	virtual void PostInitializeComponents() override;
+	virtual void PostLoad() override;
 #if WITH_EDITOR
-	virtual void PostEditUndo() OVERRIDE;
+	virtual void PostEditUndo() override;
 #endif // WITH_EDITOR
-	virtual void Destroyed() OVERRIDE;
+	virtual void Destroyed() override;
 	// End UObject Interface
 		
 	virtual void CleanUp();
@@ -77,9 +77,9 @@ class ENGINE_API ANavigationData : public AActor
 	
 	FORCEINLINE uint16 GetNavDataUniqueID() const { return NavDataUniqueID; }
 
-	virtual void TickActor(float DeltaTime, enum ELevelTick TickType, FActorTickFunction& ThisTickFunction) OVERRIDE;
+	virtual void TickActor(float DeltaTime, enum ELevelTick TickType, FActorTickFunction& ThisTickFunction) override;
 
-	virtual void RerunConstructionScripts() OVERRIDE;
+	virtual void RerunConstructionScripts() override;
 
 	virtual bool NeedsRebuild() { return false; }
 
diff --git a/Engine/Source/Runtime/Engine/Classes/AI/Navigation/NavigationGraph.h b/Engine/Source/Runtime/Engine/Classes/AI/Navigation/NavigationGraph.h
--- a/Engine/Source/Runtime/Engine/Classes/AI/Navigation/NavigationGraph.h
+++ b/Engine/Source/Runtime/Engine/Classes/AI/Navigation/NavigationGraph.h
@@ -60,7 +60,7 @@ class ANavigationGraph : public ANavigationData
 
 protected:
 #if WITH_NAVIGATION_GENERATOR
-	virtual FNavDataGenerator* ConstructGenerator(const FNavAgentProperties& AgentProps) OVERRIDE;
+	virtual FNavDataGenerator* ConstructGenerator(const FNavAgentProperties& AgentProps) override;
 #endif // WITH_NAVIGATION_GENERATOR
 
 public:
diff --git a/Engine/Source/Runtime/Engine/Classes/AI/Navigation/NavigationGraphNodeComponent.h b/Engine/Source/Runtime/Engine/Classes/AI/Navigation/NavigationGraphNodeComponent.h
--- a/Engine/Source/Runtime/Engine/Classes/AI/Navigation/NavigationGraphNodeComponent.h
+++ b/Engine/Source/Runtime/Engine/Classes/AI/Navigation/NavigationGraphNodeComponent.h
@@ -19,5 +19,5 @@ class UNavigationGraphNodeComponent : public USceneComponent
 	UNavigationGraphNodeComponent* PrevNodeComponent;
 
 	// Begin UObject interface.
-	virtual void BeginDestroy() OVERRIDE;
+	virtual void BeginDestroy() override;
 };
diff --git a/Engine/Source/Runtime/Engine/Classes/AI/Navigation/NavigationSystem.h b/Engine/Source/Runtime/Engine/Classes/AI/Navigation/NavigationSystem.h
--- a/Engine/Source/Runtime/Engine/Classes/AI/Navigation/NavigationSystem.h
+++ b/Engine/Source/Runtime/Engine/Classes/AI/Navigation/NavigationSystem.h
@@ -60,7 +60,7 @@ struct FNavigationSystemExec: public FSelfRegisteringExec
 	}
 
 	// Begin FExec Interface
-	virtual bool Exec(UWorld* Inworld, const TCHAR* Cmd, FOutputDevice& Ar) OVERRIDE;
+	virtual bool Exec(UWorld* Inworld, const TCHAR* Cmd, FOutputDevice& Ar) override;
 	// End FExec Interface
 };
 
@@ -212,7 +212,7 @@ public:
 	};
 
 	// Begin UObject Interface
-	virtual void PostInitProperties() OVERRIDE;
+	virtual void PostInitProperties() override;
 	static void AddReferencedObjects(UObject* InThis, FReferenceCollector& Collector);
 	// End UObject Interface
 
diff --git a/Engine/Source/Runtime/Engine/Classes/AI/Navigation/NavigationTestingActor.h b/Engine/Source/Runtime/Engine/Classes/AI/Navigation/NavigationTestingActor.h
--- a/Engine/Source/Runtime/Engine/Classes/AI/Navigation/NavigationTestingActor.h
+++ b/Engine/Source/Runtime/Engine/Classes/AI/Navigation/NavigationTestingActor.h
@@ -133,26 +133,26 @@ class ENGINE_API ANavigationTestingActor : public AActor, public INavAgentInterf
 	/** Dtor */
 	virtual ~ANavigationTestingActor();
 
-	virtual void BeginDestroy() OVERRIDE;
+	virtual void BeginDestroy() override;
 
 #if WITH_EDITOR
-	virtual void PreEditChange(UProperty* PropertyThatWillChange) OVERRIDE;
-	virtual void PostEditChangeProperty(FPropertyChangedEvent& PropertyChangedEvent) OVERRIDE;
-	virtual void PostEditMove(bool bFinished) OVERRIDE;
+	virtual void PreEditChange(UProperty* PropertyThatWillChange) override;
+	virtual void PostEditChangeProperty(FPropertyChangedEvent& PropertyChangedEvent) override;
+	virtual void PostEditMove(bool bFinished) override;
 	
-	virtual void PostLoad() OVERRIDE;
+	virtual void PostLoad() override;
 	void TickMe();
 #endif // WITH_EDITOR
 
-	virtual bool UpdateNavigationRelevancy() OVERRIDE { SetNavigationRelevancy(false); return false; }
+	virtual bool UpdateNavigationRelevancy() override { SetNavigationRelevancy(false); return false; }
 	// Begin INavAgentInterface Interface
-	virtual const struct FNavAgentProperties* GetNavAgentProperties() const OVERRIDE { return &NavAgentProps; }
-	virtual FVector GetNavAgentLocation() const OVERRIDE;
-	virtual void GetMoveGoalReachTest(class AActor* MovingActor, const FVector& MoveOffset, FVector& GoalOffset, float& GoalRadius, float& GoalHalfHeight) const OVERRIDE {}
+	virtual const struct FNavAgentProperties* GetNavAgentProperties() const override { return &NavAgentProps; }
+	virtual FVector GetNavAgentLocation() const override;
+	virtual void GetMoveGoalReachTest(class AActor* MovingActor, const FVector& MoveOffset, FVector& GoalOffset, float& GoalRadius, float& GoalHalfHeight) const override {}
 	// End INavAgentInterface Interface
 
 	// Begin INavPathObserverInterface Interface
-	virtual void OnPathUpdated(class INavigationPathGenerator* PathGenerator) OVERRIDE;
+	virtual void OnPathUpdated(class INavigationPathGenerator* PathGenerator) override;
 	// End INavPathObserverInterface Interface	
 
 	void UpdateNavData();
diff --git a/Engine/Source/Runtime/Engine/Classes/AI/Navigation/RecastNavMesh.h b/Engine/Source/Runtime/Engine/Classes/AI/Navigation/RecastNavMesh.h
--- a/Engine/Source/Runtime/Engine/Classes/AI/Navigation/RecastNavMesh.h
+++ b/Engine/Source/Runtime/Engine/Classes/AI/Navigation/RecastNavMesh.h
@@ -92,7 +92,7 @@ struct ENGINE_API FNavMeshPath : public FNavigationPath
 
 	FORCEINLINE void OnPathCorridorUpdated() { bCorridorEdgesGenerated = false; }
 
-	virtual void DebugDraw(const ANavigationData* NavData, FColor PathColor, UCanvas* Canvas, bool bPersistent, const uint32 NextPathPointIndex = 0) const OVERRIDE;
+	virtual void DebugDraw(const ANavigationData* NavData, FColor PathColor, UCanvas* Canvas, bool bPersistent, const uint32 NextPathPointIndex = 0) const override;
 	
 	bool ContainsWithSameEnd(const FNavMeshPath* Other) const;
 	
@@ -133,11 +133,11 @@ struct ENGINE_API FNavMeshPath : public FNavigationPath
 
 	bool IsPathSegmentANavLink(const int32 PathSegmentStartIndex) const;
 
-	virtual bool DoesIntersectBox(const FBox& Box, int32* IntersectingSegmentIndex = NULL) const OVERRIDE;
+	virtual bool DoesIntersectBox(const FBox& Box, int32* IntersectingSegmentIndex = NULL) const override;
 
 #if ENABLE_VISUAL_LOG
-	virtual void DescribeSelfToVisLog(struct FVisLogEntry* Snapshot) const OVERRIDE;
-	virtual FString GetDescription() const OVERRIDE;
+	virtual void DescribeSelfToVisLog(struct FVisLogEntry* Snapshot) const override;
+	virtual FString GetDescription() const override;
 #endif // ENABLE_VISUAL_LOG
 
 protected:
@@ -403,7 +403,7 @@ class ENGINE_API ARecastNavMesh : public ANavigationData
 
 	typedef uint16 FNavPolyFlags;
 
-	virtual void Serialize( FArchive& Ar ) OVERRIDE;
+	virtual void Serialize( FArchive& Ar ) override;
 
 	/** should we draw edges of every navmesh's triangle */
 	UPROPERTY(EditAnywhere, Category=Display)
@@ -624,11 +624,11 @@ public:
 	virtual ~ARecastNavMesh();
 
 	// Begin UObject Interface
-	virtual void PostInitProperties() OVERRIDE;
-	virtual SIZE_T GetResourceSize(EResourceSizeMode::Type Mode) OVERRIDE;	
+	virtual void PostInitProperties() override;
+	virtual SIZE_T GetResourceSize(EResourceSizeMode::Type Mode) override;	
 
 #if WITH_EDITOR
-	virtual void PostEditChangeProperty( struct FPropertyChangedEvent& PropertyChangedEvent) OVERRIDE;
+	virtual void PostEditChangeProperty( struct FPropertyChangedEvent& PropertyChangedEvent) override;
 #endif // WITH_EDITOR
 	// End UObject Interface
 
@@ -637,33 +637,33 @@ public:
 
 #if WITH_EDITOR
 	/** RecastNavMesh instances are dynamically spawned and should not be coppied */
-	virtual bool ShouldExport() OVERRIDE { return false; }
+	virtual bool ShouldExport() override { return false; }
 #endif
 
 	/** called by TickHelper, works in game AND in editor */
 	void TickMe(float DeltaTime);
 
-	virtual void CleanUp() OVERRIDE;
+	virtual void CleanUp() override;
 
 	// Begin ANavigationData Interface
-	virtual FNavLocation GetRandomPoint(TSharedPtr<const FNavigationQueryFilter> Filter = NULL, const UObject* Querier = NULL) const OVERRIDE;
-	virtual bool GetRandomPointInRadius(const FVector& Origin, float Radius, FNavLocation& OutResult, TSharedPtr<const FNavigationQueryFilter> Filter = NULL, const UObject* Querier = NULL) const OVERRIDE;
+	virtual FNavLocation GetRandomPoint(TSharedPtr<const FNavigationQueryFilter> Filter = NULL, const UObject* Querier = NULL) const override;
+	virtual bool GetRandomPointInRadius(const FVector& Origin, float Radius, FNavLocation& OutResult, TSharedPtr<const FNavigationQueryFilter> Filter = NULL, const UObject* Querier = NULL) const override;
 
-	virtual bool ProjectPoint(const FVector& Point, FNavLocation& OutLocation, const FVector& Extent, TSharedPtr<const FNavigationQueryFilter> Filter = NULL, const UObject* Querier = NULL) const OVERRIDE;
+	virtual bool ProjectPoint(const FVector& Point, FNavLocation& OutLocation, const FVector& Extent, TSharedPtr<const FNavigationQueryFilter> Filter = NULL, const UObject* Querier = NULL) const override;
 	
-	virtual ENavigationQueryResult::Type CalcPathCost(const FVector& PathStart, const FVector& PathEnd, float& OutPathCost, TSharedPtr<const FNavigationQueryFilter> Filter = NULL, const UObject* Querier = NULL) const OVERRIDE;
-	virtual ENavigationQueryResult::Type CalcPathLength(const FVector& PathStart, const FVector& PathEnd, float& OutPathLength, TSharedPtr<const FNavigationQueryFilter> QueryFilter = NULL, const UObject* Querier = NULL) const OVERRIDE;
-	virtual ENavigationQueryResult::Type CalcPathLengthAndCost(const FVector& PathStart, const FVector& PathEnd, float& OutPathLength, float& OutPathCost, TSharedPtr<const FNavigationQueryFilter> QueryFilter = NULL, const UObject* Querier = NULL) const OVERRIDE;
+	virtual ENavigationQueryResult::Type CalcPathCost(const FVector& PathStart, const FVector& PathEnd, float& OutPathCost, TSharedPtr<const FNavigationQueryFilter> Filter = NULL, const UObject* Querier = NULL) const override;
+	virtual ENavigationQueryResult::Type CalcPathLength(const FVector& PathStart, const FVector& PathEnd, float& OutPathLength, TSharedPtr<const FNavigationQueryFilter> QueryFilter = NULL, const UObject* Querier = NULL) const override;
+	virtual ENavigationQueryResult::Type CalcPathLengthAndCost(const FVector& PathStart, const FVector& PathEnd, float& OutPathLength, float& OutPathCost, TSharedPtr<const FNavigationQueryFilter> QueryFilter = NULL, const UObject* Querier = NULL) const override;
 
-	virtual UPrimitiveComponent* ConstructRenderingComponent() OVERRIDE;
+	virtual UPrimitiveComponent* ConstructRenderingComponent() override;
 	/** Returns bounding box for the navmesh. */
-	virtual FBox GetBounds() const OVERRIDE { return GetNavMeshBounds(); }
+	virtual FBox GetBounds() const override { return GetNavMeshBounds(); }
 	/** Called on world origin changes **/
-	virtual void ApplyWorldOffset(const FVector& InOffset, bool bWorldShift) OVERRIDE;
+	virtual void ApplyWorldOffset(const FVector& InOffset, bool bWorldShift) override;
 
 protected:
 #if WITH_NAVIGATION_GENERATOR
-	virtual FNavDataGenerator* ConstructGenerator(const FNavAgentProperties& AgentProps) OVERRIDE;
+	virtual FNavDataGenerator* ConstructGenerator(const FNavAgentProperties& AgentProps) override;
 	virtual void SortNavigationTiles();
 #endif // WITH_NAVIGATION_GENERATOR
 	// End ANavigationData Interface
@@ -700,9 +700,9 @@ public:
 	/** Creates a task to be executed on GameThread calling UpdateDrawing */
 	void RequestDrawingUpdate();
 
-	virtual void SetConfig(const FNavDataConfig& Src) OVERRIDE;
+	virtual void SetConfig(const FNavDataConfig& Src) override;
 protected:
-	virtual void FillConfig(FNavDataConfig& Dest) OVERRIDE;
+	virtual void FillConfig(FNavDataConfig& Dest) override;
 
 	FORCEINLINE const FNavigationQueryFilter& GetRightFilterRef(TSharedPtr<const FNavigationQueryFilter> Filter) const 
 	{
@@ -725,7 +725,7 @@ public:
 	void DrawDebugPathCorridor(NavNodeRef const* PathPolys, int32 NumPathPolys, bool bPersistent=true) const;
 
 #if !UE_BUILD_SHIPPING
-	virtual uint32 LogMemUsed() const OVERRIDE;
+	virtual uint32 LogMemUsed() const override;
 #endif // !UE_BUILD_SHIPPING
 
 	void UpdateNavMeshDrawing();
@@ -733,9 +733,9 @@ public:
 	//----------------------------------------------------------------------//
 	// Utilities
 	//----------------------------------------------------------------------//
-	virtual void OnNavAreaAdded(const UClass* NavAreaClass, int32 AgentIndex) OVERRIDE;
-	virtual int32 GetNewAreaID(const UClass* AreaClass) const OVERRIDE;
-	virtual int32 GetMaxSupportedAreas() const OVERRIDE { return RECAST_MAX_AREAS; }
+	virtual void OnNavAreaAdded(const UClass* NavAreaClass, int32 AgentIndex) override;
+	virtual int32 GetNewAreaID(const UClass* AreaClass) const override;
+	virtual int32 GetMaxSupportedAreas() const override { return RECAST_MAX_AREAS; }
 	
 	/** Get forbidden area flags from default query filter */
 	uint16 GetDefaultForbiddenFlags() const;
@@ -749,7 +749,7 @@ public:
 	// Custom navigation links
 	//----------------------------------------------------------------------//
 
-	virtual void UpdateCustomLink(const INavLinkCustomInterface* CustomLink) OVERRIDE;
+	virtual void UpdateCustomLink(const INavLinkCustomInterface* CustomLink) override;
 
 	/** update area class and poly flags for all offmesh links with given UserId */
 	void UpdateNavigationLinkArea(int32 UserId, TSubclassOf<UNavArea> AreaClass) const;
@@ -844,7 +844,7 @@ public:
 	int32 DebugPathfinding(const FPathFindingQuery& Query, TArray<FRecastDebugPathfindingStep>& Steps);
 
 	/** Checks whether this instance of navmesh supports given Agent type */
-	virtual bool DoesSupportAgent(const FNavAgentProperties& AgentProps) const OVERRIDE;
+	virtual bool DoesSupportAgent(const FNavAgentProperties& AgentProps) const override;
 
 	static const FRecastQueryFilter* GetNamedFilter(ERecastNamedFilter::Type FilterType);
 	FORCEINLINE static FNavPolyFlags GetNavLinkFlag() { return NavLinkFlag; }
@@ -858,7 +858,7 @@ public:
 	const FTileSetItem* GetTileSet() const { return TileSet.GetTypedData(); }
 	FTileSetItem* GetTileSet() { return TileSet.GetTypedData(); }
 	
-	virtual bool NeedsRebuild() OVERRIDE;
+	virtual bool NeedsRebuild() override;
 
 	/** update offset for navmesh in recast library, to current one */
 	/** @param Offset - current offset */
@@ -868,7 +868,7 @@ protected:
 	// @todo docuement
 	UPrimitiveComponent* ConstructRenderingComponentImpl();
 
-	virtual void DestroyGenerator() OVERRIDE;
+	virtual void DestroyGenerator() override;
 
 	/** Spawns an ARecastNavMesh instance, and configures it if AgentProps != NULL */
 	static ARecastNavMesh* SpawnInstance(UNavigationSystem* NavSys, const FNavDataConfig* AgentProps = NULL);
diff --git a/Engine/Source/Runtime/Engine/Classes/Animation/AimOffsetBlendSpace.h b/Engine/Source/Runtime/Engine/Classes/Animation/AimOffsetBlendSpace.h
--- a/Engine/Source/Runtime/Engine/Classes/Animation/AimOffsetBlendSpace.h
+++ b/Engine/Source/Runtime/Engine/Classes/Animation/AimOffsetBlendSpace.h
@@ -15,9 +15,9 @@ class UAimOffsetBlendSpace : public UBlendSpace
 {
 	GENERATED_UCLASS_BODY()
 
-	virtual bool IsValidAdditive() const OVERRIDE;
+	virtual bool IsValidAdditive() const override;
 
 	/** Validate sample input. Return true if it's all good to go **/
-	virtual bool ValidateSampleInput(FBlendSample & BlendSample, int32 OriginalIndex=INDEX_NONE) const OVERRIDE;
+	virtual bool ValidateSampleInput(FBlendSample & BlendSample, int32 OriginalIndex=INDEX_NONE) const override;
 };
 
diff --git a/Engine/Source/Runtime/Engine/Classes/Animation/AimOffsetBlendSpace1D.h b/Engine/Source/Runtime/Engine/Classes/Animation/AimOffsetBlendSpace1D.h
--- a/Engine/Source/Runtime/Engine/Classes/Animation/AimOffsetBlendSpace1D.h
+++ b/Engine/Source/Runtime/Engine/Classes/Animation/AimOffsetBlendSpace1D.h
@@ -15,9 +15,9 @@ class UAimOffsetBlendSpace1D : public UBlendSpace1D
 {
 	GENERATED_UCLASS_BODY()
 
-	virtual bool IsValidAdditive() const OVERRIDE;
+	virtual bool IsValidAdditive() const override;
 
 	/** Validate sample input. Return true if it's all good to go **/
-	virtual bool ValidateSampleInput(FBlendSample & BlendSample, int32 OriginalIndex=INDEX_NONE) const OVERRIDE;
+	virtual bool ValidateSampleInput(FBlendSample & BlendSample, int32 OriginalIndex=INDEX_NONE) const override;
 };
 
diff --git a/Engine/Source/Runtime/Engine/Classes/Animation/AnimBlueprint.h b/Engine/Source/Runtime/Engine/Classes/Animation/AnimBlueprint.h
--- a/Engine/Source/Runtime/Engine/Classes/Animation/AnimBlueprint.h
+++ b/Engine/Source/Runtime/Engine/Classes/Animation/AnimBlueprint.h
@@ -73,7 +73,7 @@ class ENGINE_API UAnimBlueprint : public UBlueprint
 
 #if WITH_EDITOR
 
-	virtual UClass* GetBlueprintClass() const OVERRIDE;
+	virtual UClass* GetBlueprintClass() const override;
 
 	// Inspects the hierarchy and looks for an override for the requested node GUID
 	// @param NodeGuid - Guid of the node to search for
@@ -86,7 +86,7 @@ class ENGINE_API UAnimBlueprint : public UBlueprint
 	bool GetAssetOverrides(TArray<FAnimParentNodeAssetOverride*>& OutOverrides);
 
 	// UBlueprint interface
-	virtual bool SupportedByDefaultBlueprintFactory() const OVERRIDE
+	virtual bool SupportedByDefaultBlueprintFactory() const override
 	{
 		return false;
 	}
@@ -117,7 +117,7 @@ class ENGINE_API UAnimBlueprint : public UBlueprint
 		OnOverrideChanged.Broadcast(Override.ParentNodeGuid, Override.NewAsset);
 	}
 
-	virtual void PostLoad() OVERRIDE;
+	virtual void PostLoad() override;
 
 protected:
 	// Broadcast when an override is changed, allowing derived blueprints to be updated
diff --git a/Engine/Source/Runtime/Engine/Classes/Animation/AnimBlueprintGeneratedClass.h b/Engine/Source/Runtime/Engine/Classes/Animation/AnimBlueprintGeneratedClass.h
--- a/Engine/Source/Runtime/Engine/Classes/Animation/AnimBlueprintGeneratedClass.h
+++ b/Engine/Source/Runtime/Engine/Classes/Animation/AnimBlueprintGeneratedClass.h
@@ -272,11 +272,11 @@ class ENGINE_API UAnimBlueprintGeneratedClass : public UBlueprintGeneratedClass
 #endif
 
 	// UStruct interface
-	virtual void Link(FArchive& Ar, bool bRelinkExistingProperties) OVERRIDE;
+	virtual void Link(FArchive& Ar, bool bRelinkExistingProperties) override;
 	// End of UStruct interface
 
 	// UClass interface
-	virtual void PurgeClass(bool bRecompilingOnLoad) OVERRIDE;
+	virtual void PurgeClass(bool bRecompilingOnLoad) override;
 	// End of UClass interface
 
 #if WITH_EDITOR
diff --git a/Engine/Source/Runtime/Engine/Classes/Animation/AnimComposite.h b/Engine/Source/Runtime/Engine/Classes/Animation/AnimComposite.h
--- a/Engine/Source/Runtime/Engine/Classes/Animation/AnimComposite.h
+++ b/Engine/Source/Runtime/Engine/Classes/Animation/AnimComposite.h
@@ -19,8 +19,8 @@ public:
 	struct FAnimTrack AnimationTrack;
 
 #if WITH_EDITOR
-	virtual bool GetAllAnimationSequencesReferred(TArray<UAnimSequence*>& AnimationSequences) OVERRIDE;
-	virtual void ReplaceReferredAnimations(const TMap<UAnimSequence*, UAnimSequence*>& ReplacementMap) OVERRIDE;
+	virtual bool GetAllAnimationSequencesReferred(TArray<UAnimSequence*>& AnimationSequences) override;
+	virtual void ReplaceReferredAnimations(const TMap<UAnimSequence*, UAnimSequence*>& ReplacementMap) override;
 #endif
 };
 
diff --git a/Engine/Source/Runtime/Engine/Classes/Animation/AnimCompress_Automatic.h b/Engine/Source/Runtime/Engine/Classes/Animation/AnimCompress_Automatic.h
--- a/Engine/Source/Runtime/Engine/Classes/Animation/AnimCompress_Automatic.h
+++ b/Engine/Source/Runtime/Engine/Classes/Animation/AnimCompress_Automatic.h
@@ -43,7 +43,7 @@ class UAnimCompress_Automatic : public UAnimCompress
 
 protected:
 	// Begin UAnimCompress Interface
-	virtual void DoReduction(class UAnimSequence* AnimSeq, const TArray<class FBoneData>& BoneData) OVERRIDE;
+	virtual void DoReduction(class UAnimSequence* AnimSeq, const TArray<class FBoneData>& BoneData) override;
 	// Begin UAnimCompress Interface
 };
 
diff --git a/Engine/Source/Runtime/Engine/Classes/Animation/AnimCompress_BitwiseCompressOnly.h b/Engine/Source/Runtime/Engine/Classes/Animation/AnimCompress_BitwiseCompressOnly.h
--- a/Engine/Source/Runtime/Engine/Classes/Animation/AnimCompress_BitwiseCompressOnly.h
+++ b/Engine/Source/Runtime/Engine/Classes/Animation/AnimCompress_BitwiseCompressOnly.h
@@ -17,7 +17,7 @@ class UAnimCompress_BitwiseCompressOnly : public UAnimCompress
 
 protected:
 	// Begin UAnimCompress Interface
-	virtual void DoReduction(class UAnimSequence* AnimSeq, const TArray<class FBoneData>& BoneData) OVERRIDE;
+	virtual void DoReduction(class UAnimSequence* AnimSeq, const TArray<class FBoneData>& BoneData) override;
 	// Begin UAnimCompress Interface
 };
 
diff --git a/Engine/Source/Runtime/Engine/Classes/Animation/AnimCompress_PerTrackCompression.h b/Engine/Source/Runtime/Engine/Classes/Animation/AnimCompress_PerTrackCompression.h
--- a/Engine/Source/Runtime/Engine/Classes/Animation/AnimCompress_PerTrackCompression.h
+++ b/Engine/Source/Runtime/Engine/Classes/Animation/AnimCompress_PerTrackCompression.h
@@ -132,7 +132,7 @@ public:
 
 protected:
 	// Begin UAnimCompress Interface
-	virtual void DoReduction(class UAnimSequence* AnimSeq, const TArray<class FBoneData>& BoneData) OVERRIDE;
+	virtual void DoReduction(class UAnimSequence* AnimSeq, const TArray<class FBoneData>& BoneData) override;
 	// Begin UAnimCompress Interface
 
 	// Begin UAnimCompress_RemoveLinearKeys Interface
@@ -142,13 +142,13 @@ protected:
 		const TArray<FTranslationTrack>& TranslationData,
 		const TArray<FRotationTrack>& RotationData,
 		const TArray<FScaleTrack>& ScaleData,
-		const bool bFinalPass) OVERRIDE;
+		const bool bFinalPass) override;
 	virtual void FilterBeforeMainKeyRemoval(
 		UAnimSequence* AnimSeq, 
 		const TArray<FBoneData>& BoneData, 
 		TArray<FTranslationTrack>& TranslationData,
 		TArray<FRotationTrack>& RotationData,
-		TArray<FScaleTrack>& ScaleData) OVERRIDE;
+		TArray<FScaleTrack>& ScaleData) override;
 	// End UAnimCompress_RemoveLinearKeys Interface
 };
 
diff --git a/Engine/Source/Runtime/Engine/Classes/Animation/AnimMontage.h b/Engine/Source/Runtime/Engine/Classes/Animation/AnimMontage.h
--- a/Engine/Source/Runtime/Engine/Classes/Animation/AnimMontage.h
+++ b/Engine/Source/Runtime/Engine/Classes/Animation/AnimMontage.h
@@ -288,23 +288,23 @@ class UAnimMontage : public UAnimCompositeBase
 
 public:
 	// Begin UObject Interface
-	virtual void PostLoad() OVERRIDE;
+	virtual void PostLoad() override;
 #if WITH_EDITOR
-	virtual void PostEditChangeProperty(FPropertyChangedEvent& PropertyChangedEvent) OVERRIDE;
+	virtual void PostEditChangeProperty(FPropertyChangedEvent& PropertyChangedEvent) override;
 #endif // WITH_EDITOR
 	// End UObject Interface
 
 	// Begin AnimSequenceBase Interface
-	virtual bool IsValidAdditive() const OVERRIDE;
+	virtual bool IsValidAdditive() const override;
 #if WITH_EDITOR
-	virtual EAnimEventTriggerOffsets::Type CalculateOffsetForNotify(float NotifyDisplayTime) const OVERRIDE;
+	virtual EAnimEventTriggerOffsets::Type CalculateOffsetForNotify(float NotifyDisplayTime) const override;
 #endif // WITH_EDITOR
 	// End AnimSequenceBase Interface
 
 #if WITH_EDITOR
 	// Begin UAnimationAsset interface
-	virtual bool GetAllAnimationSequencesReferred(TArray<UAnimSequence*>& AnimationSequences) OVERRIDE;
-	virtual void ReplaceReferredAnimations(const TMap<UAnimSequence*, UAnimSequence*>& ReplacementMap) OVERRIDE;
+	virtual bool GetAllAnimationSequencesReferred(TArray<UAnimSequence*>& AnimationSequences) override;
+	virtual void ReplaceReferredAnimations(const TMap<UAnimSequence*, UAnimSequence*>& ReplacementMap) override;
 	// End UAnimationAsset interface
 
 	/** Calculates what (if any) offset should be applied to the trigger time of a branch point given its display time */
@@ -377,7 +377,7 @@ private:
 #if WITH_EDITOR
 public:
 	// UAnimSequenceBase Interface
-	virtual void EvaluateCurveData(class UAnimInstance* Instance, float CurrentTime, float BlendWeight) const OVERRIDE;
+	virtual void EvaluateCurveData(class UAnimInstance* Instance, float CurrentTime, float BlendWeight) const override;
 	// End of UAnimSequenceBase Interface
 
 	/**
diff --git a/Engine/Source/Runtime/Engine/Classes/Animation/AnimNodeSpaceConversions.h b/Engine/Source/Runtime/Engine/Classes/Animation/AnimNodeSpaceConversions.h
--- a/Engine/Source/Runtime/Engine/Classes/Animation/AnimNodeSpaceConversions.h
+++ b/Engine/Source/Runtime/Engine/Classes/Animation/AnimNodeSpaceConversions.h
@@ -18,11 +18,11 @@ public:
 	FAnimNode_ConvertComponentToLocalSpace();
 
 	// FAnimNode_Base interface
-	virtual void Initialize(const FAnimationInitializeContext& Context) OVERRIDE;
-	virtual void CacheBones(const FAnimationCacheBonesContext & Context) OVERRIDE;
-	virtual void Update(const FAnimationUpdateContext& Context) OVERRIDE;
-	virtual void Evaluate(FPoseContext& Output) OVERRIDE;
-	virtual void GatherDebugData(FNodeDebugData& DebugData) OVERRIDE;
+	virtual void Initialize(const FAnimationInitializeContext& Context) override;
+	virtual void CacheBones(const FAnimationCacheBonesContext & Context) override;
+	virtual void Update(const FAnimationUpdateContext& Context) override;
+	virtual void Evaluate(FPoseContext& Output) override;
+	virtual void GatherDebugData(FNodeDebugData& DebugData) override;
 	// End of FAnimNode_Base interface
 
 };
@@ -40,10 +40,10 @@ public:
 	FAnimNode_ConvertLocalToComponentSpace();
 
 	// FAnimNode_Base interface
-	virtual void Initialize(const FAnimationInitializeContext& Context) OVERRIDE;
-	virtual void CacheBones(const FAnimationCacheBonesContext & Context) OVERRIDE;
-	virtual void Update(const FAnimationUpdateContext& Context) OVERRIDE;
-	virtual void GatherDebugData(FNodeDebugData& DebugData) OVERRIDE;
-	virtual void EvaluateComponentSpace(FComponentSpacePoseContext& Output) OVERRIDE;
+	virtual void Initialize(const FAnimationInitializeContext& Context) override;
+	virtual void CacheBones(const FAnimationCacheBonesContext & Context) override;
+	virtual void Update(const FAnimationUpdateContext& Context) override;
+	virtual void GatherDebugData(FNodeDebugData& DebugData) override;
+	virtual void EvaluateComponentSpace(FComponentSpacePoseContext& Output) override;
 	// End of FAnimNode_Base interface
 };
diff --git a/Engine/Source/Runtime/Engine/Classes/Animation/AnimNode_ApplyAdditive.h b/Engine/Source/Runtime/Engine/Classes/Animation/AnimNode_ApplyAdditive.h
--- a/Engine/Source/Runtime/Engine/Classes/Animation/AnimNode_ApplyAdditive.h
+++ b/Engine/Source/Runtime/Engine/Classes/Animation/AnimNode_ApplyAdditive.h
@@ -27,10 +27,10 @@ public:
 	FAnimNode_ApplyAdditive();
 
 	// FAnimNode_Base interface
-	virtual void Initialize(const FAnimationInitializeContext& Context) OVERRIDE;
-	virtual void CacheBones(const FAnimationCacheBonesContext & Context) OVERRIDE;
-	virtual void Update(const FAnimationUpdateContext& Context) OVERRIDE;
-	virtual void Evaluate(FPoseContext& Output) OVERRIDE;
-	virtual void GatherDebugData(FNodeDebugData& DebugData) OVERRIDE;
+	virtual void Initialize(const FAnimationInitializeContext& Context) override;
+	virtual void CacheBones(const FAnimationCacheBonesContext & Context) override;
+	virtual void Update(const FAnimationUpdateContext& Context) override;
+	virtual void Evaluate(FPoseContext& Output) override;
+	virtual void GatherDebugData(FNodeDebugData& DebugData) override;
 	// End of FAnimNode_Base interface
 };
diff --git a/Engine/Source/Runtime/Engine/Classes/Animation/AnimNode_BlendListBase.h b/Engine/Source/Runtime/Engine/Classes/Animation/AnimNode_BlendListBase.h
--- a/Engine/Source/Runtime/Engine/Classes/Animation/AnimNode_BlendListBase.h
+++ b/Engine/Source/Runtime/Engine/Classes/Animation/AnimNode_BlendListBase.h
@@ -35,11 +35,11 @@ public:
 	}
 
 	// FAnimNode_Base interface
-	virtual void Initialize(const FAnimationInitializeContext& Context) OVERRIDE;
-	virtual void CacheBones(const FAnimationCacheBonesContext & Context) OVERRIDE;
-	virtual void Update(const FAnimationUpdateContext& Context) OVERRIDE;
-	virtual void Evaluate(FPoseContext& Output) OVERRIDE;
-	virtual void GatherDebugData(FNodeDebugData& DebugData) OVERRIDE;
+	virtual void Initialize(const FAnimationInitializeContext& Context) override;
+	virtual void CacheBones(const FAnimationCacheBonesContext & Context) override;
+	virtual void Update(const FAnimationUpdateContext& Context) override;
+	virtual void Evaluate(FPoseContext& Output) override;
+	virtual void GatherDebugData(FNodeDebugData& DebugData) override;
 	// End of FAnimNode_Base interface
 
 #if WITH_EDITOR
diff --git a/Engine/Source/Runtime/Engine/Classes/Animation/AnimNode_BlendSpaceEvaluator.h b/Engine/Source/Runtime/Engine/Classes/Animation/AnimNode_BlendSpaceEvaluator.h
--- a/Engine/Source/Runtime/Engine/Classes/Animation/AnimNode_BlendSpaceEvaluator.h
+++ b/Engine/Source/Runtime/Engine/Classes/Animation/AnimNode_BlendSpaceEvaluator.h
@@ -22,7 +22,7 @@ public:
 	FAnimNode_BlendSpaceEvaluator();
 
 	// FAnimNode_Base interface
-	virtual void Update(const FAnimationUpdateContext& Context) OVERRIDE;
-	virtual void GatherDebugData(FNodeDebugData& DebugData) OVERRIDE;
+	virtual void Update(const FAnimationUpdateContext& Context) override;
+	virtual void GatherDebugData(FNodeDebugData& DebugData) override;
 	// End of FAnimNode_Base interface
 };
diff --git a/Engine/Source/Runtime/Engine/Classes/Animation/AnimNode_BlendSpacePlayer.h b/Engine/Source/Runtime/Engine/Classes/Animation/AnimNode_BlendSpacePlayer.h
--- a/Engine/Source/Runtime/Engine/Classes/Animation/AnimNode_BlendSpacePlayer.h
+++ b/Engine/Source/Runtime/Engine/Classes/Animation/AnimNode_BlendSpacePlayer.h
@@ -57,11 +57,11 @@ public:
 	FAnimNode_BlendSpacePlayer();
 
 	// FAnimNode_Base interface
-	virtual void Initialize(const FAnimationInitializeContext& Context) OVERRIDE;
-	virtual void CacheBones(const FAnimationCacheBonesContext & Context) OVERRIDE;
-	virtual void Update(const FAnimationUpdateContext& Context) OVERRIDE;
-	virtual void Evaluate(FPoseContext& Output) OVERRIDE;
-	virtual void OverrideAsset(UAnimationAsset* NewAsset) OVERRIDE;
-	virtual void GatherDebugData(FNodeDebugData& DebugData) OVERRIDE;
+	virtual void Initialize(const FAnimationInitializeContext& Context) override;
+	virtual void CacheBones(const FAnimationCacheBonesContext & Context) override;
+	virtual void Update(const FAnimationUpdateContext& Context) override;
+	virtual void Evaluate(FPoseContext& Output) override;
+	virtual void OverrideAsset(UAnimationAsset* NewAsset) override;
+	virtual void GatherDebugData(FNodeDebugData& DebugData) override;
 	// End of FAnimNode_Base interface
 };
diff --git a/Engine/Source/Runtime/Engine/Classes/Animation/AnimNode_RefPose.h b/Engine/Source/Runtime/Engine/Classes/Animation/AnimNode_RefPose.h
--- a/Engine/Source/Runtime/Engine/Classes/Animation/AnimNode_RefPose.h
+++ b/Engine/Source/Runtime/Engine/Classes/Animation/AnimNode_RefPose.h
@@ -27,8 +27,8 @@ public:
 	{
 	}
 
-	virtual void Evaluate(FPoseContext& Output) OVERRIDE;
-	virtual void GatherDebugData(FNodeDebugData& DebugData) OVERRIDE;
+	virtual void Evaluate(FPoseContext& Output) override;
+	virtual void GatherDebugData(FNodeDebugData& DebugData) override;
 
 };
 
diff --git a/Engine/Source/Runtime/Engine/Classes/Animation/AnimNode_SequenceEvaluator.h b/Engine/Source/Runtime/Engine/Classes/Animation/AnimNode_SequenceEvaluator.h
--- a/Engine/Source/Runtime/Engine/Classes/Animation/AnimNode_SequenceEvaluator.h
+++ b/Engine/Source/Runtime/Engine/Classes/Animation/AnimNode_SequenceEvaluator.h
@@ -27,11 +27,11 @@ public:
 	}
 
 	// FAnimNode_Base interface
-	virtual void Initialize(const FAnimationInitializeContext& Context) OVERRIDE;
-	virtual void CacheBones(const FAnimationCacheBonesContext & Context) OVERRIDE;
-	virtual void Update(const FAnimationUpdateContext& Context) OVERRIDE;
-	virtual void Evaluate(FPoseContext& Output) OVERRIDE;
-	virtual void OverrideAsset(UAnimationAsset* NewAsset) OVERRIDE;
-	virtual void GatherDebugData(FNodeDebugData& DebugData) OVERRIDE;
+	virtual void Initialize(const FAnimationInitializeContext& Context) override;
+	virtual void CacheBones(const FAnimationCacheBonesContext & Context) override;
+	virtual void Update(const FAnimationUpdateContext& Context) override;
+	virtual void Evaluate(FPoseContext& Output) override;
+	virtual void OverrideAsset(UAnimationAsset* NewAsset) override;
+	virtual void GatherDebugData(FNodeDebugData& DebugData) override;
 	// End of FAnimNode_Base interface
 };
diff --git a/Engine/Source/Runtime/Engine/Classes/Animation/AnimNotifies/AnimNotifyState_TimedParticleEffect.h b/Engine/Source/Runtime/Engine/Classes/Animation/AnimNotifies/AnimNotifyState_TimedParticleEffect.h
--- a/Engine/Source/Runtime/Engine/Classes/Animation/AnimNotifies/AnimNotifyState_TimedParticleEffect.h
+++ b/Engine/Source/Runtime/Engine/Classes/Animation/AnimNotifies/AnimNotifyState_TimedParticleEffect.h
@@ -44,12 +44,12 @@ class UAnimNotifyState_TimedParticleEffect : public UAnimNotifyState
 	UPROPERTY(transient)
 	TArray<FName> PreviousSocketNames;
 
-	virtual void PreEditChange(UProperty* PropertyAboutToChange) OVERRIDE;
+	virtual void PreEditChange(UProperty* PropertyAboutToChange) override;
 #endif
 
-	virtual void NotifyBegin(class USkeletalMeshComponent * MeshComp, class UAnimSequence * AnimSeq) OVERRIDE;
-	virtual void NotifyTick(class USkeletalMeshComponent * MeshComp, class UAnimSequence * AnimSeq, float FrameDeltaTime) OVERRIDE;
-	virtual void NotifyEnd(class USkeletalMeshComponent * MeshComp, class UAnimSequence * AnimSeq) OVERRIDE;
+	virtual void NotifyBegin(class USkeletalMeshComponent * MeshComp, class UAnimSequence * AnimSeq) override;
+	virtual void NotifyTick(class USkeletalMeshComponent * MeshComp, class UAnimSequence * AnimSeq, float FrameDeltaTime) override;
+	virtual void NotifyEnd(class USkeletalMeshComponent * MeshComp, class UAnimSequence * AnimSeq) override;
 
 private:
 	bool ValidateParameters(USkeletalMeshComponent* MeshComp);
diff --git a/Engine/Source/Runtime/Engine/Classes/Animation/AnimSequence.h b/Engine/Source/Runtime/Engine/Classes/Animation/AnimSequence.h
--- a/Engine/Source/Runtime/Engine/Classes/Animation/AnimSequence.h
+++ b/Engine/Source/Runtime/Engine/Classes/Animation/AnimSequence.h
@@ -496,26 +496,26 @@ class UAnimSequence : public UAnimSequenceBase
 #endif // WITH_EDITORONLY_DATA
 
 	// Begin UObject interface
-	ENGINE_API virtual void Serialize(FArchive& Ar) OVERRIDE;
-	ENGINE_API virtual void PostLoad() OVERRIDE;
+	ENGINE_API virtual void Serialize(FArchive& Ar) override;
+	ENGINE_API virtual void PostLoad() override;
 #if WITH_EDITOR
-	ENGINE_API virtual void PostEditChangeProperty(FPropertyChangedEvent& PropertyChangedEvent) OVERRIDE;
+	ENGINE_API virtual void PostEditChangeProperty(FPropertyChangedEvent& PropertyChangedEvent) override;
 #endif // WITH_EDITOR
-	ENGINE_API virtual void BeginDestroy() OVERRIDE;
-	ENGINE_API virtual SIZE_T GetResourceSize(EResourceSizeMode::Type Mode) OVERRIDE;
+	ENGINE_API virtual void BeginDestroy() override;
+	ENGINE_API virtual SIZE_T GetResourceSize(EResourceSizeMode::Type Mode) override;
 	// End of UObject interface
 
 	// Begin UAnimationAsset interface
-	ENGINE_API virtual bool IsValidAdditive() const OVERRIDE;
+	ENGINE_API virtual bool IsValidAdditive() const override;
 #if WITH_EDITOR
-	ENGINE_API virtual bool GetAllAnimationSequencesReferred(TArray<UAnimSequence*>& AnimationSequences) OVERRIDE;
-	ENGINE_API virtual int32 GetNumberOfFrames() OVERRIDE { return NumFrames; }
+	ENGINE_API virtual bool GetAllAnimationSequencesReferred(TArray<UAnimSequence*>& AnimationSequences) override;
+	ENGINE_API virtual int32 GetNumberOfFrames() override { return NumFrames; }
 #endif
 	// End of UAnimationAsset interface
 
 	// Begin UAnimSequenceBase interface
-	ENGINE_API virtual void ExtractRootTrack(float Pos, FTransform & RootTransform, const FBoneContainer * RequiredBones) const OVERRIDE;
-	ENGINE_API virtual void UpgradeMorphTargetCurves() OVERRIDE;
+	ENGINE_API virtual void ExtractRootTrack(float Pos, FTransform & RootTransform, const FBoneContainer * RequiredBones) const override;
+	ENGINE_API virtual void UpgradeMorphTargetCurves() override;
 	// End UAnimSequenceBase interface
 
 	// Begin Transform related functions 
diff --git a/Engine/Source/Runtime/Engine/Classes/Animation/AnimSequenceBase.h b/Engine/Source/Runtime/Engine/Classes/Animation/AnimSequenceBase.h
--- a/Engine/Source/Runtime/Engine/Classes/Animation/AnimSequenceBase.h
+++ b/Engine/Source/Runtime/Engine/Classes/Animation/AnimSequenceBase.h
@@ -292,7 +292,7 @@ class UAnimSequenceBase : public UAnimationAsset
 #endif // WITH_EDITORONLY_DATA
 
 	// Begin UObject interface
-	virtual void PostLoad() OVERRIDE;
+	virtual void PostLoad() override;
 	// End of UObject interface
 
 	/** Sort the Notifies array by time, earliest first. */
@@ -339,7 +339,7 @@ class UAnimSequenceBase : public UAnimationAsset
 	/** Calculates what (if any) offset should be applied to the trigger time of a notify given its display time */ 
 	virtual EAnimEventTriggerOffsets::Type CalculateOffsetForNotify(float NotifyDisplayTime) const;
 
-	virtual void GetAssetRegistryTags(TArray<FAssetRegistryTag>& OutTags) const OVERRIDE;
+	virtual void GetAssetRegistryTags(TArray<FAssetRegistryTag>& OutTags) const override;
 	
 	// Get a pointer to the data for a given Anim Notify
 	ENGINE_API uint8* FindNotifyPropertyData(int32 NotifyIndex, UArrayProperty*& ArrayProperty);
@@ -352,11 +352,11 @@ class UAnimSequenceBase : public UAnimationAsset
 	ENGINE_API virtual void UpgradeMorphTargetCurves();
 
 	// Begin UAnimationAsset interface
-	virtual void TickAssetPlayerInstance(const FAnimTickRecord& Instance, class UAnimInstance* InstanceOwner, FAnimAssetTickContext& Context) const OVERRIDE;
+	virtual void TickAssetPlayerInstance(const FAnimTickRecord& Instance, class UAnimInstance* InstanceOwner, FAnimAssetTickContext& Context) const override;
 	// this is used in editor only when used for transition getter
 	// this doesn't mean max time. In Sequence, this is SequenceLength,
 	// but for BlendSpace CurrentTime is normalized [0,1], so this is 1
-	virtual float GetMaxCurrentTime() OVERRIDE { return SequenceLength; }
+	virtual float GetMaxCurrentTime() override { return SequenceLength; }
 	// End of UAnimationAsset interface
 
 protected:
diff --git a/Engine/Source/Runtime/Engine/Classes/Animation/AnimSet.h b/Engine/Source/Runtime/Engine/Classes/Animation/AnimSet.h
--- a/Engine/Source/Runtime/Engine/Classes/Animation/AnimSet.h
+++ b/Engine/Source/Runtime/Engine/Classes/Animation/AnimSet.h
@@ -110,8 +110,8 @@ public:
 	TMap<FName,int32> SkelMesh2LinkupCache;
 
 	// Begin UObject interface
-	virtual void PostLoad() OVERRIDE;
-	virtual SIZE_T GetResourceSize(EResourceSizeMode::Type Mode) OVERRIDE;
+	virtual void PostLoad() override;
+	virtual SIZE_T GetResourceSize(EResourceSizeMode::Type Mode) override;
 	// End UObject interface
 	
 	// UAnimSet interface
diff --git a/Engine/Source/Runtime/Engine/Classes/Animation/AnimSingleNodeInstance.h b/Engine/Source/Runtime/Engine/Classes/Animation/AnimSingleNodeInstance.h
--- a/Engine/Source/Runtime/Engine/Classes/Animation/AnimSingleNodeInstance.h
+++ b/Engine/Source/Runtime/Engine/Classes/Animation/AnimSingleNodeInstance.h
@@ -56,11 +56,11 @@ class ENGINE_API UAnimSingleNodeInstance : public UAnimInstance
 	uint32 bReverse:1;
 
 	// Begin UAnimInstance interface
-	virtual void NativeInitializeAnimation() OVERRIDE;
-	virtual void NativeUpdateAnimation(float DeltaTimeX) OVERRIDE;
-	virtual bool NativeEvaluateAnimation(FPoseContext& Output) OVERRIDE;
+	virtual void NativeInitializeAnimation() override;
+	virtual void NativeUpdateAnimation(float DeltaTimeX) override;
+	virtual bool NativeEvaluateAnimation(FPoseContext& Output) override;
 protected:
-	virtual void Montage_Advance(float DeltaTime) OVERRIDE;
+	virtual void Montage_Advance(float DeltaTime) override;
 	void InternalBlendSpaceEvaluatePose(class UBlendSpaceBase* BlendSpace, TArray<FBlendSampleData>& BlendSampleDataCache, struct FA2Pose& Pose, bool bIsLooping);
 	// End UAnimInstance interface
 public:
diff --git a/Engine/Source/Runtime/Engine/Classes/Animation/AnimationAsset.h b/Engine/Source/Runtime/Engine/Classes/Animation/AnimationAsset.h
--- a/Engine/Source/Runtime/Engine/Classes/Animation/AnimationAsset.h
+++ b/Engine/Source/Runtime/Engine/Classes/Animation/AnimationAsset.h
@@ -591,12 +591,12 @@ public:
 
 	ENGINE_API void SetSkeleton(USkeleton* NewSkeleton);
 	void ResetSkeleton(USkeleton* NewSkeleton);
-	virtual void PostLoad() OVERRIDE;
+	virtual void PostLoad() override;
 
 	/** Validate our stored data against our skeleton and update accordingly */
 	void ValidateSkeleton();
 
-	virtual void Serialize(FArchive& Ar) OVERRIDE;
+	virtual void Serialize(FArchive& Ar) override;
 
 #if WITH_EDITOR
 	/** Replace Skeleton 
diff --git a/Engine/Source/Runtime/Engine/Classes/Animation/BlendSpace.h b/Engine/Source/Runtime/Engine/Classes/Animation/BlendSpace.h
--- a/Engine/Source/Runtime/Engine/Classes/Animation/BlendSpace.h
+++ b/Engine/Source/Runtime/Engine/Classes/Animation/BlendSpace.h
@@ -44,10 +44,10 @@ public:
 
 protected:
 	// Begin UBlendSpaceBase interface
-	virtual void SnapToBorder(FBlendSample& Sample) const OVERRIDE;
-	virtual EBlendSpaceAxis GetAxisToScale() const OVERRIDE;
-	virtual bool IsSameSamplePoint(const FVector& SamplePointA, const FVector& SamplePointB) const OVERRIDE;
-	virtual void GetRawSamplesFromBlendInput(const FVector &BlendInput, TArray<FGridBlendSample> & OutBlendSamples) const OVERRIDE;
+	virtual void SnapToBorder(FBlendSample& Sample) const override;
+	virtual EBlendSpaceAxis GetAxisToScale() const override;
+	virtual bool IsSameSamplePoint(const FVector& SamplePointA, const FVector& SamplePointB) const override;
+	virtual void GetRawSamplesFromBlendInput(const FVector &BlendInput, TArray<FGridBlendSample> & OutBlendSamples) const override;
 	// End UBlendSpaceBase interface
 
 private:
diff --git a/Engine/Source/Runtime/Engine/Classes/Animation/BlendSpace1D.h b/Engine/Source/Runtime/Engine/Classes/Animation/BlendSpace1D.h
--- a/Engine/Source/Runtime/Engine/Classes/Animation/BlendSpace1D.h
+++ b/Engine/Source/Runtime/Engine/Classes/Animation/BlendSpace1D.h
@@ -27,14 +27,14 @@ public:
 	bool bScaleAnimation;
 
 	/** return true if all sample data is additive **/
-	virtual bool IsValidAdditive() const OVERRIDE;
+	virtual bool IsValidAdditive() const override;
 
 protected:
 	// Begin UBlendSpaceBase interface
-	virtual void SnapToBorder(FBlendSample& Sample) const OVERRIDE;
-	virtual EBlendSpaceAxis GetAxisToScale() const OVERRIDE;
-	virtual bool IsSameSamplePoint(const FVector& SamplePointA, const FVector& SamplePointB) const OVERRIDE;
-	virtual void GetRawSamplesFromBlendInput(const FVector &BlendInput, TArray<FGridBlendSample> & OutBlendSamples) const OVERRIDE;
+	virtual void SnapToBorder(FBlendSample& Sample) const override;
+	virtual EBlendSpaceAxis GetAxisToScale() const override;
+	virtual bool IsSameSamplePoint(const FVector& SamplePointA, const FVector& SamplePointB) const override;
+	virtual void GetRawSamplesFromBlendInput(const FVector &BlendInput, TArray<FGridBlendSample> & OutBlendSamples) const override;
 	// End UBlendSpaceBase interface
 
 private:
diff --git a/Engine/Source/Runtime/Engine/Classes/Animation/BlendSpaceBase.h b/Engine/Source/Runtime/Engine/Classes/Animation/BlendSpaceBase.h
--- a/Engine/Source/Runtime/Engine/Classes/Animation/BlendSpaceBase.h
+++ b/Engine/Source/Runtime/Engine/Classes/Animation/BlendSpaceBase.h
@@ -256,19 +256,19 @@ public:
 	// Begin UObject interface
 	virtual void PostLoad();
 #if WITH_EDITOR
-	virtual void PostEditChangeProperty( struct FPropertyChangedEvent& PropertyChangedEvent) OVERRIDE;
+	virtual void PostEditChangeProperty( struct FPropertyChangedEvent& PropertyChangedEvent) override;
 #endif // WITH_EDITOR
 	// End UObject interface
 
 	// Begin UAnimationAsset interface
-	virtual void TickAssetPlayerInstance(const FAnimTickRecord& Instance, class UAnimInstance* InstanceOwner, FAnimAssetTickContext& Context) const OVERRIDE;
+	virtual void TickAssetPlayerInstance(const FAnimTickRecord& Instance, class UAnimInstance* InstanceOwner, FAnimAssetTickContext& Context) const override;
 	// this is used in editor only when used for transition getter
 	// this doesn't mean max time. In Sequence, this is SequenceLength,
 	// but for BlendSpace CurrentTime is normalized [0,1], so this is 1
-	virtual float GetMaxCurrentTime() OVERRIDE { return 1.f; }	
+	virtual float GetMaxCurrentTime() override { return 1.f; }	
 #if WITH_EDITOR
-	virtual bool GetAllAnimationSequencesReferred(TArray<UAnimSequence*>& AnimationSequences) OVERRIDE;
-	virtual void ReplaceReferredAnimations(const TMap<UAnimSequence*, UAnimSequence*>& ReplacementMap) OVERRIDE;
+	virtual bool GetAllAnimationSequencesReferred(TArray<UAnimSequence*>& AnimationSequences) override;
+	virtual void ReplaceReferredAnimations(const TMap<UAnimSequence*, UAnimSequence*>& ReplacementMap) override;
 #endif
 	// End of UAnimationAsset interface
 
