LinkerLoad will no longer reject existing objects if their class is derived from the serialized class.

diff --git a/Engine/Source/Runtime/CoreUObject/Private/UObject/LinkerLoad.cpp b/Engine/Source/Runtime/CoreUObject/Private/UObject/LinkerLoad.cpp
--- a/Engine/Source/Runtime/CoreUObject/Private/UObject/LinkerLoad.cpp
+++ b/Engine/Source/Runtime/CoreUObject/Private/UObject/LinkerLoad.cpp
@@ -3127,7 +3127,7 @@ UObject* ULinkerLoad::CreateExport( int32 Index )
 
 		// Create the export object, marking it with the appropriate flags to
 		// indicate that the object's data still needs to be loaded.
-		if (ActualObjectWithTheName && (ActualObjectWithTheName->GetClass() != LoadClass))
+		if (ActualObjectWithTheName && !ActualObjectWithTheName->GetClass()->IsChildOf(LoadClass))
 		{
 			UE_LOG(LogLinker, Error, TEXT("Failed import: class '%s' name '%s' outer '%s'. There is another object (of '%s' class) at the path."),
 				*LoadClass->GetName(), *Export.ObjectName.ToString(), *ThisParent->GetName(), *ActualObjectWithTheName->GetClass()->GetName());
