[GitHub] 233 :Object and class accessors in FPostConstructInitializeProperties

diff --git a/Engine/Source/Runtime/CoreUObject/Private/UObject/UObjectGlobals.cpp b/Engine/Source/Runtime/CoreUObject/Private/UObject/UObjectGlobals.cpp
--- a/Engine/Source/Runtime/CoreUObject/Private/UObject/UObjectGlobals.cpp
+++ b/Engine/Source/Runtime/CoreUObject/Private/UObject/UObjectGlobals.cpp
@@ -2033,6 +2033,11 @@ FPostConstructInitializeProperties::~FPostConstructInitializeProperties()
 	}
 }
 
+UClass* FPostConstructInitializeProperties::GetClass() const
+{
+	return Obj->GetClass();
+}
+
 void FSubobjectPtr::Set(UObject* InObject)
 {
 	if (Object != InObject && IsInitialized(Object) && !Object->IsPendingKill())
diff --git a/Engine/Source/Runtime/CoreUObject/Public/UObject/UObjectGlobals.h b/Engine/Source/Runtime/CoreUObject/Public/UObject/UObjectGlobals.h
--- a/Engine/Source/Runtime/CoreUObject/Public/UObject/UObjectGlobals.h
+++ b/Engine/Source/Runtime/CoreUObject/Public/UObject/UObjectGlobals.h
@@ -689,6 +689,19 @@ public:
 		return ObjectArchetype;
 	}
 
+	/**
+	* Return the object that is being constructed
+	**/
+	FORCEINLINE UObject* GetObject() const
+	{
+		return Obj;
+	}
+
+	/**
+	* Return the class of the object that is being constructed
+	**/
+	UClass* GetClass() const;
+
 	/**
 	 * Create a component or subobject
 	 * @param	TReturnType					class of return type, all overrides must be of this type
