Suppress existing VS2013 static analysis warnings in the CoreUObject module.

diff --git a/Engine/Source/Runtime/CoreUObject/Private/Misc/RedirectCollector.cpp b/Engine/Source/Runtime/CoreUObject/Private/Misc/RedirectCollector.cpp
--- a/Engine/Source/Runtime/CoreUObject/Private/Misc/RedirectCollector.cpp
+++ b/Engine/Source/Runtime/CoreUObject/Private/Misc/RedirectCollector.cpp
@@ -92,6 +92,7 @@ void FRedirectCollector::ResolveStringAssetReference()
 				Redir.PackageFilename = RefFilename;
 				Redir.RedirectorName = Redirector->GetFullName();
 				Redir.RedirectorPackageFilename = Redirector->GetLinker()->Filename;
+				CA_SUPPRESS(28182)
 				Redir.DestinationObjectName = Redirector->DestinationObject->GetFullName();
 				Redirections.AddUnique(Redir);
 				Loaded = Redirector->DestinationObject;
diff --git a/Engine/Source/Runtime/CoreUObject/Private/UObject/LinkerLoad.cpp b/Engine/Source/Runtime/CoreUObject/Private/UObject/LinkerLoad.cpp
--- a/Engine/Source/Runtime/CoreUObject/Private/UObject/LinkerLoad.cpp
+++ b/Engine/Source/Runtime/CoreUObject/Private/UObject/LinkerLoad.cpp
@@ -3636,6 +3636,7 @@ void ULinkerLoad::AsyncPreloadPackage(const TCHAR* PackageName)
 	const int32 FileSize = IFileManager::Get().FileSize(*PackageFilename);
 
 	// if we were compressed, the size we care about on the other end is the uncompressed size
+	CA_SUPPRESS(6326)
 	PrecacheInfo.PackageDataSize = UncompressedSize == -1 ? FileSize : UncompressedSize;
 	
 	// allocate enough space
@@ -3643,6 +3644,7 @@ void ULinkerLoad::AsyncPreloadPackage(const TCHAR* PackageName)
 
 	uint64 RequestId;
 	// kick off the async read (uncompressing if needed) of the whole file and make sure it worked
+	CA_SUPPRESS(6326)
 	if (UncompressedSize != -1)
 	{
 		PrecacheInfo.PackageDataSize = UncompressedSize;
diff --git a/Engine/Source/Runtime/CoreUObject/Private/UObject/Obj.cpp b/Engine/Source/Runtime/CoreUObject/Private/UObject/Obj.cpp
--- a/Engine/Source/Runtime/CoreUObject/Private/UObject/Obj.cpp
+++ b/Engine/Source/Runtime/CoreUObject/Private/UObject/Obj.cpp
@@ -399,6 +399,7 @@ void UObject::GetArchetypeInstances( TArray<UObject*>& Instances )
 				UObject* Obj = *It;
 
 				// if this object is the correct type and its archetype is this object, add it to the list
+				CA_SUPPRESS(6011)
 				if ( Obj != this && Obj->IsA(GetClass()) && Obj->IsBasedOnArchetype(this) && !Obj->IsPendingKill() )
 				{
 					Instances.Add(Obj);
@@ -2045,7 +2046,11 @@ static void ShowIntrinsicClasses( FOutputDevice& Ar )
 static void ShowClasses( UClass* Class, FOutputDevice& Ar, int32 Indent )
 {
 	Ar.Logf( TEXT("%s%s (%d)"), FCString::Spc(Indent), *Class->GetName(), Class->GetPropertiesSize() );
-	for( auto Obj : TObjectRange<UClass>() )
+
+	// Workaround for Visual Studio 2013 analyzer bug. Using a temporary directly in the range-for
+	// errors if the analyzer is enabled.
+	TObjectRange<UClass> Range;
+	for( auto Obj : Range )
 	{
 		if( Obj->GetSuperClass() == Class )
 		{
@@ -2180,6 +2185,7 @@ void UObject::RetrieveReferencers( TArray<FReferencerInformation>* OutInternalRe
 		int32 Count = ArFind.GetCount(Referencers);
 		if ( Count > 0 )
 		{
+			CA_SUPPRESS(6011)
 			if ( Object->IsIn(this) )
 			{
 				if (OutInternalReferencers != NULL)
@@ -3703,7 +3709,11 @@ void MarkObjectsToDisregardForGC()
 	// Iterate over all class objects and force the default objects to be created. Additionally also
 	// assembles the token reference stream at this point. This is required for class objects that are
 	// not taken into account for garbage collection but have instances that are.
-	for( auto* Class : TObjectRange<UClass>() )
+	
+	// Workaround for Visual Studio 2013 analyzer bug. Using a temporary directly in the range-for
+	// errors if the analyzer is enabled.
+	TObjectRange<UClass> Range;
+	for( auto* Class : Range )
 	{
 		// Force the default object to be created.
 		Class->GetDefaultObject(); // Force the default object to be constructed if it isn't already
diff --git a/Engine/Source/Runtime/CoreUObject/Private/UObject/PropertyArray.cpp b/Engine/Source/Runtime/CoreUObject/Private/UObject/PropertyArray.cpp
--- a/Engine/Source/Runtime/CoreUObject/Private/UObject/PropertyArray.cpp
+++ b/Engine/Source/Runtime/CoreUObject/Private/UObject/PropertyArray.cpp
@@ -248,6 +248,7 @@ const TCHAR* UArrayProperty::ImportText_Internal( const TCHAR* Buffer, void* Dat
 	}
 
 	// Make sure we ended on a )
+	CA_SUPPRESS(6011)
 	if (*Buffer++ != TCHAR(')'))
 	{
 		return NULL;
diff --git a/Engine/Source/Runtime/CoreUObject/Private/UObject/ReferenceChainSearch.cpp b/Engine/Source/Runtime/CoreUObject/Private/UObject/ReferenceChainSearch.cpp
--- a/Engine/Source/Runtime/CoreUObject/Private/UObject/ReferenceChainSearch.cpp
+++ b/Engine/Source/Runtime/CoreUObject/Private/UObject/ReferenceChainSearch.cpp
@@ -117,6 +117,7 @@ void FReferenceChainSearch::PrintReferencers( FReferenceChain& Referencer )
 			ObjectReachability += TEXT("(root) ");
 		}
 		
+		CA_SUPPRESS(6011)
 		if( RefInfo.ReferencedBy->HasAnyFlags(RF_Native) )
 		{
 			ObjectReachability += TEXT("(native) ");
@@ -641,6 +642,7 @@ FReferenceChainSearch::FReferenceChainSearch( UObject* InObjectToFind, uint32 Mo
 		{
 			UObject* Obj = Referencers[i].RefChain[0].ReferencedBy;
 
+			CA_SUPPRESS(6011)
 			if (Obj->IsIn(ObjectToFind) || Obj == ObjectToFind)
 			{
 				if (bIsFirst)
diff --git a/Engine/Source/Runtime/CoreUObject/Private/UObject/ScriptCore.cpp b/Engine/Source/Runtime/CoreUObject/Private/UObject/ScriptCore.cpp
--- a/Engine/Source/Runtime/CoreUObject/Private/UObject/ScriptCore.cpp
+++ b/Engine/Source/Runtime/CoreUObject/Private/UObject/ScriptCore.cpp
@@ -264,6 +264,7 @@ COREUOBJECT_API uint8 GRegisterNative( int32 NativeBytecodeIndex, const Native&
 #if !IS_MONOLITHIC
 			if (GIsHotReload)
 			{
+				CA_SUPPRESS(6385)
 				AddHotReloadFunctionRemap(Func, GNatives[NativeBytecodeIndex]);
 			}
 			else
@@ -272,6 +273,7 @@ COREUOBJECT_API uint8 GRegisterNative( int32 NativeBytecodeIndex, const Native&
 			GNativeDuplicate = NativeBytecodeIndex;
 			}
 		}
+		CA_SUPPRESS(6386)
 		GNatives[NativeBytecodeIndex] = Func;
 	}
 
@@ -403,6 +405,7 @@ void UObject::CallFunction( FFrame& Stack, RESULT_DECL, UFunction* Function )
  				Property = *ParmIt;
  				if( Property->HasAnyPropertyFlags(CPF_ReturnParm) )
  				{
+					CA_SUPPRESS(6263)
  					FOutParmRec* RetVal = (FOutParmRec*)FMemory_Alloca(sizeof(FOutParmRec));
  
  					// Our context should be that we're in a variable assignment to the return value, so ensure that we have a valid property to return to
@@ -435,6 +438,7 @@ void UObject::CallFunction( FFrame& Stack, RESULT_DECL, UFunction* Function )
 				// evaluate the expression for this parameter, which sets Stack.MostRecentPropertyAddress to the address of the property accessed
 				Stack.Step(Stack.Object, NULL);
 
+				CA_SUPPRESS(6263)
 				FOutParmRec* Out = (FOutParmRec*)FMemory_Alloca(sizeof(FOutParmRec));
 				// set the address and property in the out param info
 				// warning: Stack.MostRecentPropertyAddress could be NULL for optional out parameters
@@ -822,6 +826,7 @@ void UObject::ProcessEvent( UFunction* Function, void* Parms )
 				// bytecode
 				if ( Property->HasAnyPropertyFlags(CPF_OutParm) )
 				{
+					CA_SUPPRESS(6263)
 					FOutParmRec* Out = (FOutParmRec*)FMemory_Alloca(sizeof(FOutParmRec));
 					// set the address and property in the out param info
 					// note that since C++ doesn't support "optional out" we can ignore that here
diff --git a/Engine/Source/Runtime/CoreUObject/Private/UObject/UObjectGlobals.cpp b/Engine/Source/Runtime/CoreUObject/Private/UObject/UObjectGlobals.cpp
--- a/Engine/Source/Runtime/CoreUObject/Private/UObject/UObjectGlobals.cpp
+++ b/Engine/Source/Runtime/CoreUObject/Private/UObject/UObjectGlobals.cpp
@@ -803,10 +803,12 @@ UPackage* LoadPackage( UPackage* InOuter, const TCHAR* InLongPackageName, uint32
 		}
 
 		// Save the filename we load from
+		CA_SUPPRESS(28182)
 		Result->FileName = FName(*FileToLoad);
 
 		// is there a script SHA hash for this package?
 		uint8 SavedScriptSHA[20];
+		CA_SUPPRESS(6011)
 		bool bHasScriptSHAHash = FSHA1::GetFileSHAHash(*Linker->LinkerRoot->GetName(), SavedScriptSHA, false);
 		if (bHasScriptSHAHash)
 		{
