General - Making SaveToTransactionBuffer easier to debug.

diff --git a/Engine/Source/Runtime/CoreUObject/Private/UObject/UObjectGlobals.cpp b/Engine/Source/Runtime/CoreUObject/Private/UObject/UObjectGlobals.cpp
--- a/Engine/Source/Runtime/CoreUObject/Private/UObject/UObjectGlobals.cpp
+++ b/Engine/Source/Runtime/CoreUObject/Private/UObject/UObjectGlobals.cpp
@@ -875,8 +875,8 @@ UPackage* LoadPackage( UPackage* InOuter, const TCHAR* InLongPackageName, uint32
 				{
 					ResetLoaders(Result);
 				}
-  				delete Linker->Loader;
-  				Linker->Loader = NULL;
+				delete Linker->Loader;
+				Linker->Loader = NULL;
 			}
 			else
 			{
@@ -1562,9 +1562,10 @@ bool SaveToTransactionBuffer(UObject* Object, bool bMarkDirty)
 
 	// Neither PIE world objects nor script packages should end up in the transaction buffer. Additionally, in order
 	// to save a copy of the object, we must have a transactor and the object must be transactional.
-	if ( GUndo && 
-		 Object->HasAnyFlags( RF_Transactional ) &&
-		 ( ( Object->GetOutermost()->PackageFlags & ( PKG_PlayInEditor|PKG_ContainsScript ) ) == 0 ) )
+	const bool IsTransactional = Object->HasAnyFlags(RF_Transactional);
+	const bool IsNotPIEOrContainsScriptObject = ( ( Object->GetOutermost()->PackageFlags & ( PKG_PlayInEditor | PKG_ContainsScript ) ) == 0 );
+
+	if ( GUndo && IsTransactional && IsNotPIEOrContainsScriptObject )
 	{
 		// Mark the package dirty, if requested
 		if ( bMarkDirty )
@@ -1909,15 +1910,15 @@ FPostConstructInitializeProperties::~FPostConstructInitializeProperties()
 	}
 	if (bShouldIntializePropsFromArchetype)
 	{
-	    UClass* BaseClass = (bIsCDO && !GIsDuplicatingClassForReinstancing) ? Obj->GetClass()->GetSuperClass() : Class;
-	    if (BaseClass == NULL)
-	    {
-		    check(Class==UObject::StaticClass());
-		    BaseClass = Class;
-	    }
-    
-	    UObject* Defaults = ObjectArchetype ? ObjectArchetype : BaseClass->GetDefaultObject(false); // we don't create the CDO here if it doesn't already exist
-	    InitProperties(Obj, BaseClass, Defaults, bCopyTransientsFromClassDefaults);
+		UClass* BaseClass = (bIsCDO && !GIsDuplicatingClassForReinstancing) ? Obj->GetClass()->GetSuperClass() : Class;
+		if (BaseClass == NULL)
+		{
+			check(Class==UObject::StaticClass());
+			BaseClass = Class;
+		}
+	
+		UObject* Defaults = ObjectArchetype ? ObjectArchetype : BaseClass->GetDefaultObject(false); // we don't create the CDO here if it doesn't already exist
+		InitProperties(Obj, BaseClass, Defaults, bCopyTransientsFromClassDefaults);
 	}
 
 	bool bAllowInstancing = (InstanceGraph == NULL) || InstanceGraph->IsSubobjectInstancingEnabled();
