Making sure 'NewObject cannot be used to construct subobjects' does not prevent new objects from being constructed in PostInitProperties

diff --git a/Engine/Source/Runtime/CoreUObject/Private/UObject/UObjectGlobals.cpp b/Engine/Source/Runtime/CoreUObject/Private/UObject/UObjectGlobals.cpp
--- a/Engine/Source/Runtime/CoreUObject/Private/UObject/UObjectGlobals.cpp
+++ b/Engine/Source/Runtime/CoreUObject/Private/UObject/UObjectGlobals.cpp
@@ -1862,10 +1862,11 @@ FPostConstructInitializeProperties::FPostConstructInitializeProperties() :
 	bCopyTransientsFromClassDefaults(false),
 	bShouldIntializePropsFromArchetype(true),
 	InstanceGraph(NULL),
-	LastConstructedObject(GConstructedObject, Obj)
+	LastConstructedObject(GConstructedObject)
 {
-	// Mark we're in the constructor now.
+	// Mark we're in the constructor now.	
 	GIsInConstructor++;
+	GConstructedObject = Obj;
 }	
 
 FPostConstructInitializeProperties::FPostConstructInitializeProperties(UObject* InObj, UObject* InObjectArchetype, bool bInCopyTransientsFromClassDefaults, bool bInShouldIntializeProps, struct FObjectInstancingGraph* InInstanceGraph) :
@@ -1875,10 +1876,11 @@ FPostConstructInitializeProperties::FPostConstructInitializeProperties(UObject*
 	bCopyTransientsFromClassDefaults(bInCopyTransientsFromClassDefaults),
 	bShouldIntializePropsFromArchetype(bInShouldIntializeProps),
 	InstanceGraph(InInstanceGraph),
-	LastConstructedObject(GConstructedObject, Obj)
+	LastConstructedObject(GConstructedObject)
 {
 	// Mark we're in the constructor now.
 	GIsInConstructor++;
+	GConstructedObject = Obj;
 }
 
 /**
@@ -1889,6 +1891,7 @@ FPostConstructInitializeProperties::~FPostConstructInitializeProperties()
 	// Let the FObjectFinders know we left the constructor.
 	GIsInConstructor--;
 	check(GIsInConstructor >= 0);
+	GConstructedObject = LastConstructedObject;
 
 //	SCOPE_CYCLE_COUNTER(STAT_PostConstructInitializeProperties);
 	check(Obj);
diff --git a/Engine/Source/Runtime/CoreUObject/Public/UObject/UObjectGlobals.h b/Engine/Source/Runtime/CoreUObject/Public/UObject/UObjectGlobals.h
--- a/Engine/Source/Runtime/CoreUObject/Public/UObject/UObjectGlobals.h
+++ b/Engine/Source/Runtime/CoreUObject/Public/UObject/UObjectGlobals.h
@@ -978,7 +978,7 @@ private:
 	/**  List of component classes to intialize after the C++ constructors **/
 	mutable FSubobjectsToInit ComponentInits;
 	/**  Previously constructed object in the callstack */
-	TGuardValue<UObject*> LastConstructedObject;
+	UObject* LastConstructedObject;
 };
 
 /**
