Expose AddUninitializedValue(s) in FScriptArrayHelper

diff --git a/Engine/Source/Runtime/CoreUObject/Public/UObject/UnrealType.h b/Engine/Source/Runtime/CoreUObject/Public/UObject/UnrealType.h
--- a/Engine/Source/Runtime/CoreUObject/Public/UObject/UnrealType.h
+++ b/Engine/Source/Runtime/CoreUObject/Public/UObject/UnrealType.h
@@ -2275,6 +2275,26 @@ public:
 	{ 
 		return AddValues(1);
 	}
+	/**
+	*	Add unintialized values to the end of the array.
+	*	@param	Count: the number of items to insert.
+	*	@return	the index of the first newly added item.
+	**/
+	int32 AddUninitializedValues(int32 Count)
+	{
+		check(Count>0);
+		checkSlow(Num() >= 0);
+		const int32 OldNum = Array->Add(Count, ElementSize);
+		return OldNum;
+	}
+	/**
+	*	Add an uninitialized value to the end of the array.
+	*	@return	the index of the newly added item.
+	**/
+	FORCEINLINE int32 AddUninitializedValue()
+	{
+		return AddUninitializedValues(1);
+	}
 	/**
 	 *	Insert blank, constructed values into the array.
 	 *	@param	Index: index of the first inserted item after completion
@@ -2348,18 +2368,6 @@ public:
 	}		
 private:
 
-	/**
-	*	Add unintialized values to the end of the array.
-	*	@param	Count: the number of items to insert.
-	*	@return	the index of the first newly added item.
-	**/
-	int32 AddUninitializedValues(int32 Count)
-	{ 
-		check(Count>0);
-		checkSlow(Num() >= 0); 
-		const int32 OldNum = Array->Add(Count, ElementSize);		
-		return OldNum;
-	}
 	/**
 	 *	Internal function to call into the property system to construct / initialize elements.
 	 *	@param Index: first item to .
