UE4: Fixing an edge case where during async loading a BP generated class could be created, but not yet preloaded at the time another object is loading it as an import. This only happens during async loading in cooked builds, but is safe to call all the time since it early-outs if it was already called.

diff --git a/Engine/Source/Runtime/CoreUObject/Private/UObject/LinkerLoad.cpp b/Engine/Source/Runtime/CoreUObject/Private/UObject/LinkerLoad.cpp
--- a/Engine/Source/Runtime/CoreUObject/Private/UObject/LinkerLoad.cpp
+++ b/Engine/Source/Runtime/CoreUObject/Private/UObject/LinkerLoad.cpp
@@ -3262,6 +3262,12 @@ UObject* ULinkerLoad::CreateImport( int32 Index )
 			{
 				if( UClass*	FindClass = FindObjectFast<UClass>( ClassPackage, Import.ClassName, false, false ) )
 				{
+					// Make sure the class has been loaded and linked before creating a CDO.
+					// This is an edge case, but can happen if a blueprint package has not finished creating exports for a class
+					// during async loading, and another package creates the class via CreateImport while in cooked builds because
+					// we don't call preload immediately after creating a class in CreateExport like in non-cooked builds.
+					Preload( FindClass );
+
 					FindClass->GetDefaultObject(); // build the CDO if it isn't already built
 					UObject*	FindObject		= NULL;
 	
