Core: Add UStruct::GetStringMetaDataHierarchical

diff --git a/Engine/Source/Runtime/CoreUObject/Private/UObject/Class.cpp b/Engine/Source/Runtime/CoreUObject/Private/UObject/Class.cpp
--- a/Engine/Source/Runtime/CoreUObject/Private/UObject/Class.cpp
+++ b/Engine/Source/Runtime/CoreUObject/Private/UObject/Class.cpp
@@ -1383,6 +1383,25 @@ bool UStruct::GetBoolMetaDataHierarchical(const FName& Key) const
 	}
 	return bResult;
 }
+
+bool UStruct::GetStringMetaDataHierarchical(const FName& Key, FString* OutValue) const
+{
+	for (const UStruct* TestStruct = this; TestStruct != nullptr; TestStruct = TestStruct->GetSuperStruct())
+	{
+		if (TestStruct->HasMetaData(Key))
+		{
+			if (OutValue != nullptr)
+			{
+				*OutValue = TestStruct->GetMetaData(Key);
+			}
+
+			return true;
+		}
+	}
+
+	return false;
+}
+
 #endif
 //
 // Serialize an expression to an archive.
diff --git a/Engine/Source/Runtime/CoreUObject/Public/UObject/Class.h b/Engine/Source/Runtime/CoreUObject/Public/UObject/Class.h
--- a/Engine/Source/Runtime/CoreUObject/Public/UObject/Class.h
+++ b/Engine/Source/Runtime/CoreUObject/Public/UObject/Class.h
@@ -366,8 +366,11 @@ public:
 	}
 
 #if WITH_EDITOR
-	/** Try and find metadata with the given key. If not found on this class, work up hierarchy looking for it. */
+	/** Try and find boolean metadata with the given key. If not found on this class, work up hierarchy looking for it. */
 	bool GetBoolMetaDataHierarchical(const FName& Key) const;
+
+	/** Try and find string metadata with the given key. If not found on this class, work up hierarchy looking for it. */
+	bool GetStringMetaDataHierarchical(const FName& Key, FString* OutValue = nullptr) const;
 #endif
 };
 
