Adding '+' as a valid property token character to allow references to assets that contain '+' in their (file)name

diff --git a/Engine/Source/Runtime/CoreUObject/Private/UObject/Property.cpp b/Engine/Source/Runtime/CoreUObject/Private/UObject/Property.cpp
--- a/Engine/Source/Runtime/CoreUObject/Private/UObject/Property.cpp
+++ b/Engine/Source/Runtime/CoreUObject/Private/UObject/Property.cpp
@@ -304,7 +304,7 @@ const TCHAR* UPropertyHelpers::ReadToken( const TCHAR* Buffer, FString& String,
 	else if( FChar::IsAlnum( *Buffer ) || (DottedNames && (*Buffer==TCHAR('/'))) || (*Buffer > 255) )
 	{
 		// Get identifier.
-		while( (FChar::IsAlnum(*Buffer) || (*Buffer > 255) || *Buffer==TCHAR('_') || *Buffer==TCHAR('-') || (DottedNames && (*Buffer==TCHAR('.') || *Buffer==TCHAR('/') || *Buffer==SUBOBJECT_DELIMITER_CHAR) )) )
+		while ((FChar::IsAlnum(*Buffer) || (*Buffer > 255) || *Buffer == TCHAR('_') || *Buffer == TCHAR('-') || *Buffer == TCHAR('+') || (DottedNames && (*Buffer == TCHAR('.') || *Buffer == TCHAR('/') || *Buffer == SUBOBJECT_DELIMITER_CHAR))))
 		{
 			String += *Buffer++;
 		}
