Corrected issue with structs containing arrays of structs not using the SDO as the delta when exporting T3D text.
#branch UE4
#proj Runtime.CoreUObject

diff --git a/Engine/Source/Runtime/CoreUObject/Private/UObject/PropertyArray.cpp b/Engine/Source/Runtime/CoreUObject/Private/UObject/PropertyArray.cpp
--- a/Engine/Source/Runtime/CoreUObject/Private/UObject/PropertyArray.cpp
+++ b/Engine/Source/Runtime/CoreUObject/Private/UObject/PropertyArray.cpp
@@ -167,9 +167,10 @@ void UArrayProperty::ExportTextItem( FString& ValueStr, const void* PropertyValu
 		}
 
 		uint8* PropData = ArrayHelper.GetRawPtr(i);
-		uint8* PropDefault = (DefaultValue && DefaultArrayHelper.Num() > i)
-			? DefaultArrayHelper.GetRawPtr(i)
-			: StructDefaults;
+
+		// Always use struct defaults if the inner is a struct, for symmetry with the import of array inner struct defaults
+		uint8* PropDefault = ( StructProperty != NULL ) ? StructDefaults :
+			( ( DefaultValue && DefaultArrayHelper.Num() > i ) ? DefaultArrayHelper.GetRawPtr(i) : NULL );
 
 		// Do not re-export duplicate data from superclass when exporting to .int file
 		if ( (PortFlags & PPF_LocalizedOnly) != 0 && Inner->Identical(PropData, PropDefault) )
