Added logic for updating serialized array properties between FString and FText.

diff --git a/Engine/Source/Runtime/CoreUObject/Private/UObject/Class.cpp b/Engine/Source/Runtime/CoreUObject/Private/UObject/Class.cpp
--- a/Engine/Source/Runtime/CoreUObject/Private/UObject/Class.cpp
+++ b/Engine/Source/Runtime/CoreUObject/Private/UObject/Class.cpp
@@ -997,7 +997,44 @@ void UStruct::SerializeTaggedProperties(FArchive& Ar, uint8* Data, UStruct* Defa
 			}
 			else if( Tag.Type == NAME_ArrayProperty && Tag.InnerType != NAME_None && Tag.InnerType != CastChecked<UArrayProperty>(Property)->Inner->GetID() )
 			{
-				UE_LOG(LogClass, Warning, TEXT("Array Inner Type mismatch in %s of %s - Previous (%s) Current(%s) for package:  %s"), *Tag.Name.ToString(), *GetName(), *Tag.InnerType.ToString(), *CastChecked<UArrayProperty>(Property)->Inner->GetID().ToString(), *Ar.GetArchiveName() );
+				UArrayProperty* ArrayProperty = Cast<UArrayProperty>(Property);
+				void* ArrayPropertyData = ArrayProperty->ContainerPtrToValuePtr<void>(Data);
+
+				int32 ElementCount = 0;
+				Ar << ElementCount;
+
+				FScriptArrayHelper ScriptArrayHelper(ArrayProperty, ArrayPropertyData);
+				ScriptArrayHelper.EmptyAndAddValues(ElementCount);
+
+				if( Tag.InnerType==NAME_StrProperty && Cast<UTextProperty>(ArrayProperty->Inner) ) // Convert serialized string to text.
+				{ 
+					for(int32 i = 0; i < ElementCount; ++i)
+					{
+						FString str;
+						Ar << str;
+						FText Text = FText::FromString(str);
+						Text.Flags |= ETextFlag::ConvertedProperty;
+						CastChecked<UTextProperty>(ArrayProperty->Inner)->SetPropertyValue(ScriptArrayHelper.GetRawPtr(i), Text);
+						AdvanceProperty = true;
+					}
+					continue;
+				}
+				else if( Tag.InnerType==NAME_TextProperty && Cast<UStrProperty>(ArrayProperty->Inner) ) // Convert serialized text to string.
+				{ 
+					for(int32 i = 0; i < ElementCount; ++i)
+					{
+						FText Text;  
+						Ar << Text;
+						FString String = FTextInspector::GetSourceString(Text) ? *FTextInspector::GetSourceString(Text) : TEXT("");
+						CastChecked<UStrProperty>(ArrayProperty->Inner)->SetPropertyValue(ScriptArrayHelper.GetRawPtr(i), String);
+						AdvanceProperty = true;
+					}
+					continue; 
+				}
+				else
+				{
+					UE_LOG(LogClass, Warning, TEXT("Array Inner Type mismatch in %s of %s - Previous (%s) Current(%s) for package:  %s"), *Tag.Name.ToString(), *GetName(), *Tag.InnerType.ToString(), *CastChecked<UArrayProperty>(Property)->Inner->GetID().ToString(), *Ar.GetArchiveName() );
+				}
 			}
 			else if( Tag.Type==NAME_StructProperty && Tag.StructName!=CastChecked<UStructProperty>(Property)->Struct->GetFName() 
 				&& CastChecked<UStructProperty>(Property)->UseBinaryOrNativeSerialization(Ar) )
@@ -1053,14 +1090,14 @@ void UStruct::SerializeTaggedProperties(FArchive& Ar, uint8* Data, UStruct* Defa
 			}
 			else
 			{
-					uint8* DestAddress = Property->ContainerPtrToValuePtr<uint8>(Data, Tag.ArrayIndex);  
+				uint8* DestAddress = Property->ContainerPtrToValuePtr<uint8>(Data, Tag.ArrayIndex);  
 
-					// This property is ok.			
-					Tag.SerializeTaggedProperty( Ar, Property, DestAddress, Tag.Size, NULL );
+				// This property is ok.			
+				Tag.SerializeTaggedProperty( Ar, Property, DestAddress, Tag.Size, NULL );
 
-					AdvanceProperty = true;
-					continue;
-				}
+				AdvanceProperty = true;
+				continue;
+			}
 
 			AdvanceProperty = false;
 
