Core: Convert ULinkerLoad::CreateActiveRedirectsMap to be static and externally callable

diff --git a/Engine/Source/Runtime/CoreUObject/Private/UObject/LinkerLoad.cpp b/Engine/Source/Runtime/CoreUObject/Private/UObject/LinkerLoad.cpp
--- a/Engine/Source/Runtime/CoreUObject/Private/UObject/LinkerLoad.cpp
+++ b/Engine/Source/Runtime/CoreUObject/Private/UObject/LinkerLoad.cpp
@@ -68,7 +68,17 @@ TMap<FName, ULinkerLoad::FSubobjectRedirect> ULinkerLoad::SubobjectNameRedirects
  */
 void ULinkerLoad::CreateActiveRedirectsMap(const FString& GEngineIniName)
 {		
-	if( GConfig )
+	static bool bAlreadyInitialized_CreateActiveRedirectsMap = false;
+	if (bAlreadyInitialized_CreateActiveRedirectsMap)
+	{
+		return;
+	}
+	else
+	{
+		bAlreadyInitialized_CreateActiveRedirectsMap = true;
+	}
+
+	if (GConfig)
 	{
 		FConfigSection* PackageRedirects = GConfig->GetSectionPrivate( TEXT("/Script/Engine.Engine"), false, true, GEngineIniName );
 		for( FConfigSection::TIterator It(*PackageRedirects); It; ++It )
@@ -657,12 +667,7 @@ void UpdateObjectLoadingStatusMessage()
  */
 ULinkerLoad::ELinkerStatus ULinkerLoad::CreateLoader()
 {
-	static bool bAlreadyInitialized_CreateActiveRedirectsMap = false;
-	if( bAlreadyInitialized_CreateActiveRedirectsMap == false )
-	{
-		CreateActiveRedirectsMap( GEngineIni );
-		bAlreadyInitialized_CreateActiveRedirectsMap = true;
-	}
+	CreateActiveRedirectsMap( GEngineIni );
 
 	if( !Loader )
 	{
diff --git a/Engine/Source/Runtime/CoreUObject/Public/UObject/Linker.h b/Engine/Source/Runtime/CoreUObject/Public/UObject/Linker.h
--- a/Engine/Source/Runtime/CoreUObject/Public/UObject/Linker.h
+++ b/Engine/Source/Runtime/CoreUObject/Public/UObject/Linker.h
@@ -1325,7 +1325,7 @@ public:
 	/**
 	 * Add redirects to ULinkerLoad static map
 	 */
-	COREUOBJECT_API void CreateActiveRedirectsMap(const FString& GEngineIniName);
+	COREUOBJECT_API static void CreateActiveRedirectsMap(const FString& GEngineIniName);
 
 	/**
 	 * Locates the class adjusted index and its package adjusted index for a given class name in the import map
