Adding guards against using SetDefaultSubobjectClass and DoNotCreateOptionalDefaultSubobject outside of class initializer list

diff --git a/Engine/Source/Runtime/CoreUObject/Private/UObject/UObjectGlobals.cpp b/Engine/Source/Runtime/CoreUObject/Private/UObject/UObjectGlobals.cpp
--- a/Engine/Source/Runtime/CoreUObject/Private/UObject/UObjectGlobals.cpp
+++ b/Engine/Source/Runtime/CoreUObject/Private/UObject/UObjectGlobals.cpp
@@ -1849,6 +1849,7 @@ UObject::UObject(const class FPostConstructInitializeProperties& PCIP)
 {
 	check(!PCIP.Obj || PCIP.Obj == this);
 	const_cast<FPostConstructInitializeProperties&>(PCIP).Obj = this;
+	const_cast<FPostConstructInitializeProperties&>(PCIP).FinalizeSubobjectClassInitialization();
 }
 
 /* Global flag so that FObjectFinders know if they are called from inside the UObject constructors or not. */
@@ -1861,6 +1862,7 @@ FPostConstructInitializeProperties::FPostConstructInitializeProperties() :
 	ObjectArchetype(NULL),
 	bCopyTransientsFromClassDefaults(false),
 	bShouldIntializePropsFromArchetype(true),
+	bSubobjectClassInitializationAllowed(true),
 	InstanceGraph(NULL),
 	LastConstructedObject(GConstructedObject)
 {
@@ -1875,6 +1877,7 @@ FPostConstructInitializeProperties::FPostConstructInitializeProperties(UObject*
 	// if the SubobjectRoot NULL, then we want to copy the transients from the template, otherwise we are doing a duplicate and we want to copy the transients from the class defaults
 	bCopyTransientsFromClassDefaults(bInCopyTransientsFromClassDefaults),
 	bShouldIntializePropsFromArchetype(bInShouldIntializeProps),
+	bSubobjectClassInitializationAllowed(true),
 	InstanceGraph(InInstanceGraph),
 	LastConstructedObject(GConstructedObject)
 {
@@ -2131,6 +2134,12 @@ bool FPostConstructInitializeProperties::IslegalOverride(FName InComponentName,
 	return true;
 }
 
+void FPostConstructInitializeProperties::AssertIfSubobjectSetupIsNotAllowed(const TCHAR* SubobjectName) const
+{
+	UE_CLOG(!bSubobjectClassInitializationAllowed, LogUObjectGlobals, Fatal,
+		TEXT("%s.%s: Subobject class setup is only allowed in base class constructor call (in the initialization list)"), Obj ? *Obj->GetFullName() : TEXT("NULL"), SubobjectName);
+}
+
 UObject* StaticConstructObject
 (
 	UClass*			InClass,
diff --git a/Engine/Source/Runtime/CoreUObject/Public/UObject/UObjectGlobals.h b/Engine/Source/Runtime/CoreUObject/Public/UObject/UObjectGlobals.h
--- a/Engine/Source/Runtime/CoreUObject/Public/UObject/UObjectGlobals.h
+++ b/Engine/Source/Runtime/CoreUObject/Public/UObject/UObjectGlobals.h
@@ -801,6 +801,7 @@ public:
 	template<class T>
 	FPostConstructInitializeProperties const& SetDefaultSubobjectClass(FName SubobjectName) const
 	{
+		AssertIfSubobjectSetupIsNotAllowed(*SubobjectName.GetPlainNameString());
 		ComponentOverrides.Add(SubobjectName, T::StaticClass(), *this);
 		return *this;
 	}
@@ -811,6 +812,7 @@ public:
 	template<class T>
 	FORCEINLINE FPostConstructInitializeProperties const& SetDefaultSubobjectClass(TCHAR const*SubobjectName) const
 	{
+		AssertIfSubobjectSetupIsNotAllowed(SubobjectName);
 		ComponentOverrides.Add(SubobjectName, T::StaticClass(), *this);
 		return *this;
 	}
@@ -821,6 +823,7 @@ public:
 	 */
 	FPostConstructInitializeProperties const& DoNotCreateDefaultSubobject(FName SubobjectName) const
 	{
+		AssertIfSubobjectSetupIsNotAllowed(*SubobjectName.GetPlainNameString());
 		ComponentOverrides.Add(SubobjectName, NULL, *this);
 		return *this;
 	}
@@ -831,6 +834,7 @@ public:
 	 */
 	FORCEINLINE FPostConstructInitializeProperties const& DoNotCreateDefaultSubobject(TCHAR const*SubobjectName) const
 	{
+		AssertIfSubobjectSetupIsNotAllowed(SubobjectName);
 		ComponentOverrides.Add(SubobjectName, NULL, *this);
 		return *this;
 	}
@@ -845,6 +849,11 @@ public:
 	 **/
 	static void AssertIfInConstructor(UObject* Outer, const TCHAR* ErrorMessage);
 
+	FORCEINLINE void FinalizeSubobjectClassInitialization()
+	{
+		bSubobjectClassInitializationAllowed = false;
+	}
+
 private:
 
 	friend class UObject; 
@@ -962,6 +971,8 @@ private:
 		TArray<FSubobjectInit, TInlineAllocator<8> > SubobjectInits;
 	};
 
+	/** Asserts if SetDefaultSubobjectClass or DoNotCreateOptionalDefaultSuobject are called inside of the constructor body */
+	void AssertIfSubobjectSetupIsNotAllowed(const TCHAR* SubobjectName) const;
 
 	/**  object to initialize, from static allocate object, after construction **/
 	UObject* Obj;
@@ -971,6 +982,8 @@ private:
 	bool bCopyTransientsFromClassDefaults;
 	/**  If true, intialize the properties **/
 	bool bShouldIntializePropsFromArchetype;
+	/**  Only true until PCIP has not reached the base UObject class */
+	bool bSubobjectClassInitializationAllowed;
 	/**  Instance graph **/
 	struct FObjectInstancingGraph* InstanceGraph;
 	/**  List of component classes to override from derived classes **/
