Added code to gracefully handle cases where unknown script package is in the import list of another streaming package.

#ttp 339974: ELEMENTALVR: AUTOMATION: CRASH: Running StaticMeshUV editor test

diff --git a/Engine/Source/Runtime/CoreUObject/Private/Serialization/AsyncLoading.cpp b/Engine/Source/Runtime/CoreUObject/Private/Serialization/AsyncLoading.cpp
--- a/Engine/Source/Runtime/CoreUObject/Private/Serialization/AsyncLoading.cpp
+++ b/Engine/Source/Runtime/CoreUObject/Private/Serialization/AsyncLoading.cpp
@@ -530,8 +530,16 @@ EAsyncPackageState::Type FAsyncPackage::LoadImports()
 		if (!ExistingPackage && ContainsDependencyPackage(PendingImportedPackages, ImportPackageName) == INDEX_NONE)
 		{			
 			// The package doesn't exist and this import is not in the dependency list so add it now.
-			UE_LOG(LogStreaming, Verbose, TEXT("FAsyncPackage::LoadImports for %s: Loading %s"), *PackageNameToLoad, *ImportPackageName);			
-			AddImportDependency(AsyncQueueIndex, ImportPackageName);
+			if (!FPackageName::IsShortPackageName(ImportPackageName))
+			{
+				UE_LOG(LogStreaming, Verbose, TEXT("FAsyncPackage::LoadImports for %s: Loading %s"), *PackageNameToLoad, *ImportPackageName);
+				AddImportDependency(AsyncQueueIndex, ImportPackageName);
+			}
+			else
+			{
+				// This usually means there's a reference to a script package from another project
+				UE_LOG(LogStreaming, Warning, TEXT("FAsyncPackage::LoadImports for %s: Short package name in imports list: %s"), *PackageNameToLoad, *ImportPackageName);
+			}
 		}
 	}
 
