#UE4 - fix for issue raised on answerhub, deep class hierarchy with load order enforced by other dependencies fails to load correctly (https://answers.unrealengine.com/questions/60535/editor-crashes-on-load-after-editing-a-blueprint.html)

diff --git a/Engine/Source/Runtime/CoreUObject/Private/Blueprint/BlueprintSupport.cpp b/Engine/Source/Runtime/CoreUObject/Private/Blueprint/BlueprintSupport.cpp
--- a/Engine/Source/Runtime/CoreUObject/Private/Blueprint/BlueprintSupport.cpp
+++ b/Engine/Source/Runtime/CoreUObject/Private/Blueprint/BlueprintSupport.cpp
@@ -66,23 +66,31 @@ bool ULinkerLoad::RegenerateBlueprintClass(UClass* LoadClass, UObject* ExportObj
 	LoadClass->ClassDefaultObject = CurrentCDO;
 
 	// Make sure that we regenerate any parent classes first before attempting to build a child
- 	UClass* ClassChain = LoadClass->GetSuperClass();
- 	while(ClassChain && ClassChain->ClassGeneratedBy)
- 	{
- 		if(ClassChain->ClassGeneratedBy->HasAnyFlags(RF_BeingRegenerated))
- 		{
+	TArray<UClass*> ClassChainOrdered;
+	{
+		// Just ordering the class hierarchy from root to leafs:
+		UClass* ClassChain = LoadClass->GetSuperClass();
+		while (ClassChain && ClassChain->ClassGeneratedBy)
+		{
+			// O(n) insert, but n is tiny because this is a class hierarchy...
+			ClassChainOrdered.Insert(ClassChain, 0);
+			ClassChain = ClassChain->GetSuperClass();
+		}
+	}
+ 	for( auto Class : ClassChainOrdered )
+	{
+		if (Class->ClassGeneratedBy->HasAnyFlags(RF_BeingRegenerated))
+		{
 			// Always load the parent blueprint here in case there is a circular dependency. This will
 			// ensure that the blueprint is fully serialized before attempting to regenerate the class.
-			FPreloadMembersHelper::PreloadMembers(ClassChain->ClassGeneratedBy);
+			FPreloadMembersHelper::PreloadMembers(Class->ClassGeneratedBy);
 
-			ClassChain->ClassGeneratedBy->SetFlags(RF_NeedLoad);
-			ClassChain->ClassGeneratedBy->GetLinker()->Preload(ClassChain->ClassGeneratedBy);
+			Class->ClassGeneratedBy->SetFlags(RF_NeedLoad);
+			Class->ClassGeneratedBy->GetLinker()->Preload(Class->ClassGeneratedBy);
 
-			ClassChain->ClassGeneratedBy->RegenerateClass(ClassChain, NULL, GObjLoaded);
- 		}
- 
- 		ClassChain = ClassChain->GetSuperClass();
- 	}
+			Class->ClassGeneratedBy->RegenerateClass(Class, NULL, GObjLoaded);
+		}
+	}
 
 	// Preload the blueprint to make sure it has all the data the class needs for regeneration
 	if( LoadClass->ClassGeneratedBy->HasAnyFlags(RF_NeedLoad) )
