Fixed world composition loading from a compressed content

diff --git a/Engine/Source/Runtime/CoreUObject/Private/Misc/WorldCompositionUtility.cpp b/Engine/Source/Runtime/CoreUObject/Private/Misc/WorldCompositionUtility.cpp
--- a/Engine/Source/Runtime/CoreUObject/Private/Misc/WorldCompositionUtility.cpp
+++ b/Engine/Source/Runtime/CoreUObject/Private/Misc/WorldCompositionUtility.cpp
@@ -66,7 +66,7 @@ bool FWorldTileInfo::Read(const FString& InPackageFileName, FWorldTileInfo& OutI
 {
 	// Fill with default information
 	OutInfo = FWorldTileInfo();
-	
+
 	// Create a file reader to load the file
 	TScopedPointer<FArchive> FileReader(IFileManager::Get().CreateFileReader(*InPackageFileName));
 	if (FileReader == NULL)
@@ -87,8 +87,19 @@ bool FWorldTileInfo::Read(const FString& InPackageFileName, FWorldTileInfo& OutI
 	}
 
 	// Does the package contains a level info?
-	if(FileSummary.WorldTileInfoDataOffset != 0)
+	if (FileSummary.WorldTileInfoDataOffset != 0)
 	{
+		if (!!(FileSummary.PackageFlags & PKG_StoreCompressed))
+		{
+			check(FileSummary.CompressedChunks.Num() > 0);
+			if (!FileReader->SetCompressionMap(&FileSummary.CompressedChunks, (ECompressionFlags)FileSummary.CompressionFlags))
+			{
+				FileReader = new FArchiveAsync(*InPackageFileName); // re-assign scope pointer
+				check(!FileReader->IsError());
+				verify(FileReader->SetCompressionMap(&FileSummary.CompressedChunks, (ECompressionFlags)FileSummary.CompressionFlags));
+			}
+		}
+				
 		// Seek the the part of the file where the structure lives
 		FileReader->Seek(FileSummary.WorldTileInfoDataOffset);
 
