Make sure UPackage has proper Guid set on load from the summary

diff --git a/Engine/Source/Runtime/CoreUObject/Private/UObject/LinkerLoad.cpp b/Engine/Source/Runtime/CoreUObject/Private/UObject/LinkerLoad.cpp
--- a/Engine/Source/Runtime/CoreUObject/Private/UObject/LinkerLoad.cpp
+++ b/Engine/Source/Runtime/CoreUObject/Private/UObject/LinkerLoad.cpp
@@ -994,6 +994,9 @@ ULinkerLoad::ELinkerStatus ULinkerLoad::SerializePackageFileSummary()
 			
 			// Propagate package file size
 			LinkerRootPackage->FileSize = TotalSize();
+
+			// Propagate package Guid
+			LinkerRootPackage->SetGuid( Summary.Guid );
 		}
 		
 		// Propagate fact that package cannot use lazy loading to archive (aka this).
