AActor::OwnedComponents will no longer be exposed as UPROPERTY.

#change made OwnedComponents native, referenced through AActor::AddReferencedObjects
#change Added FReferenceCollectorAddReferencedObject overload that accepts TArrays

diff --git a/Engine/Source/Runtime/CoreUObject/Public/UObject/UObjectGlobals.h b/Engine/Source/Runtime/CoreUObject/Public/UObject/UObjectGlobals.h
--- a/Engine/Source/Runtime/CoreUObject/Public/UObject/UObjectGlobals.h
+++ b/Engine/Source/Runtime/CoreUObject/Public/UObject/UObjectGlobals.h
@@ -1365,6 +1365,22 @@ public:
 		HandleObjectReference(*(UObject**)&Object, ReferencingObject, ReferencingProperty);
 	}
 
+	/**
+	* Adds references to an array of objects.
+	*
+	* @param ObjectArray Referenced objects array.
+	* @param ReferencingObject Referencing object (if available).
+	* @param ReferencingProperty Referencing property (if available).
+	*/
+	template<class UObjectType>
+	void AddReferencedObject(TArray<UObjectType>& ObjectArray, const UObject* ReferencingObject = NULL, const UObject* ReferencingProperty = NULL)
+	{
+		for (auto& Object : ObjectArray)
+		{
+			HandleObjectReference(*(UObject**)&Object, ReferencingObject, ReferencingProperty);
+		}
+	}
+
 	/**
 	 * If true archetype references should not be added to this collector.
 	 */
diff --git a/Engine/Source/Runtime/Engine/Classes/GameFramework/Actor.h b/Engine/Source/Runtime/Engine/Classes/GameFramework/Actor.h
--- a/Engine/Source/Runtime/Engine/Classes/GameFramework/Actor.h
+++ b/Engine/Source/Runtime/Engine/Classes/GameFramework/Actor.h
@@ -1061,6 +1061,7 @@ public:
 	virtual void BeginDestroy() override;
 	virtual bool IsReadyForFinishDestroy() override;
 	virtual bool Rename( const TCHAR* NewName=NULL, UObject* NewOuter=NULL, ERenameFlags Flags=REN_None ) override;
+	static void AddReferencedObjects(UObject* InThis, FReferenceCollector& Collector);
 #if WITH_EDITOR
 	virtual void PreEditChange(UProperty* PropertyThatWillChange) override;
 	virtual void PostEditChangeProperty(FPropertyChangedEvent& PropertyChangedEvent) override;
@@ -2071,7 +2072,6 @@ public:
 	void ResetOwnedComponents();
 
 private:
-	UPROPERTY(transient, duplicatetransient)
 	TArray<UActorComponent*> OwnedComponents;
 	
 public:
diff --git a/Engine/Source/Runtime/Engine/Private/Actor.cpp b/Engine/Source/Runtime/Engine/Private/Actor.cpp
--- a/Engine/Source/Runtime/Engine/Private/Actor.cpp
+++ b/Engine/Source/Runtime/Engine/Private/Actor.cpp
@@ -307,6 +307,13 @@ void AActor::PostInitProperties()
 	}
 }
 
+void AActor::AddReferencedObjects(UObject* InThis, FReferenceCollector& Collector)
+{
+	AActor* This = CastChecked<AActor>(InThis);
+	Collector.AddReferencedObject(This->OwnedComponents);
+	Super::AddReferencedObjects(InThis, Collector);
+}
+
 UWorld* AActor::GetWorld() const
 {
 	// CDO objects do not belong to a world
diff --git a/Engine/Source/Runtime/Engine/Private/ActorComponent.cpp b/Engine/Source/Runtime/Engine/Private/ActorComponent.cpp
--- a/Engine/Source/Runtime/Engine/Private/ActorComponent.cpp
+++ b/Engine/Source/Runtime/Engine/Private/ActorComponent.cpp
@@ -294,6 +294,14 @@ void UActorComponent::BeginDestroy()
 
 	ExecuteUnregisterEvents();
 	World = NULL;
+
+	// Remove from the parent's OwnedComponents list
+	AActor* Owner = GetOwner();
+	if (Owner)
+	{
+		Owner->RemoveOwnedComponent(this);
+	}
+
 	Super::BeginDestroy();
 }
 
