Making sure Linker is not trying to re-create an import if it's already been created by VerifyImport.

diff --git a/Engine/Source/Runtime/CoreUObject/Private/UObject/LinkerLoad.cpp b/Engine/Source/Runtime/CoreUObject/Private/UObject/LinkerLoad.cpp
--- a/Engine/Source/Runtime/CoreUObject/Private/UObject/LinkerLoad.cpp
+++ b/Engine/Source/Runtime/CoreUObject/Private/UObject/LinkerLoad.cpp
@@ -3346,7 +3346,11 @@ UObject* ULinkerLoad::CreateImport( int32 Index )
 			if(Import.SourceIndex != INDEX_NONE)
 			{
 				check(Import.SourceLinker);
-				Import.XObject = Import.SourceLinker->CreateExport( Import.SourceIndex );
+				// VerifyImport may have already created the import and SourceIndex has changed to point to the actual redirected object
+				if (!Import.XObject)
+				{
+					Import.XObject = Import.SourceLinker->CreateExport(Import.SourceIndex);
+				}
 				// If an object has been replaced (consolidated) in the editor and its package hasn't been saved yet
 				// it's possible to get UbjectRedirector here as the original export is dynamically replaced
 				// with the redirector (the original object has been deleted but the data on disk hasn't been updated)
