Fix for Regression: BLUEPRINT: Infinite loop detected but wont go away

diff --git a/Engine/Source/Runtime/CoreUObject/Private/UObject/ScriptCore.cpp b/Engine/Source/Runtime/CoreUObject/Private/UObject/ScriptCore.cpp
--- a/Engine/Source/Runtime/CoreUObject/Private/UObject/ScriptCore.cpp
+++ b/Engine/Source/Runtime/CoreUObject/Private/UObject/ScriptCore.cpp
@@ -541,7 +541,7 @@ void UObject::ProcessInternal( FFrame& Stack, RESULT_DECL )
 		MS_ALIGN(16) uint8 Buffer[MAX_SIMPLE_RETURN_VALUE_SIZE] GCC_ALIGN(16);
 
 #if DO_GUARD
-		if (++Recurse > RECURSE_LIMIT)
+		if (++Recurse == RECURSE_LIMIT)
 		{
 			// We've hit the recursion limit, so print out the stack, warn, and then continue with a zeroed return value.
 			UE_LOG(LogScriptCore, Log, TEXT("%s"), *Stack.GetStackTrace());
@@ -555,6 +555,9 @@ void UObject::ProcessInternal( FFrame& Stack, RESULT_DECL )
 			FBlueprintExceptionInfo InfiniteRecursionExceptionInfo(EBlueprintExceptionType::InfiniteLoop, Desc);
 			FBlueprintCoreDelegates::ThrowScriptException(this, Stack, InfiniteRecursionExceptionInfo);
 
+			// Bump the recursion so we don't re-enter after we've thrown an exception
+			++Recurse;
+
 			return;
 		}
 #endif
