A Field's tooltip should be its name if one hasn't been explicitly provided

diff --git a/Engine/Source/Runtime/CoreUObject/Private/UObject/Class.cpp b/Engine/Source/Runtime/CoreUObject/Private/UObject/Class.cpp
--- a/Engine/Source/Runtime/CoreUObject/Private/UObject/Class.cpp
+++ b/Engine/Source/Runtime/CoreUObject/Private/UObject/Class.cpp
@@ -218,12 +218,16 @@ FText UField::GetDisplayNameText() const
 FText UField::GetToolTipText() const
 {
 	FText LocalizedToolTip;
-	const FString NativeToolTip = GetMetaData( TEXT("Tooltip") );
+	FString NativeToolTip = GetMetaData( TEXT("Tooltip") );
 
 	static const FString Namespace = TEXT("UObjectToolTips");
 	const FString Key = GetFullGroupName(true) + TEXT(".") + GetName();
 	if ( !(FText::FindText( Namespace, Key, /*OUT*/LocalizedToolTip )) || *FTextInspector::GetSourceString(LocalizedToolTip) != NativeToolTip)
 	{
+		if (NativeToolTip.IsEmpty())
+		{
+			NativeToolTip = GetName();
+		}
 		LocalizedToolTip = FText::FromString(NativeToolTip);
 	}
 
