Optional bUseFullTimeLimit for AsyncLoading, controlled by UEngine::bAsyncLoadingUseFullTimeLimit config variable (default is true)

diff --git a/Engine/Source/Editor/UnrealEd/Private/Editor.cpp b/Engine/Source/Editor/UnrealEd/Private/Editor.cpp
--- a/Engine/Source/Editor/UnrealEd/Private/Editor.cpp
+++ b/Engine/Source/Editor/UnrealEd/Private/Editor.cpp
@@ -854,7 +854,7 @@ void UEditorEngine::Tick( float DeltaSeconds, bool bIdleMode )
 	// Update subsystems.
 	{
 		// This assumes that UObject::StaticTick only calls ProcessAsyncLoading.	
-		StaticTick(DeltaSeconds, AsyncLoadingTimeLimit / 1000.f);
+		StaticTick(DeltaSeconds, bAsyncLoadingUseFullTimeLimit, AsyncLoadingTimeLimit / 1000.f);
 	}
 
 	// Look for realtime flags.
diff --git a/Engine/Source/Runtime/CoreUObject/Private/UObject/UObjectGlobals.cpp b/Engine/Source/Runtime/CoreUObject/Private/UObject/UObjectGlobals.cpp
--- a/Engine/Source/Runtime/CoreUObject/Private/UObject/UObjectGlobals.cpp
+++ b/Engine/Source/Runtime/CoreUObject/Private/UObject/UObjectGlobals.cpp
@@ -279,14 +279,15 @@ void GlobalSetProperty( const TCHAR* Value, UClass* Class, UProperty* Property,
  * @todo: Move stats code into core?
  *
  * @param DeltaTime	Time in seconds since last call
+ * @param bUseFullTimeLimit	If true, use the entire time limit even if blocked on I/O
  * @param AsyncLoadingTime Time in seconds to use for async loading limit
  */
-void StaticTick( float DeltaTime, float AsyncLoadingTime )
+void StaticTick( float DeltaTime, bool bUseFullTimeLimit, float AsyncLoadingTime )
 {
 	check(!IsLoading());
 
 	// Spend a bit of time (pre)loading packages - currently 5 ms.
-	ProcessAsyncLoading(true, false, AsyncLoadingTime);
+	ProcessAsyncLoading(true, bUseFullTimeLimit, AsyncLoadingTime);
 
 	// Check natives.
 	extern int32 GNativeDuplicate;
diff --git a/Engine/Source/Runtime/CoreUObject/Public/UObject/UObjectGlobals.h b/Engine/Source/Runtime/CoreUObject/Public/UObject/UObjectGlobals.h
--- a/Engine/Source/Runtime/CoreUObject/Public/UObject/UObjectGlobals.h
+++ b/Engine/Source/Runtime/CoreUObject/Public/UObject/UObjectGlobals.h
@@ -240,7 +240,7 @@ COREUOBJECT_API void PreInitUObject();
  */
 COREUOBJECT_API void MarkObjectsToDisregardForGC();
 COREUOBJECT_API bool StaticExec( UWorld* InWorld, const TCHAR* Cmd, FOutputDevice& Ar=*GLog );
-COREUOBJECT_API void StaticTick( float DeltaTime, float AsyncLoadingTime = 0.005f );
+COREUOBJECT_API void StaticTick( float DeltaTime, bool bUseFullTimeLimit = true, float AsyncLoadingTime = 0.005f );
 
 /**
  * Loads a package and all contained objects that match context flags.
diff --git a/Engine/Source/Runtime/Engine/Classes/Engine/Engine.h b/Engine/Source/Runtime/Engine/Classes/Engine/Engine.h
--- a/Engine/Source/Runtime/Engine/Classes/Engine/Engine.h
+++ b/Engine/Source/Runtime/Engine/Classes/Engine/Engine.h
@@ -1097,6 +1097,10 @@ public:
 	UPROPERTY(EditAnywhere, config, Category = LevelStreaming, AdvancedDisplay)
 	float AsyncLoadingTimeLimit;
 
+	/** Whether to use the entire time limit even if blocked on I/O */
+	UPROPERTY(EditAnywhere, config, Category = LevelStreaming, AdvancedDisplay)
+	uint32 bAsyncLoadingUseFullTimeLimit:1;
+	
 	/** Additional time to spend asynchronous loading during a "high priority" load */
 	UPROPERTY(EditAnywhere, config, Category = LevelStreaming, AdvancedDisplay)
 	float PriorityAsyncLoadingExtraTime;
diff --git a/Engine/Source/Runtime/Engine/Private/GameEngine.cpp b/Engine/Source/Runtime/Engine/Private/GameEngine.cpp
--- a/Engine/Source/Runtime/Engine/Private/GameEngine.cpp
+++ b/Engine/Source/Runtime/Engine/Private/GameEngine.cpp
@@ -370,6 +370,7 @@ UEngine::UEngine(const class FPostConstructInitializeProperties& PCIP)
 	: Super(PCIP)
 {
 	AsyncLoadingTimeLimit = 5.0f;
+	bAsyncLoadingUseFullTimeLimit = true;
 	PriorityAsyncLoadingExtraTime = 20.0f;
 
 	C_WorldBox = FColor(0, 0, 40, 255);
@@ -863,7 +864,7 @@ void UGameEngine::Tick( float DeltaSeconds, bool bIdleMode )
 	// Update subsystems.
 	{
 		// This assumes that UObject::StaticTick only calls ProcessAsyncLoading.
-		StaticTick(DeltaSeconds, AsyncLoadingTimeLimit / 1000.f);
+		StaticTick(DeltaSeconds, bAsyncLoadingUseFullTimeLimit, AsyncLoadingTimeLimit / 1000.f);
 	}
 
 	// -----------------------------------------------------
