General - UFunction::GetReturnProperty is now const (always was, just adding the keyword so it's callable in a const context).

diff --git a/Engine/Source/Runtime/CoreUObject/Private/UObject/Class.cpp b/Engine/Source/Runtime/CoreUObject/Private/UObject/Class.cpp
--- a/Engine/Source/Runtime/CoreUObject/Private/UObject/Class.cpp
+++ b/Engine/Source/Runtime/CoreUObject/Private/UObject/Class.cpp
@@ -3621,7 +3621,7 @@ void UFunction::Serialize( FArchive& Ar )
 	}
 }
 
-UProperty* UFunction::GetReturnProperty()
+UProperty* UFunction::GetReturnProperty() const
 {
 	for( TFieldIterator<UProperty> It(this); It && (It->PropertyFlags & CPF_Parm); ++It )
 	{
diff --git a/Engine/Source/Runtime/CoreUObject/Public/UObject/Class.h b/Engine/Source/Runtime/CoreUObject/Public/UObject/Class.h
--- a/Engine/Source/Runtime/CoreUObject/Public/UObject/Class.h
+++ b/Engine/Source/Runtime/CoreUObject/Public/UObject/Class.h
@@ -1211,7 +1211,8 @@ public:
 		checkSlow(!SuperStruct||SuperStruct->IsA<UFunction>());
 		return (UFunction*)SuperStruct;
 	}
-	UProperty* GetReturnProperty();
+
+	UProperty* GetReturnProperty() const;
 
 	/**
 	 * Used to safely check whether the passed in flag is set.
