Helper functions for TMaps and TSets in FReferenceCollector

diff --git a/Engine/Source/Editor/Kismet/Private/BlueprintEditor.cpp b/Engine/Source/Editor/Kismet/Private/BlueprintEditor.cpp
--- a/Engine/Source/Editor/Kismet/Private/BlueprintEditor.cpp
+++ b/Engine/Source/Editor/Kismet/Private/BlueprintEditor.cpp
@@ -2498,10 +2498,8 @@ void FBlueprintEditor::AddReferencedObjects( FReferenceCollector& Collector )
 		UObject* EditingObject = GetEditingObjects()[i];
 		Collector.AddReferencedObject(EditingObject);
 	}
-	for (int32 Index = 0; Index < StandardLibraries.Num(); Index++)
-	{
-		Collector.AddReferencedObject(StandardLibraries[Index]);
-	}
+
+	Collector.AddReferencedObjects(StandardLibraries);
 
 	UserDefinedEnumerators.Remove(TWeakObjectPtr<UUserDefinedEnum>()); // Remove NULLs
 	for (auto ObjectPtr : UserDefinedEnumerators)
diff --git a/Engine/Source/Editor/LandscapeEditor/Private/LandscapeEdModeBrushes.cpp b/Engine/Source/Editor/LandscapeEditor/Private/LandscapeEdModeBrushes.cpp
--- a/Engine/Source/Editor/LandscapeEditor/Private/LandscapeEdModeBrushes.cpp
+++ b/Engine/Source/Editor/LandscapeEditor/Private/LandscapeEdModeBrushes.cpp
@@ -71,10 +71,7 @@ public:
 		// Allow any currently unused material instances to be GC'd
 		BrushMaterialFreeInstances.Empty();
 
-		for (auto It = BrushMaterialInstanceMap.CreateIterator(); It; ++It)
-		{
-			Collector.AddReferencedObject(It.Value());
-		}
+		Collector.AddReferencedObjects(BrushMaterialInstanceMap);
 	}
 
 	virtual void LeaveBrush() override
diff --git a/Engine/Source/Editor/UnrealEd/Private/EditorActorFolders.cpp b/Engine/Source/Editor/UnrealEd/Private/EditorActorFolders.cpp
--- a/Engine/Source/Editor/UnrealEd/Private/EditorActorFolders.cpp
+++ b/Engine/Source/Editor/UnrealEd/Private/EditorActorFolders.cpp
@@ -56,10 +56,7 @@ FActorFolders::~FActorFolders()
 void FActorFolders::AddReferencedObjects(FReferenceCollector& Collector)
 {
 	// Add references for all our UObjects so they don't get collected
-	for (auto Pair : TemporaryWorldFolders)
-	{
-		Collector.AddReferencedObject(Pair.Value);
-	}
+	Collector.AddReferencedObjects(TemporaryWorldFolders);
 }
 
 FActorFolders& FActorFolders::Get()
diff --git a/Engine/Source/Editor/UnrealEd/Private/EditorTransaction.cpp b/Engine/Source/Editor/UnrealEd/Private/EditorTransaction.cpp
--- a/Engine/Source/Editor/UnrealEd/Private/EditorTransaction.cpp
+++ b/Engine/Source/Editor/UnrealEd/Private/EditorTransaction.cpp
@@ -176,10 +176,7 @@ void FTransaction::AddReferencedObjects( FReferenceCollector& Collector )
 	{
 		Records[ Index ].AddReferencedObjects( Collector );
 	}
-	for( ObjectMapType::TIterator It( ObjectMap ); It; ++It )
-	{
-		Collector.AddReferencedObject( It.Key() );
-	}
+	Collector.AddReferencedObjects(ObjectMap);
 }
 
 // FTransactionBase interface.
diff --git a/Engine/Source/Runtime/CoreUObject/Public/UObject/UObjectGlobals.h b/Engine/Source/Runtime/CoreUObject/Public/UObject/UObjectGlobals.h
--- a/Engine/Source/Runtime/CoreUObject/Public/UObject/UObjectGlobals.h
+++ b/Engine/Source/Runtime/CoreUObject/Public/UObject/UObjectGlobals.h
@@ -1373,7 +1373,7 @@ public:
 	* @param ReferencingProperty Referencing property (if available).
 	*/
 	template<class UObjectType>
-	void AddReferencedObject(TArray<UObjectType>& ObjectArray, const UObject* ReferencingObject = NULL, const UObject* ReferencingProperty = NULL)
+	void AddReferencedObjects(TArray<UObjectType>& ObjectArray, const UObject* ReferencingObject = NULL, const UObject* ReferencingProperty = NULL)
 	{
 		for (auto& Object : ObjectArray)
 		{
@@ -1381,6 +1381,76 @@ public:
 		}
 	}
 
+	/**
+	* Adds references to a set of objects.
+	*
+	* @param ObjectSet Referenced objects set.
+	* @param ReferencingObject Referencing object (if available).
+	* @param ReferencingProperty Referencing property (if available).
+	*/
+	template<class UObjectType>
+	void AddReferencedObjects(TSet<UObjectType>& ObjectSet, const UObject* ReferencingObject = NULL, const UObject* ReferencingProperty = NULL)
+	{
+		for (auto& Object : ObjectSet)
+		{
+			HandleObjectReference(*(UObject**)&Object, ReferencingObject, ReferencingProperty);
+		}
+	}
+
+private:
+
+	/** Compile time check if a type can be automatically converted to another one */
+	template<class From, class To>
+	class CanConverFromTo
+	{
+		static uint8 Test(...);
+		static uint16 Test(To);
+	public:
+		enum Type
+		{
+			Result = sizeof(Test(From())) - 1
+		};
+	};
+
+	/**
+	* Functions used by AddReferencedObject (TMap version). Adds references to UObjects, ignores value types
+	*/
+	template<class UObjectType>
+	void AddReferencedObjectOrIgnoreValue(UObjectType& Object, const UObject* ReferencingObject, const UObject* ReferencingProperty)
+	{
+	}
+	template<class UObjectType>
+	void AddReferencedObjectOrIgnoreValue(UObjectType*& Object, const UObject* ReferencingObject, const UObject* ReferencingProperty)
+	{
+		HandleObjectReference(*(UObject**)&Object, ReferencingObject, ReferencingProperty);
+	}
+
+public:
+
+	/**
+	* Adds references to a map of objects.
+	*
+	* @param ObjectArray Referenced objects map.
+	* @param ReferencingObject Referencing object (if available).
+	* @param ReferencingProperty Referencing property (if available).
+	*/
+	template <typename TKeyType, typename TValueType>
+	void AddReferencedObjects(TMap<TKeyType, TValueType>& Map, const UObject* ReferencingObject = NULL, const UObject* ReferencingProperty = NULL)
+	{
+		static_assert(CanConverFromTo<TKeyType, UObjectBase*>::Result || CanConverFromTo<TValueType, UObjectBase*>::Result, "At least one of TMap template types must be derived from UObject");
+		for (auto& It : Map)
+		{
+			if (CanConverFromTo<TKeyType, UObjectBase*>::Result)
+			{
+				AddReferencedObjectOrIgnoreValue(It.Key, ReferencingObject, ReferencingProperty);
+			}
+			if (CanConverFromTo<TValueType, UObjectBase*>::Result)
+			{
+				AddReferencedObjectOrIgnoreValue(It.Value, ReferencingObject, ReferencingProperty);
+			}
+		}
+	}
+
 	/**
 	 * If true archetype references should not be added to this collector.
 	 */
diff --git a/Engine/Source/Runtime/Engine/Private/Landscape/Landscape.cpp b/Engine/Source/Runtime/Engine/Private/Landscape/Landscape.cpp
--- a/Engine/Source/Runtime/Engine/Private/Landscape/Landscape.cpp
+++ b/Engine/Source/Runtime/Engine/Private/Landscape/Landscape.cpp
@@ -916,10 +916,7 @@ void ALandscapeProxy::AddReferencedObjects(UObject* InThis, FReferenceCollector&
 
 	Super::AddReferencedObjects(InThis, Collector);
 
-	for (auto It = This->MaterialInstanceConstantMap.CreateIterator(); It; ++It)
-	{
-		Collector.AddReferencedObject(It.Value(), This);
-	}
+	Collector.AddReferencedObjects(This->MaterialInstanceConstantMap, This);
 
 	for (auto It = This->WeightmapUsageMap.CreateIterator(); It; ++It)
 	{
diff --git a/Engine/Source/Runtime/Engine/Private/World.cpp b/Engine/Source/Runtime/Engine/Private/World.cpp
--- a/Engine/Source/Runtime/Engine/Private/World.cpp
+++ b/Engine/Source/Runtime/Engine/Private/World.cpp
@@ -193,11 +193,8 @@ void UWorld::AddReferencedObjects(UObject* InThis, FReferenceCollector& Collecto
 		Collector.AddReferencedObject( This->AvoidanceManager, This );
 
 #if WITH_EDITORONLY_DATA
-		for( auto It = This->LandscapeInfoMap.CreateIterator(); It; ++It )
-		{
-			Collector.AddReferencedObject( It.Value(), This );
-		}
-		[REDACTED_EMAIL]: How to handle w/out having a cooker/pakcager?
+		Collector.AddReferencedObjects(This->LandscapeInfoMap, This);
+		[REDACTED_EMAIL]: How to handle w/out having a cooker/packager?
 #endif	//#if WITH_EDITORONLY_DATA
 
 		for( int32 Index = 0; Index < This->ExtraReferencedObjects.Num(); Index++ )
