#UE4 Fix for 341705, as comment explains we need to link an object's super struct before serializing  the object

diff --git a/Engine/Source/Runtime/CoreUObject/Private/UObject/LinkerLoad.cpp b/Engine/Source/Runtime/CoreUObject/Private/UObject/LinkerLoad.cpp
--- a/Engine/Source/Runtime/CoreUObject/Private/UObject/LinkerLoad.cpp
+++ b/Engine/Source/Runtime/CoreUObject/Private/UObject/LinkerLoad.cpp
@@ -2956,6 +2956,17 @@ UObject* ULinkerLoad::CreateExport( int32 Index )
 					return NULL;
 				}
 			}
+			else
+			{
+				// SuperStruct needs to be fully linked so that UStruct::Link will have access to UObject::SuperStruct->PropertySize. 
+				// There are other attempts to force our super struct to load, and I have not verified that they can all be removed
+				// in favor of this one:
+				if( !SuperStruct->HasAnyFlags(RF_LoadCompleted) && SuperStruct->GetLinker() && Export.SuperIndex.IsImport() )
+				{
+					SuperStruct->SetFlags(RF_NeedLoad);
+					Preload(SuperStruct);
+				}
+			}
 		}
 
 		// Only UClass objects and UProperty objects of intrinsic classes can have RF_Native set. Those property objects are never
