#UE4 Licensee has provided me with a test case for which this actually crashes, reverting to old behavior until I have an understanding of why (SuperStruct appears to be completely loaded, although RF_LoadCompleted is not set)

Back out revision 124 from //depot/UE4/Engine/Source/Runtime/CoreUObject/Private/UObject/LinkerLoad.cpp

diff --git a/Engine/Source/Runtime/CoreUObject/Private/UObject/LinkerLoad.cpp b/Engine/Source/Runtime/CoreUObject/Private/UObject/LinkerLoad.cpp
--- a/Engine/Source/Runtime/CoreUObject/Private/UObject/LinkerLoad.cpp
+++ b/Engine/Source/Runtime/CoreUObject/Private/UObject/LinkerLoad.cpp
@@ -2956,17 +2956,6 @@ UObject* ULinkerLoad::CreateExport( int32 Index )
 					return NULL;
 				}
 			}
-			else
-			{
-				// SuperStruct needs to be fully linked so that UStruct::Link will have access to UObject::SuperStruct->PropertySize. 
-				// There are other attempts to force our super struct to load, and I have not verified that they can all be removed
-				// in favor of this one:
-				if( !SuperStruct->HasAnyFlags(RF_LoadCompleted) && SuperStruct->GetLinker() && Export.SuperIndex.IsImport() )
-				{
-					SuperStruct->SetFlags(RF_NeedLoad);
-					Preload(SuperStruct);
-				}
-			}
 		}
 
 		// Only UClass objects and UProperty objects of intrinsic classes can have RF_Native set. Those property objects are never
