Document metadata tags

diff --git a/Engine/Source/Runtime/CoreUObject/Public/UObject/ObjectBase.h b/Engine/Source/Runtime/CoreUObject/Public/UObject/ObjectBase.h
--- a/Engine/Source/Runtime/CoreUObject/Public/UObject/ObjectBase.h
+++ b/Engine/Source/Runtime/CoreUObject/Public/UObject/ObjectBase.h
@@ -841,31 +841,251 @@ namespace US
 // Metadata specifiers
 namespace UM
 {
+	// Metadata usable in UCLASS
 	enum
 	{
-		/// [FunctionMetadata] This function can only be called on 'this' in a blueprint. It cannot be called on another instance.
-		BlueprintProtected,
+		/// [ClassMetadata] Used for Actor Component classes. If present indicates that it can be spawned by a Blueprint.
+		BlueprintSpawnableComponent,
 
-		/// [FunctionMetadata] This function is not safe to call during Actor construction.
-		UnsafeDuringActorConstruction,
+		/// [ClassMetadata] Used for Actor classes. If the native class cannot tick, Blueprint generated classes based this Actor can have bCanEverTick flag overridden.
+		ChildCanTick,
+
+		/// [ClassMetadata] For BehaviorTree nodes indicates that the class is deprecated and will display a warning when compiled.
+		DeprecatedNode,
+
+		/// [ClassMetadata] [FunctionMetadata] Used in conjunction with DeprecatedNode or DeprecatedFunction to customize the warning message displayed to the user.
+		DeprecationMessage,
+
+		/// [ClassMetadata] [PropertyMetadata] The name to display for this class or property instead of auto-generating it from the name.
+		DisplayName,
+
+		/// [ClassMetadata]
+		IsBlueprintBase,
+
+		/// [ClassMetadata] Comma delimited list of blueprint events that are not be allowed to be overriden in classes of this type
+		KismetHideOverrides,
+
+		/// [ClassMetadata] Specifies interfaces that are not compatible with the class.
+		ProhibitedInterfaces,
+
+		/// [ClassMetadata] Used by BlueprintFunctionLibrary classes to restrict the graphs the functions in the library can be used in to the classes specified.
+		RestrictedToClasses,
+
+		/// [ClassMetadata] Indicates that when placing blueprint nodes in graphs owned by this class that any hidden self pins (generally used to determine the world) should be visible because the self context of the class cannot
+		///                 provide the world context and it must be wired in manually
+		ShowHiddenSelfPins,
+
+	};
+
+	// Metadata usable in USTRUCT
+	enum
+	{
+		/// [StructMetadata] Indicates that the struct has a custom break node (and what the path to the BlueprintCallable UFunction is) that should be used instead of the default BreakStruct node.  
+		HasNativeBreak,
+
+		/// [StructMetadata] Indicates that the struct has a custom make node (and what the path to the BlueprintCallable UFunction is) that should be used instead of the default MakeStruct node.  
+		HasNativeMake,
+	};
+
+	// Metadata usable in UPROPERTY for customizing the behavior when displaying the property in a property panel
+	enum
+	{
+		/// [PropertyMetadata] Used for FStringClassReference properties.  Indicates whether abstract class types should be shown in the class picker.
+		AllowAbstract,
+
+		/// [PropertyMetadata] Used for FStringAssetReference properties.  Comma delimited list that indicates the class type(s) of assets to be displayed in the asset picker.
+		AllowedClasses,
+
+		/// [PropertyMetadata] Used for FVector properties.  It causes a ratio lock to be added when displaying this property in details panels.
+		AllowPreserveRatio,
+
+		/// [PropertyMetadata] Used for integer properties.  Clamps the valid values that can be entered in the UI to be between 0 and the length of the array specified.
+		ArrayClamp,
+
+		/// [PropertyMetadata] Used for float and integer properties.  Specifies the minimum value that may be entered for the property.
+		ClampMin,
+
+		/// [PropertyMetadata] Used for float and integer properties.  Specifies the maximum value that may be entered for the property.
+		ClampMax,
+
+		/// [ClassMetadata] [PropertyMetadata] The name to display for this class or property instead of auto-generating it from the name.
+		// DisplayName, (Commented out so as to avoid duplicate name with version in the Class section, but still show in the property section)
+
+		/// [PropertyMetadata] Indicates that the property is an asset type and it should display the thumbnail of the selected asset.
+		DisplayThumbnail,	
+	
+		/// [PropertyMetadata] Species a boolean property that is used to indicate whether editing of this property is disabled.
+		EditCondition,
+		
+		/// [PropertyMetadata] Used for FStringAssetReference properties in conjunction with AllowedClasses. Indicates whether only the exact classes specified in AllowedClasses can be used or whether subclasses are valid.
+		ExactClass,
+
+		/// [PropertyMetadata]
+		ExposeFunctionCategories,
+
+		/// [PropertyMetadata] Specifies whether the property should be exposed on a Spawn Actor for the class type.
+		ExposeOnSpawn,
+
+		/// [PropertyMetadata]
+		FixedIncrement,
+
+		/// [PropertyMetadata] Used for FColor and FLinearColor properties. Indicates that the Alpha property should be hidden when displaying the property widget in the details.
+		HideAlphaChannel,
+
+		/// [PropertyMetadata] Used for FStringClassReference properties.  Indicates whether only blueprint classes should be shown in the class picker.
+		IsBlueprintBaseOnly,
+
+		/// [PropertyMetadata]
+		MakeEditWidget,
+
+		/// [PropertyMetadata] For properties in a structure indicates the default value of the property in a blueprint make structure node.
+		MakeStructureDefaultValue,
+
+		/// [PropertyMetadata] Used FStringClassReference properties. Indicates the parent class that the class picker will use when filtering which classes to display.
+		MetaClass,
+
+		/// [PropertyMetadata]
+		Multiple,
+
+		/// [PropertyMetadata] Used for array properties. Indicates that the duplicate icon should not be shown for entries of this array in the property panel.
+		NoElementDuplicate,
+
+		/// [PropertyMetadata] Used for integer and float properties. Indicates that the spin box element of the number editing widget should not be displayed.
+		NoSpinbox,
+
+		/// [PropertyMetadata] Used by FFilePath properties. Indicates the path filter to display in the file picker.
+		FilePathFilter,
+
+		/// [PropertyMetadta] Used by FDirectoryPath properties. 
+		RelativePath,
+
+		/// [PropertyMetadta] Used by FDirectoryPath properties. 
+		RelativeToGameContentDir,
+
+		// [PropertyMetadata]
+		ShowOnlyInnerProperties,
+
+		// [PropertyMetadata]
+		SliderExponent,
+
+		// [PropertyMetadata] Overrides the automatically generated tooltip from the property comment
+		ToolTip,
+
+		/// [PropertyMetadata] Used for float and integer properties.  Specifies the lowest that the value slider should represent.
+		UIMin,
+
+		/// [PropertyMetadata] Used for float and integer properties.  Specifies the highest that the value slider should represent.
+		UIMax,
+	};
+
+	// Metadata usable in UPROPERTY for customizing the behavior of Persona and UMG
+	// TODO: Move this to be contained in those modules specifically?
+	enum 
+	{
+		/// [PropertyMetadata] The property is not exposed as a data pin and is only be editable in the details panel. Applicable only to properties that will be displayed in Persona and UMG.
+		NeverAsPin, 
+
+		/// [PropertyMetadata] The property can be exposed as a data pin, but is hidden by default. Applicable only to properties that will be displayed in Persona and UMG.
+		PinHiddenByDefault, 
+
+		/// [PropertyMetadata] The property can be exposed as a data pin and is visible by default. Applicable only to properties that will be displayed in Persona and UMG.
+		PinShownByDefault, 
+
+		/// [PropertyMetadata] The property is always exposed as a data pin. Applicable only to properties that will be displayed in Persona and UMG.
+		AlwaysAsPin, 
+
+		/// [PropertyMetadata] Indicates that the property has custom code to display and should not generate a standard property widget int he details panel. Applicable only to properties that will be displayed in Persona.
+		CustomizeProperty,
+	};
+
+	// Metadata usable in UPROPERTY for customizing the behavior of Material Expressions
+	// TODO: Move this to be contained in that module?
+	enum
+	{
+		/// [PropertyMetadata] Used for float properties in MaterialExpression classes. If the specified FMaterialExpression pin is not connected, this value is used instead.
+		OverridingInputProperty,
+
+		/// [PropertyMetadata] Used for FMaterialExpression properties in MaterialExpression classes. If specified the pin need not be connected and the value of the property marked as OverridingInputProperty will be used instead.
+		RequiredInput,
+	};
+
+
+	// Metadata usable in UFUNCTION
+	enum
+	{
+		/// [FunctionMetadata] Indicates that a BlueprintCallable function should use a Call Array Function node and that the parameters specified in the comma delimited list should be treated as wild card array properties.
+		ArrayParm,
+
+		/// [FunctionMetadata] Used when ArrayParm has been specified to indicate other function parameters that should be treated as wild card properties linked to the type of the array parameter.
+		ArrayTypeDependentParams,
+
+		/// [FunctionMetadata]
+		AutoCreateRefTerm,
 
 		/// [FunctionMetadata] This function is an internal implementation detail, used to implement another function or node.  It is never directly exposed in a graph.
 		BlueprintInternalUseOnly,
 
+		/// [FunctionMetadata] This function can only be called on 'this' in a blueprint. It cannot be called on another instance.
+		BlueprintProtected,
+
+		/// [FunctionMetadata] Indicates that a BlueprintCallable function should use the Commutative Associative Binary node.
+		CommutativeAssociativeBinaryOperator,
+
+		/// [FunctionMetadata] Indicates that a BlueprintCallable function should display in the compact display mode and the name to use in that mode.
+		CompactNodeTitle,
+
+		/// [FunctionMetadata]
+		CustomStructureParam,
+
+		/// [FunctionMetadata] For BlueprintCallable functions indicates that the object property named's default value should be the self context of the node
+		DefaultToSelf,
+
 		/// [FunctionMetadata] This function is deprecated, any blueprint references to it cause a compilation warning.
 		DeprecatedFunction,
 
-		/// [FunctionMetadata] Supplies the custom message to use for deprecation
-		DeprecationMessage,
+		/// [ClassMetadata] [FunctionMetadata] Used in conjunction with DeprecatedNode or DeprecatedFunction to customize the warning message displayed to the user.
+		// DeprecationMessage, (Commented out so as to avoid duplicate name with version in the Class section, but still show in the function section)
 
-		/// [ClassMetadata] If present, the component class can be spawned by a blueprint
-		BlueprintSpawnableComponent,
+		/// [FunctionMetadata] For BlueprintCallable functions indicates that an input exec pin should be created for each entry in the enum specified.
+		ExpandEnumAsExecs,
+
+		/// [FunctionMetadata] The name to display for node's of this function in the blueprints.
+		FriendlyName,
+
+		/// [FunctionMetadata] For BlueprintCallable functions indicates that the parameter pin should be hidden from the user's view.
+		HidePin,
+
+		/// [FunctionMetadata]
+		HideSpawnParms,
+
+		/// [FunctionMetadata] For BlueprintCallable functions provides additional keywords to be associated with the function for search purposes.
+		Keywords,
+
+		/// [FunctionMetadata] Indicates that a BlueprintCallable function is Latent
+		Latent,
 
+		/// [FunctionMetadata] For Latent BlueprintCallable functions indicates which parameter is the LatentInfo parameter
+		LatentInfo,
+
+		/// [FunctionMetadata] For BlueprintCallable functions indicates that the material override node should be used
+		MaterialParameterCollectionFunction,
+
+		/// [FunctionMetadata] For BlueprintCallable functions indicates that the function should be displayed the same as the implicit Break Struct nodes
+		NativeBreakFunc,
+
+		/// [FunctionMetadata] For BlueprintCallable functions indicates that the function should be displayed the same as the implicit Make Struct nodes
+		NativeMakeFunc,
+
+		// [FunctionMetadata] Used by BlueprintCallable functions to indicate that this function is not to be allowed in the Construction Script.
+		UnsafeDuringActorConstruction,
+
+	};
+
+	// Metadata usable in UINTERFACE
+	enum
+	{
 		/// [InterfaceMetadata] This interface cannot be implemented by a blueprint (e.g., it has only non-exposed C++ member methods)
 		CannotImplementInterfaceInBlueprint,
-
-		/// [ClassMetadata] if the native class cannot tick, the blueprint generated class (based on exactly this native class) can have bCanEverTick flag overridden.
-		ChildCanTick,
 	};
 }
 
