#UE4 corrected fix for the previous circular dependency load issue that I 'fixed'... I was looking at the wrong CDO, so I hadn't fixed one flavor of the bug.

diff --git a/Engine/Source/Runtime/CoreUObject/Private/UObject/LinkerLoad.cpp b/Engine/Source/Runtime/CoreUObject/Private/UObject/LinkerLoad.cpp
--- a/Engine/Source/Runtime/CoreUObject/Private/UObject/LinkerLoad.cpp
+++ b/Engine/Source/Runtime/CoreUObject/Private/UObject/LinkerLoad.cpp
@@ -2961,13 +2961,17 @@ UObject* ULinkerLoad::CreateExport( int32 Index )
 				// SuperStruct needs to be fully linked so that UStruct::Link will have access to UObject::SuperStruct->PropertySize. 
 				// There are other attempts to force our super struct to load, and I have not verified that they can all be removed
 				// in favor of this one:
-				if( !SuperStruct->HasAnyFlags(RF_LoadCompleted) 
-					&& !SuperStruct->GetClass()->ClassDefaultObject
-					&& SuperStruct->GetLinker() 
-					&& Export.SuperIndex.IsImport() )
+				if (!SuperStruct->HasAnyFlags(RF_LoadCompleted)
+					&& !SuperStruct->HasAnyFlags(RF_Native)
+					&& SuperStruct->GetLinker()
+					&& Export.SuperIndex.IsImport())
 				{
-					SuperStruct->SetFlags(RF_NeedLoad);
-					Preload(SuperStruct);
+					const UClass* AsClass = Cast<UClass>(SuperStruct);
+					if (AsClass && !AsClass->ClassDefaultObject)
+					{
+						SuperStruct->SetFlags(RF_NeedLoad);
+						Preload(SuperStruct);
+					}
 				}
 			}
 		}
