Missing change for string asset ref. Not critical, but could have slightly change behavior.

diff --git a/Engine/Source/Runtime/CoreUObject/Private/UObject/LinkerLoad.cpp b/Engine/Source/Runtime/CoreUObject/Private/UObject/LinkerLoad.cpp
--- a/Engine/Source/Runtime/CoreUObject/Private/UObject/LinkerLoad.cpp
+++ b/Engine/Source/Runtime/CoreUObject/Private/UObject/LinkerLoad.cpp
@@ -3593,7 +3593,7 @@ FArchive& ULinkerLoad::operator<<( FAssetPtr& AssetPtr)
 {
 	FArchive& Ar = *this;
 	FStringAssetReference ID;
-	Ar << ID;
+	ID.Serialize(Ar);
 	AssetPtr = ID;
 	return Ar;
 }
