LinkerLoad: only skip CreateExport when actually following a redirector (and only in non-cooked builds)

#ttp 341935: MAIN: CRASH: Clients crash when loading into a zone

diff --git a/Engine/Source/Runtime/CoreUObject/Private/UObject/LinkerLoad.cpp b/Engine/Source/Runtime/CoreUObject/Private/UObject/LinkerLoad.cpp
--- a/Engine/Source/Runtime/CoreUObject/Private/UObject/LinkerLoad.cpp
+++ b/Engine/Source/Runtime/CoreUObject/Private/UObject/LinkerLoad.cpp
@@ -1726,12 +1726,12 @@ void ULinkerLoad::Verify()
 {
 	if(!FApp::IsGame() || GIsEditor || IsRunningCommandlet())
 	{
-		if( !bHaveImportsBeenVerified )
+		if (!bHaveImportsBeenVerified)
 		{
 			// Validate all imports and map them to their remote linkers.
-			for( int32 i=0; i<Summary.ImportCount; i++ )
+			for (int32 ImportIndex = 0; ImportIndex < Summary.ImportCount; ImportIndex++)
 			{
-				VerifyImport( i );
+				VerifyImport(ImportIndex);
 			}
 		}
 	}
@@ -1925,35 +1925,23 @@ void ULinkerLoad::GatherImportDependencies(int32 ImportIndex, TSet<FDependencyRe
 }
 
 
-
-
-/**
- * A wrapper around VerifyImportInner. If the VerifyImportInner (previously VerifyImport) fails, this function
- * will look for a UObjectRedirector that will point to the real location of the object. You will see this if
- * an object was renamed to a different package or group, but something that was referencing the object was not
- * not currently open. (Rename fixes up references of all loaded objects, but naturally not for ones that aren't
- * loaded).
- *
- * @param	i	The index into this packages ImportMap to verify
- */
-void ULinkerLoad::VerifyImport( int32 i )
+ULinkerLoad::EVerifyResult ULinkerLoad::VerifyImport(int32 ImportIndex)
 {
-	FObjectImport& Import = ImportMap[i];
+	FObjectImport& Import = ImportMap[ImportIndex];
 
 	// keep a string of modifiers to add to the Editor Warning dialog
 	FString WarningAppend;
 
 	// try to load the object, but don't print any warnings on error (so we can try the redirector first)
 	// note that a true return value here does not mean it failed or succeeded, just tells it how to respond to a further failure
-	bool bCrashOnFail = VerifyImportInner(i,WarningAppend);
+	bool bCrashOnFail = VerifyImportInner(ImportIndex, WarningAppend);
 	if (FPlatformProperties::HasEditorOnlyData() == false)
 	{
 		bCrashOnFail = false;
 	}
 
 	// by default, we haven't failed yet
-	bool bFailed = false;
-	bool bRedir = false;
+	EVerifyResult Result = VERIFY_Success;
 
 	// these checks find out if the VerifyImportInner was successful or not 
 	if (Import.SourceLinker && Import.SourceIndex == INDEX_NONE && Import.XObject == NULL && !Import.OuterIndex.IsNull() && Import.ObjectName != NAME_ObjectRedirector)
@@ -1964,12 +1952,12 @@ void ULinkerLoad::VerifyImport( int32 i )
 		Import.ClassPackage = GLongCoreUObjectPackageName;
 
 		// try again for the redirector
-		VerifyImportInner(i,WarningAppend);
+		VerifyImportInner(ImportIndex, WarningAppend);
 
 		// if the redirector wasn't found, then it truly doesn't exist
 		if (Import.SourceIndex == INDEX_NONE)
 		{
-			bFailed = true;
+			Result = VERIFY_Failed;
 		}
 		// otherwise, we found that the redirector exists
 		else
@@ -1983,7 +1971,7 @@ void ULinkerLoad::VerifyImport( int32 i )
 			// this should probably never fail, but just in case
 			if (!Redir)
 			{
-				bFailed = true;
+				Result = VERIFY_Failed;
 			}
 			else
 			{
@@ -1998,7 +1986,7 @@ void ULinkerLoad::VerifyImport( int32 i )
 				// check to make sure the destination obj was loaded,
 				if ( DestObject == NULL )
 				{
-					bFailed = true;
+					Result = VERIFY_Failed;
 				}
 				// check that in fact it was the type we thought it should be
 				else if ( DestObject->GetClass()->GetFName() != OriginalImport.ClassName
@@ -2006,7 +1994,7 @@ void ULinkerLoad::VerifyImport( int32 i )
 					// if the destination object is a CDO, allow class changes
 					&&	!DestObject->HasAnyFlags(RF_ClassDefaultObject) )
 				{
-					bFailed = true;
+					Result = VERIFY_Failed;
 					// if the destination is a ObjectRedirector you've most likely made a nasty circular loop
 					if( Redir->DestinationObject->GetClass() == UObjectRedirector::StaticClass() )
 					{
@@ -2015,6 +2003,8 @@ void ULinkerLoad::VerifyImport( int32 i )
 				}
 				else
 				{
+					Result = VERIFY_Redirected;
+
 					// send a callback saying we followed a redirector successfully
 					FCoreDelegates::RedirectorFollowed.Broadcast(Filename, Redir);
 
@@ -2032,7 +2022,7 @@ void ULinkerLoad::VerifyImport( int32 i )
 		Import.ClassPackage = OriginalImport.ClassPackage;
 
 		// if nothing above failed, then we are good to go
-		if (!bFailed)
+		if (Result != VERIFY_Failed)
 		{
 			// we update the runtime information (SourceIndex, SourceLinker) to point to the object the redirector pointed to
 			Import.SourceIndex = Import.XObject->GetLinkerIndex();
@@ -2046,8 +2036,8 @@ void ULinkerLoad::VerifyImport( int32 i )
 			// then do the throw here
 			if (bCrashOnFail)
 			{
-				UE_LOG(LogLinker, Fatal,  TEXT("Failed import: %s %s (file %s)"), *Import.ClassName.ToString(), *GetImportFullName(i), *Import.SourceLinker->Filename );
-				return;
+				UE_LOG(LogLinker, Fatal,  TEXT("Failed import: %s %s (file %s)"), *Import.ClassName.ToString(), *GetImportFullName(ImportIndex), *Import.SourceLinker->Filename );
+				return Result;
 			}
 			// otherwise just printout warnings, and if in the editor, popup the EdLoadWarnings box
 			else
@@ -2058,12 +2048,12 @@ void ULinkerLoad::VerifyImport( int32 i )
 				if( GIsEditor && !IsRunningCommandlet() )
 				{
 					FFormatNamedArguments Arguments[2];
-					Arguments[0].Add(TEXT("ImportClass"), FText::FromName(GetImportClassName(i)));
+					Arguments[0].Add(TEXT("ImportClass"), FText::FromName(GetImportClassName(ImportIndex)));
 					Arguments[1].Add(TEXT("Warning"), FText::FromString(WarningAppend));
 
 					// put something into the load warnings dialog, with any extra information from above (in WarningAppend)
 					FMessageLog(NAME_LoadErrors).Error(FText::Format(LOCTEXT("ImportFailure", "Failed import: {ImportClass}"), Arguments[0]))
-						->AddToken(FAssetNameToken::Create(GetImportPathName(i)))
+						->AddToken(FAssetNameToken::Create(GetImportPathName(ImportIndex)))
 						->AddToken(FTextToken::Create(FText::Format(LOCTEXT("ImportFailure_WarningIn", "{Warning} in"), Arguments[1])))
 						->AddToken(FAssetNameToken::Create(LinkerRoot->GetName()));
 				}
@@ -2075,7 +2065,7 @@ void ULinkerLoad::VerifyImport( int32 i )
 					if ( (!GIsEditor || IsRunningCommandlet()) && (FindClass->IsChildOf(UClass::StaticClass())) )
 					{
 						UE_LOG(LogLinker, Warning, TEXT("Missing Class '%s' referenced by package '%s' ('%s').  Classes should not be removed if referenced by content; mark the class 'deprecated' instead."),
-							*GetImportFullName(i),
+							*GetImportFullName(ImportIndex),
 							*LinkerRoot->GetName(),
 							GSerializedExportLinker ? *GSerializedExportLinker->GetExportPathName(GSerializedExportIndex) : TEXT("Unknown") );
 					}
@@ -2083,7 +2073,7 @@ void ULinkerLoad::VerifyImport( int32 i )
 					else if ( FindClass == NULL || !FindClass->HasAnyClassFlags(CLASS_Deprecated) )
 					{
 						UE_LOG(LogLinker, Warning, TEXT("Missing Class '%s' referenced by package '%s' ('%s')."),
-							*GetImportFullName(i),
+							*GetImportFullName(ImportIndex),
 							*LinkerRoot->GetName(),
 							GSerializedExportLinker ? *GSerializedExportLinker->GetExportPathName(GSerializedExportIndex) : TEXT("Unknown") );
 					}
@@ -2092,6 +2082,8 @@ void ULinkerLoad::VerifyImport( int32 i )
 			}
 		}
 	}
+
+	return Result;
 }
 
 /**
@@ -3357,15 +3349,18 @@ UObject* ULinkerLoad::CreateImport( int32 Index )
 
 		if( Import.XObject == NULL )
 		{
+			EVerifyResult VerifyImportResult = VERIFY_Success;
 			if( Import.SourceLinker == NULL )
 			{
-				VerifyImport(Index);
+				VerifyImportResult = VerifyImport(Index);
 			}
 			if(Import.SourceIndex != INDEX_NONE)
 			{
 				check(Import.SourceLinker);
-				// VerifyImport may have already created the import and SourceIndex has changed to point to the actual redirected object
-				if (!Import.XObject)
+				// VerifyImport may have already created the import and SourceIndex has changed to point to the actual redirected object.
+				// This can only happen in non-cooked builds since cooked builds don't have redirects and other cases are valid.
+				// We also don't want to call CreateExport only when there was an actual redirector involved.
+				if (FPlatformProperties::RequiresCookedData() || !Import.XObject || VerifyImportResult != VERIFY_Redirected)
 				{
 					Import.XObject = Import.SourceLinker->CreateExport(Import.SourceIndex);
 				}
diff --git a/Engine/Source/Runtime/CoreUObject/Public/UObject/Linker.h b/Engine/Source/Runtime/CoreUObject/Public/UObject/Linker.h
--- a/Engine/Source/Runtime/CoreUObject/Public/UObject/Linker.h
+++ b/Engine/Source/Runtime/CoreUObject/Public/UObject/Linker.h
@@ -1189,6 +1189,17 @@ class ULinkerLoad : public ULinker, public FArchiveUObject
 		LINKER_TimedOut = 2
 	};
 
+	/** Verify result. */
+	enum EVerifyResult
+	{
+		/** Error occured when verifying import (can be fatal). */
+		VERIFY_Failed = 0,
+		/** Verify completed successfully. */
+		VERIFY_Success = 1,
+		/** Verify completed successfully and followed a redirector. */
+		VERIFY_Redirected = 2
+	};
+
 	// Variables.
 public:
 	/** Flags determining loading behavior.																					*/
@@ -1420,9 +1431,10 @@ public:
 	 * not currently open. (Rename fixes up references of all loaded objects, but naturally not for ones that aren't
 	 * loaded).
 	 *
-	 * @param	i	The index into this package's ImportMap to verify
+	 * @param	ImportIndex	The index into this package's ImportMap to verify
+	 * @return Verify import result
 	 */
-	void VerifyImport( int32 i );
+	EVerifyResult VerifyImport(int32 ImportIndex);
 	
 	/**
 	 * Loads all objects in package.
