Fixed desynchronization of archive position when reading serialized property tags after reading an array property that can not be converted.

diff --git a/Engine/Source/Runtime/CoreUObject/Private/UObject/Class.cpp b/Engine/Source/Runtime/CoreUObject/Private/UObject/Class.cpp
--- a/Engine/Source/Runtime/CoreUObject/Private/UObject/Class.cpp
+++ b/Engine/Source/Runtime/CoreUObject/Private/UObject/Class.cpp
@@ -1,4 +1,4 @@
-// Copyright 1998-2014 Epic Games, Inc. All Rights Reserved.
+﻿// Copyright 1998-2014 Epic Games, Inc. All Rights Reserved.
 
 /*=============================================================================
 	UnClass.cpp: Object class implementation.
@@ -829,6 +829,7 @@ void UStruct::SerializeTaggedProperties(FArchive& Ar, uint8* Data, UStruct* Defa
 
 			bool bSkipSkipWarning = false;
 
+			const int64 StartOfProperty = Ar.Tell();
 			if( !Property )
 			{
 				//UE_LOG(LogClass, Warning, TEXT("Property %s of %s not found for package:  %s"), *Tag.Name.ToString(), *GetFullName(), *Ar.GetArchiveName() );
@@ -1035,8 +1036,8 @@ void UStruct::SerializeTaggedProperties(FArchive& Ar, uint8* Data, UStruct* Defa
 				}
 				else
 				{
-				UE_LOG(LogClass, Warning, TEXT("Array Inner Type mismatch in %s of %s - Previous (%s) Current(%s) for package:  %s"), *Tag.Name.ToString(), *GetName(), *Tag.InnerType.ToString(), *CastChecked<UArrayProperty>(Property)->Inner->GetID().ToString(), *Ar.GetArchiveName() );
-			}
+					UE_LOG(LogClass, Warning, TEXT("Array Inner Type mismatch in %s of %s - Previous (%s) Current(%s) for package:  %s"), *Tag.Name.ToString(), *GetName(), *Tag.InnerType.ToString(), *CastChecked<UArrayProperty>(Property)->Inner->GetID().ToString(), *Ar.GetArchiveName() );
+				}
 			}
 			else if( Tag.Type==NAME_StructProperty && Tag.StructName!=CastChecked<UStructProperty>(Property)->Struct->GetFName() 
 				&& CastChecked<UStructProperty>(Property)->UseBinaryOrNativeSerialization(Ar) )
@@ -1092,20 +1093,21 @@ void UStruct::SerializeTaggedProperties(FArchive& Ar, uint8* Data, UStruct* Defa
 			}
 			else
 			{
-					uint8* DestAddress = Property->ContainerPtrToValuePtr<uint8>(Data, Tag.ArrayIndex);  
+				uint8* DestAddress = Property->ContainerPtrToValuePtr<uint8>(Data, Tag.ArrayIndex);  
 
-					// This property is ok.			
-					Tag.SerializeTaggedProperty( Ar, Property, DestAddress, Tag.Size, NULL );
+				// This property is ok.			
+				Tag.SerializeTaggedProperty( Ar, Property, DestAddress, Tag.Size, NULL );
 
-					AdvanceProperty = true;
-					continue;
-				}
+				AdvanceProperty = true;
+				continue;
+			}
 
 			AdvanceProperty = false;
 
 			// Skip unknown or bad property.
+			const int64 RemainingSize = Tag.Size - (Ar.Tell() - StartOfProperty);
 			uint8 B;
-			for( int32 i=0; i<Tag.Size; i++ )
+			for( int64 i=0; i<RemainingSize; i++ )
 			{
 				Ar << B;
 			}
