Fix missing property names in obj refs output.
TTP #341661 CORE: 'Obj refs' output is sometimes missing property names.

diff --git a/Engine/Source/Runtime/CoreUObject/Private/UObject/ReferenceChainSearch.cpp b/Engine/Source/Runtime/CoreUObject/Private/UObject/ReferenceChainSearch.cpp
--- a/Engine/Source/Runtime/CoreUObject/Private/UObject/ReferenceChainSearch.cpp
+++ b/Engine/Source/Runtime/CoreUObject/Private/UObject/ReferenceChainSearch.cpp
@@ -58,19 +58,41 @@ bool FReferenceChainSearch::FReferenceChain::Contains( const FReferenceChain& Ot
 	return true;
 }
 
+int32 FReferenceChainSearch::FFindReferencerCollector::FindReferencedObjectIndex(const UObject& ReferencedBy, const UObject& ReferencedObject)
+{
+	int32 Result = INDEX_NONE;
+	auto & TokenMap = ReferencedBy.GetClass()->DebugTokenMap;
+
+	for (int32 Index = 0, End = TokenMap.GetTokenMapSize(); Index < End; ++Index)
+	{
+		auto TokenName = TokenMap.GetTokenInfo(Index).Name;
+		if (ReferencedObject.GetFName() == TokenName)
+		{
+			Result = Index;
+			break;
+		}
+	}
+
+	return Result;
+}
 
 void FReferenceChainSearch::FFindReferencerCollector::HandleObjectReference( UObject*& InObject, const UObject* RefObject, const UObject* ReferencingProperty )
 {
 	UObject* RefSrc = RefObject != NULL ? const_cast<UObject*>(RefObject) : ReferencingObject;
+	int32 ReferencedObjectIndex = INDEX_NONE;
+	if (RefSrc != nullptr && InObject != nullptr)
+	{
+		ReferencedObjectIndex = FindReferencedObjectIndex(*RefSrc, *InObject);
+	}
 
 	if (ReferencingProperty != NULL)
 	{
-		FReferenceChainLink RefInfo(EReferenceType::Property, RefSrc, (void*)ReferencingProperty, InObject);
+		FReferenceChainLink RefInfo(ReferencedObjectIndex, EReferenceType::Property, RefSrc, (void*)ReferencingProperty, InObject);
 		References.Push(RefInfo);
 	}
 	else
 	{
-		FReferenceChainLink RefInfo(RefType, RefSrc, AROFuncPtr, InObject);
+		FReferenceChainLink RefInfo(ReferencedObjectIndex, RefType, RefSrc, AROFuncPtr, InObject);
 		References.Push(RefInfo);
 	}
 }
@@ -95,6 +117,7 @@ void FReferenceChainSearch::PrintReferencers( FReferenceChain& Referencer )
 
 		FString ReferencedThrough = Internal::RefereceTypeToString(RefInfo.ReferenceType);
 
+#if !(UE_BUILD_TEST || UE_BUILD_SHIPPING)
 		if (RefInfo.ReferencedThrough != NULL)
 		{
 			if ( RefInfo.IsProperty() )
@@ -110,7 +133,11 @@ void FReferenceChainSearch::PrintReferencers( FReferenceChain& Referencer )
 				ReferencedThrough = ANSI_TO_TCHAR(Str);
 			}
 		}
-
+		else if (RefInfo.ReferencedObjectIndex != INDEX_NONE)
+		{
+			ReferencedThrough = RefInfo.GetReferencedByName();
+		}
+#endif
 		FString ObjectReachability;
 		if( RefInfo.ReferencedBy->IsRooted() )
 		{
@@ -496,7 +523,7 @@ bool FReferenceChainSearch::ProcessObject( UObject* CurrentObject )
 			UObject*&	Object		= *ObjectPtr;
 			TokenReturnCount		= REFERENCE_INFO.ReturnCount;
 
-			FReferenceChainLink TopRef(InArrayProp != NULL ? EReferenceType::ArrayProperty : EReferenceType::Property, CurrentObject, Prop, Object);
+			FReferenceChainLink TopRef(ReferenceTokenStreamIndex, InArrayProp != NULL ? EReferenceType::ArrayProperty : EReferenceType::Property, CurrentObject, Prop, Object);
 			AddToReferenceList(ReferenceList, TopRef);
 		}
 		else if( REFERENCE_INFO.Type == GCRT_ArrayObject )
@@ -509,7 +536,7 @@ bool FReferenceChainSearch::ProcessObject( UObject* CurrentObject )
 			{
 				UObject*& Object = ObjectArray[ObjectIndex];
 				
-				FReferenceChainLink TopRef(EReferenceType::ArrayProperty, CurrentObject, Prop, Object, ObjectIndex);
+				FReferenceChainLink TopRef(ReferenceTokenStreamIndex, EReferenceType::ArrayProperty, CurrentObject, Prop, Object, ObjectIndex);
 				AddToReferenceList(ReferenceList, TopRef);
 			}
 		}
@@ -548,7 +575,7 @@ bool FReferenceChainSearch::ProcessObject( UObject* CurrentObject )
 			UObject*&	Object		= *ObjectPtr;
 			TokenReturnCount		= REFERENCE_INFO.ReturnCount;
 
-			FReferenceChainLink TopRef(InArrayProp != NULL ? EReferenceType::ArrayProperty : EReferenceType::Property, CurrentObject, Prop, Object);
+			FReferenceChainLink TopRef(ReferenceTokenStreamIndex, InArrayProp != NULL ? EReferenceType::ArrayProperty : EReferenceType::Property, CurrentObject, Prop, Object);
 			AddToReferenceList(ReferenceList, TopRef);
 		}
 		else if( REFERENCE_INFO.Type == GCRT_FixedArray )
diff --git a/Engine/Source/Runtime/CoreUObject/Public/UObject/GarbageCollection.h b/Engine/Source/Runtime/CoreUObject/Public/UObject/GarbageCollection.h
--- a/Engine/Source/Runtime/CoreUObject/Public/UObject/GarbageCollection.h
+++ b/Engine/Source/Runtime/CoreUObject/Public/UObject/GarbageCollection.h
@@ -183,6 +183,11 @@ public:
 		TokenMap.Empty();
 	}
 
+	FORCEINLINE int32 GetTokenMapSize() const
+	{
+		return TokenMap.Num();
+	}
+
 private:
 	/* Token map. */
 	TArray<FTokenInfo> TokenMap;
diff --git a/Engine/Source/Runtime/CoreUObject/Public/UObject/ReferenceChainSearch.h b/Engine/Source/Runtime/CoreUObject/Public/UObject/ReferenceChainSearch.h
--- a/Engine/Source/Runtime/CoreUObject/Public/UObject/ReferenceChainSearch.h
+++ b/Engine/Source/Runtime/CoreUObject/Public/UObject/ReferenceChainSearch.h
@@ -27,6 +27,8 @@ public:
 	/** Entry in the reference chain */
 	struct FReferenceChainLink
 	{
+		/** Index of ReferencedObj token in ReferencedBy->ReferenceTokenStream **/
+		int32 ReferencedObjectIndex;
 		/** Describes in which way an object is referenced */
 		EReferenceType::Type ReferenceType;
 		/** The object that is referencing */
@@ -40,10 +42,16 @@ public:
 		int32 ArrayIndex;
 
 		FReferenceChainLink()
-			:ReferenceType(EReferenceType::Invalid), ReferencedBy(NULL), ReferencedThrough(NULL), ArrayIndex(-1) {}
+			:ReferencedObjectIndex(INDEX_NONE), ReferenceType(EReferenceType::Invalid), ReferencedBy(NULL), ReferencedThrough(NULL), ArrayIndex(INDEX_NONE) {}
 
-		FReferenceChainLink(EReferenceType::Type RefType, UObject* InReferencedBy, void* InReferencedThrough, UObject* InReferencedObj, int32 InArrayIndex = -1)
-			:ReferenceType(RefType), ReferencedBy(InReferencedBy), ReferencedThrough(InReferencedThrough), ReferencedObj(InReferencedObj), ArrayIndex(InArrayIndex) {}
+		FReferenceChainLink(int32 InReferencedObjectIndex, EReferenceType::Type RefType, UObject* InReferencedBy, void* InReferencedThrough, UObject* InReferencedObj, int32 InArrayIndex = INDEX_NONE)
+			:ReferencedObjectIndex(InReferencedObjectIndex), ReferenceType(RefType), ReferencedBy(InReferencedBy), ReferencedThrough(InReferencedThrough), ReferencedObj(InReferencedObj), ArrayIndex(InArrayIndex)
+		{ }
+
+		FORCEINLINE FString GetReferencedByName() const
+		{
+			return ReferencedBy->GetClass()->DebugTokenMap.GetTokenInfo(ReferencedObjectIndex).Name.ToString();
+		}
 
 		/** Returns whether this link is a property reference or not */
 		FORCEINLINE bool IsProperty() const { return ReferenceType == EReferenceType::Property || ReferenceType == EReferenceType::ArrayProperty; }
@@ -73,6 +81,8 @@ public:
 		void* AROFuncPtr;
 		/** Referencing object being AROed*/
 		UObject* ReferencingObject;
+		/** Finds index of ReferencedObject in ReferencedBy.ReferenceTokenStream */
+		int32 FindReferencedObjectIndex(const UObject& ReferencedBy, const UObject& ReferencedObject);
 
 	public:
 
