Don't mark packages as dirty in PIE

diff --git a/Engine/Source/Runtime/CoreUObject/Private/UObject/UObjectBaseUtility.cpp b/Engine/Source/Runtime/CoreUObject/Private/UObject/UObjectBaseUtility.cpp
--- a/Engine/Source/Runtime/CoreUObject/Private/UObject/UObjectBaseUtility.cpp
+++ b/Engine/Source/Runtime/CoreUObject/Private/UObject/UObjectBaseUtility.cpp
@@ -132,7 +132,7 @@ void UObjectBaseUtility::MarkPackageDirty() const
 			// we explicitly disable the ability to dirty a package or map during load.  Commandlets can still
 			// set the dirty state on load.
 			if( (!GIsEditor || IsRunningCommandlet() || 
-				(GIsEditor && !GIsEditorLoadingPackage)))
+				(GIsEditor && !GIsEditorLoadingPackage && !GIsPlayInEditorWorld)))
 			{
 				const bool bIsDirty = Package->IsDirty();
 
