Add missing macro around FReferenceChainLink::GetReferencedByName

diff --git a/Engine/Source/Runtime/CoreUObject/Public/UObject/ReferenceChainSearch.h b/Engine/Source/Runtime/CoreUObject/Public/UObject/ReferenceChainSearch.h
--- a/Engine/Source/Runtime/CoreUObject/Public/UObject/ReferenceChainSearch.h
+++ b/Engine/Source/Runtime/CoreUObject/Public/UObject/ReferenceChainSearch.h
@@ -47,11 +47,12 @@ public:
 		FReferenceChainLink(int32 InReferencedObjectIndex, EReferenceType::Type RefType, UObject* InReferencedBy, void* InReferencedThrough, UObject* InReferencedObj, int32 InArrayIndex = INDEX_NONE)
 			:ReferencedObjectIndex(InReferencedObjectIndex), ReferenceType(RefType), ReferencedBy(InReferencedBy), ReferencedThrough(InReferencedThrough), ReferencedObj(InReferencedObj), ArrayIndex(InArrayIndex)
 		{ }
-
+#if !(UE_BUILD_TEST || UE_BUILD_SHIPPING)
 		FORCEINLINE FString GetReferencedByName() const
 		{
 			return ReferencedBy->GetClass()->DebugTokenMap.GetTokenInfo(ReferencedObjectIndex).Name.ToString();
 		}
+#endif // !(UE_BUILD_TEST || UE_BUILD_SHIPPING)
 
 		/** Returns whether this link is a property reference or not */
 		FORCEINLINE bool IsProperty() const { return ReferenceType == EReferenceType::Property || ReferenceType == EReferenceType::ArrayProperty; }
