Add missing #if guard.

diff --git a/Engine/Source/Runtime/CoreUObject/Private/UObject/ReferenceChainSearch.cpp b/Engine/Source/Runtime/CoreUObject/Private/UObject/ReferenceChainSearch.cpp
--- a/Engine/Source/Runtime/CoreUObject/Private/UObject/ReferenceChainSearch.cpp
+++ b/Engine/Source/Runtime/CoreUObject/Private/UObject/ReferenceChainSearch.cpp
@@ -58,6 +58,7 @@ bool FReferenceChainSearch::FReferenceChain::Contains( const FReferenceChain& Ot
 	return true;
 }
 
+#if !(UE_BUILD_TEST || UE_BUILD_SHIPPING)
 int32 FReferenceChainSearch::FFindReferencerCollector::FindReferencedObjectIndex(const UObject& ReferencedBy, const UObject& ReferencedObject)
 {
 	int32 Result = INDEX_NONE;
@@ -75,16 +76,17 @@ int32 FReferenceChainSearch::FFindReferencerCollector::FindReferencedObjectIndex
 
 	return Result;
 }
-
+#endif // !(UE_BUILD_TEST || UE_BUILD_SHIPPING)
 void FReferenceChainSearch::FFindReferencerCollector::HandleObjectReference( UObject*& InObject, const UObject* RefObject, const UObject* ReferencingProperty )
 {
 	UObject* RefSrc = RefObject != NULL ? const_cast<UObject*>(RefObject) : ReferencingObject;
 	int32 ReferencedObjectIndex = INDEX_NONE;
+#if !(UE_BUILD_TEST || UE_BUILD_SHIPPING)
 	if (RefSrc != nullptr && InObject != nullptr)
 	{
 		ReferencedObjectIndex = FindReferencedObjectIndex(*RefSrc, *InObject);
 	}
-
+#endif // !(UE_BUILD_TEST || UE_BUILD_SHIPPING)
 	if (ReferencingProperty != NULL)
 	{
 		FReferenceChainLink RefInfo(ReferencedObjectIndex, EReferenceType::Property, RefSrc, (void*)ReferencingProperty, InObject);
diff --git a/Engine/Source/Runtime/CoreUObject/Public/UObject/ReferenceChainSearch.h b/Engine/Source/Runtime/CoreUObject/Public/UObject/ReferenceChainSearch.h
--- a/Engine/Source/Runtime/CoreUObject/Public/UObject/ReferenceChainSearch.h
+++ b/Engine/Source/Runtime/CoreUObject/Public/UObject/ReferenceChainSearch.h
@@ -82,9 +82,10 @@ public:
 		void* AROFuncPtr;
 		/** Referencing object being AROed*/
 		UObject* ReferencingObject;
+#if !(UE_BUILD_TEST || UE_BUILD_SHIPPING)
 		/** Finds index of ReferencedObject in ReferencedBy.ReferenceTokenStream */
 		int32 FindReferencedObjectIndex(const UObject& ReferencedBy, const UObject& ReferencedObject);
-
+#endif // !(UE_BUILD_TEST || UE_BUILD_SHIPPING)
 	public:
 
 		FFindReferencerCollector( FReferenceChainSearch* InRefSearchArc, EReferenceType::Type InRefType, void* InAROFuncPtr, UObject* InReferencingObject = NULL)
