* added comment

diff --git a/Engine/Source/Runtime/CoreUObject/Private/UObject/Obj.cpp b/Engine/Source/Runtime/CoreUObject/Private/UObject/Obj.cpp
--- a/Engine/Source/Runtime/CoreUObject/Private/UObject/Obj.cpp
+++ b/Engine/Source/Runtime/CoreUObject/Private/UObject/Obj.cpp
@@ -531,6 +531,7 @@ bool UObject::ConditionalBeginDestroy()
 #if !(UE_BUILD_SHIPPING || UE_BUILD_TEST)
 		if( DebugBeginDestroyed.Contains(this) )
 		{
+			// class might override BeginDestroy without calling Super::BeginDestroy();
 			UE_LOG(LogObj, Fatal, TEXT("%s failed to route BeginDestroy"), *GetFullName() );
 		}
 #endif
