Break redundant recursive loading.

ttp341943 LIVE: BUILD: Class Array breaks shipping build

diff --git a/Engine/Config/BaseEditor.ini b/Engine/Config/BaseEditor.ini
--- a/Engine/Config/BaseEditor.ini
+++ b/Engine/Config/BaseEditor.ini
@@ -282,4 +282,5 @@ bDontLoadBlueprintOutsideEditor=false
 bBlueprintIsNotBlueprintType=false
 
 [StructSerialization]
-SkipByteCodeSerialization=true
\ No newline at end of file
+SkipByteCodeSerialization=true
+BreakSerializationRecursion=true
\ No newline at end of file
diff --git a/Engine/Source/Runtime/CoreUObject/Private/UObject/Class.cpp b/Engine/Source/Runtime/CoreUObject/Private/UObject/Class.cpp
--- a/Engine/Source/Runtime/CoreUObject/Private/UObject/Class.cpp
+++ b/Engine/Source/Runtime/CoreUObject/Private/UObject/Class.cpp
@@ -691,7 +691,7 @@ void UStruct::InitTaggedPropertyRedirectsMap()
 	}
 }
 
-void UStruct::SerializeTaggedProperties(FArchive& Ar, uint8* Data, UStruct* DefaultsStruct, uint8* Defaults) const
+void UStruct::SerializeTaggedProperties(FArchive& Ar, uint8* Data, UStruct* DefaultsStruct, uint8* Defaults, const UObject* BreakRecursionIfFullyLoad) const
 {
 	FName PropertyName(NAME_None);
 
@@ -712,7 +712,7 @@ void UStruct::SerializeTaggedProperties(FArchive& Ar, uint8* Data, UStruct* Defa
 		int32		RemainingArrayDim	= Property ? Property->ArrayDim : 0;
 
 		// Load all stored properties, potentially skipping unknown ones.
-		while( 1 )
+		while (1)
 		{
 			FPropertyTag Tag;
 			Ar << Tag;
@@ -827,18 +827,20 @@ void UStruct::SerializeTaggedProperties(FArchive& Ar, uint8* Data, UStruct* Defa
 				}
 			}
 
-			bool bSkipSkipWarning = false;
-
 			const int64 StartOfProperty = Ar.Tell();
 			if( !Property )
 			{
 				//UE_LOG(LogClass, Warning, TEXT("Property %s of %s not found for package:  %s"), *Tag.Name.ToString(), *GetFullName(), *Ar.GetArchiveName() );
 			}
+#if WITH_EDITOR
+			else if (BreakRecursionIfFullyLoad && BreakRecursionIfFullyLoad->HasAllFlags(RF_LoadCompleted))
+			{
+			}
+#endif // WITH_EDITOR
 			// editoronly properties should be skipped if we are NOT the editor, or we are 
 			// the editor but are cooking for console (editoronly implies notforconsole)
 			else if ((Property->PropertyFlags & CPF_EditorOnly) && !FPlatformProperties::HasEditorOnlyData() && !GForceLoadEditorOnly)
 			{
-				bSkipSkipWarning = true;
 			}
 			// check for valid array index
 			else if( Tag.ArrayIndex >= Property->ArrayDim || Tag.ArrayIndex < 0 )
diff --git a/Engine/Source/Runtime/CoreUObject/Private/UObject/Obj.cpp b/Engine/Source/Runtime/CoreUObject/Private/UObject/Obj.cpp
--- a/Engine/Source/Runtime/CoreUObject/Private/UObject/Obj.cpp
+++ b/Engine/Source/Runtime/CoreUObject/Private/UObject/Obj.cpp
@@ -854,7 +854,13 @@ void UObject::SerializeScriptProperties( FArchive& Ar ) const
 	if( (Ar.IsLoading() || Ar.IsSaving()) && !Ar.WantBinaryPropertySerialization() )
 	{
 		UObject* DiffObject = GetArchetype();
-		Class->SerializeTaggedProperties( Ar, (uint8*)this, HasAnyFlags(RF_ClassDefaultObject) ? Class->GetSuperClass() : Class, (uint8*)DiffObject );		
+#if WITH_EDITOR
+		static const FBoolConfigValueHelper BreakSerializationRecursion(TEXT("StructSerialization"), TEXT("BreakSerializationRecursion"));
+		const bool bBreakSerializationRecursion = BreakSerializationRecursion && Ar.IsLoading() && Ar.GetLinker();
+#else 
+		const bool bBreakSerializationRecursion = false;
+#endif
+		Class->SerializeTaggedProperties(Ar, (uint8*)this, HasAnyFlags(RF_ClassDefaultObject) ? Class->GetSuperClass() : Class, (uint8*)DiffObject, bBreakSerializationRecursion ? this : NULL);
 	}
 	else if ( Ar.GetPortFlags() != 0 )
 	{
diff --git a/Engine/Source/Runtime/CoreUObject/Public/UObject/Class.h b/Engine/Source/Runtime/CoreUObject/Public/UObject/Class.h
--- a/Engine/Source/Runtime/CoreUObject/Public/UObject/Class.h
+++ b/Engine/Source/Runtime/CoreUObject/Public/UObject/Class.h
@@ -299,7 +299,7 @@ public:
 	 */
 	void SerializeBinEx( FArchive& Ar, void* Data, void const* DefaultData, UStruct* DefaultStruct ) const;
 
-	virtual void SerializeTaggedProperties( FArchive& Ar, uint8* Data, UStruct* DefaultsStruct, uint8* Defaults ) const;
+	virtual void SerializeTaggedProperties( FArchive& Ar, uint8* Data, UStruct* DefaultsStruct, uint8* Defaults, const UObject* BreakRecursionIfFullyLoad=NULL) const;
 
 #if WITH_EDITOR
 private:
diff --git a/Engine/Source/Runtime/Engine/Classes/Engine/UserDefinedStruct.h b/Engine/Source/Runtime/Engine/Classes/Engine/UserDefinedStruct.h
--- a/Engine/Source/Runtime/Engine/Classes/Engine/UserDefinedStruct.h
+++ b/Engine/Source/Runtime/Engine/Classes/Engine/UserDefinedStruct.h
@@ -48,7 +48,7 @@ public:
 	// End of UObject interface.
 #endif	// WITH_EDITOR
 
-	virtual void SerializeTaggedProperties( FArchive& Ar, uint8* Data, UStruct* DefaultsStruct, uint8* Defaults ) const override;
+	virtual void SerializeTaggedProperties(FArchive& Ar, uint8* Data, UStruct* DefaultsStruct, uint8* Defaults, const UObject* BreakRecursionIfFullyLoad=NULL) const override;
 
 	// UScriptStruct interface.
 	virtual void RecursivelyPreload() override;
diff --git a/Engine/Source/Runtime/Engine/Private/UserDefinedStruct.cpp b/Engine/Source/Runtime/Engine/Private/UserDefinedStruct.cpp
--- a/Engine/Source/Runtime/Engine/Private/UserDefinedStruct.cpp
+++ b/Engine/Source/Runtime/Engine/Private/UserDefinedStruct.cpp
@@ -55,7 +55,7 @@ UProperty* UUserDefinedStruct::CustomFindProperty(const FName Name) const
 
 #endif	// WITH_EDITOR
 
-void UUserDefinedStruct::SerializeTaggedProperties(FArchive& Ar, uint8* Data, UStruct* DefaultsStruct, uint8* Defaults) const
+void UUserDefinedStruct::SerializeTaggedProperties(FArchive& Ar, uint8* Data, UStruct* DefaultsStruct, uint8* Defaults, const UObject* BreakRecursionIfFullyLoad) const
 {
 #if WITH_EDITOR
 	/*	The following code is responsible for UUserDefinedStruct's default values serialization.	*/
