BP: validation and optimization related to RefreshAllNodes

diff --git a/Engine/Source/Editor/BlueprintGraph/Private/EdGraphSchema_K2.cpp b/Engine/Source/Editor/BlueprintGraph/Private/EdGraphSchema_K2.cpp
--- a/Engine/Source/Editor/BlueprintGraph/Private/EdGraphSchema_K2.cpp
+++ b/Engine/Source/Editor/BlueprintGraph/Private/EdGraphSchema_K2.cpp
@@ -2878,6 +2878,24 @@ void UEdGraphSchema_K2::BreakSinglePinLink(UEdGraphPin* SourcePin, UEdGraphPin*
 
 void UEdGraphSchema_K2::ReconstructNode(UEdGraphNode& TargetNode, bool bIsBatchRequest/*=false*/) const
 {
+	{
+		TArray<UObject *> NodeChildren;
+		GetObjectsWithOuter(&TargetNode, NodeChildren, false);
+		for (int32 Iter = 0; Iter < NodeChildren.Num(); ++Iter)
+		{
+			UEdGraphPin* Pin = Cast<UEdGraphPin>(NodeChildren[Iter]);
+			const bool bIsValidPin = !Pin 
+				|| (Pin->HasAllFlags(RF_PendingKill) && !Pin->LinkedTo.Num()) 
+				|| TargetNode.Pins.Contains(Pin);
+			if (!bIsValidPin)
+			{
+				UE_LOG(LogBlueprint, Error, 
+					TEXT("Broken Node: %s keeps removed/invalid pin: %s. Try refresh all nodes."),
+					*TargetNode.GetFullName(), *Pin->GetFullName());
+			}
+		}
+	}
+
 	Super::ReconstructNode(TargetNode, bIsBatchRequest);
 
 	// If the reconstruction is being handled by something doing a batch (i.e. the blueprint autoregenerating itself), defer marking the blueprint as modified to prevent multiple recompiles
diff --git a/Engine/Source/Editor/UnrealEd/Private/Kismet2/BlueprintEditorUtils.cpp b/Engine/Source/Editor/UnrealEd/Private/Kismet2/BlueprintEditorUtils.cpp
--- a/Engine/Source/Editor/UnrealEd/Private/Kismet2/BlueprintEditorUtils.cpp
+++ b/Engine/Source/Editor/UnrealEd/Private/Kismet2/BlueprintEditorUtils.cpp
@@ -259,6 +259,14 @@ static void RenameVariableReferences(UBlueprint* Blueprint, UClass* VariableClas
 
 void FBlueprintEditorUtils::RefreshAllNodes(UBlueprint* Blueprint)
 {
+	if (!Blueprint || !Blueprint->HasAllFlags(RF_LoadCompleted))
+	{
+		UE_LOG(LogBlueprint, Warning, 
+			TEXT("RefreshAllNodes called on incompletly loaded blueprint '%s'"), 
+			Blueprint ? *Blueprint->GetFullName() : TEXT("NULL"));
+		return;
+	}
+
 	TArray<UK2Node*> AllNodes;
 	FBlueprintEditorUtils::GetAllNodesOfClass(Blueprint, AllNodes);
 
@@ -303,6 +311,14 @@ void FBlueprintEditorUtils::RefreshAllNodes(UBlueprint* Blueprint)
 
 void FBlueprintEditorUtils::RefreshExternalBlueprintDependencyNodes(UBlueprint* Blueprint, UStruct* RefreshOnlyChild)
 {
+	if (!Blueprint || !Blueprint->HasAllFlags(RF_LoadCompleted))
+	{
+		UE_LOG(LogBlueprint, Warning,
+			TEXT("RefreshAllNodes called on incompletly loaded blueprint '%s'"),
+			Blueprint ? *Blueprint->GetFullName() : TEXT("NULL"));
+		return;
+	}
+
 	TArray<UK2Node*> AllNodes;
 	FBlueprintEditorUtils::GetAllNodesOfClass(Blueprint, AllNodes);
 
diff --git a/Engine/Source/Editor/UnrealEd/Private/Kismet2/Kismet2.cpp b/Engine/Source/Editor/UnrealEd/Private/Kismet2/Kismet2.cpp
--- a/Engine/Source/Editor/UnrealEd/Private/Kismet2/Kismet2.cpp
+++ b/Engine/Source/Editor/UnrealEd/Private/Kismet2/Kismet2.cpp
@@ -186,7 +186,7 @@ UBlueprint* FKismetEditorUtilities::CreateBlueprint(UClass* ParentClass, UObject
 	}
 	
 	// Create new UBlueprint object
-	UBlueprint* NewBP = ConstructObject<UBlueprint>(*BlueprintClassType, Outer, NewBPName, RF_Public|RF_Standalone|RF_Transactional);
+	UBlueprint* NewBP = ConstructObject<UBlueprint>(*BlueprintClassType, Outer, NewBPName, RF_Public|RF_Standalone|RF_Transactional|RF_LoadCompleted);
 	NewBP->Status = BS_BeingCreated;
 	NewBP->BlueprintType = BlueprintType;
 	NewBP->ParentClass = ParentClass;
diff --git a/Engine/Source/Runtime/CoreUObject/Private/Blueprint/BlueprintSupport.cpp b/Engine/Source/Runtime/CoreUObject/Private/Blueprint/BlueprintSupport.cpp
--- a/Engine/Source/Runtime/CoreUObject/Private/Blueprint/BlueprintSupport.cpp
+++ b/Engine/Source/Runtime/CoreUObject/Private/Blueprint/BlueprintSupport.cpp
@@ -79,16 +79,22 @@ bool ULinkerLoad::RegenerateBlueprintClass(UClass* LoadClass, UObject* ExportObj
 	}
  	for( auto Class : ClassChainOrdered )
 	{
-		if (Class->ClassGeneratedBy->HasAnyFlags(RF_BeingRegenerated))
+		UObject* BlueprintObject = Class->ClassGeneratedBy;
+		if (BlueprintObject && BlueprintObject->HasAnyFlags(RF_BeingRegenerated))
 		{
 			// Always load the parent blueprint here in case there is a circular dependency. This will
 			// ensure that the blueprint is fully serialized before attempting to regenerate the class.
-			FPreloadMembersHelper::PreloadMembers(Class->ClassGeneratedBy);
-
-			Class->ClassGeneratedBy->SetFlags(RF_NeedLoad);
-			Class->ClassGeneratedBy->GetLinker()->Preload(Class->ClassGeneratedBy);
+			if (!BlueprintObject->HasAnyFlags(RF_LoadCompleted))
+			{
+				BlueprintObject->SetFlags(RF_NeedLoad);
+				if (auto Linker = BlueprintObject->GetLinker())
+				{
+					Linker->Preload(BlueprintObject);
+				}
+			}
 
-			Class->ClassGeneratedBy->RegenerateClass(Class, NULL, GObjLoaded);
+			FPreloadMembersHelper::PreloadMembers(BlueprintObject);
+			BlueprintObject->RegenerateClass(Class, NULL, GObjLoaded);
 		}
 	}
 
