Blueprint should be fully loaded before regeneration.

ttp342952 BP: CRASH: UEdGraphNode::FindPinChecked()
ttp342956 CRITICAL: Regression: QAGAME: CRASH: Opening tracetest map (UEdGraphNode::FindPinChecked)

diff --git a/Engine/Source/Runtime/CoreUObject/Private/Blueprint/BlueprintSupport.cpp b/Engine/Source/Runtime/CoreUObject/Private/Blueprint/BlueprintSupport.cpp
--- a/Engine/Source/Runtime/CoreUObject/Private/Blueprint/BlueprintSupport.cpp
+++ b/Engine/Source/Runtime/CoreUObject/Private/Blueprint/BlueprintSupport.cpp
@@ -26,6 +26,18 @@ struct FPreloadMembersHelper
 			}
 		}
 	}
+
+	static void PreloadObject(UObject* InObject)
+	{
+		if (InObject && !InObject->HasAnyFlags(RF_LoadCompleted))
+		{
+			InObject->SetFlags(RF_NeedLoad);
+			if (ULinkerLoad* Linker = InObject->GetLinker())
+			{
+				Linker->Preload(InObject);
+			}
+		}
+	}
 };
 
 /**
@@ -94,14 +106,7 @@ bool ULinkerLoad::RegenerateBlueprintClass(UClass* LoadClass, UObject* ExportObj
 			{
 				// Always load the parent blueprint here in case there is a circular dependency. This will
 				// ensure that the blueprint is fully serialized before attempting to regenerate the class.
-				if (!BlueprintObject->HasAnyFlags(RF_LoadCompleted))
-				{
-					BlueprintObject->SetFlags(RF_NeedLoad);
-					if (ULinkerLoad* Linker = BlueprintObject->GetLinker())
-					{
-						Linker->Preload(BlueprintObject);
-					}
-				}
+				FPreloadMembersHelper::PreloadObject(BlueprintObject);
 			
 				FPreloadMembersHelper::PreloadMembers(BlueprintObject);
 				// recurse into this function for this parent class; 
@@ -110,18 +115,18 @@ bool ULinkerLoad::RegenerateBlueprintClass(UClass* LoadClass, UObject* ExportObj
 			}
 		}
 
-		// Preload the blueprint to make sure it has all the data the class needs for regeneration
-		if (LoadClass->ClassGeneratedBy->HasAnyFlags(RF_NeedLoad))
 		{
-			LoadClass->ClassGeneratedBy->GetLinker()->Preload(LoadClass->ClassGeneratedBy);
-		}
+			UObject* BlueprintObject = LoadClass->ClassGeneratedBy;
+			// Preload the blueprint to make sure it has all the data the class needs for regeneration
+			FPreloadMembersHelper::PreloadObject(BlueprintObject);
 
-		UClass* RegeneratedClass = LoadClass->ClassGeneratedBy->RegenerateClass(LoadClass, CurrentCDO, GObjLoaded);
-		if (RegeneratedClass)
-		{
-			LoadClass->ClassGeneratedBy->ClearFlags(RF_BeingRegenerated);
-			// Fix up the linker so that the RegeneratedClass is used
-			LoadClass->ClearFlags(RF_NeedLoad | RF_NeedPostLoad);
+			UClass* RegeneratedClass = BlueprintObject->RegenerateClass(LoadClass, CurrentCDO, GObjLoaded);
+			if (RegeneratedClass)
+			{
+				BlueprintObject->ClearFlags(RF_BeingRegenerated);
+				// Fix up the linker so that the RegeneratedClass is used
+				LoadClass->ClearFlags(RF_NeedLoad | RF_NeedPostLoad);
+			}
 		}
 	}
 
