The "_C" postif should not be visible in editor UI.

diff --git a/Engine/Source/Editor/BlueprintGraph/Private/EdGraphSchema_K2.cpp b/Engine/Source/Editor/BlueprintGraph/Private/EdGraphSchema_K2.cpp
--- a/Engine/Source/Editor/BlueprintGraph/Private/EdGraphSchema_K2.cpp
+++ b/Engine/Source/Editor/BlueprintGraph/Private/EdGraphSchema_K2.cpp
@@ -175,10 +175,6 @@ void UEdGraphSchema_K2::FPinTypeTreeInfo::Init(const FString& InFriendlyName, co
 			if(Struct != NULL)
 			{
 				SubtypeTooltip = Struct->GetToolTipText().ToString();
-				if(SubtypeTooltip.Len() == 0)
-				{
-					SubtypeTooltip = Struct->GetName();
-				}
 			}
 
 			Children.Add( MakeShareable(new FPinTypeTreeInfo(CategoryName, *it, SubtypeTooltip)) );
diff --git a/Engine/Source/Runtime/CoreUObject/Private/UObject/Class.cpp b/Engine/Source/Runtime/CoreUObject/Private/UObject/Class.cpp
--- a/Engine/Source/Runtime/CoreUObject/Private/UObject/Class.cpp
+++ b/Engine/Source/Runtime/CoreUObject/Private/UObject/Class.cpp
@@ -180,6 +180,21 @@ void UField::AddCppProperty( UProperty* Property )
 }
 
 #if WITH_EDITOR || HACK_HEADER_GENERATOR
+
+struct FDisplayNameHelper
+{
+	static FString Get(const UObject& Object)
+	{
+		FString Name = Object.GetName();
+		const UClass* Class = Cast<UClass>(&Object);
+		if (Class && !Class->HasAnyClassFlags(CLASS_Native))
+		{
+			Name.RemoveFromEnd(TEXT("_C"));
+		}
+		return Name;
+	}
+};
+
 /**
  * Finds the localized display name or native display name as a fallback.
  *
@@ -199,7 +214,7 @@ FText UField::GetDisplayNameText() const
 	}
 	else
 	{
-		NativeDisplayName = FName::NameToDisplayString(GetName(), IsA<UBoolProperty>());
+		NativeDisplayName = FName::NameToDisplayString(FDisplayNameHelper::Get(*this), IsA<UBoolProperty>());
 	}
 
 	if ( !( FText::FindText( Namespace, Key, /*OUT*/LocalizedDisplayName, &NativeDisplayName ) ) )
@@ -226,7 +241,7 @@ FText UField::GetToolTipText() const
 	{
 		if (NativeToolTip.IsEmpty())
 		{
-			NativeToolTip = GetName();
+			NativeToolTip = FDisplayNameHelper::Get(*this);
 		}
 		LocalizedToolTip = FText::FromString(NativeToolTip);
 	}
