Heading off some potential blueprints issues that came out of a merge (fixing it up so that we properly pass the Blueprint ExportObject along in a circular load situation).

diff --git a/Engine/Source/Runtime/CoreUObject/Private/Blueprint/BlueprintSupport.cpp b/Engine/Source/Runtime/CoreUObject/Private/Blueprint/BlueprintSupport.cpp
--- a/Engine/Source/Runtime/CoreUObject/Private/Blueprint/BlueprintSupport.cpp
+++ b/Engine/Source/Runtime/CoreUObject/Private/Blueprint/BlueprintSupport.cpp
@@ -111,7 +111,7 @@ bool ULinkerLoad::RegenerateBlueprintClass(UClass* LoadClass, UObject* ExportObj
 				FPreloadMembersHelper::PreloadMembers(BlueprintObject);
 				// recurse into this function for this parent class; 
 				// 'ClassDefaultObject' should be the class's original ExportObject
-				RegenerateBlueprintClass(Class, BlueprintObject);
+				RegenerateBlueprintClass(Class, Class->ClassDefaultObject);
 			}
 		}
 
diff --git a/Engine/Source/Runtime/CoreUObject/Private/UObject/LinkerLoad.cpp b/Engine/Source/Runtime/CoreUObject/Private/UObject/LinkerLoad.cpp
--- a/Engine/Source/Runtime/CoreUObject/Private/UObject/LinkerLoad.cpp
+++ b/Engine/Source/Runtime/CoreUObject/Private/UObject/LinkerLoad.cpp
@@ -2736,6 +2736,7 @@ void ULinkerLoad::Preload( UObject* Object )
 					if ( Object->HasAnyFlags(RF_ClassDefaultObject) )
 					{
 						Object->GetClass()->SerializeDefaultObject(Object, *this);
+						Object->SetFlags(RF_LoadCompleted);
 					}
 					else
 					{
