ensured uncooked builds behave similarly with regards to find object in memory first when loading packages.
#TTP 342432 - CRITICAL: SWINGNINJA: CRASH: Opening P_City2 in uncooked game
#Branch UE4
#Proj CoreUObject

#Change simplified the check for attempting to find exports in memory first to be just a test for a cooked data build.

#info The assert was caused by UBlueprintGeneratedClass::ConditionalRecompileClass being triggered during the FScopedDependencyGather destructor on MyCharacter blueprint. The blueprint in question was not serialised in and the GeneratedClass member was null ( which pushed the load to EndLoad in LoadPackage ). The game can be forced to load without asserts but still has issues and I've found that only by forcing RegenerateBlueprintClass on MyCharacter_C results in the maps loading correctly. This change basically ensures that uncooked data builds follow the same load path, really shouldn't affect cooked data builds.

diff --git a/Engine/Source/Runtime/CoreUObject/Private/UObject/LinkerLoad.cpp b/Engine/Source/Runtime/CoreUObject/Private/UObject/LinkerLoad.cpp
--- a/Engine/Source/Runtime/CoreUObject/Private/UObject/LinkerLoad.cpp
+++ b/Engine/Source/Runtime/CoreUObject/Private/UObject/LinkerLoad.cpp
@@ -3080,8 +3080,9 @@ UObject* ULinkerLoad::CreateExport( int32 Index )
 		// case of async loading as we cannot in-place replace objects.
 
 		UObject* ActualObjectWithTheName = StaticFindObjectFastInternal(NULL, ThisParent, Export.ObjectName, true);
-
-		if(	(FApp::IsGame() && !GIsEditor && !IsRunningCommandlet()) 
+		
+		// if we require cooked data, attempt to find exports in memory first
+		if(	FPlatformProperties::RequiresCookedData() 
 		||	GIsAsyncLoading 
 		||	Export.bForcedExport 
 		||	LinkerRoot->ShouldFindExportsInMemoryFirst()
