BP: passing parameters by ref optimization - part 1 - native code (similar change should be done for VM calling).
The reference should always be retrieved from MostRecentPropertyAddress. There is no need to copy the parameter.

diff --git a/Engine/Source/Runtime/CoreUObject/Private/UObject/ScriptCore.cpp b/Engine/Source/Runtime/CoreUObject/Private/UObject/ScriptCore.cpp
--- a/Engine/Source/Runtime/CoreUObject/Private/UObject/ScriptCore.cpp
+++ b/Engine/Source/Runtime/CoreUObject/Private/UObject/ScriptCore.cpp
@@ -116,8 +116,6 @@ void FFrame::Step(UObject *Context, RESULT_DECL)
 
 void FFrame::StepExplicitProperty(void*const Result, UProperty* Property)
 {
-	checkSlow(Result != NULL);
-
 	if (Property->PropertyFlags & CPF_OutParm)
 	{
 		// look through the out parameter infos and find the one that has the address of this property
@@ -134,7 +132,10 @@ void FFrame::StepExplicitProperty(void*const Result, UProperty* Property)
 	else
 	{
 		MostRecentPropertyAddress = Property->ContainerPtrToValuePtr<uint8>(Locals);
-		Property->CopyCompleteValueToScriptVM(Result, MostRecentPropertyAddress);
+		if (Result)
+		{
+			Property->CopyCompleteValueToScriptVM(Result, MostRecentPropertyAddress);
+		}
 	}
 }
 
diff --git a/Engine/Source/Runtime/CoreUObject/Public/UObject/ScriptMacros.h b/Engine/Source/Runtime/CoreUObject/Public/UObject/ScriptMacros.h
--- a/Engine/Source/Runtime/CoreUObject/Public/UObject/ScriptMacros.h
+++ b/Engine/Source/Runtime/CoreUObject/Public/UObject/ScriptMacros.h
@@ -27,20 +27,17 @@ enum {MAX_VARIABLE_SIZE = 0x0FFF };
 	Stack.StepCompiledIn<PropertyType>(&ParamName);
 
 #define PARAM_PASSED_BY_REF(ParamName, PropertyType, ParamType)									\
-	ParamType ParamName##Temp;																	\
-	ParamType& ParamName = Stack.StepCompiledInRef<PropertyType, ParamType>(&ParamName##Temp);
+	ParamType& ParamName = Stack.StepCompiledInRef<PropertyType, ParamType>(NULL);
 
 #define PARAM_PASSED_BY_REF_ZEROED(ParamName, PropertyType, ParamType)							\
-	ParamType ParamName##Temp = (ParamType)0;													\
-	ParamType& ParamName = Stack.StepCompiledInRef<PropertyType, ParamType>(&ParamName##Temp);
+	ParamType& ParamName = Stack.StepCompiledInRef<PropertyType, ParamType>(NULL);
 
 #define P_GET_PROPERTY(PropertyType, ParamName)													\
 	PropertyType::TCppType ParamName = PropertyType::GetDefaultPropertyValue();					\
 	Stack.StepCompiledIn<PropertyType>(&ParamName);
 
 #define P_GET_PROPERTY_REF(PropertyType, ParamName)												\
-	PropertyType::TCppType ParamName##Temp = PropertyType::GetDefaultPropertyValue();			\
-	PropertyType::TCppType& ParamName = Stack.StepCompiledInRef<PropertyType, PropertyType::TCppType>(&ParamName##Temp);
+	PropertyType::TCppType& ParamName = Stack.StepCompiledInRef<PropertyType, PropertyType::TCppType>(NULL);
 
 
 
diff --git a/Engine/Source/Runtime/CoreUObject/Public/UObject/Stack.h b/Engine/Source/Runtime/CoreUObject/Public/UObject/Stack.h
--- a/Engine/Source/Runtime/CoreUObject/Public/UObject/Stack.h
+++ b/Engine/Source/Runtime/CoreUObject/Public/UObject/Stack.h
@@ -421,5 +421,6 @@ FORCEINLINE_DEBUGGABLE TNativeType& FFrame::StepCompiledInRef(void*const Tempora
 		StepExplicitProperty(TemporaryBuffer, Property);
 	}
 
-	return (MostRecentPropertyAddress != NULL) ? *(TNativeType*)(MostRecentPropertyAddress) : *(TNativeType*)TemporaryBuffer;
+	check(MostRecentPropertyAddress);
+	return *(TNativeType*)(MostRecentPropertyAddress);
 }
