#UE4 final fix for TTP#341121. Infinite loop detection could fire repeatedly if we managed to write a script that could rubberband around the runaway count. Revied by Nick W.

diff --git a/Engine/Source/Runtime/CoreUObject/Private/UObject/ScriptCore.cpp b/Engine/Source/Runtime/CoreUObject/Private/UObject/ScriptCore.cpp
--- a/Engine/Source/Runtime/CoreUObject/Private/UObject/ScriptCore.cpp
+++ b/Engine/Source/Runtime/CoreUObject/Private/UObject/ScriptCore.cpp
@@ -41,8 +41,9 @@ COREUOBJECT_API int32 GCastDuplicate=0;
 #if DO_GUARD
 	static int32 Runaway=0;
 	static int32 Recurse=0;
+	static bool Ranaway = false;
 	#define CHECK_RUNAWAY {++Runaway;}
-	COREUOBJECT_API void GInitRunaway() {Recurse=Runaway=0;}
+	COREUOBJECT_API void GInitRunaway() {Recurse=Runaway=0; Ranaway = false;}
 #else
 	#define CHECK_RUNAWAY
 	COREUOBJECT_API void GInitRunaway() {}
@@ -542,7 +543,7 @@ void UObject::ProcessInternal( FFrame& Stack, RESULT_DECL )
 		MS_ALIGN(16) uint8 Buffer[MAX_SIMPLE_RETURN_VALUE_SIZE] GCC_ALIGN(16);
 
 #if DO_GUARD
-		if (++Recurse == RECURSE_LIMIT)
+		if (++Recurse == RECURSE_LIMIT && !Ranaway)
 		{
 			// We've hit the recursion limit, so print out the stack, warn, and then continue with a zeroed return value.
 			UE_LOG(LogScriptCore, Log, TEXT("%s"), *Stack.GetStackTrace());
@@ -556,8 +557,9 @@ void UObject::ProcessInternal( FFrame& Stack, RESULT_DECL )
 			FBlueprintExceptionInfo InfiniteRecursionExceptionInfo(EBlueprintExceptionType::InfiniteLoop, Desc);
 			FBlueprintCoreDelegates::ThrowScriptException(this, Stack, InfiniteRecursionExceptionInfo);
 
-			// Bump the recursion so we don't re-enter after we've thrown an exception
-			++Recurse;
+			// This flag prevents repeated warnings of ininite loop, script exception handler 
+			// is expected to have terminated execution appropriately:
+			Ranaway = true;
 
 			return;
 		}
