This adds an editor setting in the general project settings under blueprints to allow users can define the script maximum loop iterations
#TTP 340083 - +UE4: K2: Add editor exposed configuration of script max loop iterations count
#Branch UE4
#Proj Engine, CoreUObject

#Change added base engine config value MaximumLoopIterationCount defaulted to the same value as RUNAWAY_LIMIT in scriptcore.cpp
#Change added UEngine::MaximumLoopIterationCount UProperty to receive and store the config value
#Change replaced RUNAWAY_LIMIT define in sciptcore.cpp with GMaximumScriptLoopIterations and switched checks to use this variable
#Change added FBlueprintCoreDelegates::SetScriptMaximumLoopIterations as a feeder function to set GMaximumScriptLoopIterations in the CoreUObjectModule.
#Change Added call to update GMaximumScriptLoopIterations in UEngine::Init and UEditorEngine::PlayInEditor to forward any user changes at key points.

diff --git a/Engine/Config/BaseEngine.ini b/Engine/Config/BaseEngine.ini
--- a/Engine/Config/BaseEngine.ini
+++ b/Engine/Config/BaseEngine.ini
@@ -124,6 +124,7 @@ LevelStreamingComponentsRegistrationGranularity = 10
 bSubtitlesEnabled=True
 bSubtitlesForcedOff=false
 DefaultSoundName=/Engine/EngineSounds/WhiteNoise.WhiteNoise
+MaximumLoopIterationCount = 1000000
 ; If setting bAllowDebugViewmodesOnConsoles=true, must also add EngineDebugMaterials to the StartupPackages for console platforms
 bAllowDebugViewmodesOnConsoles=false
 CameraRotationThreshold=45.0
diff --git a/Engine/Source/Editor/UnrealEd/Private/PlayLevel.cpp b/Engine/Source/Editor/UnrealEd/Private/PlayLevel.cpp
--- a/Engine/Source/Editor/UnrealEd/Private/PlayLevel.cpp
+++ b/Engine/Source/Editor/UnrealEd/Private/PlayLevel.cpp
@@ -1615,6 +1615,9 @@ void UEditorEngine::PlayInEditor( UWorld* InWorld, bool bInSimulateInEditor )
 	ULevelEditorPlaySettings const* EditorPlayInSettings = GetDefault<ULevelEditorPlaySettings>();
 	check(EditorPlayInSettings);
 
+	// Update Script Maximum loop iteration count
+	FBlueprintCoreDelegates::SetScriptMaximumLoopIterations( GEngine->MaximumLoopIterationCount );
+
 	// Prompt the user to compile any dirty Blueprints before PIE can occur.
 	TArray< UBlueprint* > ErrorBlueprintList;
 	bool bAnyBlueprintsDirty = false;
diff --git a/Engine/Source/Runtime/CoreUObject/Private/UObject/ScriptCore.cpp b/Engine/Source/Runtime/CoreUObject/Private/UObject/ScriptCore.cpp
--- a/Engine/Source/Runtime/CoreUObject/Private/UObject/ScriptCore.cpp
+++ b/Engine/Source/Runtime/CoreUObject/Private/UObject/ScriptCore.cpp
@@ -26,7 +26,8 @@ COREUOBJECT_API int32 GNativeDuplicate=0;
 COREUOBJECT_API void (UObject::*GCasts[CST_Max])( FFrame &Stack, RESULT_DECL );
 COREUOBJECT_API int32 GCastDuplicate=0;
 
-#define RUNAWAY_LIMIT 1000000
+COREUOBJECT_API int32 GMaximumScriptLoopIterations = 1000000;
+
 #if !PLATFORM_DESKTOP
 	#define RECURSE_LIMIT 120
 #else
@@ -87,6 +88,11 @@ void FBlueprintCoreDelegates::ThrowScriptException(const UObject* ActiveObject,
 	}
 }
 
+void FBlueprintCoreDelegates::SetScriptMaximumLoopIterations( const int32 MaximumLoopIterations )
+{
+	GMaximumScriptLoopIterations = MaximumLoopIterations;
+}
+
 //////////////////////////////////////////////////////////////////////////
 // FEditorScriptExecutionGuard
 FEditorScriptExecutionGuard::FEditorScriptExecutionGuard()
@@ -571,7 +577,7 @@ void UObject::ProcessInternal( FFrame& Stack, RESULT_DECL )
 		while (*Stack.Code != EX_Return)
 		{
 #if DO_GUARD
-			if( Runaway > RUNAWAY_LIMIT )
+			if( Runaway > GMaximumScriptLoopIterations )
 			{
 				// We've hit the recursion limit, so print out the stack, warn, and then continue with a zeroed return value.
 				UE_LOG(LogScriptCore, Log, TEXT("%s"), *Stack.GetStackTrace());
@@ -581,7 +587,7 @@ void UObject::ProcessInternal( FFrame& Stack, RESULT_DECL )
 				ClearReturnValue(ReturnProp, Result);
 
 				// Notify anyone who cares that we've had a fatal error, so we can shut down PIE, etc
-				const FString Desc = FString::Printf(TEXT("Runaway loop detected (over %i iterations)"), RUNAWAY_LIMIT);
+				const FString Desc = FString::Printf(TEXT("Runaway loop detected (over %i iterations)"), GMaximumScriptLoopIterations );
 				FBlueprintExceptionInfo RunawayLoopExceptionInfo(EBlueprintExceptionType::InfiniteLoop, Desc);
 				FBlueprintCoreDelegates::ThrowScriptException(this, Stack, RunawayLoopExceptionInfo);
 
diff --git a/Engine/Source/Runtime/CoreUObject/Public/UObject/Script.h b/Engine/Source/Runtime/CoreUObject/Public/UObject/Script.h
--- a/Engine/Source/Runtime/CoreUObject/Public/UObject/Script.h
+++ b/Engine/Source/Runtime/CoreUObject/Public/UObject/Script.h
@@ -303,6 +303,7 @@ public:
 
 public:
 	static void ThrowScriptException(const UObject* ActiveObject, const struct FFrame& StackFrame, const FBlueprintExceptionInfo& Info);
+	static void SetScriptMaximumLoopIterations( const int32 MaximumLoopIterations );
 };
 
 // Scoped struct to allow execution of script in editor, while resetting the runaway loop counts
diff --git a/Engine/Source/Runtime/Engine/Classes/Engine/Engine.h b/Engine/Source/Runtime/Engine/Classes/Engine/Engine.h
--- a/Engine/Source/Runtime/Engine/Classes/Engine/Engine.h
+++ b/Engine/Source/Runtime/Engine/Classes/Engine/Engine.h
@@ -1104,7 +1104,11 @@ public:
 	/** Batching granularity used to register actor components during level streaming */
 	UPROPERTY(EditAnywhere, config, Category=LevelStreaming, AdvancedDisplay)
 	int32 LevelStreamingComponentsRegistrationGranularity;
-	
+
+	/** Script maximum loop iteration count used as a threshold to warn users about script execution runaway */
+	UPROPERTY(EditAnywhere, config, Category=Blueprints)
+	int32 MaximumLoopIterationCount;
+
 	/** @todo document */
 	UPROPERTY(config)
 	uint32 bEnableEditorPSysRealtimeLOD:1;
diff --git a/Engine/Source/Runtime/Engine/Private/UnrealEngine.cpp b/Engine/Source/Runtime/Engine/Private/UnrealEngine.cpp
--- a/Engine/Source/Runtime/Engine/Private/UnrealEngine.cpp
+++ b/Engine/Source/Runtime/Engine/Private/UnrealEngine.cpp
@@ -709,6 +709,9 @@ void UEngine::Init(IEngineLoop* InEngineLoop)
 		GConfig->GetBool(TEXT("DevOptions.Debug"), TEXT("ShowSelectedLightmap"), GShowDebugSelectedLightmap, GEngineIni);
 	}
 
+	// Update Script Maximum loop iteration count
+	FBlueprintCoreDelegates::SetScriptMaximumLoopIterations( GEngine->MaximumLoopIterationCount );
+
 	GNearClippingPlane = NearClipPlane;
 
 	// Initialize the audio device
