ttp342904 CRASH:EDITOR: Attached character blueprint cause immediate crash upon opening blueprint editor

diff --git a/Engine/Source/Editor/UnrealEd/Private/Kismet2/BlueprintEditorUtils.cpp b/Engine/Source/Editor/UnrealEd/Private/Kismet2/BlueprintEditorUtils.cpp
--- a/Engine/Source/Editor/UnrealEd/Private/Kismet2/BlueprintEditorUtils.cpp
+++ b/Engine/Source/Editor/UnrealEd/Private/Kismet2/BlueprintEditorUtils.cpp
@@ -624,6 +624,16 @@ struct FRegenerationHelper
 		}
 	}
 
+	static void ForceLoadMetaData(UObject* InObject)
+	{
+		checkSlow(InObject);
+		UPackage* Package = InObject->GetOutermost();
+		checkSlow(Package);
+		UMetaData* MetaData = Package->GetMetaData();
+		checkSlow(MetaData);
+		ForcedLoad(MetaData);
+	}
+
 	static void PreloadAndLinkIfNecessary(UStruct* Struct)
 	{
 		bool bChanged = false;
@@ -636,13 +646,7 @@ struct FRegenerationHelper
 			}
 		}
 
-		{
-			UPackage* Package = Struct->GetOutermost();
-			check(Package);
-			UMetaData* MetaData = Package->GetMetaData();
-			check(MetaData);
-			ForcedLoad(MetaData);
-		}
+		ForceLoadMetaData(Struct);
 
 		const int32 OldPropertiesSize = Struct->GetPropertiesSize();
 		for (UField* Field = Struct->Children; Field; Field = Field->Next)
@@ -1035,6 +1039,7 @@ UClass* FBlueprintEditorUtils::RegenerateBlueprintClass(UBlueprint* Blueprint, U
 	if (!Blueprint->bHasBeenRegenerated)
 	{
 		check(PreviousCDO != nullptr);
+		FRegenerationHelper::ForceLoadMetaData(Blueprint);
 		FRegenerationHelper::ForcedLoadMembers(PreviousCDO);
 		FRegenerationHelper::ForcedLoadMembers(Blueprint);
 	}
diff --git a/Engine/Source/Runtime/CoreUObject/Private/UObject/MetaData.cpp b/Engine/Source/Runtime/CoreUObject/Private/UObject/MetaData.cpp
--- a/Engine/Source/Runtime/CoreUObject/Private/UObject/MetaData.cpp
+++ b/Engine/Source/Runtime/CoreUObject/Private/UObject/MetaData.cpp
@@ -114,10 +114,31 @@ void UMetaData::Serialize(FArchive& Ar)
 		}
 	}
 	
-	Ar << ObjectMetaDataMap;
-
-	if (Ar.IsLoading())
+	if (!Ar.IsLoading())
+	{
+		Ar << ObjectMetaDataMap;
+	}
+	else
 	{
+		{
+			TMap< FWeakObjectPtr, TMap<FName, FString> > TempMap;
+			Ar << TempMap;
+
+			const bool bLoadFromLinker = (NULL != Ar.GetLinker());
+			if (bLoadFromLinker && HasAnyFlags(RF_LoadCompleted))
+			{
+				UE_LOG(LogMetaData, Verbose, TEXT("Metadata was already loaded by linker. %s"), *GetFullName());
+			}
+			else
+			{
+				if (bLoadFromLinker && ObjectMetaDataMap.Num())
+				{
+					UE_LOG(LogMetaData, Verbose, TEXT("Metadata: Some values, filled while serialization, may be lost. %s"), *GetFullName());
+				}
+				Swap(ObjectMetaDataMap, TempMap);
+			}
+		}
+
 		// Run redirects on loaded keys
 		InitializeRedirectMap();
 
@@ -270,6 +291,11 @@ void UMetaData::SetValue(const UObject* Object, FName Key, const TCHAR* Value)
 {
 	check(Key != NAME_None);
 
+	if (!HasAllFlags(RF_LoadCompleted))
+	{
+		UE_LOG(LogMetaData, Error, TEXT("MetaData::SetValue called before meta data is completely loaded. %s'"), *GetPathName());
+	}
+
 	// look up the existing map if we have it
 	TMap<FName, FString>* ObjectValues = ObjectMetaDataMap.Find(Object);
 
diff --git a/Engine/Source/Runtime/CoreUObject/Private/UObject/Package.cpp b/Engine/Source/Runtime/CoreUObject/Private/UObject/Package.cpp
--- a/Engine/Source/Runtime/CoreUObject/Private/UObject/Package.cpp
+++ b/Engine/Source/Runtime/CoreUObject/Private/UObject/Package.cpp
@@ -113,7 +113,7 @@ UMetaData* UPackage::GetMetaData()
 		// If MetaData is NULL then it wasn't loaded by linker, so we have to create it.
 		if(MetaData == NULL)
 		{
-			MetaData = ConstructObject<UMetaData>(UMetaData::StaticClass(), this, NAME_PackageMetaData, RF_Standalone);
+			MetaData = ConstructObject<UMetaData>(UMetaData::StaticClass(), this, NAME_PackageMetaData, RF_Standalone | RF_LoadCompleted);
 		}
 	}
 
