#UE4 consistently bailing once we've ranaway, old behavior would bail only when we exactly hit the stack size limit. Fix for TTP#343294

diff --git a/Engine/Source/Runtime/CoreUObject/Private/UObject/ScriptCore.cpp b/Engine/Source/Runtime/CoreUObject/Private/UObject/ScriptCore.cpp
--- a/Engine/Source/Runtime/CoreUObject/Private/UObject/ScriptCore.cpp
+++ b/Engine/Source/Runtime/CoreUObject/Private/UObject/ScriptCore.cpp
@@ -549,7 +549,14 @@ void UObject::ProcessInternal( FFrame& Stack, RESULT_DECL )
 		MS_ALIGN(16) uint8 Buffer[MAX_SIMPLE_RETURN_VALUE_SIZE] GCC_ALIGN(16);
 
 #if DO_GUARD
-		if (++Recurse == RECURSE_LIMIT && !Ranaway)
+		if(Ranaway)
+		{
+			// If we have a return property, return a zeroed value in it, to try and save execution as much as possible
+			UProperty* ReturnProp = ((UFunction*)Stack.Node)->GetReturnProperty();
+			ClearReturnValue(ReturnProp, Result);
+			return;
+		}
+		else if (++Recurse == RECURSE_LIMIT)
 		{
 			// We've hit the recursion limit, so print out the stack, warn, and then continue with a zeroed return value.
 			UE_LOG(LogScriptCore, Log, TEXT("%s"), *Stack.GetStackTrace());
