BP: we need to avoid a case, when while class regeneration a different CDO is set, and now we set it to a stale one

ttp#343166 CRITICAL: ElementalVR: CRASH: Occurs when compiling BP_SimplePlanet.

diff --git a/Engine/Source/Runtime/CoreUObject/Private/UObject/Class.cpp b/Engine/Source/Runtime/CoreUObject/Private/UObject/Class.cpp
--- a/Engine/Source/Runtime/CoreUObject/Private/UObject/Class.cpp
+++ b/Engine/Source/Runtime/CoreUObject/Private/UObject/Class.cpp
@@ -3023,7 +3023,21 @@ void UClass::Serialize( FArchive& Ar )
 	{
 		check(GetStructureSize() >= sizeof(UObject));
 		check(!GetSuperClass() || !GetSuperClass()->HasAnyFlags(RF_NeedLoad));
-		Ar << ClassDefaultObject;
+		UObject* const OldCDO = ClassDefaultObject;
+		UObject* TempClassDefaultObject = NULL;
+		Ar << TempClassDefaultObject;
+		// we need to avoid a case, when while class regeneration a different CDO is set, and now we set it to an stale one (ttp#343166)
+		if (ClassDefaultObject == OldCDO)
+		{
+			ClassDefaultObject = TempClassDefaultObject;
+		}
+		else if (TempClassDefaultObject != ClassDefaultObject)
+		{
+			UE_LOG(LogClass, Log, TEXT("CDO was changed while class serialization.\n\tOld: '%s'\n\tSerialized: '%s'\n\tActual: '%s'")
+				, OldCDO ? *OldCDO->GetFullName() : TEXT("NULL")
+				, TempClassDefaultObject ? *TempClassDefaultObject->GetFullName() : TEXT("NULL")
+				, ClassDefaultObject ? *ClassDefaultObject->GetFullName() : TEXT("NULL"));
+		}
 		ClassUnique = 0;
 	}
 	else
