The log output was Modified to use 'Parent' instead of 'Super' as requested.
#TTP 304083 - ROCKET: TASK: K2:  Change verbage on the super call nodes
#Branch UE4
#Proj CoreUObject

diff --git a/Engine/Source/Runtime/CoreUObject/Private/UObject/LinkerLoad.cpp b/Engine/Source/Runtime/CoreUObject/Private/UObject/LinkerLoad.cpp
--- a/Engine/Source/Runtime/CoreUObject/Private/UObject/LinkerLoad.cpp
+++ b/Engine/Source/Runtime/CoreUObject/Private/UObject/LinkerLoad.cpp
@@ -2939,13 +2939,13 @@ UObject* ULinkerLoad::CreateExport( int32 Index )
 			{
 				if( LoadClass->IsChildOf(UFunction::StaticClass()) )
 				{
-					// If this is a function whose super has been removed, give it a NULL super, as we would have in the script compiler
-					UE_LOG(LogLinker, Warning, TEXT("CreateExport: Failed to load Super for %s; removing super information, but keeping function"), *GetExportFullName(Index));
+					// If this is a function whose parent has been removed, give it a NULL parent, as we would have in the script compiler
+					UE_LOG(LogLinker, Warning, TEXT("CreateExport: Failed to load Parent for %s; removing parent information, but keeping function"), *GetExportFullName(Index));
 					Export.SuperIndex = FPackageIndex();
 				}
 				else
 				{
-					UE_LOG(LogLinker, Warning, TEXT("CreateExport: Failed to load Super for %s"), *GetExportFullName(Index));
+					UE_LOG(LogLinker, Warning, TEXT("CreateExport: Failed to load Parent for %s"), *GetExportFullName(Index));
 					return NULL;
 				}
 			}
