Don't copy the outer linker when the linker already exists in non-cooked builds as we may be overwriting a redirected linker.

#ttp 343661: BLOCKER: MAIN: EDITOR: CRASH: Editor crashes after opening aismoke.umap after renaming huskgeneric blueprint

diff --git a/Engine/Source/Runtime/CoreUObject/Private/UObject/LinkerLoad.cpp b/Engine/Source/Runtime/CoreUObject/Private/UObject/LinkerLoad.cpp
--- a/Engine/Source/Runtime/CoreUObject/Private/UObject/LinkerLoad.cpp
+++ b/Engine/Source/Runtime/CoreUObject/Private/UObject/LinkerLoad.cpp
@@ -2191,8 +2191,14 @@ bool ULinkerLoad::VerifyImportInner(const int32 ImportIndex, FString& WarningSuf
 			 }
 		}
 
-		// Copy the SourceLinker from the FObjectImport for our Outer
-		Import.SourceLinker = OuterImport.SourceLinker;
+		// Copy the SourceLinker from the FObjectImport for our Outer if the SourceLinker hasn't been set yet,
+		// Otherwise we may be overwriting a re-directed linker and SourceIndex is already from the redirected one.
+		// This can only happen in non-cooked builds though.
+		if (FPlatformProperties::RequiresCookedData() || !Import.SourceLinker)
+		{
+			Import.SourceLinker = OuterImport.SourceLinker;
+		}
+
 
 		//check(Import.SourceLinker);
 		//@todo what does it mean if we don't have a SourceLinker here?
