More log output when changing linkers in non-editor builds

diff --git a/Engine/Source/Runtime/CoreUObject/Private/UObject/UObjectLinker.cpp b/Engine/Source/Runtime/CoreUObject/Private/UObject/UObjectLinker.cpp
--- a/Engine/Source/Runtime/CoreUObject/Private/UObject/UObjectLinker.cpp
+++ b/Engine/Source/Runtime/CoreUObject/Private/UObject/UObjectLinker.cpp
@@ -8,6 +8,8 @@
 
 #include "UObjectAnnotation.h"
 
+DEFINE_LOG_CATEGORY_STATIC(LogUObjectLinker, Log, All);
+
 //@todo UE4 Console - Check that the mapping of UObjects to linkers is sparse and that we aren't spending a ton of time with these lookups.
 
 struct FLinkerIndexPair
@@ -107,7 +109,11 @@ void UObject::SetLinker( ULinkerLoad* LinkerLoad, int32 LinkerIndex, bool bShoul
 #if WITH_EDITOR
 		PostLinkerChange();
 #else
-		check(!Existing.Linker || !LinkerLoad); // It is only legal to change linkers in the editor
+		UE_CLOG(Existing.Linker && LinkerLoad, LogUObjectLinker, Fatal,
+			TEXT("It is only legal to change linkers in the editor. Trying to change linker on %s from %s to %s"),
+			*GetFullName(),
+			Existing.Linker ? *Existing.Linker->Filename : TEXT("NULL"),
+			LinkerLoad ? *LinkerLoad->Filename : TEXT("NULL"));
 #endif
 	}
 }
