Core: Added an ensure >0 to FBlueprintCoreDelegates::SetScriptMaximumLoopIterations

diff --git a/Engine/Source/Runtime/CoreUObject/Private/UObject/ScriptCore.cpp b/Engine/Source/Runtime/CoreUObject/Private/UObject/ScriptCore.cpp
--- a/Engine/Source/Runtime/CoreUObject/Private/UObject/ScriptCore.cpp
+++ b/Engine/Source/Runtime/CoreUObject/Private/UObject/ScriptCore.cpp
@@ -91,7 +91,10 @@ void FBlueprintCoreDelegates::ThrowScriptException(const UObject* ActiveObject,
 
 void FBlueprintCoreDelegates::SetScriptMaximumLoopIterations( const int32 MaximumLoopIterations )
 {
-	GMaximumScriptLoopIterations = MaximumLoopIterations;
+	if (ensure(MaximumLoopIterations > 0))
+	{
+		GMaximumScriptLoopIterations = MaximumLoopIterations;
+	}
 }
 
 //////////////////////////////////////////////////////////////////////////
