NonPIETransient renamed to NonPIEDuplicateTransient to better reflect that it's a variant of DuplicateTransient, not a variant of Transient.
Bug fix where serialization was not occurring on properties marked as NonPIETransient.

diff --git a/Engine/Source/Editor/UnrealEd/Private/Editor.cpp b/Engine/Source/Editor/UnrealEd/Private/Editor.cpp
--- a/Engine/Source/Editor/UnrealEd/Private/Editor.cpp
+++ b/Engine/Source/Editor/UnrealEd/Private/Editor.cpp
@@ -5245,7 +5245,7 @@ static void CopyLightComponentProperties( const AActor& InOldActor, AActor& InNe
 		// Now Copy the light component properties.
 		for( UProperty* Property = CommonLightComponentClass->PropertyLink; Property != NULL; Property = Property->PropertyLinkNext )
 		{
-			bool bIsTransient = !!(Property->PropertyFlags & (CPF_Transient | CPF_DuplicateTransient | CPF_NonPIETransient));
+			bool bIsTransient = !!(Property->PropertyFlags & (CPF_Transient | CPF_DuplicateTransient | CPF_NonPIEDuplicateTransient));
 			// Properties are identical if they have not changed from the light component on the default source actor
 			bool bIsIdentical = Property->Identical_InContainer(LightComponentToCopy, DefaultLightComponent);
 			bool bIsComponent = !!(Property->PropertyFlags & (CPF_InstancedReference | CPF_ContainsInstancedReference));
diff --git a/Engine/Source/Programs/UnrealHeaderTool/Private/HeaderParser.cpp b/Engine/Source/Programs/UnrealHeaderTool/Private/HeaderParser.cpp
--- a/Engine/Source/Programs/UnrealHeaderTool/Private/HeaderParser.cpp
+++ b/Engine/Source/Programs/UnrealHeaderTool/Private/HeaderParser.cpp
@@ -2482,7 +2482,12 @@ bool FHeaderParser::GetVarType
 			}
 			else if (Specifier == TEXT("NonPIETransient"))
 			{
-				Flags |= CPF_NonPIETransient;
+				UE_LOG(LogCompile, Warning, TEXT("NonPIETransient is deprecated - NonPIEDuplicateTransient should be used instead"));
+				Flags |= CPF_NonPIEDuplicateTransient;
+			}
+			else if (Specifier == TEXT("NonPIEDuplicateTransient"))
+			{
+				Flags |= CPF_NonPIEDuplicateTransient;
 			}
 			else if (Specifier == TEXT("Export"))
 			{
@@ -3295,7 +3300,7 @@ bool FHeaderParser::GetVarType
 	}
 
 	// Check for invalid transients
-	uint64 Transients = VarProperty.PropertyFlags & (CPF_DuplicateTransient | CPF_TextExportTransient | CPF_NonPIETransient);
+	uint64 Transients = VarProperty.PropertyFlags & (CPF_DuplicateTransient | CPF_TextExportTransient | CPF_NonPIEDuplicateTransient);
 	if (Transients && !Cast<UClass>(Scope))
 	{
 		TArray<const TCHAR*> FlagStrs = ParsePropertyFlags(Transients);
diff --git a/Engine/Source/Runtime/CoreUObject/Private/UObject/Obj.cpp b/Engine/Source/Runtime/CoreUObject/Private/UObject/Obj.cpp
--- a/Engine/Source/Runtime/CoreUObject/Private/UObject/Obj.cpp
+++ b/Engine/Source/Runtime/CoreUObject/Private/UObject/Obj.cpp
@@ -2338,7 +2338,8 @@ TArray<const TCHAR*> ParsePropertyFlags(uint64 Flags)
 		TEXT("CPF_BlueprintCallable"),
 		TEXT("CPF_BlueprintAuthorityOnly"),
 		TEXT("CPF_TextExportTransient"),
-		TEXT("CPF_NonPIETransient"),
+		TEXT("CPF_NonPIEDuplicateTransient"),
+		TEXT("CPF_ExposeOnSpawn"),
 	};
 
 	for (const TCHAR* FlagName : PropertyFlags)
diff --git a/Engine/Source/Runtime/CoreUObject/Private/UObject/Property.cpp b/Engine/Source/Runtime/CoreUObject/Private/UObject/Property.cpp
--- a/Engine/Source/Runtime/CoreUObject/Private/UObject/Property.cpp
+++ b/Engine/Source/Runtime/CoreUObject/Private/UObject/Property.cpp
@@ -600,14 +600,14 @@ bool UProperty::ShouldSerializeValue( FArchive& Ar ) const
 	if (Ar.IsSaveGame() && !(PropertyFlags & CPF_SaveGame))
 		return false;
 
-	static uint64 SkipFlags = CPF_Transient | CPF_DuplicateTransient | CPF_NonPIETransient | CPF_NonTransactional | CPF_Deprecated | CPF_DevelopmentAssets;
+	static uint64 SkipFlags = CPF_Transient | CPF_DuplicateTransient | CPF_NonPIEDuplicateTransient | CPF_NonTransactional | CPF_Deprecated | CPF_DevelopmentAssets;
 	if (!(PropertyFlags & SkipFlags))
 		return true;
 
 	bool Skip =
 			((PropertyFlags & CPF_Transient) && Ar.IsPersistent() && !Ar.IsSerializingDefaults())
 		||	((PropertyFlags & CPF_DuplicateTransient) && (Ar.GetPortFlags() & PPF_Duplicate))
-		||	((PropertyFlags & CPF_NonPIETransient) && !(Ar.GetPortFlags() & PPF_DuplicateForPIE))
+		||	((PropertyFlags & CPF_NonPIEDuplicateTransient) && !(Ar.GetPortFlags() & PPF_DuplicateForPIE) && (Ar.GetPortFlags() & PPF_Duplicate))
 		||  (Ar.IsFilterEditorOnly() && IsEditorOnlyProperty())
 		||	((PropertyFlags & CPF_NonTransactional) && Ar.IsTransacting())
 		||	((PropertyFlags & CPF_Deprecated) && (Ar.IsSaving() || Ar.IsTransacting() || Ar.WantBinaryPropertySerialization()));
@@ -651,7 +651,7 @@ bool UProperty::ShouldPort( uint32 PortFlags/*=0*/ ) const
 		return false;
 
 	// if we're not copying for PIE and NonPIETransient is set, don't export
-	if (!(PortFlags & PPF_DuplicateForPIE) && HasAnyPropertyFlags(CPF_NonPIETransient))
+	if (!(PortFlags & PPF_DuplicateForPIE) && HasAnyPropertyFlags(CPF_NonPIEDuplicateTransient))
 		return false;
 
 	// if we're only supposed to export components and this isn't a component property, don't export
diff --git a/Engine/Source/Runtime/CoreUObject/Private/UObject/UObjectGlobals.cpp b/Engine/Source/Runtime/CoreUObject/Private/UObject/UObjectGlobals.cpp
--- a/Engine/Source/Runtime/CoreUObject/Private/UObject/UObjectGlobals.cpp
+++ b/Engine/Source/Runtime/CoreUObject/Private/UObject/UObjectGlobals.cpp
@@ -2116,7 +2116,7 @@ void FPostConstructInitializeProperties::InitProperties(UObject* Obj, UClass* De
 				bNeedInitialize = InitNonNativeProperty(P, Obj);
 			}
 
-			if (bCopyTransientsFromClassDefaults && P->HasAnyPropertyFlags(CPF_Transient|CPF_DuplicateTransient|CPF_NonPIETransient))
+			if (bCopyTransientsFromClassDefaults && P->HasAnyPropertyFlags(CPF_Transient|CPF_DuplicateTransient|CPF_NonPIEDuplicateTransient))
 			{
 				// This is a duplicate. The value for all transient or non-duplicatable properties should be copied
 				// from the source class's defaults.
diff --git a/Engine/Source/Runtime/CoreUObject/Public/UObject/ObjectBase.h b/Engine/Source/Runtime/CoreUObject/Public/UObject/ObjectBase.h
--- a/Engine/Source/Runtime/CoreUObject/Public/UObject/ObjectBase.h
+++ b/Engine/Source/Runtime/CoreUObject/Public/UObject/ObjectBase.h
@@ -326,9 +326,13 @@ typedef uint64 EClassCastFlags;
 #define CPF_BlueprintCallable				DECLARE_UINT64(0x0000100000000000)		// MC Delegates only.  Property should be exposed for calling in blueprint code
 #define CPF_BlueprintAuthorityOnly			DECLARE_UINT64(0x0000200000000000)		// MC Delegates only.  This delegate accepts (only in blueprint) only events with BlueprintAuthorityOnly.
 #define CPF_TextExportTransient				DECLARE_UINT64(0x0000400000000000)		// Property shouldn't be exported to text format (e.g. copy/paste)
-#define CPF_NonPIETransient					DECLARE_UINT64(0x0000800000000000)		// Property should only be copied in PIE
+#define CPF_NonPIEDuplicateTransient		DECLARE_UINT64(0x0000800000000000)		// Property should only be copied in PIE
 #define CPF_ExposeOnSpawn					DECLARE_UINT64(0x0001000000000000)		// Property is exposed on spawn
 
+#define CPF_NonPIETransient \
+	EMIT_DEPRECATED_WARNING_MESSAGE("CPF_NonPIETransient is deprecated. Please use CPF_NonPIEDuplicateTransient instead.") \
+	CPF_NonPIEDuplicateTransient
+
 /** @name Combinations flags */
 //@{
 #define CPF_ParmFlags				(CPF_Parm | CPF_OutParm | CPF_ReturnParm | CPF_ReferenceParm | CPF_ConstParm)
@@ -725,9 +729,12 @@ namespace UP
 		/// Property should always be reset to the default value during any type of duplication (copy/paste, binary duplication, etc.)
 		DuplicateTransient,
 
-		/// Property should always be reset to the default value unless it's being duplicated for a PIE session
+		/// Property should always be reset to the default value unless it's being duplicated for a PIE session - deprecated, use NonPIEDuplicateTransient instead
 		NonPIETransient,
 
+		/// Property should always be reset to the default value unless it's being duplicated for a PIE session
+		NonPIEDuplicateTransient,
+
 		/// Value is copied out after function call. Only valid on function param declaration.
 		Ref,
 
