Core - Minor code formatting

diff --git a/Engine/Source/Runtime/CoreUObject/Private/UObject/Obj.cpp b/Engine/Source/Runtime/CoreUObject/Private/UObject/Obj.cpp
--- a/Engine/Source/Runtime/CoreUObject/Private/UObject/Obj.cpp
+++ b/Engine/Source/Runtime/CoreUObject/Private/UObject/Obj.cpp
@@ -1164,9 +1164,9 @@ bool UObject::IsAsset () const
 bool UObject::IsSafeForRootSet() const
 {
 	if (IsInBlueprint())
- 	{
- 		return false;
- 	}
+	{
+		return false;
+	}
 
 	const ULinkerLoad* LinkerLoad = Cast<const ULinkerLoad>(this);
 	// Exclude linkers from root set if we're using seekfree loading		
@@ -1813,7 +1813,9 @@ void StaticShutdownAfterError()
 	{
 		static bool Shutdown=0;
 		if( Shutdown )
+		{
 			return;
+		}
 		Shutdown = 1;
 		UE_LOG(LogExit, Log, TEXT("Executing StaticShutdownAfterError") );
 		for ( FRawObjectIterator It; It; ++It )
