Fixed case insensitive source string comparison when checking if multuiple UTextProperty are identical.

diff --git a/Engine/Source/Runtime/CoreUObject/Private/UObject/UTextProperty.cpp b/Engine/Source/Runtime/CoreUObject/Private/UObject/UTextProperty.cpp
--- a/Engine/Source/Runtime/CoreUObject/Private/UObject/UTextProperty.cpp
+++ b/Engine/Source/Runtime/CoreUObject/Private/UObject/UTextProperty.cpp
@@ -26,7 +26,7 @@ bool UTextProperty::Identical( const void* A, const void* B, uint32 PortFlags )
 
 		if (GIsEditor)
 		{
-			return *FTextInspector::GetSourceString(ValueA) == *FTextInspector::GetSourceString(ValueB);
+			return FTextInspector::GetSourceString(ValueA)->Compare(*FTextInspector::GetSourceString(ValueB), ESearchCase::CaseSensitive) == 0;
 		}
 		else
 		{
