Misc refactoring.

diff --git a/Engine/Source/Runtime/CoreUObject/Private/Misc/PackageName.cpp b/Engine/Source/Runtime/CoreUObject/Private/Misc/PackageName.cpp
--- a/Engine/Source/Runtime/CoreUObject/Private/Misc/PackageName.cpp
+++ b/Engine/Source/Runtime/CoreUObject/Private/Misc/PackageName.cpp
@@ -23,12 +23,7 @@ namespace PackageNameConstants
 
 bool FPackageName::IsShortPackageName(const FString& PossiblyLongName)
 {
-	int32 IndexOfLastSlash = PossiblyLongName.Find(TEXT("/"), ESearchCase::CaseSensitive, ESearchDir::FromEnd);
-	if (IndexOfLastSlash != INDEX_NONE)
-	{
-		return false;
-	}
-	return true;
+	return !PossiblyLongName.Contains(TEXT("/"), ESearchCase::CaseSensitive, ESearchDir::FromEnd);
 }
 
 bool FPackageName::IsShortPackageName(const FName PossiblyLongName)
@@ -158,50 +153,50 @@ private:
 	FLongPackagePathsSingleton()
 	{
 		EngineRootPath = TEXT("/Engine/");
-		GameRootPath = TEXT("/Game/");
+		GameRootPath   = TEXT("/Game/");
 		ScriptRootPath = TEXT("/Script/");
-		TempRootPath = TEXT("/Temp/");
+		TempRootPath   = TEXT("/Temp/");
 
-		EngineContentPath = FPaths::EngineContentDir();
-		ContentPathShort = TEXT("../../Content/");
-		EngineShadersPath = FPaths::EngineDir() / TEXT("Shaders/");
+		EngineContentPath      = FPaths::EngineContentDir();
+		ContentPathShort       = TEXT("../../Content/");
+		EngineShadersPath      = FPaths::EngineDir() / TEXT("Shaders/");
 		EngineShadersPathShort = TEXT("../../Shaders/");
-		GameContentPath = FPaths::GameContentDir();
-		GameScriptPath = FPaths::GameDir() / TEXT("Script/");
-		GameSavedPath = FPaths::GameSavedDir();
+		GameContentPath        = FPaths::GameContentDir();
+		GameScriptPath         = FPaths::GameDir() / TEXT("Script/");
+		GameSavedPath          = FPaths::GameSavedDir();
 
 		FString RebasedGameDir = FString::Printf(TEXT("../../../%s/"), GGameName);
 
         GameContentPathRebased = RebasedGameDir / TEXT("Content/");
-        GameScriptPathRebased = RebasedGameDir / TEXT("Script/");
-        GameSavedPathRebased = RebasedGameDir / TEXT("Saved/");
+        GameScriptPathRebased  = RebasedGameDir / TEXT("Script/");
+        GameSavedPathRebased   = RebasedGameDir / TEXT("Saved/");
 		
 		ContentPathToRoot.Empty(10);
-		new (ContentPathToRoot) FPathPair(EngineRootPath, EngineContentPath);
+		ContentPathToRoot.Emplace(EngineRootPath, EngineContentPath);
 #if IS_MONOLITHIC
-		new (ContentPathToRoot) FPathPair(GameRootPath, ContentPathShort);
+		ContentPathToRoot.Emplace(GameRootPath,   ContentPathShort);
 #else
-		new (ContentPathToRoot) FPathPair(EngineRootPath, ContentPathShort);
+		ContentPathToRoot.Emplace(EngineRootPath, ContentPathShort);
 #endif
-		new (ContentPathToRoot) FPathPair(EngineRootPath, EngineShadersPath);
-		new (ContentPathToRoot) FPathPair(EngineRootPath, EngineShadersPathShort);
-		new (ContentPathToRoot) FPathPair(GameRootPath, GameContentPath);
-		new (ContentPathToRoot) FPathPair(ScriptRootPath, GameScriptPath);
-		new (ContentPathToRoot) FPathPair(TempRootPath, GameSavedPath);
-		new (ContentPathToRoot) FPathPair(GameRootPath, GameContentPathRebased);
-		new (ContentPathToRoot) FPathPair(ScriptRootPath, GameScriptPathRebased);
-		new (ContentPathToRoot) FPathPair(TempRootPath, GameSavedPathRebased);
+		ContentPathToRoot.Emplace(EngineRootPath, EngineShadersPath);
+		ContentPathToRoot.Emplace(EngineRootPath, EngineShadersPathShort);
+		ContentPathToRoot.Emplace(GameRootPath,   GameContentPath);
+		ContentPathToRoot.Emplace(ScriptRootPath, GameScriptPath);
+		ContentPathToRoot.Emplace(TempRootPath,   GameSavedPath);
+		ContentPathToRoot.Emplace(GameRootPath,   GameContentPathRebased);
+		ContentPathToRoot.Emplace(ScriptRootPath, GameScriptPathRebased);
+		ContentPathToRoot.Emplace(TempRootPath,   GameSavedPathRebased);
 
 
 		ContentRootToPath.Empty(8);
-		new (ContentRootToPath) FPathPair(EngineRootPath, EngineContentPath);
-		new (ContentRootToPath) FPathPair(EngineRootPath, EngineShadersPath);
-		new (ContentRootToPath) FPathPair(GameRootPath, GameContentPath);
-		new (ContentRootToPath) FPathPair(ScriptRootPath, GameScriptPath);
-		new (ContentRootToPath) FPathPair(TempRootPath, GameSavedPath);
-		new (ContentRootToPath) FPathPair(GameRootPath, GameContentPathRebased);
-		new (ContentRootToPath) FPathPair(ScriptRootPath, GameScriptPathRebased);
-		new (ContentRootToPath) FPathPair(TempRootPath, GameSavedPathRebased);
+		ContentRootToPath.Emplace(EngineRootPath, EngineContentPath);
+		ContentRootToPath.Emplace(EngineRootPath, EngineShadersPath);
+		ContentRootToPath.Emplace(GameRootPath,   GameContentPath);
+		ContentRootToPath.Emplace(ScriptRootPath, GameScriptPath);
+		ContentRootToPath.Emplace(TempRootPath,   GameSavedPath);
+		ContentRootToPath.Emplace(GameRootPath,   GameContentPathRebased);
+		ContentRootToPath.Emplace(ScriptRootPath, GameScriptPathRebased);
+		ContentRootToPath.Emplace(TempRootPath,   GameSavedPathRebased);
 
 
 		// Allow the plugin manager to mount new content paths by exposing access through a delegate.  PluginManager is 
@@ -213,16 +208,16 @@ private:
 FString FPackageName::InternalFilenameToLongPackageName(const FString& InFilename)
 {
 	const FLongPackagePathsSingleton& Paths = FLongPackagePathsSingleton::Get();
-	FString Filename = *InFilename.Replace(TEXT("\\"), TEXT("/"));
+	FString Filename = InFilename.Replace(TEXT("\\"), TEXT("/"));
 
 	// Convert to relative path if it's not already a long package name
 	bool bIsValidLongPackageName = false;
-	for (int32 PathIndex = 0; (PathIndex < Paths.ContentRootToPath.Num()) && !bIsValidLongPackageName; ++PathIndex)
+	for (const auto& Pair : Paths.ContentRootToPath)
 	{
-		const FString& TestRoot = Paths.ContentRootToPath[PathIndex].RootPath;
-		if (Filename.StartsWith(TestRoot))
+		if (Filename.StartsWith(Pair.RootPath))
 		{
 			bIsValidLongPackageName = true;
+			break;
 		}
 	}
 
@@ -231,16 +226,13 @@ FString FPackageName::InternalFilenameToLongPackageName(const FString& InFilenam
 		Filename = IFileManager::Get().ConvertToRelativePath(*Filename);
 	}
 
-	FString PackageName = FPaths::GetBaseFilename(Filename);
-	int32 PackageNameStartsAt = Filename.Len() - FPaths::GetCleanFilename(Filename).Len();
-	FString NoExt = Filename.Mid(0, PackageNameStartsAt + PackageName.Len());
-
-	FString Result = NoExt;
+	FString PackageName         = FPaths::GetBaseFilename(Filename);
+	int32   PackageNameStartsAt = Filename.Len() - FPaths::GetCleanFilename(Filename).Len();
+	FString Result              = Filename.Mid(0, PackageNameStartsAt + PackageName.Len());
 	Result.ReplaceInline(TEXT("\\"), TEXT("/"));
 
-	for (int32 PathIndex = 0; PathIndex < Paths.ContentPathToRoot.Num(); ++PathIndex)
+	for (const auto& Pair : Paths.ContentPathToRoot)
 	{
-		const FPathPair& Pair = Paths.ContentPathToRoot[PathIndex];
 		if (Result.StartsWith(Pair.ContentPath))
 		{
 			Result = Pair.RootPath + Result.Mid(Pair.ContentPath.Len());
@@ -253,32 +245,26 @@ FString FPackageName::InternalFilenameToLongPackageName(const FString& InFilenam
 
 bool FPackageName::TryConvertFilenameToLongPackageName(const FString& InFilename, FString& OutPackageName)
 {
-	const FString& Result = InternalFilenameToLongPackageName(InFilename);
+	FString Result = InternalFilenameToLongPackageName(InFilename);
 
 	// we don't support loading packages from outside of well defined places
-	const bool bContainsDot = Result.Contains(TEXT("."));
-	const bool bContainsBackslash = Result.Contains(TEXT("\\"));
-	const bool bContainsColon = Result.Contains(TEXT(":"));
-	
-	if (!(bContainsDot || bContainsBackslash || bContainsColon))
-	{
-		OutPackageName = Result;
-		return true;
-	}
-	else
+	if (Result.Contains(TEXT(".")) || Result.Contains(TEXT("\\")) || Result.Contains(TEXT(":")))
 	{
 		return false;
 	}
+
+	OutPackageName = MoveTemp(Result);
+	return true;
 }
 
 FString FPackageName::FilenameToLongPackageName(const FString& InFilename)
 {
-	const FString& Result = InternalFilenameToLongPackageName(InFilename);
+	FString Result = InternalFilenameToLongPackageName(InFilename);
 
 	// we don't support loading packages from outside of well defined places
-	const bool bContainsDot = Result.Contains(TEXT("."));
+	const bool bContainsDot       = Result.Contains(TEXT("."));
 	const bool bContainsBackslash = Result.Contains(TEXT("\\"));
-	const bool bContainsColon = Result.Contains(TEXT(":"));
+	const bool bContainsColon     = Result.Contains(TEXT(":"));
 
 	if (bContainsDot || bContainsBackslash || bContainsColon)
 	{
@@ -304,17 +290,13 @@ FString FPackageName::FilenameToLongPackageName(const FString& InFilename)
 
 FString FPackageName::LongPackageNameToFilename(const FString& InLongPackageName, const FString& InExtension)
 {
-	FString Result = InLongPackageName;
-	const FLongPackagePathsSingleton& Paths = FLongPackagePathsSingleton::Get();
+	const auto& Paths = FLongPackagePathsSingleton::Get();
 
-	for (int32 PathIndex = 0; PathIndex < Paths.ContentRootToPath.Num(); ++PathIndex)
+	for (const auto& Pair : Paths.ContentRootToPath)
 	{
-		const FPathPair& Pair = Paths.ContentRootToPath[PathIndex];
 		if (InLongPackageName.StartsWith(Pair.RootPath))
 		{
-			Result = Pair.ContentPath + Result.Mid(Pair.RootPath.Len());
-			Result += InExtension;
-
+			FString Result = Pair.ContentPath + InLongPackageName.Mid(Pair.RootPath.Len()) + InExtension;
 			return Result;
 		}
 	}
@@ -322,7 +304,7 @@ FString FPackageName::LongPackageNameToFilename(const FString& InLongPackageName
 	// This is not a long package name or the root folder is not handled in the above cases
 	UE_LOG(LogPackageName, Fatal,TEXT("LongPackageNameToFilename failed to convert '%s'. Path does not map to any roots."), *InLongPackageName);
 
-	return Result;
+	return InLongPackageName;
 }
 
 FString FPackageName::GetLongPackagePath(const FString& InLongPackageName)
@@ -545,7 +527,9 @@ FName* FPackageName::FindScriptPackageName(FName InShortName)
 }
 
 bool FPackageName::FindPackageFileWithoutExtension(const FString& InPackageFilename, FString& OutFilename)
-{	
+{
+	auto& FileManager = IFileManager::Get();
+
 	static const FString* PackageExtensions[] =
 	{
 		&AssetPackageExtension,
@@ -553,20 +537,20 @@ bool FPackageName::FindPackageFileWithoutExtension(const FString& InPackageFilen
 	};
 
 	// Loop through all known extensions and check if the file exist.
-	bool bFound = false;
-	for (int32 ExtensionIndex = 0; ExtensionIndex < ARRAY_COUNT(PackageExtensions) && !bFound; ++ExtensionIndex)
+	for (auto Extension : PackageExtensions)
 	{
-		const FString PackageFilename = InPackageFilename + *PackageExtensions[ExtensionIndex];
-		const FDateTime Timestamp = IFileManager::Get().GetTimeStamp(*PackageFilename);
+		FString   PackageFilename = InPackageFilename + *Extension;
+		FDateTime Timestamp       = FileManager.GetTimeStamp(*PackageFilename);
 		if (Timestamp != FDateTime::MinValue())
 		{
 			// The package exists so exit. From now on InPackageFilename can be equal to OutFilename so
 			// don't attempt to use it anymore (case where &InPackageFilename == &OutFilename).
-			OutFilename = PackageFilename;
-			bFound = true;
+			OutFilename = MoveTemp(PackageFilename);
+			return true;
 		}
 	}
-	return bFound;
+
+	return false;
 }
 
 bool FPackageName::DoesPackageExist(const FString& LongPackageName, const FGuid* Guid /*= NULL*/, FString* OutFilename /*= NULL*/)
diff --git a/Engine/Source/Runtime/CoreUObject/Private/UObject/Linker.cpp b/Engine/Source/Runtime/CoreUObject/Private/UObject/Linker.cpp
--- a/Engine/Source/Runtime/CoreUObject/Private/UObject/Linker.cpp
+++ b/Engine/Source/Runtime/CoreUObject/Private/UObject/Linker.cpp
@@ -574,153 +574,146 @@ static void LogGetPackageLinkerError(const TCHAR* InFilename, const FText& InFul
 //
 ULinkerLoad* GetPackageLinker
 (
- UPackage*		InOuter,
- const TCHAR*	InLongPackageName,
- uint32			LoadFlags,
- UPackageMap*	Sandbox,
- FGuid*			CompatibleGuid
- )
+	UPackage*		InOuter,
+	const TCHAR*	InLongPackageName,
+	uint32			LoadFlags,
+	UPackageMap*	Sandbox,
+	FGuid*			CompatibleGuid
+)
 {
 	// See if there is already a linker for this package.
 	ULinkerLoad* Result = ULinkerLoad::FindExistingLinkerForPackage(InOuter);
 
 	// Try to load the linker.
 	// See if the linker is already loaded.
-	FString NewFilename;
 	if( Result )
 	{
-		// Linker already found.
-		NewFilename = TEXT("");
+		return Result;
 	}
-	else
+
+	FString NewFilename;
+	if( !InLongPackageName )
 	{
-		if( !InLongPackageName )
+		// Resolve filename from package name.
+		if( !InOuter )
 		{
-			// Resolve filename from package name.
-			if( !InOuter )
-			{
-				// try to recover from this instead of throwing, it seems recoverable just by doing this
-				FText ErrorText(LOCTEXT("PackageResolveFailed", "Can't resolve asset name"));
-				LogGetPackageLinkerError(InLongPackageName, ErrorText, ErrorText, InOuter, LoadFlags);
-				return NULL;
-			}
+			// try to recover from this instead of throwing, it seems recoverable just by doing this
+			FText ErrorText(LOCTEXT("PackageResolveFailed", "Can't resolve asset name"));
+			LogGetPackageLinkerError(InLongPackageName, ErrorText, ErrorText, InOuter, LoadFlags);
+			return nullptr;
+		}
 
-			if( !FPackageName::DoesPackageExist(InOuter->GetName(), CompatibleGuid, &NewFilename) )
+		if( !FPackageName::DoesPackageExist(InOuter->GetName(), CompatibleGuid, &NewFilename) )
+		{
+			// Compiled in packages have no linker and this is ok
+			if (!(LoadFlags & LOAD_AllowDll) && !(InOuter->PackageFlags & PKG_CompiledIn))
 			{
-				if( (LoadFlags & LOAD_AllowDll) && InOuter->IsA(UPackage::StaticClass()) )
-				{
-					return NULL;
-				}
-				if( InOuter->IsA(UPackage::StaticClass()) && (((UPackage*)InOuter)->PackageFlags & PKG_CompiledIn) )
-				{
-					// this is a compiled in package and so it has no linker and this is ok
-					return NULL;
-				}
-
 				FFormatNamedArguments Arguments;
 				Arguments.Add(TEXT("AssetName"), FText::FromString(InOuter->GetName()));
 				Arguments.Add(TEXT("PackageName"), FText::FromString(GSerializedPackageLinker ? *(GSerializedPackageLinker->Filename) : TEXT("NULL")));
-				LogGetPackageLinkerError(GSerializedPackageLinker ? *GSerializedPackageLinker->Filename : NULL,
-										 FText::Format(LOCTEXT("PackageNotFound", "Can't find file for asset '{AssetName}' while loading {PackageName}."), Arguments),
-										 LOCTEXT("PackageNotFoundShort", "Can't find file for asset."),
-										 InOuter,
-										 LoadFlags);
-
-				return NULL;
+				LogGetPackageLinkerError(GSerializedPackageLinker ? *GSerializedPackageLinker->Filename : nullptr,
+											FText::Format(LOCTEXT("PackageNotFound", "Can't find file for asset '{AssetName}' while loading {PackageName}."), Arguments),
+											LOCTEXT("PackageNotFoundShort", "Can't find file for asset."),
+											InOuter,
+											LoadFlags);
 			}
+
+			return nullptr;
 		}
-		else
+	}
+	else
+	{
+		FString PackageName(InLongPackageName);
+		if (!FPackageName::TryConvertFilenameToLongPackageName(InLongPackageName, PackageName))
 		{
-			FString PackageName(InLongPackageName);
-			if (!FPackageName::TryConvertFilenameToLongPackageName(InLongPackageName, PackageName))
-			{
-				// try to recover from this instead of throwing, it seems recoverable just by doing this
-				FText ErrorText(LOCTEXT("PackageResolveFailed", "Can't resolve asset name"));
-				LogGetPackageLinkerError(InLongPackageName, ErrorText, ErrorText, InOuter, LoadFlags);
-				return NULL;
-			}
-			{			
-				UPackage* ExistingPackage = Cast<UPackage>(StaticFindObject(UPackage::StaticClass(), NULL, *PackageName) );
-				if (ExistingPackage && !ExistingPackage->GetOuter() && (ExistingPackage->PackageFlags & PKG_CompiledIn))
-				{
-					// this is a compiled in package and so it has no linker and this is ok
-					return NULL;
-				}
-			}
-			// Verify that the file exists.
-			if( !FPackageName::DoesPackageExist( PackageName, CompatibleGuid, &NewFilename ) )
-			{
-				FFormatNamedArguments Arguments;
-				Arguments.Add(TEXT("Filename"), FText::FromString(InLongPackageName));
+			// try to recover from this instead of throwing, it seems recoverable just by doing this
+			FText ErrorText(LOCTEXT("PackageResolveFailed", "Can't resolve asset name"));
+			LogGetPackageLinkerError(InLongPackageName, ErrorText, ErrorText, InOuter, LoadFlags);
+			return nullptr;
+		}
 
-				// try to recover from this instead of throwing, it seems recoverable just by doing this
-				LogGetPackageLinkerError(InLongPackageName, FText::Format(LOCTEXT("FileNotFound", "Can't find file '{Filename}'"), Arguments), LOCTEXT("FileNotFoundShort", "Can't find file"), InOuter, LoadFlags);
-				return NULL;
-			}
-				
-			// Create the package with the provided long package name.
-			UPackage* FilenamePkg = CreatePackage( NULL, *PackageName );
-			if (LoadFlags & LOAD_PackageForPIE)
+		if (UPackage* ExistingPackage = FindObject<UPackage>(nullptr, *PackageName))
+		{
+			if (!ExistingPackage->GetOuter() && (ExistingPackage->PackageFlags & PKG_CompiledIn))
 			{
-				FilenamePkg->PackageFlags |= PKG_PlayInEditor;
+				// this is a compiled in package and so it has no linker and this is ok
+				return nullptr;
 			}
+		}
 
-			// If no package specified, use package from file.
-			if( InOuter==NULL )
-			{
-				if( !FilenamePkg )
-				{
-					FFormatNamedArguments Arguments;
-					Arguments.Add(TEXT("Filename"), FText::FromString(InLongPackageName));
-					LogGetPackageLinkerError(InLongPackageName, FText::Format(LOCTEXT("FilenameToPackage", "Can't convert filename '{Filename}' to asset name"), Arguments), LOCTEXT("FilenameToPackageShort", "Can't convert filename to asset name"), InOuter, LoadFlags);
-					return NULL;
-				}
-				InOuter = FilenamePkg;
-				Result = ULinkerLoad::FindExistingLinkerForPackage(InOuter);
-			}
-			else if( InOuter != FilenamePkg )//!!should be tested and validated in new UnrealEd
+		// Verify that the file exists.
+		if( !FPackageName::DoesPackageExist( PackageName, CompatibleGuid, &NewFilename ) )
+		{
+			FFormatNamedArguments Arguments;
+			Arguments.Add(TEXT("Filename"), FText::FromString(InLongPackageName));
+
+			// try to recover from this instead of throwing, it seems recoverable just by doing this
+			LogGetPackageLinkerError(InLongPackageName, FText::Format(LOCTEXT("FileNotFound", "Can't find file '{Filename}'"), Arguments), LOCTEXT("FileNotFoundShort", "Can't find file"), InOuter, LoadFlags);
+			return nullptr;
+		}
+
+		// Create the package with the provided long package name.
+		UPackage* FilenamePkg = CreatePackage(nullptr, *PackageName);
+		if (LoadFlags & LOAD_PackageForPIE)
+		{
+			FilenamePkg->PackageFlags |= PKG_PlayInEditor;
+		}
+
+		// If no package specified, use package from file.
+		if (!InOuter)
+		{
+			if( !FilenamePkg )
 			{
-				// Loading a new file into an existing package, so reset the loader.
-				//UE_LOG(LogLinker, Log,  TEXT("New File, Existing Package (%s, %s)"), *InOuter->GetFullName(), *FilenamePkg->GetFullName() );
-				ResetLoaders( InOuter );
+				FFormatNamedArguments Arguments;
+				Arguments.Add(TEXT("Filename"), FText::FromString(InLongPackageName));
+				LogGetPackageLinkerError(InLongPackageName, FText::Format(LOCTEXT("FilenameToPackage", "Can't convert filename '{Filename}' to asset name"), Arguments), LOCTEXT("FilenameToPackageShort", "Can't convert filename to asset name"), InOuter, LoadFlags);
+				return nullptr;
 			}
+			InOuter = FilenamePkg;
+			Result = ULinkerLoad::FindExistingLinkerForPackage(InOuter);
+		}
+		else if (InOuter != FilenamePkg) //!!should be tested and validated in new UnrealEd
+		{
+			// Loading a new file into an existing package, so reset the loader.
+			//UE_LOG(LogLinker, Log,  TEXT("New File, Existing Package (%s, %s)"), *InOuter->GetFullName(), *FilenamePkg->GetFullName() );
+			ResetLoaders( InOuter );
 		}
+	}
 
 #if 0
-		// Make sure the package is accessible in the sandbox.
-		if( Sandbox && !Sandbox->SupportsPackage(InOuter) )
-		{
-			FFormatNamedArguments Arguments;
-			Arguments.Add(TEXT("AssetName"), FText::FromString(InOuter->GetName()));
+	// Make sure the package is accessible in the sandbox.
+	if( Sandbox && !Sandbox->SupportsPackage(InOuter) )
+	{
+		FFormatNamedArguments Arguments;
+		Arguments.Add(TEXT("AssetName"), FText::FromString(InOuter->GetName()));
 
-			LogGetPackageLinkerError(InLongPackageName, FText::Format(LOCTEXT("Sandbox", "Asset '{AssetName}' is not accessible in this sandbox"), Arguments), LOCTEXT("SandboxShort", "Asset is not accessible in this sandbox"), InOuter, LoadFlags);
-			return NULL;
-		}
+		LogGetPackageLinkerError(InLongPackageName, FText::Format(LOCTEXT("Sandbox", "Asset '{AssetName}' is not accessible in this sandbox"), Arguments), LOCTEXT("SandboxShort", "Asset is not accessible in this sandbox"), InOuter, LoadFlags);
+		return nullptr;
+	}
 #endif
 
-		// Create new linker.
-		if( !Result )
-		{
-			check(IsLoading());
+	// Create new linker.
+	if( !Result )
+	{
+		check(IsLoading());
 
-			// we will already have found the filename above
-			check(NewFilename.Len() > 0);
+		// we will already have found the filename above
+		check(NewFilename.Len() > 0);
 
-			Result = ULinkerLoad::CreateLinker( InOuter, *NewFilename, LoadFlags );
-		}
+		Result = ULinkerLoad::CreateLinker( InOuter, *NewFilename, LoadFlags );
+	}
 
-		// Verify compatibility.
-		if( CompatibleGuid && Result->Summary.Guid!=*CompatibleGuid )
-		{
-			FFormatNamedArguments Arguments;
-			Arguments.Add(TEXT("AssetName"), FText::FromString(InOuter->GetName()));
+	// Verify compatibility.
+	if( CompatibleGuid && Result->Summary.Guid!=*CompatibleGuid )
+	{
+		FFormatNamedArguments Arguments;
+		Arguments.Add(TEXT("AssetName"), FText::FromString(InOuter->GetName()));
 
-			// This should never fire, because FindPackageFile should never return an incompatible file
-			LogGetPackageLinkerError(InLongPackageName, FText::Format(LOCTEXT("PackageVersion", "Asset '{AssetName}' version mismatch"), Arguments), LOCTEXT("PackageVersionShort", "Asset version mismatch"), InOuter, LoadFlags);
-			return NULL;
-		}
+		// This should never fire, because FindPackageFile should never return an incompatible file
+		LogGetPackageLinkerError(InLongPackageName, FText::Format(LOCTEXT("PackageVersion", "Asset '{AssetName}' version mismatch"), Arguments), LOCTEXT("PackageVersionShort", "Asset version mismatch"), InOuter, LoadFlags);
+		return nullptr;
 	}
-	// Success.
+
 	return Result;
 }
 
diff --git a/Engine/Source/Runtime/CoreUObject/Private/UObject/UObjectGlobals.cpp b/Engine/Source/Runtime/CoreUObject/Private/UObject/UObjectGlobals.cpp
--- a/Engine/Source/Runtime/CoreUObject/Private/UObject/UObjectGlobals.cpp
+++ b/Engine/Source/Runtime/CoreUObject/Private/UObject/UObjectGlobals.cpp
@@ -465,33 +465,33 @@ UPackage* CreatePackage( UObject* InOuter, const TCHAR* PackageName )
 //
 // Resolve a package and name.
 //
-bool ResolveName( UObject*& InPackage, FString& InName, bool Create, bool Throw )
+bool ResolveName( UObject*& InPackage, FString& InOutName, bool Create, bool Throw )
 {
 	FString* IniFilename = NULL;
 
 	// See if the name is specified in the .ini file.
-	if( FCString::Strnicmp( *InName, TEXT("engine-ini:"), FCString::Strlen(TEXT("engine-ini:")) )==0 )
+	if( FCString::Strnicmp( *InOutName, TEXT("engine-ini:"), FCString::Strlen(TEXT("engine-ini:")) )==0 )
 	{
 		IniFilename = &GEngineIni;
 	}
-	else if( FCString::Strnicmp( *InName, TEXT("game-ini:"), FCString::Strlen(TEXT("game-ini:")) )==0 )
+	else if( FCString::Strnicmp( *InOutName, TEXT("game-ini:"), FCString::Strlen(TEXT("game-ini:")) )==0 )
 	{
 		IniFilename = &GGameIni;
 	}
-	else if( FCString::Strnicmp( *InName, TEXT("input-ini:"), FCString::Strlen(TEXT("input-ini:")) )==0 )
+	else if( FCString::Strnicmp( *InOutName, TEXT("input-ini:"), FCString::Strlen(TEXT("input-ini:")) )==0 )
 	{
 		IniFilename = &GInputIni;
 	}
-	else if( FCString::Strnicmp( *InName, TEXT("editor-ini:"), FCString::Strlen(TEXT("editor-ini:")) )==0 )
+	else if( FCString::Strnicmp( *InOutName, TEXT("editor-ini:"), FCString::Strlen(TEXT("editor-ini:")) )==0 )
 	{
 		IniFilename = &GEditorIni;
 	}
 
 
-	if( IniFilename && InName.Contains(TEXT(".")) )
+	if( IniFilename && InOutName.Contains(TEXT(".")) )
 	{
 		// Get .ini key and section.
-		FString Section = InName.Mid(1+InName.Find(TEXT(":"), ESearchCase::CaseSensitive));
+		FString Section = InOutName.Mid(1+InOutName.Find(TEXT(":"), ESearchCase::CaseSensitive));
 		int32 i = Section.Find(TEXT("."), ESearchCase::CaseSensitive, ESearchDir::FromEnd);
 		FString Key;
 		if( i != -1)
@@ -506,15 +506,15 @@ bool ResolveName( UObject*& InPackage, FString& InName, bool Create, bool Throw
 		{
 			if( Throw == true )
 			{
-				UE_LOG(LogUObjectGlobals, Error, TEXT( " %s %s " ), *FString::Printf( TEXT("Can't find '%s' in configuration file section=%s key=%s"), *InName, *Section, *Key), **IniFilename );
+				UE_LOG(LogUObjectGlobals, Error, TEXT( " %s %s " ), *FString::Printf( TEXT("Can't find '%s' in configuration file section=%s key=%s"), *InOutName, *Section, *Key), **IniFilename );
 			}
 			return false;
 		}
-		InName = Result;
+		InOutName = Result;
 	}
 
 	// Strip off the object class.
-	ConstructorHelpers::StripObjectClass( InName );
+	ConstructorHelpers::StripObjectClass( InOutName );
 
 	// Handle specified packages.
 	int32 i;
@@ -529,58 +529,56 @@ bool ResolveName( UObject*& InPackage, FString& InName, bool Create, bool Throw
 	bool bSubobjectPath = false;
 
 	// to make parsing the name easier, replace the subobject delimiter with an extra dot
-	InName.ReplaceInline(SUBOBJECT_DELIMITER, TEXT(".."));
-	while( (i=InName.Find(TEXT("."))) != -1 )
+	InOutName.ReplaceInline(SUBOBJECT_DELIMITER, TEXT(".."));
+	while( (i=InOutName.Find(TEXT("."))) != -1 )
 	{
-		FString PartialName = InName.Left(i);
+		FString PartialName = InOutName.Left(i);
 
-		// if the next part of InName ends in two dots, it indicates that the next object in the path name
+		// if the next part of InOutName ends in two dots, it indicates that the next object in the path name
 		// is not a top-level object (i.e. it's a subobject).  e.g. SomePackage.SomeGroup.SomeObject..Subobject
-		if ( InName.Mid(i+1,1) == TEXT(".") )
+		if ( InOutName.Mid(i+1,1) == TEXT(".") )
 		{
-			InName = PartialName + InName.Mid(i+1);
+			InOutName      = PartialName + InOutName.Mid(i+1);
 			bSubobjectPath = true;
-			Create = false;
+			Create         = false;
 		}
 
 		// In case this is a short script package name, convert to long name before passing to CreatePackage/FindObject.
 		FName* ScriptPackageName = FPackageName::FindScriptPackageName(*PartialName);
-		if (ScriptPackageName != NULL)
+		if (ScriptPackageName)
 		{
 			PartialName = ScriptPackageName->ToString();
 		}
 		// Only long package names are allowed so don't even attempt to create one because whatever the name represents
 		// it's not a valid package name anyway.
-		if (!Create || FPackageName::IsShortPackageName(PartialName) == false)
+		
+		if (!Create)
 		{
-			if( Create )
+			UObject* NewPackage = FindObject<UPackage>( InPackage, *PartialName );
+			if( !NewPackage )
 			{
-				if (!ScriptPackageName)
-				{
-					InPackage = LoadPackage(Cast<UPackage>(InPackage), *PartialName, 0);
-				}
-				if (!InPackage)
+				NewPackage = FindObject<UObject>( InPackage == NULL ? ANY_PACKAGE : InPackage, *PartialName );
+				if( !NewPackage )
 				{
-					InPackage = CreatePackage(InPackage, *PartialName);
+					return bSubobjectPath;
 				}
-
-				check(InPackage);
 			}
-			else
+			InPackage = NewPackage;
+		}
+		else if (!FPackageName::IsShortPackageName(PartialName))
+		{
+			if (!ScriptPackageName)
 			{
-				UObject* NewPackage = FindObject<UPackage>( InPackage, *PartialName );
-				if( !NewPackage )
-				{
-					NewPackage = FindObject<UObject>( InPackage == NULL ? ANY_PACKAGE : InPackage, *PartialName );
-					if( !NewPackage )
-					{
-						return bSubobjectPath;
-					}
-				}
-				InPackage = NewPackage;
+				InPackage = LoadPackage(Cast<UPackage>(InPackage), *PartialName, 0);
+			}
+			if (!InPackage)
+			{
+				InPackage = CreatePackage(InPackage, *PartialName);
 			}
+
+			check(InPackage);
 		}
-		InName = InName.Mid(i+1);
+		InOutName = InOutName.Mid(i+1);
 	}
 
 	return true;
@@ -654,12 +652,11 @@ UObject* StaticLoadObject(UClass* ObjectClass, UObject* InOuter, const TCHAR* In
 	check(ObjectClass);
 	check(InName);
 
-	FString		StrName				= InName;
-	UObject*	Result				=NULL;
+	FString StrName = InName;
 
 	// break up the name into packages, returning the innermost name and its outer
 	ResolveName(InOuter, StrName, true, true);
-	if ( InOuter != NULL )
+	if (InOuter)
 	{
 		if( bAllowObjectReconciliation && ((FApp::IsGame() && !GIsEditor && !IsRunningCommandlet()) 
 #if WITH_EDITOR
@@ -667,52 +664,51 @@ UObject* StaticLoadObject(UClass* ObjectClass, UObject* InOuter, const TCHAR* In
 #endif
 			))
 		{
-			Result = StaticFindObjectFast(ObjectClass, InOuter, *StrName);
+			if (UObject* Result = StaticFindObjectFast(ObjectClass, InOuter, *StrName))
+			{
+				return Result;
+			}
 		}
 
-		if( !Result )
+		if( (FPlatformProperties::RequiresCookedData() || FPlatformProperties::IsServerOnly()) && GUseSeekFreeLoading )
 		{
-			if( (FPlatformProperties::RequiresCookedData() || FPlatformProperties::IsServerOnly()) && GUseSeekFreeLoading )
+			if ( (LoadFlags&LOAD_NoWarn) == 0 )
 			{
-				if ( (LoadFlags&LOAD_NoWarn) == 0 )
-				{
-					UE_LOG(LogUObjectGlobals, Warning ,TEXT("StaticLoadObject for %s %s %s couldn't find object in memory!"),
-						*ObjectClass->GetName(),
-						*InOuter->GetName(),
-						*StrName);
-				}
+				UE_LOG(LogUObjectGlobals, Warning ,TEXT("StaticLoadObject for %s %s %s couldn't find object in memory!"),
+					*ObjectClass->GetName(),
+					*InOuter->GetName(),
+					*StrName);
 			}
-			else
+		}
+		else
+		{
+			// now that we have one asset per package, we load the entire package whenever a single object is requested
+			LoadPackage(NULL, *InOuter->GetOutermost()->GetName(), LoadFlags & ~LOAD_Verify);
+
+			// now, find the object in the package
+			if (UObject* Result = Result = StaticFindObjectFast(ObjectClass, InOuter, *StrName))
 			{
-				// now that we have one asset per package, we load the entire package whenever a single object is requested
-				LoadPackage(NULL, *InOuter->GetOutermost()->GetName(), LoadFlags & ~LOAD_Verify);
-				// now, find the object in the package
-				Result = StaticFindObjectFast(ObjectClass, InOuter, *StrName);
+				return Result;
+			}
 
-				if ( !Result )
-				{
-					// If the object was not found, check for a redirector and follow it if the class matches
-					UObjectRedirector* Redirector = Cast<UObjectRedirector>(StaticFindObjectFast(UObjectRedirector::StaticClass(), InOuter, *StrName));
-					if ( Redirector && Redirector->DestinationObject && Redirector->DestinationObject->IsA(ObjectClass) )
-					{
-						Result = Redirector->DestinationObject;
-					}
-				}
+			// If the object was not found, check for a redirector and follow it if the class matches
+			UObjectRedirector* Redirector = Cast<UObjectRedirector>(StaticFindObjectFast(UObjectRedirector::StaticClass(), InOuter, *StrName));
+			if ( Redirector && Redirector->DestinationObject && Redirector->DestinationObject->IsA(ObjectClass) )
+			{
+				return Redirector->DestinationObject;
 			}
 		}
 	}
-	// if we haven't created or found the object, error
-	if (!Result)
-	{
-		FFormatNamedArguments Arguments;
-		Arguments.Add(TEXT("ClassName"), FText::FromString( ObjectClass->GetName() ));
-		Arguments.Add(TEXT("OuterName"), InOuter ? FText::FromString( InOuter->GetPathName() ) : NSLOCTEXT( "Core", "None", "None" ));
-		Arguments.Add(TEXT("ObjectName"), FText::FromString( StrName ));
-		const FString Error = FText::Format( NSLOCTEXT( "Core", "ObjectNotFound", "Failed to find object '{ClassName} {OuterName}.{ObjectName}'" ), Arguments ).ToString();
-		SafeLoadError(InOuter, LoadFlags, *Error, *Error);
-	}
 
-	return Result;
+	// we haven't created or found the object, error
+	FFormatNamedArguments Arguments;
+	Arguments.Add(TEXT("ClassName"), FText::FromString( ObjectClass->GetName() ));
+	Arguments.Add(TEXT("OuterName"), InOuter ? FText::FromString( InOuter->GetPathName() ) : NSLOCTEXT( "Core", "None", "None" ));
+	Arguments.Add(TEXT("ObjectName"), FText::FromString( StrName ));
+	const FString Error = FText::Format( NSLOCTEXT( "Core", "ObjectNotFound", "Failed to find object '{ClassName} {OuterName}.{ObjectName}'" ), Arguments ).ToString();
+	SafeLoadError(InOuter, LoadFlags, *Error, *Error);
+
+	return nullptr;
 }
 
 //
@@ -721,20 +717,21 @@ UObject* StaticLoadObject(UClass* ObjectClass, UObject* InOuter, const TCHAR* In
 UClass* StaticLoadClass( UClass* BaseClass, UObject* InOuter, const TCHAR* InName, const TCHAR* Filename, uint32 LoadFlags, UPackageMap* Sandbox )
 {
 	check(BaseClass);
-		UClass* Class = LoadObject<UClass>( InOuter, InName, Filename, LoadFlags, Sandbox );
-		if( Class && !Class->IsChildOf(BaseClass) )
-		{
-			FFormatNamedArguments Arguments;
-			Arguments.Add(TEXT("ClassName"), FText::FromString( Class->GetFullName() ));
-			Arguments.Add(TEXT("BaseClassName"), FText::FromString( BaseClass->GetFullName() ));
-			const FString Error = FText::Format( NSLOCTEXT( "Core", "LoadClassMismatch", "{ClassName} is not a child class of {BaseClassName}" ), Arguments ).ToString();
-			SafeLoadError(InOuter, LoadFlags,*Error, *Error);
 
-			// return NULL class due to error
-			Class = NULL;
-		}
-		return Class;
+	UClass* Class = LoadObject<UClass>( InOuter, InName, Filename, LoadFlags, Sandbox );
+	if( Class && !Class->IsChildOf(BaseClass) )
+	{
+		FFormatNamedArguments Arguments;
+		Arguments.Add(TEXT("ClassName"), FText::FromString( Class->GetFullName() ));
+		Arguments.Add(TEXT("BaseClassName"), FText::FromString( BaseClass->GetFullName() ));
+		const FString Error = FText::Format( NSLOCTEXT( "Core", "LoadClassMismatch", "{ClassName} is not a child class of {BaseClassName}" ), Arguments ).ToString();
+		SafeLoadError(InOuter, LoadFlags,*Error, *Error);
+
+		// return NULL class due to error
+		Class = NULL;
 	}
+	return Class;
+}
 
 /**
  * Loads a package and all contained objects that match context flags.
@@ -777,7 +774,7 @@ UPackage* LoadPackage( UPackage* InOuter, const TCHAR* InLongPackageName, uint32
 	}
 
 #if WITH_EDITOR
-	TGuardValue<bool> IsEditorLoadingPackage(GIsEditorLoadingPackage, (GIsEditor ? true : GIsEditorLoadingPackage));
+	TGuardValue<bool> IsEditorLoadingPackage(GIsEditorLoadingPackage, GIsEditor || GIsEditorLoadingPackage);
 #endif
 
 	// Try to load.
@@ -788,19 +785,21 @@ UPackage* LoadPackage( UPackage* InOuter, const TCHAR* InLongPackageName, uint32
 		const double StartTime = FPlatformTime::Seconds();
 
 		// Create a new linker object which goes off and tries load the file.
-		ULinkerLoad* Linker = GetPackageLinker( InOuter, *FileToLoad, LoadFlags, NULL, NULL );
-		if( !Linker )
+		ULinkerLoad* Linker = GetPackageLinker(InOuter, *FileToLoad, LoadFlags, nullptr, nullptr);
+		if (!Linker)
 		{
 			EndLoad();
-			return( NULL );
+			return nullptr;
 		}
-		else if (Linker->LinkerRoot && Linker->LinkerRoot->HasAnyFlags(RF_WasLoaded))
+
+		if (Linker->LinkerRoot && Linker->LinkerRoot->HasAnyFlags(RF_WasLoaded))
 		{
 			// The linker is associated with a package that has already been loaded.
 			// Loading packages that have already been loaded is unsupported.
 			EndLoad();
-			return Cast<UPackage>(Linker->LinkerRoot);
+			return Linker->LinkerRoot;
 		}
+
 		Result = Linker->LinkerRoot;
 
 		// If we are loading a package for diff'ing, set the package flag
