#ue4 - Optimize UObject::GetArchetypeInstances() to use GetObjectsOfClass rather than FObjectIterator with repeated casts.

diff --git a/Engine/Source/Runtime/CoreUObject/Private/UObject/Obj.cpp b/Engine/Source/Runtime/CoreUObject/Private/UObject/Obj.cpp
--- a/Engine/Source/Runtime/CoreUObject/Private/UObject/Obj.cpp
+++ b/Engine/Source/Runtime/CoreUObject/Private/UObject/Obj.cpp
@@ -377,17 +377,18 @@ void UObject::GetArchetypeInstances( TArray<UObject*>& Instances )
 
 	if ( HasAnyFlags(RF_ArchetypeObject|RF_ClassDefaultObject) )
 	{
-		// use an FObjectIterator because we need to evaluate CDOs as well.
+		// we need to evaluate CDOs as well, but nothing pending kill
+		TArray<UObject*> IterObjects;
+		GetObjectsOfClass(GetClass(), IterObjects, true, RF_PendingKill);
 
 		// if this object is the class default object, any object of the same class (or derived classes) could potentially be affected
 		if ( !HasAnyFlags(RF_ArchetypeObject) )
 		{
-			for ( FObjectIterator It; It; ++It )
+			Instances.Reserve(IterObjects.Num()-1);
+			for (auto It : IterObjects)
 			{
-				UObject* Obj = *It;
-
-				// if this object is the correct type
-				if ( Obj != this && Obj->IsA(GetClass()) && !Obj->IsPendingKill() )
+				UObject* Obj = It;
+				if ( Obj != this )
 				{
 					Instances.Add(Obj);
 				}
@@ -395,14 +396,13 @@ void UObject::GetArchetypeInstances( TArray<UObject*>& Instances )
 		}
 		else
 		{
-			// editing an archetype object - objects of child classes won't be affected
-			for ( FObjectIterator It; It; ++It )
+			for (auto It : IterObjects)
 			{
-				UObject* Obj = *It;
-
+				UObject* Obj = It;
+				
 				// if this object is the correct type and its archetype is this object, add it to the list
 				CA_SUPPRESS(6011)
-				if ( Obj != this && Obj->IsA(GetClass()) && Obj->IsBasedOnArchetype(this) && !Obj->IsPendingKill() )
+				if ( Obj != this && Obj->IsBasedOnArchetype(this) )
 				{
 					Instances.Add(Obj);
 				}
