UE4: TAssetPtr UProperties now properly report string asset references when saving packages. Also, string asset references with empty values no longer include "None" in the list of references.

diff --git a/Engine/Source/Runtime/CoreUObject/Private/UObject/PropertyAssetObject.cpp b/Engine/Source/Runtime/CoreUObject/Private/UObject/PropertyAssetObject.cpp
--- a/Engine/Source/Runtime/CoreUObject/Private/UObject/PropertyAssetObject.cpp
+++ b/Engine/Source/Runtime/CoreUObject/Private/UObject/PropertyAssetObject.cpp
@@ -34,7 +34,8 @@ void UAssetObjectProperty::SerializeItem( FArchive& Ar, void* Value, int32 MaxRe
 {
 	// We never serialize our reference while the garbage collector is harvesting references
 	// to objects, because we don't want asset pointers to keep objects from being garbage collected
-	if( !Ar.IsObjectReferenceCollector() || Ar.IsModifyingWeakAndStrongReferences() )
+	// Allow persistent archives so they can keep track of string references. (e.g. FArchiveSaveTagImports)
+	if( !Ar.IsObjectReferenceCollector() || Ar.IsModifyingWeakAndStrongReferences() || Ar.IsPersistent() )
 	{
 		FAssetPtr OldValue = *(FAssetPtr*)Value;
 		Ar << *(FAssetPtr*)Value;
diff --git a/Engine/Source/Runtime/CoreUObject/Private/UObject/SavePackage.cpp b/Engine/Source/Runtime/CoreUObject/Private/UObject/SavePackage.cpp
--- a/Engine/Source/Runtime/CoreUObject/Private/UObject/SavePackage.cpp
+++ b/Engine/Source/Runtime/CoreUObject/Private/UObject/SavePackage.cpp
@@ -535,7 +535,10 @@ public:
 	FArchive& operator<<(FAssetPtr& AssetPtr) override;
 	FArchive& operator<<(FStringAssetReference& Value) override
 	{
-		StringAssetReferencesMap.Add(Value.ToString());
+		if ( Value.IsValid() )
+		{
+			StringAssetReferencesMap.Add(Value.ToString());
+		}
 		return *this;
 	}
 	FArchive& operator<<(FName& Name) override
