Disregard for GC fixes.

Added code to make sure the first GC index is never greater than the last non-GC index after DisregardForGC is closed.

Enabled DisregardForGC in ShooterGame so that it gets tested regularly

diff --git a/Engine/Source/Runtime/CoreUObject/Private/UObject/UObjectArray.cpp b/Engine/Source/Runtime/CoreUObject/Private/UObject/UObjectArray.cpp
--- a/Engine/Source/Runtime/CoreUObject/Private/UObject/UObjectArray.cpp
+++ b/Engine/Source/Runtime/CoreUObject/Private/UObject/UObjectArray.cpp
@@ -11,11 +11,6 @@ DEFINE_LOG_CATEGORY_STATIC(LogUObjectArray, Log, All);
 /** Global UObject array							*/
 FUObjectArray GUObjectArray;
 
-/**
- * Allocates and initializes the permanent object pool
- *
- * @param InPermanentObjectPoolSize size of permanent object pool
- */
 void FUObjectArray::AllocatePermanentObjectPool(int32 MaxObjectsNotConsideredByGC)
 {
 	// GObjFirstGCIndex is the index at which the garbage collector will start for the mark phase.
@@ -30,11 +25,14 @@ void FUObjectArray::AllocatePermanentObjectPool(int32 MaxObjectsNotConsideredByG
 	FWeakObjectPtr::Init(); // this adds a delete listener
 }
 
-/**
- * Adds a uobject to the global array which is used for uobject iteration
- *
- * @param	Object Object to allocate an index for
- */
+void FUObjectArray::CloseDisregardForGC()
+{
+	OpenForDisregardForGC = false;
+	GIsInitialLoad = false;
+	// Make sure the first GC index matches the last non-GC index after DisregardForGC is closed
+	ObjFirstGCIndex = ObjLastNonGCIndex + 1;
+}
+
 void FUObjectArray::AllocateUObjectIndex(UObjectBase* Object)
 {
 	int32 Index;
diff --git a/Engine/Source/Runtime/CoreUObject/Public/UObject/UObjectArray.h b/Engine/Source/Runtime/CoreUObject/Public/UObject/UObjectArray.h
--- a/Engine/Source/Runtime/CoreUObject/Public/UObject/UObjectArray.h
+++ b/Engine/Source/Runtime/CoreUObject/Public/UObject/UObjectArray.h
@@ -90,11 +90,7 @@ public:
 	 * After the initial load, this closes the disregard pool so that new object are GC-able
 	 *
 	 */
-	void CloseDisregardForGC()
-	{
-		OpenForDisregardForGC = false;
-		GIsInitialLoad = false;
-	}
+	void CloseDisregardForGC();
 
 	/**
 	 * indicates if the disregard for GC optimization is active
