Avoiding checking this for null in UObject::IsBasedOnArchetype().

This is reported as warning by clang 3.5.0. Arguably not a compiler's call, but it's indeed better to check that condition before calling the function and not inside it.

diff --git a/Engine/Source/Runtime/CoreUObject/Private/UObject/CoreNative.cpp b/Engine/Source/Runtime/CoreUObject/Private/UObject/CoreNative.cpp
--- a/Engine/Source/Runtime/CoreUObject/Private/UObject/CoreNative.cpp
+++ b/Engine/Source/Runtime/CoreUObject/Private/UObject/CoreNative.cpp
@@ -233,7 +233,7 @@ UObject* FObjectInstancingGraph::InstancePropertyValue( class UObject* Component
 	// if the object we're instancing the components for (Owner) has the current component's outer in its archetype chain, and its archetype has a NULL value
 	// for this component property it means that the archetype didn't instance its component, so we shouldn't either.
 
-	if (ComponentTemplate == NULL && CurrentValue != NULL && Owner->IsBasedOnArchetype(CurrentValue->GetOuter()))
+	if (ComponentTemplate == NULL && CurrentValue != NULL && (Owner && Owner->IsBasedOnArchetype(CurrentValue->GetOuter())))
 	{
 		NewValue = NULL;
 	}
diff --git a/Engine/Source/Runtime/CoreUObject/Private/UObject/Obj.cpp b/Engine/Source/Runtime/CoreUObject/Private/UObject/Obj.cpp
--- a/Engine/Source/Runtime/CoreUObject/Private/UObject/Obj.cpp
+++ b/Engine/Source/Runtime/CoreUObject/Private/UObject/Obj.cpp
@@ -401,8 +401,7 @@ void UObject::GetArchetypeInstances( TArray<UObject*>& Instances )
 				UObject* Obj = It;
 				
 				// if this object is the correct type and its archetype is this object, add it to the list
-				CA_SUPPRESS(6011)
-				if ( Obj != this && Obj->IsBasedOnArchetype(this) )
+				if ( Obj != this && Obj && Obj->IsBasedOnArchetype(this) )
 				{
 					Instances.Add(Obj);
 				}
diff --git a/Engine/Source/Runtime/CoreUObject/Public/UObject/UnrealType.h b/Engine/Source/Runtime/CoreUObject/Public/UObject/UnrealType.h
--- a/Engine/Source/Runtime/CoreUObject/Public/UObject/UnrealType.h
+++ b/Engine/Source/Runtime/CoreUObject/Public/UObject/UnrealType.h
@@ -3053,7 +3053,8 @@ T* SmartCastProperty( UProperty* Src )
  */
 inline bool UObject::IsBasedOnArchetype(  const UObject* const SomeObject ) const
 {
-	if ( this && SomeObject != this )
+	checkfSlow(this, TEXT("IsBasedOnArchetype() is called on a null pointer. Fix the call site."));
+	if ( SomeObject != this )
 	{
 		for ( UObject* Template = GetArchetype(); Template; Template = Template->GetArchetype() )
 		{
