Don't check for object flags if this == NULL in CanCheckDefaultSuobjects

diff --git a/Engine/Source/Runtime/CoreUObject/Private/UObject/Obj.cpp b/Engine/Source/Runtime/CoreUObject/Private/UObject/Obj.cpp
--- a/Engine/Source/Runtime/CoreUObject/Private/UObject/Obj.cpp
+++ b/Engine/Source/Runtime/CoreUObject/Private/UObject/Obj.cpp
@@ -986,7 +986,7 @@ bool UObject::CanCheckDefaultSubObjects(bool bForceCheck, bool& bResult)
 		bResult = false; // these aren't in a suitable spot in their lifetime for testing
 		bCanCheck = false;
 	}
-	if (HasAnyFlags(RF_NeedLoad | RF_NeedPostLoad | RF_NeedPostLoadSubobjects | RF_Unreachable | RF_PendingKill) || GIsDuplicatingClassForReinstancing)
+	if (bCanCheck && (HasAnyFlags(RF_NeedLoad | RF_NeedPostLoad | RF_NeedPostLoadSubobjects | RF_Unreachable | RF_PendingKill) || GIsDuplicatingClassForReinstancing))
 	{
 		bResult = true; // these aren't in a suitable spot in their lifetime for testing
 		bCanCheck = false;
