Changing CheckDefaultSuobjects to not be a virtual function due to if (!this) check in CanCheckDefaultSuobjects

diff --git a/Engine/Source/Runtime/CoreUObject/Private/UObject/Obj.cpp b/Engine/Source/Runtime/CoreUObject/Private/UObject/Obj.cpp
--- a/Engine/Source/Runtime/CoreUObject/Private/UObject/Obj.cpp
+++ b/Engine/Source/Runtime/CoreUObject/Private/UObject/Obj.cpp
@@ -998,55 +998,64 @@ bool UObject::CanCheckDefaultSubObjects(bool bForceCheck, bool& bResult)
 
 bool UObject::CheckDefaultSubobjects(bool bForceCheck /*= false*/)
 {
-	bool Result = true;	
+	bool Result = true;
 	if (CanCheckDefaultSubObjects(bForceCheck, Result))
 	{
-		CompCheck(this);
-		UClass* Class = GetClass();
+		Result = CheckDefaultSubobjectsInternal();
+	}
+	return Result;
+}
+
+bool UObject::CheckDefaultSubobjectsInternal()
+{
+	bool Result = true;	
+
+	CompCheck(this);
+	UClass* Class = GetClass();
 
-		if (Class != UFunction::StaticClass() && Class->GetName() != TEXT("EdGraphPin"))
+	if (Class != UFunction::StaticClass() && Class->GetName() != TEXT("EdGraphPin"))
+	{
+		// Check for references to default subobjects of other objects.
+		// There should never be a pointer to a subobject from outside of the outer (chain) it belongs to.
+		TArray<const UObject*> OtherReferencedSubobjects;
+		FSubobjectReferenceFinder DefaultSubobjectCollector(OtherReferencedSubobjects, this);
+		for (int32 Index = 0; Index < OtherReferencedSubobjects.Num(); ++Index)
 		{
-			// Check for references to default subobjects of other objects.
-			// There should never be a pointer to a subobject from outside of the outer (chain) it belongs to.
-			TArray<const UObject*> OtherReferencedSubobjects;
-			FSubobjectReferenceFinder DefaultSubobjectCollector(OtherReferencedSubobjects, this);
-			for (int32 Index = 0; Index < OtherReferencedSubobjects.Num(); ++Index)
-			{
-				const UObject* TestObject = OtherReferencedSubobjects[Index];
-				UE_LOG(LogCheckSubobjects, Error, TEXT("%s has a reference to default subobject (%s) of %s."), *GetFullName(), *TestObject->GetFullName(), *TestObject->GetOuter()->GetFullName());
-			}
-			CompCheck(OtherReferencedSubobjects.Num() == 0);
+			const UObject* TestObject = OtherReferencedSubobjects[Index];
+			UE_LOG(LogCheckSubobjects, Error, TEXT("%s has a reference to default subobject (%s) of %s."), *GetFullName(), *TestObject->GetFullName(), *TestObject->GetOuter()->GetFullName());
 		}
+		CompCheck(OtherReferencedSubobjects.Num() == 0);
+	}
 
 #if 0 // usually overkill, but valid tests
-		if (!HasAnyFlags(RF_ClassDefaultObject) && Class->HasAnyClassFlags(CLASS_HasInstancedReference))
+	if (!HasAnyFlags(RF_ClassDefaultObject) && Class->HasAnyClassFlags(CLASS_HasInstancedReference))
+	{
+		UObject *Archetype = GetArchetype();
+		CompCheck(this != Archetype);
+		Archetype->CheckDefaultSubobjects();
+		if (Archetype != Class->GetDefaultObject())
 		{
-			UObject *Archetype = GetArchetype();
-			CompCheck(this != Archetype);
-			Archetype->CheckDefaultSubobjects();
-			if (Archetype != Class->GetDefaultObject())
-			{
-				Class->GetDefaultObject()->CheckDefaultSubobjects();
-			}
+			Class->GetDefaultObject()->CheckDefaultSubobjects();
 		}
+	}
 #endif
 
-		if (HasAnyFlags(RF_ClassDefaultObject))
-		{
-			CompCheck(GetFName() == Class->GetDefaultObjectName());
-		}
+	if (HasAnyFlags(RF_ClassDefaultObject))
+	{
+		CompCheck(GetFName() == Class->GetDefaultObjectName());
+	}
 
 
-		TArray<UObject *> AllCollectedComponents;
-		CollectDefaultSubobjects(AllCollectedComponents, true);
-		TArray<UObject *> DirectCollectedComponents;
-		CollectDefaultSubobjects(DirectCollectedComponents, false);
+	TArray<UObject *> AllCollectedComponents;
+	CollectDefaultSubobjects(AllCollectedComponents, true);
+	TArray<UObject *> DirectCollectedComponents;
+	CollectDefaultSubobjects(DirectCollectedComponents, false);
 		
-		AllCollectedComponents.Sort();
-		DirectCollectedComponents.Sort();
+	AllCollectedComponents.Sort();
+	DirectCollectedComponents.Sort();
+
+	CompCheck(ALLOW_SUB_SUB_OBJECTS || AllCollectedComponents == DirectCollectedComponents); // just say no to subobjects of subobjects
 
-		CompCheck(ALLOW_SUB_SUB_OBJECTS || AllCollectedComponents == DirectCollectedComponents); // just say no to subobjects of subobjects
-	}
 	return Result;
 }
 
diff --git a/Engine/Source/Runtime/CoreUObject/Public/UObject/UObject.h b/Engine/Source/Runtime/CoreUObject/Public/UObject/UObject.h
--- a/Engine/Source/Runtime/CoreUObject/Public/UObject/UObject.h
+++ b/Engine/Source/Runtime/CoreUObject/Public/UObject/UObject.h
@@ -612,7 +612,7 @@ public:
 	 * @param bForceCheck Force checks even if not enabled globally.
 	 * @return true if the assumptions are met, false otherwise.
 	 */
-	virtual bool CheckDefaultSubobjects(bool bForceCheck = false);
+	bool CheckDefaultSubobjects(bool bForceCheck = false);
 
 	/**
 	 * Save configuration.
@@ -981,6 +981,13 @@ protected:
 	/** Checks it's ok to perform subobjects check at this time. */
 	bool CanCheckDefaultSubObjects(bool bForceCheck, bool& bResult);
 
+	/**
+	* Checks default sub-object assumptions.
+	*
+	* @return true if the assumptions are met, false otherwise.
+	*/
+	virtual bool CheckDefaultSubobjectsInternal();
+
 private:
 	void ProcessContextOpcode(FFrame& Stack, RESULT_DECL, bool bCanFailSilent);
 };
diff --git a/Engine/Source/Runtime/Engine/Classes/GameFramework/Actor.h b/Engine/Source/Runtime/Engine/Classes/GameFramework/Actor.h
--- a/Engine/Source/Runtime/Engine/Classes/GameFramework/Actor.h
+++ b/Engine/Source/Runtime/Engine/Classes/GameFramework/Actor.h
@@ -1041,7 +1041,7 @@ public:
 	FActorEndPlaySignature OnEndPlay;
 	
 	// Begin UObject Interface
-	virtual bool CheckDefaultSubobjects(bool bForceCheck = false) override;
+	virtual bool CheckDefaultSubobjectsInternal() override;
 	virtual void PostInitProperties() override;
 	virtual bool Modify( bool bAlwaysMarkDirty=true ) override;
 	virtual void ProcessEvent( UFunction* Function, void* Parameters ) override;
diff --git a/Engine/Source/Runtime/Engine/Private/Actor.cpp b/Engine/Source/Runtime/Engine/Private/Actor.cpp
--- a/Engine/Source/Runtime/Engine/Private/Actor.cpp
+++ b/Engine/Source/Runtime/Engine/Private/Actor.cpp
@@ -104,13 +104,12 @@ FString FActorTickFunction::DiagnosticMessage()
 	return Target->GetFullName() + TEXT("[TickActor]");
 }
 
-bool AActor::CheckDefaultSubobjects(bool bForceCheck /*= false*/)
+bool AActor::CheckDefaultSubobjectsInternal()
 {
-	bool Result = true;
-	if (CanCheckDefaultSubObjects(bForceCheck, Result))
+	bool Result = Super::CheckDefaultSubobjectsInternal();
+	if (Result)
 	{
-		Result = Super::CheckDefaultSubobjects(bForceCheck);
-		Result = Result && CheckActorComponents();
+		Result = CheckActorComponents();
 	}
 	return Result;
 }
