Making sure old classes are not visible in the class picker after hot-reload.

#ttp 345524: HotReload: Strange temporary class names are visible in Placement Browser's "All Classes" list after a HotReload

#change Marked old classes as transient, made sure they're not in root set
#change Made sure class picker rejects transient UClass objects

diff --git a/Engine/Source/Editor/ClassViewer/Private/SClassViewer.cpp b/Engine/Source/Editor/ClassViewer/Private/SClassViewer.cpp
--- a/Engine/Source/Editor/ClassViewer/Private/SClassViewer.cpp
+++ b/Engine/Source/Editor/ClassViewer/Private/SClassViewer.cpp
@@ -1676,7 +1676,9 @@ void FClassHierarchy::AddChildren_NoFilter( TSharedPtr< FClassViewerNode >& InOu
 		UClass* CurrentClass = *ClassIt;
 
 		// Ignore deprecated and temporary trash classes.
-		if (CurrentClass->HasAnyClassFlags(CLASS_Deprecated | CLASS_NewerVersionExists) || FKismetEditorUtilities::IsClassABlueprintSkeleton(CurrentClass))
+		if (CurrentClass->HasAnyClassFlags(CLASS_Deprecated | CLASS_NewerVersionExists) || 
+			FKismetEditorUtilities::IsClassABlueprintSkeleton(CurrentClass) ||
+			CurrentClass->HasAnyFlags(RF_Transient))
 		{
 			continue;
 		}
diff --git a/Engine/Source/Editor/UnrealEd/Private/Kismet2/KismetReinstanceUtilities.cpp b/Engine/Source/Editor/UnrealEd/Private/Kismet2/KismetReinstanceUtilities.cpp
--- a/Engine/Source/Editor/UnrealEd/Private/Kismet2/KismetReinstanceUtilities.cpp
+++ b/Engine/Source/Editor/UnrealEd/Private/Kismet2/KismetReinstanceUtilities.cpp
@@ -68,6 +68,7 @@ FBlueprintCompileReinstancer::FBlueprintCompileReinstancer(UClass* InClassToRein
 		// Add the class to the root set, so that it doesn't get prematurely garbage collected
 		DuplicatedClass->AddToRoot();
 		ClassToReinstance->ClassFlags &= ~CLASS_NewerVersionExists;
+		InClassToReinstance->RemoveFromRoot();
 		GIsDuplicatingClassForReinstancing = false;
 
 		// Bind and link the duplicate class, so that it has the proper duplicate property offsets
diff --git a/Engine/Source/Runtime/CoreUObject/Private/UObject/UObjectBase.cpp b/Engine/Source/Runtime/CoreUObject/Private/UObject/UObjectBase.cpp
--- a/Engine/Source/Runtime/CoreUObject/Private/UObject/UObjectBase.cpp
+++ b/Engine/Source/Runtime/CoreUObject/Private/UObject/UObjectBase.cpp
@@ -580,6 +580,7 @@ void UClassCompiledInDefer(FFieldCompiledInInfo* ClassInfo, const TCHAR* Name, S
 			ExistingClass->GetDefaultObject()->ClearFlags(RF_RootSet | RF_Standalone | RF_Public);
 			const FName OldClassRename = MakeUniqueObjectName(GetTransientPackage(), ExistingClass->GetClass(), *FString::Printf(TEXT("HOTRELOADED_%s"), *NameWithoutPrefix));
 			ExistingClass->Rename(*OldClassRename.ToString(), GetTransientPackage());
+			ExistingClass->SetFlags(RF_Transient);
 			ExistingClass->AddToRoot();
 			ClassInfo->OldField = ExistingClass;
 
