UGCObjectReferencer's references should not have to meet the disregard for GC assumptions, otherwise the referencer would be pretty useless at runtime.

diff --git a/Engine/Source/Runtime/CoreUObject/Private/UObject/GarbageCollection.cpp b/Engine/Source/Runtime/CoreUObject/Private/UObject/GarbageCollection.cpp
--- a/Engine/Source/Runtime/CoreUObject/Private/UObject/GarbageCollection.cpp
+++ b/Engine/Source/Runtime/CoreUObject/Private/UObject/GarbageCollection.cpp
@@ -1084,7 +1084,10 @@ void CollectGarbage( EObjectFlags KeepFlags, bool bPerformFullPurge )
 		{
 			bool bShouldAssert = false;
 			UObject* Object = *It;
-			if( GUObjectArray.IsDisregardForGC(Object) )
+			// Don't require UGCObjectReferencer's references to adhere to the assumptions.
+			// Although we want the referencer itself to sit in the disregard for gc set, most of the objects
+			// it's referencing will not be in the root set.
+			if (GUObjectArray.IsDisregardForGC(Object) && !Object->IsA(UGCObjectReferencer::StaticClass()))
 			{
 				// Serialize object with reference collector.
 				TArray<UObject*> CollectedReferences;
