Core - Adding a way in UClasses to check if they contain a UProperty directly in it or a base class.

diff --git a/Engine/Source/Runtime/CoreUObject/Private/UObject/Class.cpp b/Engine/Source/Runtime/CoreUObject/Private/UObject/Class.cpp
--- a/Engine/Source/Runtime/CoreUObject/Private/UObject/Class.cpp
+++ b/Engine/Source/Runtime/CoreUObject/Private/UObject/Class.cpp
@@ -3230,6 +3230,17 @@ bool UClass::IsFunctionImplementedInBlueprint(FName InFunctionName) const
 	return false;
 }
 
+bool UClass::HasProperty(UProperty* InProperty) const
+{
+	if ( UClass* PropertiesClass = Cast<UClass>(InProperty->GetOuter()) )
+	{
+		return PropertiesClass->FindNearestCommonBaseClass(this) != nullptr;
+	}
+
+	return false;
+}
+
+
 /*-----------------------------------------------------------------------------
 	UClass constructors.
 -----------------------------------------------------------------------------*/
diff --git a/Engine/Source/Runtime/CoreUObject/Public/UObject/Class.h b/Engine/Source/Runtime/CoreUObject/Public/UObject/Class.h
--- a/Engine/Source/Runtime/CoreUObject/Public/UObject/Class.h
+++ b/Engine/Source/Runtime/CoreUObject/Public/UObject/Class.h
@@ -2132,6 +2132,12 @@ public:
 	 */
 	virtual bool IsFunctionImplementedInBlueprint(FName InFunctionName) const;
 
+	/**
+	 * Checks if the property exists on this class or a parent class.
+	 * @param InProperty	The property to check if it is contained in this or a parent class.
+	 * @return				True if the property exists on this or a parent class.
+	 */
+	virtual bool HasProperty(UProperty* InProperty) const;
 
 private:
 	// This signature intentionally hides the method declared in UObjectBaseUtility to make it private.
