Core: Reformatted import error message for MessageLog

diff --git a/Engine/Source/Runtime/CoreUObject/Private/UObject/LinkerLoad.cpp b/Engine/Source/Runtime/CoreUObject/Private/UObject/LinkerLoad.cpp
--- a/Engine/Source/Runtime/CoreUObject/Private/UObject/LinkerLoad.cpp
+++ b/Engine/Source/Runtime/CoreUObject/Private/UObject/LinkerLoad.cpp
@@ -2057,15 +2057,17 @@ ULinkerLoad::EVerifyResult ULinkerLoad::VerifyImport(int32 ImportIndex)
 				UClass* FindClass = ClassPackage ? FindObject<UClass>( ClassPackage, *OriginalImport.ClassName.ToString() ) : NULL;
 				if( GIsEditor && !IsRunningCommandlet() )
 				{
-					FFormatNamedArguments Arguments[2];
-					Arguments[0].Add(TEXT("ImportClass"), FText::FromName(GetImportClassName(ImportIndex)));
-					Arguments[1].Add(TEXT("Warning"), FText::FromString(WarningAppend));
-
 					// put something into the load warnings dialog, with any extra information from above (in WarningAppend)
-					FMessageLog(NAME_LoadErrors).Error(FText::Format(LOCTEXT("ImportFailure", "Failed import: {ImportClass}"), Arguments[0]))
-						->AddToken(FAssetNameToken::Create(GetImportPathName(ImportIndex)))
-						->AddToken(FTextToken::Create(FText::Format(LOCTEXT("ImportFailure_WarningIn", "{Warning} in"), Arguments[1])))
-						->AddToken(FAssetNameToken::Create(LinkerRoot->GetName()));
+					TSharedRef<FTokenizedMessage> TokenizedMessage = FMessageLog(NAME_LoadErrors).Error(FText::Format(LOCTEXT("ImportFailure", "Failed import for {ImportClass}"), FText::FromName(GetImportClassName(ImportIndex))));
+					TokenizedMessage->AddToken(FAssetNameToken::Create(GetImportPathName(ImportIndex)));
+
+					if (!WarningAppend.IsEmpty())
+					{
+						TokenizedMessage->AddToken(FTextToken::Create(FText::Format(LOCTEXT("ImportFailure_WarningIn", "{0} in {1}"),
+							FText::FromString(WarningAppend),
+							FText::FromString(LinkerRoot->GetName())))
+						);
+					}
 				}
 
 #if UE_BUILD_DEBUG
