#ue4 - Strip "[REDACTED_HANDLE]" from the end of all tooltips in the editor.

diff --git a/Engine/Source/Runtime/CoreUObject/Private/UObject/Class.cpp b/Engine/Source/Runtime/CoreUObject/Private/UObject/Class.cpp
--- a/Engine/Source/Runtime/CoreUObject/Private/UObject/Class.cpp
+++ b/Engine/Source/Runtime/CoreUObject/Private/UObject/Class.cpp
@@ -243,6 +243,13 @@ FText UField::GetToolTipText() const
 		{
 			NativeToolTip = FDisplayNameHelper::Get(*this);
 		}
+		else
+		{
+			if (NativeToolTip.Split(TEXT("@see"), &NativeToolTip, nullptr, ESearchCase::IgnoreCase, ESearchDir::FromStart))
+			{
+				NativeToolTip.TrimTrailing();
+			}
+		}
 		LocalizedToolTip = FText::FromString(NativeToolTip);
 	}
 
