The CPF_EditInline flag, will force the subobject instance to be persistent, but CLASS_DefaultToInstanced wonÆt imply CPF_EditInline anymore.

diff --git a/Engine/Source/Programs/UnrealHeaderTool/Private/HeaderParser.cpp b/Engine/Source/Programs/UnrealHeaderTool/Private/HeaderParser.cpp
--- a/Engine/Source/Programs/UnrealHeaderTool/Private/HeaderParser.cpp
+++ b/Engine/Source/Programs/UnrealHeaderTool/Private/HeaderParser.cpp
@@ -3637,7 +3637,7 @@ UProperty* FHeaderParser::GetVarNameAndDim
 			{
 				if ( DoesAnythingInHierarchyHaveDefaultToInstanced( VarProperty.PropertyClass ) )
 				{
-					VarProperty.PropertyFlags |= CPF_EditInline | CPF_InstancedReference;
+					VarProperty.PropertyFlags |= CPF_InstancedReference;
 				}
 				NewProperty = new(NewScope,PropertyName,ObjectFlags)UObjectProperty(FPostConstructInitializeProperties());
 			}
diff --git a/Engine/Source/Runtime/CoreUObject/Public/UObject/ObjectBase.h b/Engine/Source/Runtime/CoreUObject/Public/UObject/ObjectBase.h
--- a/Engine/Source/Runtime/CoreUObject/Public/UObject/ObjectBase.h
+++ b/Engine/Source/Runtime/CoreUObject/Public/UObject/ObjectBase.h
@@ -305,7 +305,7 @@ typedef uint64 EClassCastFlags;
 //#define CPF_    							DECLARE_UINT64(0x0000000000800000)		// 
 #define CPF_SaveGame						DECLARE_UINT64(0x0000000001000000)		// Property should be serialized for save games
 #define CPF_NoClear							DECLARE_UINT64(0x0000000002000000)		// Hide clear (and browse) button.
-#define CPF_EditInline						DECLARE_UINT64(0x0000000004000000)		// Edit this object reference inline.
+#define CPF_EditInline						DECLARE_UINT64(0x0000000004000000)		// Edit this object reference inline. Property must have also CPF_InstancedReference. The sub-object instance will be persistent.
 #define CPF_ReferenceParm					DECLARE_UINT64(0x0000000008000000)		// Value is passed by reference; CPF_OutParam and CPF_Param should also be set.
 #define CPF_BlueprintAssignable				DECLARE_UINT64(0x0000000010000000)		// MC Delegates only.  Property should be exposed for assigning in blueprint code
 #define CPF_Deprecated  					DECLARE_UINT64(0x0000000020000000)		// Property is deprecated.  Read it from an archive, but don't save it.
diff --git a/Engine/Source/Runtime/Engine/Private/UnrealEngine.cpp b/Engine/Source/Runtime/Engine/Private/UnrealEngine.cpp
--- a/Engine/Source/Runtime/Engine/Private/UnrealEngine.cpp
+++ b/Engine/Source/Runtime/Engine/Private/UnrealEngine.cpp
@@ -9671,7 +9671,7 @@ struct FFindInstancedReferenceSubobjectHelper
 		check(ContainerAddress);
 		for (UProperty* Prop = (Struct ? Struct->RefLink : NULL); Prop; Prop = Prop->NextRef)
 		{
-			if (Prop->HasAnyPropertyFlags(CPF_InstancedReference))
+			if (Prop->HasAllPropertyFlags(CPF_InstancedReference | CPF_EditInline))
 			{
 				auto ObjectProperty = Cast<const UObjectProperty>(Prop);
 				if (ObjectProperty)
