UObjects: no longer require dummy UClasses in UStruct-only modules.

#change Added code to make sure the script package is created before UEnums and UStructs get registered

diff --git a/Engine/Source/Programs/UnrealHeaderTool/Private/CodeGenerator.cpp b/Engine/Source/Programs/UnrealHeaderTool/Private/CodeGenerator.cpp
--- a/Engine/Source/Programs/UnrealHeaderTool/Private/CodeGenerator.cpp
+++ b/Engine/Source/Programs/UnrealHeaderTool/Private/CodeGenerator.cpp
@@ -2340,8 +2340,8 @@ void FNativeClassHeaderGenerator::ExportGeneratedStructBodyMacros(const TArray<U
 			GeneratedPackageCPP.Logf(TEXT("    }\r\n"));
 			GeneratedPackageCPP.Logf(TEXT("    return Singleton;\r\n"));
 			GeneratedPackageCPP.Logf(TEXT("}\r\n"));
-			GeneratedPackageCPP.Logf(TEXT("static FCompiledInDeferStruct Z_CompiledInDeferStruct_UScriptStruct_%s(%s::StaticStruct);\r\n"), 
-				StructNameCPP, StructNameCPP);
+			GeneratedPackageCPP.Logf(TEXT("static FCompiledInDeferStruct Z_CompiledInDeferStruct_UScriptStruct_%s(%s::StaticStruct, TEXT(\"%s\"));\r\n"), 
+				StructNameCPP, StructNameCPP, *Struct->GetOutermost()->GetName());
 
 			// Generate StaticRegisterNatives equivalent for structs without classes.
 			if (!Struct->GetOuter()->IsA(UStruct::StaticClass()))
@@ -2390,7 +2390,7 @@ void FNativeClassHeaderGenerator::ExportGeneratedEnumsInitCode(const TArray<UEnu
 			GeneratedPackageCPP.Logf(TEXT("    }\r\n"));
 			GeneratedPackageCPP.Logf(TEXT("    return Singleton;\r\n"));
 			GeneratedPackageCPP.Logf(TEXT("}\r\n"));
-			GeneratedPackageCPP.Logf(TEXT("static FCompiledInDeferEnum Z_CompiledInDeferEnum_UEnum_%s(%s_StaticEnum);\r\n"), *Enum->GetName(), *Enum->GetName());
+			GeneratedPackageCPP.Logf(TEXT("static FCompiledInDeferEnum Z_CompiledInDeferEnum_UEnum_%s(%s_StaticEnum, TEXT(\"%s\"));\r\n"), *Enum->GetName(), *Enum->GetName(), *Enum->GetOutermost()->GetName());
 		}
 	}
 }
diff --git a/Engine/Source/Runtime/CoreUObject/Private/UObject/UObjectBase.cpp b/Engine/Source/Runtime/CoreUObject/Private/UObject/UObjectBase.cpp
--- a/Engine/Source/Runtime/CoreUObject/Private/UObject/UObjectBase.cpp
+++ b/Engine/Source/Runtime/CoreUObject/Private/UObject/UObjectBase.cpp
@@ -414,17 +414,38 @@ void UObjectForceRegistration(UObjectBase* Object)
 	}
 }
 
-static TArray<class UScriptStruct *(*)()>& GetDeferredCompiledInStructRegistration()
+/**
+ * Struct containing the function pointer and package name of a UStruct to be registered with UObject system
+ */
+struct FPendingStructRegistrant
+{	
+	class UScriptStruct *(*RegisterFn)();
+	const TCHAR* PackageName;
+
+	FPendingStructRegistrant() {}
+	FPendingStructRegistrant(class UScriptStruct *(*Fn)(), const TCHAR* InPackageName)
+		: RegisterFn(Fn)
+		, PackageName(InPackageName)
+	{
+	}
+	FORCEINLINE bool operator==(const FPendingStructRegistrant& Other) const
+	{
+		return RegisterFn == Other.RegisterFn;
+	}
+};
+
+static TArray<FPendingStructRegistrant>& GetDeferredCompiledInStructRegistration()
 {
-	static TArray<class UScriptStruct *(*)()> DeferredCompiledInRegistration;
+	static TArray<FPendingStructRegistrant> DeferredCompiledInRegistration;
 	return DeferredCompiledInRegistration;
 }
 
-void UObjectCompiledInDeferStruct(class UScriptStruct *(*InRegister)())
+void UObjectCompiledInDeferStruct(class UScriptStruct *(*InRegister)(), const TCHAR* PackageName)
 {
 	// we do reregister StaticStruct in hot reload
-	checkSlow(!GetDeferredCompiledInStructRegistration().Contains(InRegister));
-	GetDeferredCompiledInStructRegistration().Add(InRegister);
+	FPendingStructRegistrant Registrant(InRegister, PackageName);
+	checkSlow(!GetDeferredCompiledInStructRegistration().Contains(Registrant));
+	GetDeferredCompiledInStructRegistration().Add(Registrant);
 }
 
 #if WITH_HOT_RELOAD
@@ -497,18 +518,39 @@ class UScriptStruct *GetStaticStruct(class UScriptStruct *(*InRegister)(), UObje
 	return (*InRegister)();
 }
 
+/**
+ * Struct containing the function pointer and package name of a UEnum to be registered with UObject system
+ */
+struct FPendingEnumRegistrant
+{
+	class UEnum *(*RegisterFn)();
+	const TCHAR* PackageName;
+
+	FPendingEnumRegistrant() {}
+	FPendingEnumRegistrant(class UEnum *(*Fn)(), const TCHAR* InPackageName)
+		: RegisterFn(Fn)
+		, PackageName(InPackageName)
+	{
+	}
+	FORCEINLINE bool operator==(const FPendingEnumRegistrant& Other) const
+	{
+		return RegisterFn == Other.RegisterFn;
+	}
+};
+
 // Same thing as GetDeferredCompiledInStructRegistration but for UEnums declared in header files without UClasses.
-static TArray<class UEnum *(*)()>& GetDeferredCompiledInEnumRegistration()
+static TArray<FPendingEnumRegistrant>& GetDeferredCompiledInEnumRegistration()
 {
-	static TArray<class UEnum *(*)()> DeferredCompiledInRegistration;
+	static TArray<FPendingEnumRegistrant> DeferredCompiledInRegistration;
 	return DeferredCompiledInRegistration;
 }
 
-void UObjectCompiledInDeferEnum(class UEnum *(*InRegister)())
+void UObjectCompiledInDeferEnum(class UEnum *(*InRegister)(), const TCHAR* PackageName)
 {
 	// we do reregister StaticStruct in hot reload
-	checkSlow(!GetDeferredCompiledInEnumRegistration().Contains(InRegister));
-	GetDeferredCompiledInEnumRegistration().Add(InRegister);
+	FPendingEnumRegistrant Registrant(InRegister, PackageName);
+	checkSlow(!GetDeferredCompiledInEnumRegistration().Contains(Registrant));
+	GetDeferredCompiledInEnumRegistration().Add(Registrant);
 }
 
 class UEnum *GetStaticEnum(class UEnum *(*InRegister)(), UObject* EnumOuter, const TCHAR* EnumName)
@@ -706,12 +748,14 @@ static void UObjectLoadAllCompiledInStructs()
 	const bool bHaveEnumRegistrants = GetDeferredCompiledInEnumRegistration().Num() != 0;
 	if( bHaveEnumRegistrants )
 	{
-		TArray<class UEnum *(*)()> PendingRegistrants;
+		TArray<FPendingEnumRegistrant> PendingRegistrants;
 		PendingRegistrants = GetDeferredCompiledInEnumRegistration();
 		GetDeferredCompiledInEnumRegistration().Empty();
-		for(int32 RegistrantIndex = 0;RegistrantIndex < PendingRegistrants.Num();++RegistrantIndex)
+		for (auto& EnumRegistrant : PendingRegistrants)
 		{
-			PendingRegistrants[RegistrantIndex]();
+			// Make sure the package exists in case it does not contain any UObjects
+			CreatePackage(NULL, EnumRegistrant.PackageName);
+			EnumRegistrant.RegisterFn();
 		}
 	}
 
@@ -719,12 +763,14 @@ static void UObjectLoadAllCompiledInStructs()
 	const bool bHaveRegistrants = GetDeferredCompiledInStructRegistration().Num() != 0;
 	if( bHaveRegistrants )
 	{
-		TArray<class UScriptStruct *(*)()> PendingRegistrants;
+		TArray<FPendingStructRegistrant> PendingRegistrants;
 		PendingRegistrants = GetDeferredCompiledInStructRegistration();
 		GetDeferredCompiledInStructRegistration().Empty();
-		for(int32 RegistrantIndex = 0;RegistrantIndex < PendingRegistrants.Num();++RegistrantIndex)
+		for (auto& StructRegistrant : PendingRegistrants)
 		{
-			PendingRegistrants[RegistrantIndex]();
+			// Make sure the package exists in case it does not contain any UObjects or UEnums
+			CreatePackage(NULL, StructRegistrant.PackageName);
+			StructRegistrant.RegisterFn();
 		}
 	}
 }
diff --git a/Engine/Source/Runtime/CoreUObject/Public/UObject/UObjectBase.h b/Engine/Source/Runtime/CoreUObject/Public/UObject/UObjectBase.h
--- a/Engine/Source/Runtime/CoreUObject/Public/UObject/UObjectBase.h
+++ b/Engine/Source/Runtime/CoreUObject/Public/UObject/UObjectBase.h
@@ -312,13 +312,13 @@ struct FCompiledInDefer
 /**
  * Stashes the singleton function that builds a compiled in struct (StaticStruct). Later, this is executed.
  */
-COREUOBJECT_API void UObjectCompiledInDeferStruct(class UScriptStruct *(*InRegister)());
+COREUOBJECT_API void UObjectCompiledInDeferStruct(class UScriptStruct *(*InRegister)(), const TCHAR* PackageName);
 
 struct FCompiledInDeferStruct
 {
-	FCompiledInDeferStruct(class UScriptStruct *(*InRegister)())
+	FCompiledInDeferStruct(class UScriptStruct *(*InRegister)(), const TCHAR* PackageName)
 	{
-		UObjectCompiledInDeferStruct(InRegister);
+		UObjectCompiledInDeferStruct(InRegister, PackageName);
 	}
 };
 
@@ -330,13 +330,13 @@ COREUOBJECT_API class UScriptStruct *GetStaticStruct(class UScriptStruct *(*InRe
 /**
  * Stashes the singleton function that builds a compiled in enum. Later, this is executed.
  */
-COREUOBJECT_API void UObjectCompiledInDeferEnum(class UEnum *(*InRegister)());
+COREUOBJECT_API void UObjectCompiledInDeferEnum(class UEnum *(*InRegister)(), const TCHAR* PackageName);
 
 struct FCompiledInDeferEnum
 {
-	FCompiledInDeferEnum(class UEnum *(*InRegister)())
+	FCompiledInDeferEnum(class UEnum *(*InRegister)(), const TCHAR* PackageName)
 	{
-		UObjectCompiledInDeferEnum(InRegister);
+		UObjectCompiledInDeferEnum(InRegister, PackageName);
 	}
 };
 
