Removing unnecessary log message (old debug code) from UObjectBase::AddObject that resulted in an extra FName -> FString conversion each time new object was constructed.

diff --git a/Engine/Source/Runtime/CoreUObject/Private/UObject/UObjectBase.cpp b/Engine/Source/Runtime/CoreUObject/Private/UObject/UObjectBase.cpp
--- a/Engine/Source/Runtime/CoreUObject/Private/UObject/UObjectBase.cpp
+++ b/Engine/Source/Runtime/CoreUObject/Private/UObject/UObjectBase.cpp
@@ -218,10 +218,6 @@ void UObjectBase::DeferredRegister(UClass *UClassStaticClass,const TCHAR* Packag
  */
 void UObjectBase::AddObject(FName InName)
 {
-	if ( !FPackageName::IsShortPackageName(InName) )
-	{
-		UE_LOG(LogLongPackageNames, Verbose, TEXT("Package created with long name %s"), *InName.ToString());
-	}
 #if EXTERNAL_OBJECT_NAMES
 	NameAnnotation.AddAnnotation(InternalIndex,InName);
 #else
