CoreUObject - Added ability to reserve memory for the object annotations

diff --git a/Engine/Source/Runtime/CoreUObject/Private/UObject/Obj.cpp b/Engine/Source/Runtime/CoreUObject/Private/UObject/Obj.cpp
--- a/Engine/Source/Runtime/CoreUObject/Private/UObject/Obj.cpp
+++ b/Engine/Source/Runtime/CoreUObject/Private/UObject/Obj.cpp
@@ -915,7 +915,7 @@ void UObject::CollectDefaultSubobjects( TArray<UObject*>& OutSubobjectArray, boo
 
 /**
  * FSubobjectReferenceFinder.
- * Helper class used to collect default subobjects of other objecets than the referencing object.
+ * Helper class used to collect default subobjects of other objects than the referencing object.
  */
 class FSubobjectReferenceFinder : public FReferenceCollector
 {
@@ -1841,6 +1841,9 @@ void StaticShutdownAfterError()
 		}
 		Shutdown = 1;
 		UE_LOG(LogExit, Log, TEXT("Executing StaticShutdownAfterError") );
+
+		ShutdownAfterErrorAnnotation.Reserve( GUObjectArray.GetObjectArrayNum() );
+
 		for ( FRawObjectIterator It; It; ++It )
 		{
 			It->ConditionalShutdownAfterError();
@@ -2388,6 +2391,7 @@ TArray<const TCHAR*> ParsePropertyFlags(uint64 Flags)
 	return Results;
 }
 
+// @TODO yrx 2014-09-15 Move to ObjectCommads.cpp or ObjectExec.cpp
 bool StaticExec( UWorld* InWorld, const TCHAR* Cmd, FOutputDevice& Ar )
 {
 	const TCHAR *Str = Cmd;
diff --git a/Engine/Source/Runtime/CoreUObject/Public/UObject/UObjectAnnotation.h b/Engine/Source/Runtime/CoreUObject/Public/UObject/UObjectAnnotation.h
--- a/Engine/Source/Runtime/CoreUObject/Public/UObject/UObjectAnnotation.h
+++ b/Engine/Source/Runtime/CoreUObject/Public/UObject/UObjectAnnotation.h
@@ -198,6 +198,14 @@ public:
 		return AnnotationMap;
 	}
 
+	/** 
+	 * Reserves memory for the annotation map for the specified number of elements, used to avoid reallocations. 
+	 */
+	void Reserve(int32 ExpectedNumElements)
+	{
+		AnnotationMap.Empty(ExpectedNumElements);
+	}
+
 private:
 
 	/**
@@ -407,6 +415,13 @@ public:
 		return this->GetAnnotation(Object).Mark;
 	}
 
+	/** 
+	 * Reserves memory for the annotation map for the specified number of elements, used to avoid reallocations. 
+	 */
+	FORCEINLINE void Reserve(int32 ExpectedNumElements)
+	{
+		FUObjectAnnotationSparse<FBoolAnnotation,true>::Reserve(ExpectedNumElements);
+	}
 };
 
 
