[GitHub] 304 : Fixed inherited interface BP method crash.
https://github.com/EpicGames/UnrealEngine/pull/304/files

UClass::FindFunctionByNam: Function in interface parent can be found.

diff --git a/Engine/Source/Runtime/CoreUObject/Private/UObject/Class.cpp b/Engine/Source/Runtime/CoreUObject/Private/UObject/Class.cpp
--- a/Engine/Source/Runtime/CoreUObject/Private/UObject/Class.cpp
+++ b/Engine/Source/Runtime/CoreUObject/Private/UObject/Class.cpp
@@ -3465,7 +3465,7 @@ UFunction* UClass::FindFunctionByName(FName InName, EIncludeSuperFlag::Type Incl
 
 		for (auto& Inter : SearchClass->Interfaces)
 		{
-			if (UFunction* Result = Inter.Class->FuncMap.FindRef(InName))
+			if (UFunction* Result = Inter.Class->FindFunctionByName(InName))
 				return Result;
 		}
 	}
