[GitHub] 332 : Script integration UObject helpers

#ttp 343283

diff --git a/Engine/Source/Runtime/CoreUObject/Private/UObject/Class.cpp b/Engine/Source/Runtime/CoreUObject/Private/UObject/Class.cpp
--- a/Engine/Source/Runtime/CoreUObject/Private/UObject/Class.cpp
+++ b/Engine/Source/Runtime/CoreUObject/Private/UObject/Class.cpp
@@ -3426,6 +3426,33 @@ bool UClass::HotReloadPrivateStaticClass(
 	}
 	return true;
 }
+
+bool UClass::ReplaceNativeFunction(FName InFName, Native InPointer, bool bAddToFunctionRemapTable)
+{
+	IHotReloadInterface* HotReloadSupport = nullptr;
+
+	if(bAddToFunctionRemapTable)
+	{
+		HotReloadSupport = &FModuleManager::LoadModuleChecked<IHotReloadInterface>("HotReload");
+	}
+
+	// Find the function in the class's native function lookup table.
+	for (int32 FunctionIndex = 0; FunctionIndex < NativeFunctionLookupTable.Num(); ++FunctionIndex)
+	{
+		FNativeFunctionLookup& NativeFunctionLookup = NativeFunctionLookupTable[FunctionIndex];
+		if (NativeFunctionLookup.Name == InFName)
+		{
+			if (bAddToFunctionRemapTable)
+			{
+				HotReloadSupport->AddHotReloadFunctionRemap(InPointer, NativeFunctionLookup.Pointer);
+			}
+			NativeFunctionLookup.Pointer = InPointer;
+			return true;
+		}
+	}
+	return false;
+}
+
 #endif
 
 void UClass::AddNativeFunction(const ANSICHAR* InName,Native InPointer)
@@ -3434,20 +3461,16 @@ void UClass::AddNativeFunction(const ANSICHAR* InName,Native InPointer)
 #if WITH_HOT_RELOAD
 	if (GIsHotReload)
 	{
-		IHotReloadInterface& HotReloadSupport = FModuleManager::LoadModuleChecked<IHotReloadInterface>("HotReload");
-
 		// Find the function in the class's native function lookup table.
-		for(int32 FunctionIndex = 0;FunctionIndex < NativeFunctionLookupTable.Num();++FunctionIndex)
+		if (ReplaceNativeFunction(InFName, InPointer, true))
 		{
-			FNativeFunctionLookup& NativeFunctionLookup = NativeFunctionLookupTable[FunctionIndex];
-			if(NativeFunctionLookup.Name == InFName)
-			{
-				HotReloadSupport.AddHotReloadFunctionRemap(InPointer, NativeFunctionLookup.Pointer);
-				NativeFunctionLookup.Pointer = InPointer;
-				return;
-			}
+			return;
+		}
+		else
+		{
+			// function was not found, so it's new
+			UE_LOG(LogClass, Log, TEXT("Function %s is new."), *InFName.ToString());
 		}
-		UE_LOG(LogClass, Log, TEXT("Function %s is new."),* InFName.ToString());
 	}
 #endif
 	new(NativeFunctionLookupTable) FNativeFunctionLookup(InFName,InPointer);
diff --git a/Engine/Source/Runtime/CoreUObject/Private/UObject/UObjectGlobals.cpp b/Engine/Source/Runtime/CoreUObject/Private/UObject/UObjectGlobals.cpp
--- a/Engine/Source/Runtime/CoreUObject/Private/UObject/UObjectGlobals.cpp
+++ b/Engine/Source/Runtime/CoreUObject/Private/UObject/UObjectGlobals.cpp
@@ -1932,19 +1932,8 @@ FPostConstructInitializeProperties::~FPostConstructInitializeProperties()
 		InitProperties(Obj, BaseClass, Defaults, bCopyTransientsFromClassDefaults);
 	}
 
-	bool bAllowInstancing = (InstanceGraph == NULL) || InstanceGraph->IsSubobjectInstancingEnabled();
-	bool bNeedSubobjectInstancing = false;
-	// initialize any subobjects, now that the constructors have run
-	for (int32 Index = 0; Index < ComponentInits.SubobjectInits.Num(); Index++)
-	{
-		UObject* Subobject = ComponentInits.SubobjectInits[Index].Subobject;
-		UObject* Template = ComponentInits.SubobjectInits[Index].Template;
-		InitProperties(Subobject, Template->GetClass(), Template, false);
-		if (bAllowInstancing && !Subobject->HasAnyFlags(RF_NeedLoad))
-		{
-			bNeedSubobjectInstancing = true;
-		}
-	}
+	bool bAllowInstancing = IsInstancingAllowed();
+	bool bNeedSubobjectInstancing = InitSubobjectProperties(bAllowInstancing);
 
 	// Restore class information if replacing native class.
 	if (ObjectRestoreAfterInitProps != NULL)
@@ -1977,32 +1966,7 @@ FPostConstructInitializeProperties::~FPostConstructInitializeProperties()
 	}
 	if (bNeedInstancing || bNeedSubobjectInstancing)
 	{
-		FObjectInstancingGraph TempInstancingGraph;
-		FObjectInstancingGraph* UseInstancingGraph = InstanceGraph ? InstanceGraph : &TempInstancingGraph;
-		UseInstancingGraph->AddNewObject(Obj);
-		// Add any default subobjects
-		for (int32 Index = 0; Index < ComponentInits.SubobjectInits.Num(); Index++)
-		{
-			UseInstancingGraph->AddNewObject(ComponentInits.SubobjectInits[Index].Subobject);
-		}
-		if (bNeedInstancing)
-		{
-			UObject* Archetype = Obj->GetArchetype();
-			Class->InstanceSubobjectTemplates(Obj, Archetype, Archetype ? Archetype->GetClass() : NULL, Obj, UseInstancingGraph);
-		}
-		if (bNeedSubobjectInstancing)
-		{
-			// initialize any subobjects, now that the constructors have run
-			for (int32 Index = 0; Index < ComponentInits.SubobjectInits.Num(); Index++)
-			{
-				UObject* Subobject = ComponentInits.SubobjectInits[Index].Subobject;
-				UObject* Template = ComponentInits.SubobjectInits[Index].Template;
-				if (!Subobject->HasAnyFlags(RF_NeedLoad))
-				{
-					Subobject->GetClass()->InstanceSubobjectTemplates(Subobject, Template, Template->GetClass(), Subobject, UseInstancingGraph);
-				}
-			}
-		}
+		InstanceSubobjects(Class, bNeedInstancing, bNeedSubobjectInstancing);
 	}
 
 	Obj->PostInitProperties();
@@ -2047,6 +2011,59 @@ FPostConstructInitializeProperties::~FPostConstructInitializeProperties()
 	}
 }
 
+bool FPostConstructInitializeProperties::IsInstancingAllowed() const
+{
+	return (InstanceGraph == NULL) || InstanceGraph->IsSubobjectInstancingEnabled();
+}
+
+bool FPostConstructInitializeProperties::InitSubobjectProperties(bool bAllowInstancing) const
+{
+	bool bNeedSubobjectInstancing = false;
+	// initialize any subobjects, now that the constructors have run
+	for (int32 Index = 0; Index < ComponentInits.SubobjectInits.Num(); Index++)
+	{
+		UObject* Subobject = ComponentInits.SubobjectInits[Index].Subobject;
+		UObject* Template = ComponentInits.SubobjectInits[Index].Template;
+		InitProperties(Subobject, Template->GetClass(), Template, false);
+		if (bAllowInstancing && !Subobject->HasAnyFlags(RF_NeedLoad))
+		{
+			bNeedSubobjectInstancing = true;
+		}
+	}
+
+	return bNeedSubobjectInstancing;
+}
+
+void FPostConstructInitializeProperties::InstanceSubobjects(UClass* Class, bool bNeedInstancing, bool bNeedSubobjectInstancing) const
+{
+	FObjectInstancingGraph TempInstancingGraph;
+	FObjectInstancingGraph* UseInstancingGraph = InstanceGraph ? InstanceGraph : &TempInstancingGraph;
+	UseInstancingGraph->AddNewObject(Obj);
+	// Add any default subobjects
+	for (int32 Index = 0; Index < ComponentInits.SubobjectInits.Num(); Index++)
+	{
+		UseInstancingGraph->AddNewObject(ComponentInits.SubobjectInits[Index].Subobject);
+	}
+	if (bNeedInstancing)
+	{
+		UObject* Archetype = Obj->GetArchetype();
+		Class->InstanceSubobjectTemplates(Obj, Archetype, Archetype ? Archetype->GetClass() : NULL, Obj, UseInstancingGraph);
+	}
+	if (bNeedSubobjectInstancing)
+	{
+		// initialize any subobjects, now that the constructors have run
+		for (int32 Index = 0; Index < ComponentInits.SubobjectInits.Num(); Index++)
+		{
+			UObject* Subobject = ComponentInits.SubobjectInits[Index].Subobject;
+			UObject* Template = ComponentInits.SubobjectInits[Index].Template;
+			if (!Subobject->HasAnyFlags(RF_NeedLoad))
+			{
+				Subobject->GetClass()->InstanceSubobjectTemplates(Subobject, Template, Template->GetClass(), Subobject, UseInstancingGraph);
+			}
+		}
+	}
+}
+
 UClass* FPostConstructInitializeProperties::GetClass() const
 {
 	return Obj->GetClass();
@@ -2549,4 +2566,47 @@ UScriptStruct* GetFallbackStruct()
 {
 	static UScriptStruct* FallbackStruct = FindObjectChecked<UScriptStruct>(UObject::StaticClass(), TEXT("FallbackStruct"));
 	return FallbackStruct;
+}
+
+UObject* FPostConstructInitializeProperties::CreateDefaultSubobject(UObject* Outer, FName SubobjectFName, UClass* ReturnType, UClass* ClassToCreateByDefault, bool bIsRequired, bool bAbstract, bool bIsTransient) const
+{
+	if (SubobjectFName == NAME_None)
+	{
+		UE_LOG(LogClass, Fatal, TEXT("Illegal default subobject name: %s"), *SubobjectFName.ToString());
+	}
+
+	UObject* Result = NULL;
+	UClass* OverrideClass = ComponentOverrides.Get(SubobjectFName, ReturnType, ClassToCreateByDefault, *this);
+	if (!OverrideClass && bIsRequired)
+	{
+		OverrideClass = ClassToCreateByDefault;
+		UE_LOG(LogClass, Warning, TEXT("Ignored DoNotCreateDefaultSubobject for %s as it's marked as required. Creating %s."), *SubobjectFName.ToString(), *OverrideClass->GetName());
+	}
+	if (OverrideClass)
+	{
+		check(OverrideClass->IsChildOf(ReturnType));
+
+		// Abstract sub-objects are only allowed when explicitly created with CreateAbstractDefaultSubobject.
+		if (!OverrideClass->HasAnyClassFlags(CLASS_Abstract) || !bAbstract)
+		{
+			UObject* Template = OverrideClass->GetDefaultObject(); // force the CDO to be created if it hasn't already
+			const EObjectFlags SubobjectFlags = Outer->GetMaskedFlags(RF_PropagateToSubObjects);
+			Result = StaticConstructObject(OverrideClass, Outer, SubobjectFName, SubobjectFlags);
+			if (!bIsTransient && !Outer->GetArchetype()->GetClass()->HasAnyClassFlags(CLASS_Native | CLASS_Intrinsic))
+			{
+				// The archetype of the outer is not native, so we need to copy properties to the subobjects after the C++ constructor chain for the outer has run (because those sets properties on the subobjects)
+				UObject* MaybeTemplate = Outer->GetArchetype()->GetClass()->GetDefaultSubobjectByName(SubobjectFName);
+				if (MaybeTemplate && MaybeTemplate->IsA(ReturnType) && Template != MaybeTemplate)
+				{
+					ComponentInits.Add(Result, MaybeTemplate);
+				}
+			}
+			if (Outer->HasAnyFlags(RF_ClassDefaultObject) && Outer->GetClass()->GetSuperClass())
+			{
+				Outer->GetClass()->AddDefaultSubobject(Result, ReturnType);
+			}
+			Result->SetFlags(RF_DefaultSubObject);
+		}
+	}
+	return Result;
 }
\ No newline at end of file
diff --git a/Engine/Source/Runtime/CoreUObject/Public/UObject/Class.h b/Engine/Source/Runtime/CoreUObject/Public/UObject/Class.h
--- a/Engine/Source/Runtime/CoreUObject/Public/UObject/Class.h
+++ b/Engine/Source/Runtime/CoreUObject/Public/UObject/Class.h
@@ -1844,6 +1844,18 @@ public:
 		class UClass* TClass_Super_StaticClass,
 		class UClass* TClass_WithinClass_StaticClass
 		);
+
+
+	/**
+	* Replace a native function in the  internal native function table
+	* @param	InName							name of the function
+	* @param	InPointer						pointer to the function
+	* @param	bAddToFunctionRemapTable		For C++ hot-reloading, UFunctions are patched in a deferred manner and this should be true
+	*											For script hot-reloading, script integrations may have a many to 1 mapping of UFunction to native pointer
+	*											because dispatch is shared, so the C++ remap table does not work in this case, and this should be false
+	* @return	true if the function was found and replaced, false if it was not
+	*/
+	bool ReplaceNativeFunction(FName InName, Native InPointer, bool bAddToFunctionRemapTable);
 #endif
 
 #if WITH_EDITOR
@@ -2161,11 +2173,12 @@ private:
 		return UObject::FindFunctionChecked(InName);
 	}
 
+protected:
 	/**
 	 * Get the default object from the class, creating it if missing, if requested or under a few other circumstances
 	 * @return		the CDO for this class
 	 **/
-	UObject* CreateDefaultObject();
+	virtual UObject* CreateDefaultObject();
 };
 
 
@@ -2540,47 +2553,3 @@ inline T* GetMutableDefault(UClass *Class)
 	checkSlow(Class->GetDefaultObject()->IsA(T::StaticClass()));
 	return (T*)Class->GetDefaultObject();
 }
-
-template<class TReturnType, class TClassToConstructByDefault>
-inline TReturnType* FPostConstructInitializeProperties::CreateDefaultSubobject(UObject* Outer, FName SubobjectFName, bool bIsRequired, bool bAbstract, bool bIsTransient) const
-{
-	if(SubobjectFName == NAME_None)
-	{
-		UE_LOG(LogClass, Fatal, TEXT("Illegal default subobject name: %s"), *SubobjectFName.ToString());
-	}
-
-	TReturnType* Result = NULL;
-	UClass* OverrideClass = ComponentOverrides.Get<TReturnType, TClassToConstructByDefault>(SubobjectFName, *this);
-	if (!OverrideClass && bIsRequired)
-	{
-		OverrideClass = TClassToConstructByDefault::StaticClass();
-		UE_LOG(LogClass, Warning, TEXT("Ignored DoNotCreateDefaultSubobject for %s as it's marked as required. Creating %s."), *SubobjectFName.ToString(), *OverrideClass->GetName());
-	}
-	if (OverrideClass)
-	{
-		check(OverrideClass->IsChildOf(TReturnType::StaticClass()));
-
-		// Abstract sub-objects are only allowed when explicitly created with CreateAbstractDefaultSubobject.
-		if (!OverrideClass->HasAnyClassFlags(CLASS_Abstract) || !bAbstract)
-		{
-			UObject* Template = OverrideClass->GetDefaultObject(); // force the CDO to be created if it hasn't already
-			const EObjectFlags SubobjectFlags = Outer->GetMaskedFlags(RF_PropagateToSubObjects);
-			Result = ConstructObject<TReturnType>(OverrideClass, Outer, SubobjectFName, SubobjectFlags);
-			if ( !bIsTransient && !Outer->GetArchetype()->GetClass()->HasAnyClassFlags(CLASS_Native | CLASS_Intrinsic) )
-			{
-				// The archetype of the outer is not native, so we need to copy properties to the subobjects after the C++ constructor chain for the outer has run (because those sets properties on the subobjects)
-				UObject* MaybeTemplate = Outer->GetArchetype()->GetClass()->GetDefaultSubobjectByName(SubobjectFName);
-				if (MaybeTemplate && MaybeTemplate->IsA(TReturnType::StaticClass()) && Template != MaybeTemplate)
-				{
-					ComponentInits.Add(Result, MaybeTemplate);
-				}
-			}
-			if (Outer->HasAnyFlags(RF_ClassDefaultObject) && Outer->GetClass()->GetSuperClass())
-			{
-				Outer->GetClass()->AddDefaultSubobject(Result, TReturnType::StaticClass());
-			}
-			Result->SetFlags(RF_DefaultSubObject);
-		}
-	}
-	return Result;
-}
diff --git a/Engine/Source/Runtime/CoreUObject/Public/UObject/UObjectGlobals.h b/Engine/Source/Runtime/CoreUObject/Public/UObject/UObjectGlobals.h
--- a/Engine/Source/Runtime/CoreUObject/Public/UObject/UObjectGlobals.h
+++ b/Engine/Source/Runtime/CoreUObject/Public/UObject/UObjectGlobals.h
@@ -698,7 +698,8 @@ public:
 	template<class TReturnType>
 	TSubobjectPtrConstructor<TReturnType> CreateDefaultSubobject(UObject* Outer, FName SubobjectName, bool bTransient = false) const
 	{
-		return TSubobjectPtrConstructor<TReturnType>(CreateDefaultSubobject<TReturnType, TReturnType>(Outer, SubobjectName, /*bIsRequired =*/ true, /*bIsAbstract =*/ false, bTransient));
+		UClass* ReturnType = TReturnType::StaticClass();
+		return TSubobjectPtrConstructor<TReturnType>(static_cast<TReturnType*>(CreateDefaultSubobject(Outer, SubobjectName, ReturnType, ReturnType, /*bIsRequired =*/ true, /*bIsAbstract =*/ false, bTransient)));
 	}
 
 	/**
@@ -712,7 +713,8 @@ public:
 	template<class TReturnType>
 	TSubobjectPtrConstructor<TReturnType> CreateOptionalDefaultSubobject(UObject* Outer, FName SubobjectName, bool bTransient = false) const
 	{
-		return TSubobjectPtrConstructor<TReturnType>(CreateDefaultSubobject<TReturnType, TReturnType>(Outer, SubobjectName, /*bIsRequired =*/ false, /*bIsAbstract =*/ false, bTransient));
+		UClass* ReturnType = TReturnType::StaticClass();
+		return TSubobjectPtrConstructor<TReturnType>(static_cast<TReturnType*>(CreateDefaultSubobject(Outer, SubobjectName, ReturnType, ReturnType, /*bIsRequired =*/ false, /*bIsAbstract =*/ false, bTransient)));
 	}
 
 	/**
@@ -726,7 +728,8 @@ public:
 	template<class TReturnType>
 	TSubobjectPtrConstructor<TReturnType> CreateAbstractDefaultSubobject(UObject* Outer, FName SubobjectName, bool bTransient = false) const
 	{
-		return TSubobjectPtrConstructor<TReturnType>(CreateDefaultSubobject<TReturnType, TReturnType>(Outer, SubobjectName, /*bIsRequired =*/ true, /*bIsAbstract =*/ true, bTransient));
+		UClass* ReturnType = TReturnType::StaticClass();
+		return TSubobjectPtrConstructor<TReturnType>(static_cast<TReturnType*>(CreateDefaultSubobject(Outer, SubobjectName, ReturnType, ReturnType, /*bIsRequired =*/ true, /*bIsAbstract =*/ true, bTransient)));
 	}
 
 	/** 
@@ -740,7 +743,7 @@ public:
 	template<class TReturnType, class TClassToConstructByDefault> 
 	TSubobjectPtrConstructor<TReturnType> CreateDefaultSubobject(UObject* Outer, FName SubobjectName, bool bTransient = false) const 
 	{ 
-		return TSubobjectPtrConstructor<TReturnType>(CreateDefaultSubobject<TReturnType, TClassToConstructByDefault>(Outer, SubobjectName, /*bIsRequired =*/ true, /*bIsAbstract =*/ false, bTransient));
+		return TSubobjectPtrConstructor<TReturnType>(static_cast<TReturnType*>(CreateDefaultSubobject(Outer, SubobjectName, TReturnType::StaticClass(), TClassToConstructByDefault::StaticClass(), /*bIsRequired =*/ true, /*bIsAbstract =*/ false, bTransient)));
 	}
 
 	/**
@@ -756,7 +759,8 @@ public:
 #if WITH_EDITOR
 		if (GIsEditor)
 		{
-			TReturnType* EditorSubobject = CreateDefaultSubobject<TReturnType, TReturnType>(Outer, SubobjectName, /*bIsRequired =*/ false, /*bIsAbstract =*/ false, bTransient);
+			UClass* ReturnType = TReturnType::StaticClass();
+			TReturnType* EditorSubobject = static_cast<TReturnType*>(CreateDefaultSubobject(Outer, SubobjectName, ReturnType, ReturnType, /*bIsRequired =*/ false, /*bIsAbstract =*/ false, bTransient));
 			if (EditorSubobject)
 			{
 				EditorSubobject->AlwaysLoadOnClient = false;
@@ -777,8 +781,7 @@ public:
 	 * @param bIsRequired			true if the component is required and will always be created even if DoNotCreateDefaultSubobject was sepcified.
 	 * @param bIsTransient		true if the component is being assigned to a transient property
 	 */
-	template<class TReturnType, class TClassToConstructByDefault>
-	TReturnType* CreateDefaultSubobject(UObject* Outer, FName SubobjectName, bool bIsRequired, bool bIsAbstract, bool bIsTransient) const;
+	UObject* CreateDefaultSubobject(UObject* Outer, FName SubobjectFName, UClass* ReturnType, UClass* ClassToCreateByDefault, bool bIsRequired, bool bAbstract, bool bIsTransient) const;
 
 	/**
 	 * Sets the class of a subobject for a base class
@@ -843,6 +846,7 @@ public:
 private:
 
 	friend class UObject; 
+	friend class FScriptIntegrationObjectHelper;
 
 	template<class T>
 	friend void InternalConstructor( const class FPostConstructInitializeProperties& X );
@@ -857,6 +861,23 @@ private:
 	 */
 	static void InitProperties(UObject* Obj, UClass* DefaultsClass, UObject* DefaultData, bool bCopyTransientsFromClassDefaults);
 
+	bool IsInstancingAllowed() const;
+
+	/**
+	 * Calls InitProperties for any default subobjects created through this PCIP.
+	 * @param bAllowInstancing	Indicates whether the object's components may be copied from their templates.
+	 * @return true if there are any subobjects which require instancing.
+	*/
+	bool InitSubobjectProperties(bool bAllowInstancing) const;
+
+	/**
+	 * Create copies of the object's components from their templates.
+	 * @param Class						Class of the object we are initializing
+	 * @param bNeedInstancing			Indicates whether the object's components need to be instanced
+	 * @param bNeedSubobjectInstancing	Indicates whether subobjects of the object's components need to be instanced
+	 */
+	void InstanceSubobjects(UClass* Class, bool bNeedInstancing, bool bNeedSubobjectInstancing) const;
+
 	/** 
 	 * Initializes a non-native property, according to the initialization rules. If the property is non-native
 	 * and does not have a zero contructor, it is inialized with the default value.
@@ -884,18 +905,17 @@ private:
 			}
 		}
 		/**  Retrieve an override, or TClassToConstructByDefault::StaticClass or NULL if this was removed by a derived class **/
-		template<class TReturnType, class TClassToConstructByDefault>
-		UClass* Get(FName InComponentName, FPostConstructInitializeProperties const& PCIP)
+		UClass* Get(FName InComponentName, UClass* ReturnType, UClass* ClassToConstructByDefault, FPostConstructInitializeProperties const& PCIP)
 		{
 			int32 Index = Find(InComponentName);
-			UClass *BaseComponentClass = TClassToConstructByDefault::StaticClass();
+			UClass *BaseComponentClass = ClassToConstructByDefault;
 			if (Index == INDEX_NONE)
 			{
 				return BaseComponentClass; // no override so just do what the base class wanted
 			}
 			else if (Overrides[Index].ComponentClass)
 			{
-				if (PCIP.IslegalOverride(InComponentName, Overrides[Index].ComponentClass, TReturnType::StaticClass())) // if THE base class is asking for a T, the existing override (which we are going to use) had better be derived
+				if (PCIP.IslegalOverride(InComponentName, Overrides[Index].ComponentClass, ReturnType)) // if THE base class is asking for a T, the existing override (which we are going to use) had better be derived
 				{
 					return Overrides[Index].ComponentClass; // the override is of an acceptable class, so use it
 				}
@@ -980,6 +1000,48 @@ private:
 	UObject* LastConstructedObject;
 };
 
+/**
+* Helper class for script integrations to access some UObject innards. Needed for script-generated UObject classes
+*/
+class FScriptIntegrationObjectHelper
+{
+public:
+	/**
+	* Binary initialize object properties to zero or defaults.
+	*
+	* @param	PCIP				FPostConstructInitializeProperties helper
+	* @param	Obj					object to initialize data for
+	* @param	DefaultsClass		the class to use for initializing the data
+	* @param	DefaultData			the buffer containing the source data for the initialization
+	*/
+	inline static void InitProperties(const FPostConstructInitializeProperties& PCIP, UObject* Obj, UClass* DefaultsClass, UObject* DefaultData)
+	{
+		FPostConstructInitializeProperties::InitProperties(Obj, DefaultsClass, DefaultData, PCIP.bCopyTransientsFromClassDefaults);
+	}
+
+	/**
+	* Calls InitProperties for any default subobjects created through this PCIP.
+	* @param bAllowInstancing	Indicates whether the object's components may be copied from their templates.
+	* @return true if there are any subobjects which require instancing.
+	*/
+	inline static bool InitSubobjectProperties(const FPostConstructInitializeProperties& PCIP)
+	{
+		return PCIP.InitSubobjectProperties(PCIP.IsInstancingAllowed());
+	}
+
+	/**
+	* Create copies of the object's components from their templates.
+	* @param PCIP					FPostConstructInitializeProperties helper
+	* @param Class						Class of the object we are initializing
+	* @param bNeedInstancing			Indicates whether the object's components need to be instanced
+	* @param bNeedSubobjectInstancing	Indicates whether subobjects of the object's components need to be instanced
+	*/
+	inline static void InstanceSubobjects(const FPostConstructInitializeProperties& PCIP, UClass* Class, bool bNeedInstancing, bool bNeedSubobjectInstancing)
+	{
+		PCIP.InstanceSubobjects(Class, bNeedInstancing, bNeedSubobjectInstancing);
+	}
+};
+
 /**
  * Helper class for deferred execution of 
 */
@@ -1547,4 +1609,3 @@ extern COREUOBJECT_API bool GShouldVerifyGCAssumptions;
 COREUOBJECT_API UScriptStruct* GetFallbackStruct();
 
 #endif	// __UNOBJGLOBALS_H__
-
