BP: Blueprint code cannot execute any code on an object, that is pending kill.

diff --git a/Engine/Source/Runtime/CoreUObject/Private/UObject/ScriptCore.cpp b/Engine/Source/Runtime/CoreUObject/Private/UObject/ScriptCore.cpp
--- a/Engine/Source/Runtime/CoreUObject/Private/UObject/ScriptCore.cpp
+++ b/Engine/Source/Runtime/CoreUObject/Private/UObject/ScriptCore.cpp
@@ -1392,7 +1392,7 @@ void UObject::ProcessContextOpcode( FFrame& Stack, RESULT_DECL, bool bCanFailSil
 	Stack.Step( this, &NewContext );
 
 	// Execute or skip the following expression in the object's context.
-	if (NewContext != NULL)
+	if (IsValid(NewContext))
 	{
 		Stack.Code += sizeof(CodeSkipSizeType) + sizeof(ScriptPointerType) + sizeof(uint8);
 		Stack.Step( NewContext, Result );
@@ -1401,7 +1401,14 @@ void UObject::ProcessContextOpcode( FFrame& Stack, RESULT_DECL, bool bCanFailSil
 	{
 		if (!bCanFailSilently)
 		{
-			if (Stack.MostRecentProperty != NULL)
+			if (NewContext && NewContext->IsPendingKill())
+			{
+				const FString Desc = FString::Printf(TEXT("Cannot access '%s'. It is pending kill. Property: '%s'")
+					, *GetNameSafe(NewContext), *GetNameSafe(Stack.MostRecentProperty));
+				FBlueprintExceptionInfo ExceptionInfo(EBlueprintExceptionType::AccessViolation, Desc);
+				FBlueprintCoreDelegates::ThrowScriptException(this, Stack, ExceptionInfo);
+			}
+			else if (Stack.MostRecentProperty != NULL)
 			{
 				const FString Desc = FString::Printf(TEXT("Accessed None '%s'"), *Stack.MostRecentProperty->GetName());
 				FBlueprintExceptionInfo ExceptionInfo(EBlueprintExceptionType::AccessViolation, Desc);
