Set RF_WasLoaded flag when saving packages to make sure this flag is consistently set on objects that exist on disk and in memory.

diff --git a/Engine/Source/Runtime/CoreUObject/Private/UObject/SavePackage.cpp b/Engine/Source/Runtime/CoreUObject/Private/UObject/SavePackage.cpp
--- a/Engine/Source/Runtime/CoreUObject/Private/UObject/SavePackage.cpp
+++ b/Engine/Source/Runtime/CoreUObject/Private/UObject/SavePackage.cpp
@@ -3607,6 +3607,22 @@ bool UPackage::SavePackage( UPackage* InOuter, UObject* Base, EObjectFlags TopLe
 					}
 					else
 					{
+						// Mark exports and the package as RF_Loaded after they've been serialized
+						// This is to ensue that newly created packages are properly marked as loaded (since they now exist on disk and 
+						// in memory in the exact same state).
+						for (auto& Export : Linker->ExportMap)
+						{
+							if (Export.Object)
+							{
+								Export.Object->SetFlags(RF_WasLoaded);
+							}
+						}
+						if (Linker->LinkerRoot)
+						{
+							// And finally set the flag on the package itself.
+							Linker->LinkerRoot->SetFlags(RF_WasLoaded);
+						}
+
 						// Clear dirty flag if desired
 						if (!(SaveFlags & SAVE_KeepDirty))
 						{
