CoreUObject - Removed ConditionalShutdownAfterError, replaced with RegisterForShutdownAfterError, which needs to be called if an object requires such funcionality, this change will dramatically reduce the amount of memory required for calling FUObjectShutdown::ShutdownAfterError, from hundreds of MBs to a few KBs

diff --git a/Engine/Source/Runtime/CoreUObject/Private/UObject/Obj.cpp b/Engine/Source/Runtime/CoreUObject/Private/UObject/Obj.cpp
--- a/Engine/Source/Runtime/CoreUObject/Private/UObject/Obj.cpp
+++ b/Engine/Source/Runtime/CoreUObject/Private/UObject/Obj.cpp
@@ -1812,41 +1812,26 @@ void UObject::CultureChange()
    Shutdown.
 -----------------------------------------------------------------------------*/
 
-/** Used to track which objects have been shutdown after error			*/
-static FUObjectAnnotationSparseBool ShutdownAfterErrorAnnotation;
-
-
-void UObject::ConditionalShutdownAfterError()
-{
-	if( !ShutdownAfterErrorAnnotation.Get(this) )
-	{
-		ShutdownAfterErrorAnnotation.Set(this);
-		ShutdownAfterError();
-	}
-}
-
 /**
  * After a critical error, shutdown all objects which require
  * mission-critical cleanup, such as restoring the video mode,
  * releasing hardware resources.
  */
-void StaticShutdownAfterError()
+static void StaticShutdownAfterError()
 {
 	if( UObjectInitialized() )
 	{
-		static bool Shutdown=0;
-		if( Shutdown )
+		static bool bShutdown = false;
+		if( bShutdown )
 		{
 			return;
 		}
-		Shutdown = 1;
+		bShutdown = true;
 		UE_LOG(LogExit, Log, TEXT("Executing StaticShutdownAfterError") );
 
-		ShutdownAfterErrorAnnotation.Reserve( GUObjectArray.GetObjectArrayNum() );
-
-		for ( FRawObjectIterator It; It; ++It )
+		for( FRawObjectIterator It; It; ++It )
 		{
-			It->ConditionalShutdownAfterError();
+			It->ShutdownAfterError();
 		}
 	}
 }
diff --git a/Engine/Source/Runtime/CoreUObject/Public/UObject/UObject.h b/Engine/Source/Runtime/CoreUObject/Public/UObject/UObject.h
--- a/Engine/Source/Runtime/CoreUObject/Public/UObject/UObject.h
+++ b/Engine/Source/Runtime/CoreUObject/Public/UObject/UObject.h
@@ -561,10 +561,6 @@ public:
 	 */
 	void ConditionalPostLoadSubobjects( struct FObjectInstancingGraph* OuterInstanceGraph=NULL );
 
-	/** Make sure this object has been shut down. */
-	void ConditionalShutdownAfterError();
-
-
 	/**
 	 * Starts caching of platform specific data for the target platform
 	 * Called when cooking before serialization so that object can prepare platform specific data
diff --git a/Engine/Source/Runtime/CoreUObject/Public/UObject/UObjectGlobals.h b/Engine/Source/Runtime/CoreUObject/Public/UObject/UObjectGlobals.h
--- a/Engine/Source/Runtime/CoreUObject/Public/UObject/UObjectGlobals.h
+++ b/Engine/Source/Runtime/CoreUObject/Public/UObject/UObjectGlobals.h
@@ -419,7 +419,6 @@ COREUOBJECT_API UPackage* FindPackage(UObject* InOuter, const TCHAR* PackageName
  */
 COREUOBJECT_API UPackage* CreatePackage( UObject* InOuter, const TCHAR* PackageName );
 
-void StaticShutdownAfterError();
 void GlobalSetProperty( const TCHAR* Value, UClass* Class, UProperty* Property, bool bNotifyObjectOfChange );
 
 /**
