Cast unification.

* Cast can be used for every UClass/interface cast:
  - Up/down UObject hierarchy.
  - From interfaces to UObjects.
  - From UObjects to interfaces.
  - Between interfaces.
* InterfaceCast has been deprecated.
* dynamic_cast can now be used instead of Cast.
* dynamic_cast between two UType pointers will be equivalent to Cast, but with dynamic_cast syntax and semantics.
* dynamic_cast of references is equivalent to a CastChecked - if it fails, it will fatally log rather than throw a std::bad_cast exception (UE doesn't support exceptions).
* If the two types involved are not UTypes, it should continue to act like a normal dynamic_cast.
* dynamic_cast<void*> works as expected.
* All/many Casts in CoreUObject changed to dynamic_cast to show this behaviour.
* Some refactoring.

diff --git a/Engine/Source/Editor/Matinee/Private/Matinee.cpp b/Engine/Source/Editor/Matinee/Private/Matinee.cpp
--- a/Engine/Source/Editor/Matinee/Private/Matinee.cpp
+++ b/Engine/Source/Editor/Matinee/Private/Matinee.cpp
@@ -2606,7 +2606,7 @@ void FMatinee::NotifyPostChange( const FPropertyChangedEvent& PropertyChangedEve
 					AActor* Actor = Inst->GetGroupActor();
 					if(Actor)
 					{
-						IMatineeAnimInterface * MatineeAnimInterface = InterfaceCast<IMatineeAnimInterface>(Actor);
+						IMatineeAnimInterface * MatineeAnimInterface = Cast<IMatineeAnimInterface>(Actor);
 						if (MatineeAnimInterface)
 						{
 							MatineeAnimInterface->PreviewBeginAnimControl(CurrentSelectedGroup);
diff --git a/Engine/Source/Editor/Matinee/Private/MatineeTrackHelpers.cpp b/Engine/Source/Editor/Matinee/Private/MatineeTrackHelpers.cpp
--- a/Engine/Source/Editor/Matinee/Private/MatineeTrackHelpers.cpp
+++ b/Engine/Source/Editor/Matinee/Private/MatineeTrackHelpers.cpp
@@ -94,7 +94,7 @@ bool UMatineeTrackAnimControlHelper::PreCreateTrack( UInterpGroup* Group, const
 	AActor* Actor = GrInst->GetGroupActor();
 	if ( Actor != NULL )
 	{
-		IMatineeAnimInterface* MatineeAnimInterface = InterfaceCast<IMatineeAnimInterface>(Actor);
+		IMatineeAnimInterface* MatineeAnimInterface = Cast<IMatineeAnimInterface>(Actor);
 		if (!MatineeAnimInterface)
 		{
 			UE_LOG(LogSlateMatinee, Log, TEXT("IntepGroup : MatineeAnimInterface is missing for (%s)"), *Actor->GetName());
diff --git a/Engine/Source/Editor/UnrealEd/Private/Kismet2/EnumEditorUtils.cpp b/Engine/Source/Editor/UnrealEd/Private/Kismet2/EnumEditorUtils.cpp
--- a/Engine/Source/Editor/UnrealEd/Private/Kismet2/EnumEditorUtils.cpp
+++ b/Engine/Source/Editor/UnrealEd/Private/Kismet2/EnumEditorUtils.cpp
@@ -253,7 +253,7 @@ void FEnumEditorUtils::BroadcastChanges(const UUserDefinedEnum* Enum, const TArr
 		for (TObjectIterator<UK2Node> It(RF_Transient); It; ++It)
 		{
 			UK2Node* Node = *It;
-			INodeDependingOnEnumInterface* NodeDependingOnEnum = InterfaceCast<INodeDependingOnEnumInterface>(Node);
+			INodeDependingOnEnumInterface* NodeDependingOnEnum = Cast<INodeDependingOnEnumInterface>(Node);
 			if (FNodeValidatorHelper::IsValid(Node) && NodeDependingOnEnum && (Enum == NodeDependingOnEnum->GetEnum()))
 			{
 				if (UBlueprint* Blueprint = Node->GetBlueprint())
diff --git a/Engine/Source/Programs/UnrealHeaderTool/Private/CodeGenerator.cpp b/Engine/Source/Programs/UnrealHeaderTool/Private/CodeGenerator.cpp
--- a/Engine/Source/Programs/UnrealHeaderTool/Private/CodeGenerator.cpp
+++ b/Engine/Source/Programs/UnrealHeaderTool/Private/CodeGenerator.cpp
@@ -1755,8 +1755,7 @@ void FNativeClassHeaderGenerator::ExportInterfaceClassDeclaration( FClass* Class
 		// to script VM functions
 		if (SuperClass->IsChildOf(UInterface::StaticClass()))
 		{
-			InterfaceBoilerplate.Logf(TEXT("\tvirtual UObject* GetUObjectInterface%s()=0;\r\n"), *Class->GetName());
-			InterfaceBoilerplate.Logf(TEXT("\tvirtual const UObject* GetUObjectInterface%s() const=0;\r\n"), *Class->GetName());
+			InterfaceBoilerplate.Logf(TEXT("\tvirtual UObject* __GetUObject() const = 0;\r\n"));
 		}
 
 		// Replication, add in the declaration for GetLifetimeReplicatedProps() automatically if there are any net flagged properties
@@ -1919,30 +1918,7 @@ void FNativeClassHeaderGenerator::ExportClassHeaderInner(FClass* Class, bool bVa
 					ClassBoilerplate.Logf(TEXT("    static const TCHAR* StaticConfigName() {return TEXT(\"%s\");}\r\n\r\n"), *Class->ClassConfigName.ToString());
 				}
 
-				{
-					// arrays for preventing multiple export of accessor function in situations where the native class
-					// implements multiple children of the same interface base class
-					TArray<UClass*> UObjectExportedInterfaces;
-					TArray<UClass*> ExportedInterfaces;
-
-					for (TArray<FImplementedInterface>::TIterator It(Class->Interfaces); It; ++It)
-					{
-						UClass* InterfaceClass = It->Class;
-						if ( InterfaceClass->HasAnyClassFlags(CLASS_Native) )
-						{
-							for ( UClass* IClass = InterfaceClass; IClass && IClass->HasAnyClassFlags(CLASS_Interface); IClass = IClass->GetSuperClass() )
-							{
-								if ( IClass != UInterface::StaticClass() && !UObjectExportedInterfaces.Contains(IClass) )
-								{
-									UObjectExportedInterfaces.Add(IClass);
-									ClassBoilerplate.Logf(TEXT("%svirtual UObject* GetUObjectInterface%s(){return this;}\r\n"), FCString::Spc(4), *IClass->GetName());
-									ClassBoilerplate.Logf(TEXT("%svirtual const UObject* GetUObjectInterface%s() const{return this;}\r\n"), FCString::Spc(4), *IClass->GetName());
-								}
-							}
-						}
-					}
-				}
-
+				ClassBoilerplate.Logf(TEXT("%sUObject* __GetUObject() const { return const_cast<%s*>(this); }\r\n"), FCString::Spc(4), ClassCPPName);
 
 				// Replication, add in the declaration for GetLifetimeReplicatedProps() automatically if there are any net flagged properties
 				for( TFieldIterator<UProperty> It(Class,EFieldIteratorFlags::ExcludeSuper); It; ++It )
diff --git a/Engine/Source/Runtime/AIModule/Classes/Navigation/NavigationComponent.h b/Engine/Source/Runtime/AIModule/Classes/Navigation/NavigationComponent.h
--- a/Engine/Source/Runtime/AIModule/Classes/Navigation/NavigationComponent.h
+++ b/Engine/Source/Runtime/AIModule/Classes/Navigation/NavigationComponent.h
@@ -126,7 +126,7 @@ public:
 	/** Returns current move goal for actor, using data from GetMoveGoalOffset() */
 	FORCEINLINE FVector GetCurrentMoveGoal(const AActor* InGoalActor, AActor* MovingActor) const
 	{
-		const INavAgentInterface* NavAgent = InterfaceCast<const INavAgentInterface>(InGoalActor);
+		const INavAgentInterface* NavAgent = Cast<const INavAgentInterface>(InGoalActor);
 		if (NavAgent)
 		{
 			const FVector Offset = NavAgent->GetMoveGoalOffset(MovingActor);
diff --git a/Engine/Source/Runtime/AIModule/Private/AIInterfaces.cpp b/Engine/Source/Runtime/AIModule/Private/AIInterfaces.cpp
--- a/Engine/Source/Runtime/AIModule/Private/AIInterfaces.cpp
+++ b/Engine/Source/Runtime/AIModule/Private/AIInterfaces.cpp
@@ -13,7 +13,7 @@ FGenericTeamId FGenericTeamId::NoTeam(FGenericTeamId::NoTeamId);
 
 FGenericTeamId FGenericTeamId::GetTeamIdentifier(const AActor* TeamMember)
 {
-	const IGenericTeamAgentInterface* TeamAgent = InterfaceCast<const IGenericTeamAgentInterface>(TeamMember);
+	const IGenericTeamAgentInterface* TeamAgent = Cast<const IGenericTeamAgentInterface>(TeamMember);
 	if (TeamAgent)
 	{
 		return TeamAgent->GetGenericTeamId();
@@ -24,7 +24,7 @@ FGenericTeamId FGenericTeamId::GetTeamIdentifier(const AActor* TeamMember)
 
 ETeamAttitude::Type FGenericTeamId::GetAttitude(const AActor* A, const AActor* B)
 {
-	const IGenericTeamAgentInterface* TeamAgentA = InterfaceCast<const IGenericTeamAgentInterface>(A);
+	const IGenericTeamAgentInterface* TeamAgentA = Cast<const IGenericTeamAgentInterface>(A);
 
 	return TeamAgentA == NULL || B == NULL ? ETeamAttitude::Neutral : TeamAgentA->GetTeamAttitudeTowards(*B);
 }
diff --git a/Engine/Source/Runtime/AIModule/Private/EnvironmentQuery/EQSRenderingComponent.cpp b/Engine/Source/Runtime/AIModule/Private/EnvironmentQuery/EQSRenderingComponent.cpp
--- a/Engine/Source/Runtime/AIModule/Private/EnvironmentQuery/EQSRenderingComponent.cpp
+++ b/Engine/Source/Runtime/AIModule/Private/EnvironmentQuery/EQSRenderingComponent.cpp
@@ -43,10 +43,10 @@ FEQSSceneProxy::FEQSSceneProxy(const UPrimitiveComponent* InComponent, const FSt
 	}
 
 	ActorOwner = InComponent ? InComponent->GetOwner() : NULL;
-	QueryDataSource = InterfaceCast<const IEQSQueryResultSourceInterface>(ActorOwner);
+	QueryDataSource = Cast<const IEQSQueryResultSourceInterface>(ActorOwner);
 	if (QueryDataSource == NULL)
 	{
-		QueryDataSource = InterfaceCast<const IEQSQueryResultSourceInterface>(InComponent);
+		QueryDataSource = Cast<const IEQSQueryResultSourceInterface>(InComponent);
 		if (QueryDataSource == NULL)
 		{
 			return;
@@ -78,10 +78,10 @@ FEQSSceneProxy::FEQSSceneProxy(const UPrimitiveComponent* InComponent, const FSt
 	}
 
 	ActorOwner = InComponent ? InComponent->GetOwner() : NULL;
-	QueryDataSource = InterfaceCast<const IEQSQueryResultSourceInterface>(ActorOwner);
+	QueryDataSource = Cast<const IEQSQueryResultSourceInterface>(ActorOwner);
 	if (QueryDataSource == NULL)
 	{
-		QueryDataSource = InterfaceCast<const IEQSQueryResultSourceInterface>(InComponent);
+		QueryDataSource = Cast<const IEQSQueryResultSourceInterface>(InComponent);
 		if (QueryDataSource == NULL)
 		{
 			return;
@@ -96,10 +96,10 @@ void FEQSSceneProxy::CollectEQSData(const UPrimitiveComponent* InComponent, cons
 	IEQSQueryResultSourceInterface* QueryDataSource = const_cast<IEQSQueryResultSourceInterface*>(InQueryDataSource);
 	if (QueryDataSource == NULL)
 	{
-		QueryDataSource = InterfaceCast<IEQSQueryResultSourceInterface>(ActorOwner);
+		QueryDataSource = Cast<IEQSQueryResultSourceInterface>(ActorOwner);
 		if (QueryDataSource == NULL)
 		{
-			QueryDataSource = InterfaceCast<IEQSQueryResultSourceInterface>(const_cast<UPrimitiveComponent*>(InComponent));
+			QueryDataSource = Cast<IEQSQueryResultSourceInterface>(const_cast<UPrimitiveComponent*>(InComponent));
 			if (QueryDataSource == NULL)
 			{
 				return;
diff --git a/Engine/Source/Runtime/AIModule/Private/EnvironmentQuery/EnvQueryTypes.cpp b/Engine/Source/Runtime/AIModule/Private/EnvironmentQuery/EnvQueryTypes.cpp
--- a/Engine/Source/Runtime/AIModule/Private/EnvironmentQuery/EnvQueryTypes.cpp
+++ b/Engine/Source/Runtime/AIModule/Private/EnvironmentQuery/EnvQueryTypes.cpp
@@ -235,7 +235,7 @@ namespace FEQSHelpers
 		const UNavigationSystem* NavSys = QueryInstance.World->GetNavigationSystem();
 
 		// try to match navigation agent for querier
-		INavAgentInterface* NavAgent = QueryInstance.Owner.IsValid() ? InterfaceCast<INavAgentInterface>(QueryInstance.Owner.Get()) : NULL;
+		INavAgentInterface* NavAgent = QueryInstance.Owner.IsValid() ? Cast<INavAgentInterface>(QueryInstance.Owner.Get()) : NULL;
 		if (NavAgent)
 		{
 			const FNavAgentProperties* NavAgentProps = NavAgent ? NavAgent->GetNavAgentProperties() : NULL;
diff --git a/Engine/Source/Runtime/AIModule/Private/EnvironmentQuery/Tests/EnvQueryTest_Pathfinding.cpp b/Engine/Source/Runtime/AIModule/Private/EnvironmentQuery/Tests/EnvQueryTest_Pathfinding.cpp
--- a/Engine/Source/Runtime/AIModule/Private/EnvironmentQuery/Tests/EnvQueryTest_Pathfinding.cpp
+++ b/Engine/Source/Runtime/AIModule/Private/EnvironmentQuery/Tests/EnvQueryTest_Pathfinding.cpp
@@ -242,7 +242,7 @@ float UEnvQueryTest_Pathfinding::FindPathLengthTo(const FVector& ItemPos, const
 
 ANavigationData* UEnvQueryTest_Pathfinding::FindNavigationData(UNavigationSystem* NavSys, UObject* Owner) const
 {
-	INavAgentInterface* NavAgent = InterfaceCast<INavAgentInterface>(Owner);
+	INavAgentInterface* NavAgent = Cast<INavAgentInterface>(Owner);
 	if (NavAgent)
 	{
 		return NavSys->GetNavDataForProps(*NavAgent->GetNavAgentProperties());
diff --git a/Engine/Source/Runtime/AIModule/Private/Navigation/CrowdFollowingComponent.cpp b/Engine/Source/Runtime/AIModule/Private/Navigation/CrowdFollowingComponent.cpp
--- a/Engine/Source/Runtime/AIModule/Private/Navigation/CrowdFollowingComponent.cpp
+++ b/Engine/Source/Runtime/AIModule/Private/Navigation/CrowdFollowingComponent.cpp
@@ -204,7 +204,7 @@ void UCrowdFollowingComponent::UpdateCrowdAgentParams() const
 	UCrowdManager* CrowdManager = UCrowdManager::GetCurrent(GetWorld());
 	if (CrowdManager)
 	{
-		const ICrowdAgentInterface* IAgent = InterfaceCast<ICrowdAgentInterface>(this);
+		const ICrowdAgentInterface* IAgent = Cast<ICrowdAgentInterface>(this);
 		CrowdManager->UpdateAgentParams(IAgent);
 	}
 }
@@ -284,7 +284,7 @@ void UCrowdFollowingComponent::Initialize()
 	UCrowdManager* CrowdManager = UCrowdManager::GetCurrent(GetWorld());
 	if (CrowdManager)
 	{
-		const ICrowdAgentInterface* IAgent = InterfaceCast<ICrowdAgentInterface>(this);
+		const ICrowdAgentInterface* IAgent = Cast<ICrowdAgentInterface>(this);
 		CrowdManager->RegisterAgent(IAgent);
 	}
 	else
@@ -307,7 +307,7 @@ void UCrowdFollowingComponent::Cleanup()
 	UCrowdManager* CrowdManager = UCrowdManager::GetCurrent(GetWorld());
 	if (CrowdManager)
 	{
-		const ICrowdAgentInterface* IAgent = InterfaceCast<ICrowdAgentInterface>(this);
+		const ICrowdAgentInterface* IAgent = Cast<ICrowdAgentInterface>(this);
 		CrowdManager->UnregisterAgent(IAgent);
 	}
 }
@@ -379,7 +379,7 @@ void UCrowdFollowingComponent::OnLanded()
 	UCrowdManager* CrowdManager = UCrowdManager::GetCurrent(GetWorld());
 	if (bEnableCrowdSimulation && CrowdManager)
 	{
-		const ICrowdAgentInterface* IAgent = InterfaceCast<ICrowdAgentInterface>(this);
+		const ICrowdAgentInterface* IAgent = Cast<ICrowdAgentInterface>(this);
 		CrowdManager->UpdateAgentState(IAgent);
 	}
 }
@@ -395,7 +395,7 @@ void UCrowdFollowingComponent::FinishUsingCustomLink(class INavLinkCustomInterfa
 		UCrowdManager* CrowdManager = UCrowdManager::GetCurrent(GetWorld());
 		if (bPrevCustomLink && !bCurrentCustomLink && CrowdManager)
 		{
-			const ICrowdAgentInterface* IAgent = InterfaceCast<ICrowdAgentInterface>(this);
+			const ICrowdAgentInterface* IAgent = Cast<ICrowdAgentInterface>(this);
 			CrowdManager->OnAgentFinishedCustomLink(IAgent);
 		}
 	}
diff --git a/Engine/Source/Runtime/AIModule/Private/Navigation/CrowdManager.cpp b/Engine/Source/Runtime/AIModule/Private/Navigation/CrowdManager.cpp
--- a/Engine/Source/Runtime/AIModule/Private/Navigation/CrowdManager.cpp
+++ b/Engine/Source/Runtime/AIModule/Private/Navigation/CrowdManager.cpp
@@ -248,7 +248,7 @@ void UCrowdManager::Tick(float DeltaTime)
 					const FCrowdAgentData& AgentData = It.Value();
 					if (AgentData.bIsSimulated && AgentData.IsValid())
 					{
-						const UCrowdFollowingComponent* CrowdComponent = Cast<const UCrowdFollowingComponent>(It.Key()->GetUObjectInterfaceCrowdAgentInterface());
+						const UCrowdFollowingComponent* CrowdComponent = Cast<const UCrowdFollowingComponent>(It.Key());
 						if (CrowdComponent)
 						{
 							ApplyVelocity(CrowdComponent, AgentData.AgentIndex);
@@ -610,7 +610,7 @@ void UCrowdManager::AddAgent(const ICrowdAgentInterface* Agent, FCrowdAgentData&
 	// create link filter for fully simulated agents
 	// (used to determine if agent can traverse smart links)
 	FRecastSpeciaLinkFilter* MyLinkFilter = NULL;
-	const UCrowdFollowingComponent* CrowdComponent = Cast<const UCrowdFollowingComponent>(Agent->GetUObjectInterfaceCrowdAgentInterface());
+	const UCrowdFollowingComponent* CrowdComponent = Cast<const UCrowdFollowingComponent>(Agent);
 	if (CrowdComponent)
 	{
 		UNavigationSystem* NavSys = Cast<UNavigationSystem>(GetOuter());
@@ -646,7 +646,7 @@ void UCrowdManager::GetAgentParams(const ICrowdAgentInterface* Agent, struct dtC
 	// skip maxSpeed, it will be constantly updated in every tick
 	// skip maxAcceleration, we don't use Detour's movement code
 
-	const UCrowdFollowingComponent* CrowdComponent = Cast<const UCrowdFollowingComponent>(Agent->GetUObjectInterfaceCrowdAgentInterface());
+	const UCrowdFollowingComponent* CrowdComponent = Cast<const UCrowdFollowingComponent>(Agent);
 	if (CrowdComponent)
 	{
 		AgentParams->collisionQueryRange = CrowdComponent->GetCrowdCollisionQueryRange();
@@ -729,7 +729,7 @@ void UCrowdManager::UpdateAgentPaths()
 		FCrowdAgentData& AgentData = It.Value();
 		if (AgentData.bIsSimulated && AgentData.IsValid())
 		{
-			UCrowdFollowingComponent* CrowdComponent = (UCrowdFollowingComponent*)Cast<const UCrowdFollowingComponent>(It.Key()->GetUObjectInterfaceCrowdAgentInterface());
+			UCrowdFollowingComponent* CrowdComponent = (UCrowdFollowingComponent*)Cast<const UCrowdFollowingComponent>(It.Key());
 
 			const dtCrowdAgent* Agent = DetourCrowd->getAgent(AgentData.AgentIndex);
 			dtPolyRef AgentPolyRef = Agent->corridor.getFirstPoly();
@@ -775,7 +775,7 @@ void UCrowdManager::UpdateAgentPaths()
 void UCrowdManager::UpdateSelectedDebug(const ICrowdAgentInterface* Agent, int32 AgentIndex) const
 {
 #if WITH_EDITOR
-	const UObject* Obj = Agent->GetUObjectInterfaceCrowdAgentInterface();
+	const UObject* Obj = Cast<const UObject>(Agent);
 	if (GIsEditor && Obj)
 	{
 		const AController* TestController = Cast<const AController>(Obj->GetOuter());
@@ -1026,7 +1026,7 @@ void UCrowdManager::DebugTick() const
 		for (auto It = ActiveAgents.CreateConstIterator(); It; ++It)
 		{
 			const ICrowdAgentInterface* IAgent = It.Key();
-			const UObject* AgentOb = IAgent ? IAgent->GetUObjectInterfaceCrowdAgentInterface() : NULL;
+			const UObject* AgentOb = IAgent ?  Cast<const UObject>(IAgent) : NULL;
 			const AActor* LogOwner = AgentOb ? Cast<const AActor>(AgentOb->GetOuter()) : NULL;
 
 			const FCrowdAgentData& AgentData = It.Value();
diff --git a/Engine/Source/Runtime/AIModule/Private/Navigation/NavigationComponent.cpp b/Engine/Source/Runtime/AIModule/Private/Navigation/NavigationComponent.cpp
--- a/Engine/Source/Runtime/AIModule/Private/Navigation/NavigationComponent.cpp
+++ b/Engine/Source/Runtime/AIModule/Private/Navigation/NavigationComponent.cpp
@@ -824,12 +824,12 @@ ANavigationData* UNavigationComponent::PickNavData() const
 
 	if (MyNavAgent == NULL)
 	{
-		MyNavAgent = InterfaceCast<INavAgentInterface>( GetOuter() );
+		MyNavAgent = Cast<INavAgentInterface>( GetOuter() );
 	}
 
 	if (MyPathObserver == NULL)
 	{
-		MyPathObserver = InterfaceCast<INavPathObserverInterface>( GetOuter() );
+		MyPathObserver = Cast<INavPathObserverInterface>( GetOuter() );
 	}
 
 	if (GetWorld() != NULL && GetWorld()->GetNavigationSystem() != NULL && MyNavAgent != NULL)
diff --git a/Engine/Source/Runtime/AIModule/Private/Navigation/PathFollowingComponent.cpp b/Engine/Source/Runtime/AIModule/Private/Navigation/PathFollowingComponent.cpp
--- a/Engine/Source/Runtime/AIModule/Private/Navigation/PathFollowingComponent.cpp
+++ b/Engine/Source/Runtime/AIModule/Private/Navigation/PathFollowingComponent.cpp
@@ -243,7 +243,7 @@ void UPathFollowingComponent::AbortMove(const FString& Reason, FAIRequestID Requ
 	{
 		const EPathFollowingResult::Type FinishResult = bSilent ? EPathFollowingResult::Skipped : EPathFollowingResult::Aborted;
 
-		INavLinkCustomInterface* CustomNavLink = InterfaceCast<INavLinkCustomInterface>(CurrentCustomLinkOb.Get());
+		INavLinkCustomInterface* CustomNavLink = Cast<INavLinkCustomInterface>(CurrentCustomLinkOb.Get());
 		if (CustomNavLink)
 		{
 			CustomNavLink->OnLinkMoveFinished(this);
@@ -546,7 +546,7 @@ int32 UPathFollowingComponent::DetermineStartingPathPoint(const FNavigationPath*
 void UPathFollowingComponent::SetDestinationActor(const AActor* InDestinationActor)
 {
 	DestinationActor = InDestinationActor;
-	DestinationAgent = InterfaceCast<const INavAgentInterface>(InDestinationActor);
+	DestinationAgent = Cast<const INavAgentInterface>(InDestinationActor);
 	MoveOffset = DestinationAgent ? DestinationAgent->GetMoveGoalOffset(GetOwner()) : FVector::ZeroVector;
 }
 
@@ -848,7 +848,7 @@ bool UPathFollowingComponent::HasReached(const AActor* TestGoal, float InAccepta
 		InAcceptanceRadius = MyDefaultAcceptanceRadius;
 	}
 
-	const INavAgentInterface* NavAgent = InterfaceCast<const INavAgentInterface>(TestGoal);
+	const INavAgentInterface* NavAgent = Cast<const INavAgentInterface>(TestGoal);
 	if (NavAgent)
 	{
 		const FVector GoalMoveOffset = NavAgent->GetMoveGoalOffset(GetOwner());
@@ -1005,7 +1005,7 @@ void UPathFollowingComponent::DebugReachTest(float& CurrentDot, float& CurrentDi
 
 void UPathFollowingComponent::StartUsingCustomLink(class INavLinkCustomInterface* CustomNavLink, const FVector& DestPoint)
 {
-	INavLinkCustomInterface* PrevNavLink = InterfaceCast<INavLinkCustomInterface>(CurrentCustomLinkOb.Get());
+	INavLinkCustomInterface* PrevNavLink = Cast<INavLinkCustomInterface>(CurrentCustomLinkOb.Get());
 	if (PrevNavLink)
 	{
 		UE_VLOG(GetOwner(), LogPathFollowing, Log, TEXT("Force finish custom move using navlink: %s"), *GetNameSafe(CurrentCustomLinkOb.Get()));
@@ -1014,7 +1014,7 @@ void UPathFollowingComponent::StartUsingCustomLink(class INavLinkCustomInterface
 		CurrentCustomLinkOb.Reset();
 	}
 
-	UObject* NewNavLinkOb = CustomNavLink ? CustomNavLink->GetUObjectInterfaceNavLinkCustomInterface() : NULL;
+	UObject* NewNavLinkOb = Cast<UObject>(CustomNavLink);
 	if (NewNavLinkOb)
 	{
 		CurrentCustomLinkOb = NewNavLinkOb;
@@ -1032,7 +1032,7 @@ void UPathFollowingComponent::StartUsingCustomLink(class INavLinkCustomInterface
 
 void UPathFollowingComponent::FinishUsingCustomLink(class INavLinkCustomInterface* CustomNavLink)
 {
-	UObject* NavLinkOb = CustomNavLink ? CustomNavLink->GetUObjectInterfaceNavLinkCustomInterface() : NULL;
+	UObject* NavLinkOb = Cast<UObject>(CustomNavLink);
 	if (CurrentCustomLinkOb == NavLinkOb)
 	{
 		UE_VLOG(GetOwner(), LogPathFollowing, Log, TEXT("Finish custom move using navlink: %s"), *GetNameSafe(NavLinkOb));
diff --git a/Engine/Source/Runtime/AIModule/Private/Perception/AIPerceptionTypes.cpp b/Engine/Source/Runtime/AIModule/Private/Perception/AIPerceptionTypes.cpp
--- a/Engine/Source/Runtime/AIModule/Private/Perception/AIPerceptionTypes.cpp
+++ b/Engine/Source/Runtime/AIModule/Private/Perception/AIPerceptionTypes.cpp
@@ -86,6 +86,6 @@ const IGenericTeamAgentInterface* FPerceptionListener::GetTeamAgent() const
 	}
 
 	const AActor* OwnerActor = PercComponent->GetOwner();
-	const IGenericTeamAgentInterface* OwnerTeamAgent = InterfaceCast<const IGenericTeamAgentInterface>(OwnerActor);
-	return OwnerTeamAgent != NULL ? OwnerTeamAgent : InterfaceCast<const IGenericTeamAgentInterface>(PercComponent->GetBodyActor());
+	const IGenericTeamAgentInterface* OwnerTeamAgent = Cast<const IGenericTeamAgentInterface>(OwnerActor);
+	return OwnerTeamAgent != NULL ? OwnerTeamAgent : Cast<const IGenericTeamAgentInterface>(PercComponent->GetBodyActor());
 }
\ No newline at end of file
diff --git a/Engine/Source/Runtime/AIModule/Private/Perception/AISense_Damage.cpp b/Engine/Source/Runtime/AIModule/Private/Perception/AISense_Damage.cpp
--- a/Engine/Source/Runtime/AIModule/Private/Perception/AISense_Damage.cpp
+++ b/Engine/Source/Runtime/AIModule/Private/Perception/AISense_Damage.cpp
@@ -17,7 +17,7 @@ float UAISense_Damage::Update()
 	{
 		const FAIDamageEvent& Event = RegisteredEvents[EventIndex];
 
-		IAIPerceptionListenerInterface* PerceptionListener = InterfaceCast<IAIPerceptionListenerInterface>(Event.DamagedActor);
+		IAIPerceptionListenerInterface* PerceptionListener = Cast<IAIPerceptionListenerInterface>(Event.DamagedActor);
 		if (PerceptionListener != NULL)
 		{
 			UAIPerceptionComponent* PerceptionComponent = PerceptionListener->GetPerceptionComponent();
diff --git a/Engine/Source/Runtime/AIModule/Private/Perception/AISense_Prediction.cpp b/Engine/Source/Runtime/AIModule/Private/Perception/AISense_Prediction.cpp
--- a/Engine/Source/Runtime/AIModule/Private/Perception/AISense_Prediction.cpp
+++ b/Engine/Source/Runtime/AIModule/Private/Perception/AISense_Prediction.cpp
@@ -18,7 +18,7 @@ float UAISense_Prediction::Update()
 
 		if (Event.Requestor != NULL && Event.PredictedActor != NULL)
 		{
-			IAIPerceptionListenerInterface* PerceptionListener = InterfaceCast<IAIPerceptionListenerInterface>(Event.Requestor);
+			IAIPerceptionListenerInterface* PerceptionListener = Cast<IAIPerceptionListenerInterface>(Event.Requestor);
 			if (PerceptionListener != NULL)
 			{
 				UAIPerceptionComponent* PerceptionComponent = PerceptionListener->GetPerceptionComponent();
diff --git a/Engine/Source/Runtime/AIModule/Private/Perception/AISense_Sight.cpp b/Engine/Source/Runtime/AIModule/Private/Perception/AISense_Sight.cpp
--- a/Engine/Source/Runtime/AIModule/Private/Perception/AISense_Sight.cpp
+++ b/Engine/Source/Runtime/AIModule/Private/Perception/AISense_Sight.cpp
@@ -243,7 +243,7 @@ void UAISense_Sight::RegisterTarget(AActor& TargetActor, FQueriesOperationPostPr
 		FAISightTarget NewSightTarget(&TargetActor);
 
 		SightTarget = &(ObservedTargets.Add(NewSightTarget.TargetId, NewSightTarget));
-		SightTarget->SightTargetInterface = InterfaceCast<IAISightTargetInterface>(&TargetActor);
+		SightTarget->SightTargetInterface = Cast<IAISightTargetInterface>(&TargetActor);
 	}
 
 	// set/update data
diff --git a/Engine/Source/Runtime/AIModule/Private/Perception/AISense_Touch.cpp b/Engine/Source/Runtime/AIModule/Private/Perception/AISense_Touch.cpp
--- a/Engine/Source/Runtime/AIModule/Private/Perception/AISense_Touch.cpp
+++ b/Engine/Source/Runtime/AIModule/Private/Perception/AISense_Touch.cpp
@@ -19,7 +19,7 @@ float UAISense_Touch::Update()
 
 		if (Event.TouchReceiver != NULL && Event.OtherActor != NULL)
 		{
-			IAIPerceptionListenerInterface* PerceptionListener = InterfaceCast<IAIPerceptionListenerInterface>(Event.TouchReceiver);
+			IAIPerceptionListenerInterface* PerceptionListener = Cast<IAIPerceptionListenerInterface>(Event.TouchReceiver);
 			if (PerceptionListener != NULL)
 			{
 				UAIPerceptionComponent* PerceptionComponent = PerceptionListener->GetPerceptionComponent();
diff --git a/Engine/Source/Runtime/Core/Public/Templates/UnrealTemplate.h b/Engine/Source/Runtime/Core/Public/Templates/UnrealTemplate.h
--- a/Engine/Source/Runtime/Core/Public/Templates/UnrealTemplate.h
+++ b/Engine/Source/Runtime/Core/Public/Templates/UnrealTemplate.h
@@ -374,6 +374,17 @@ struct TAreTypesEqual<A,A>
 #define TEMPLATE_PARAMETERS2(X,Y) X,Y
 
 
+/**
+ * Removes one level of pointer from a type, e.g.:
+ *
+ * TRemovePointer<      int32  >::Type == int32
+ * TRemovePointer<      int32* >::Type == int32
+ * TRemovePointer<      int32**>::Type == int32*
+ * TRemovePointer<const int32* >::Type == const int32
+ */
+template <typename T> struct TRemovePointer     { typedef T Type; };
+template <typename T> struct TRemovePointer<T*> { typedef T Type; };
+
 /**
  * TRemoveReference<type> will remove any references from a type.
  */
@@ -503,6 +514,18 @@ template <typename From, typename To> struct TCopyQualifiersFromTo<const
 template <typename From, typename To> struct TCopyQualifiersFromTo<      volatile From, To> { typedef       volatile To Type; };
 template <typename From, typename To> struct TCopyQualifiersFromTo<const volatile From, To> { typedef const volatile To Type; };
 
+/**
+ * Tests if qualifiers are lost between one type and another, e.g.:
+ *
+ * TCopyQualifiersFromTo<const    T1,                T2>::Value == true
+ * TCopyQualifiersFromTo<volatile T1, const volatile T2>::Value == false
+ */
+template <typename From, typename To>
+struct TLosesQualifiersFromTo
+{
+	enum { Value = !TAreTypesEqual<typename TCopyQualifiersFromTo<From, To>::Type, To>::Value };
+};
+
 /**
  * Returns the same type passed to it.  This is useful in a few cases, but mainly for inhibiting template argument deduction in function arguments, e.g.:
  *
@@ -518,6 +541,46 @@ struct TIdentity
 	typedef T Type;
 };
 
+/**
+ * Does LHS::Value && RHS::Value, but short-circuits if LHS::Value == false.
+ */
+template <bool LHSValue, typename RHS>
+struct TAndValue
+{
+	enum { Value = RHS::Value };
+};
+
+template <typename RHS>
+struct TAndValue<false, RHS>
+{
+	enum { Value = false };
+};
+
+template <typename LHS, typename RHS>
+struct TAnd : TAndValue<LHS::Value, RHS>
+{
+};
+
+/**
+ * Does LHS::Value || RHS::Value, but short-circuits if LHS::Value == true.
+ */
+template <bool LHSValue, typename RHS>
+struct TOrValue
+{
+	enum { Value = RHS::Value };
+};
+
+template <typename RHS>
+struct TOrValue<true, RHS>
+{
+	enum { Value = true };
+};
+
+template <typename LHS, typename RHS>
+struct TOr : TOrValue<LHS::Value, RHS>
+{
+};
+
 /**
  * Wraps a type.
  *
diff --git a/Engine/Source/Runtime/Core/Public/Templates/UnrealTypeTraits.h b/Engine/Source/Runtime/Core/Public/Templates/UnrealTypeTraits.h
--- a/Engine/Source/Runtime/Core/Public/Templates/UnrealTypeTraits.h
+++ b/Engine/Source/Runtime/Core/Public/Templates/UnrealTypeTraits.h
@@ -153,6 +153,25 @@ template<typename T> struct TIsPointerType<const T* const>		{ enum { Value = tru
 template<typename T> struct TIsPointerType<T* volatile>			{ enum { Value = true }; };
 template<typename T> struct TIsPointerType<T* const volatile>	{ enum { Value = true }; };
 
+/**
+ * TIsReferenceType
+ */
+template<typename T> struct TIsReferenceType      { enum { Value = false }; };
+template<typename T> struct TIsReferenceType<T&>  { enum { Value = true  }; };
+template<typename T> struct TIsReferenceType<T&&> { enum { Value = true  }; };
+
+/**
+ * TIsLValueReferenceType
+ */
+template<typename T> struct TIsLValueReferenceType     { enum { Value = false }; };
+template<typename T> struct TIsLValueReferenceType<T&> { enum { Value = true  }; };
+
+/**
+ * TIsLValueReferenceType
+ */
+template<typename T> struct TIsRValueReferenceType      { enum { Value = false }; };
+template<typename T> struct TIsRValueReferenceType<T&&> { enum { Value = true  }; };
+
 /**
  * TIsVoidType
  */
diff --git a/Engine/Source/Runtime/CoreUObject/Private/Blueprint/BlueprintSupport.cpp b/Engine/Source/Runtime/CoreUObject/Private/Blueprint/BlueprintSupport.cpp
--- a/Engine/Source/Runtime/CoreUObject/Private/Blueprint/BlueprintSupport.cpp
+++ b/Engine/Source/Runtime/CoreUObject/Private/Blueprint/BlueprintSupport.cpp
@@ -151,7 +151,7 @@ bool UObject::IsInBlueprint() const
 	const UObject* TestObject = this;
  	while (TestObject)
  	{
- 		const UClass *ClassObject = Cast<const UClass>(TestObject);
+ 		const UClass *ClassObject = dynamic_cast<const UClass*>(TestObject);
 		if (ClassObject 
 			&& ClassObject->HasAnyClassFlags(CLASS_CompiledFromBlueprint) 
 			&& ClassObject->ClassGeneratedBy)
@@ -230,7 +230,7 @@ void FBlueprintSupport::DuplicateAllFields(UStruct* StructToDuplicate, FDuplicat
 			UField* Field = *FieldIt;
 
 			// Make sure functions also do their parameters and children first
-			if (UFunction* Function = Cast<UFunction>(Field))
+			if (UFunction* Function = dynamic_cast<UFunction*>(Field))
 			{
 				for (TFieldIterator<UField> FunctionFieldIt(Function, EFieldIteratorFlags::ExcludeSuper); FunctionFieldIt; ++FunctionFieldIt)
 				{
diff --git a/Engine/Source/Runtime/CoreUObject/Private/Misc/RedirectCollector.cpp b/Engine/Source/Runtime/CoreUObject/Private/Misc/RedirectCollector.cpp
--- a/Engine/Source/Runtime/CoreUObject/Private/Misc/RedirectCollector.cpp
+++ b/Engine/Source/Runtime/CoreUObject/Private/Misc/RedirectCollector.cpp
@@ -11,21 +11,22 @@ static TArray<FString> StringAssetRefFilenameStack;
 
 void FRedirectCollector::OnRedirectorFollowed(const FString& InString, UObject* InObject)
 {
+	check(InObject);
+
 	// the object had better be a redir
-	UObjectRedirector* Redirector = Cast<UObjectRedirector>(InObject);
-	check(Redirector);
+	UObjectRedirector& Redirector = dynamic_cast<UObjectRedirector&>(*InObject);
 
 	// save the info if it matches the parameters we need to match on:
 	// if the string matches the package we were loading
 	// AND
 	// we aren't matching a particular package || the string matches the package
-	if (!FileToFixup.Len() || FileToFixup == FPackageName::FilenameToLongPackageName(Redirector->GetLinker()->Filename))
+	if (!FileToFixup.Len() || FileToFixup == FPackageName::FilenameToLongPackageName(Redirector.GetLinker()->Filename))
 	{
 		FRedirection Redir;
 		Redir.PackageFilename = InString;
-		Redir.RedirectorName = Redirector->GetFullName();
-		Redir.RedirectorPackageFilename = Redirector->GetLinker()->Filename;
-		Redir.DestinationObjectName = Redirector->DestinationObject->GetFullName();
+		Redir.RedirectorName = Redirector.GetFullName();
+		Redir.RedirectorPackageFilename = Redirector.GetLinker()->Filename;
+		Redir.DestinationObjectName = Redirector.DestinationObject->GetFullName();
 		// we only want one of each redirection reported
 		Redirections.AddUnique(Redir);
 
@@ -84,7 +85,7 @@ void FRedirectCollector::ResolveStringAssetReference()
 			UObject *Loaded = LoadObject<UObject>(NULL, *ToLoad, NULL, LOAD_None, NULL);
 			StringAssetRefFilenameStack.Pop();
 
-			UObjectRedirector* Redirector = Cast<UObjectRedirector>(Loaded);
+			UObjectRedirector* Redirector = dynamic_cast<UObjectRedirector*>(Loaded);
 			if (Redirector)
 			{
 				UE_LOG(LogRedirectors, Log, TEXT("    Found redir '%s'"), *Redirector->GetFullName());
diff --git a/Engine/Source/Runtime/CoreUObject/Private/Misc/StringAssetReference.cpp b/Engine/Source/Runtime/CoreUObject/Private/Misc/StringAssetReference.cpp
--- a/Engine/Source/Runtime/CoreUObject/Private/Misc/StringAssetReference.cpp
+++ b/Engine/Source/Runtime/CoreUObject/Private/Misc/StringAssetReference.cpp
@@ -115,7 +115,7 @@ UObject *FStringAssetReference::ResolveObject() const
 	FString FullPath = ToString();
 	FoundObject = StaticFindObject( UObject::StaticClass(), NULL, *FullPath);
 
-	UObjectRedirector* Redir = Cast<UObjectRedirector>(FoundObject);
+	UObjectRedirector* Redir = dynamic_cast<UObjectRedirector*>(FoundObject);
 	if (Redir)
 	{
 		// If we found a redirector, follow it
diff --git a/Engine/Source/Runtime/CoreUObject/Private/Misc/StringClassReference.cpp b/Engine/Source/Runtime/CoreUObject/Private/Misc/StringClassReference.cpp
--- a/Engine/Source/Runtime/CoreUObject/Private/Misc/StringClassReference.cpp
+++ b/Engine/Source/Runtime/CoreUObject/Private/Misc/StringClassReference.cpp
@@ -20,7 +20,7 @@ bool FStringClassReference::SerializeFromMismatchedTag(struct FPropertyTag const
 
 UClass *FStringClassReference::ResolveClass() const
 {
-	return Cast<UClass>(ResolveObject());
+	return dynamic_cast<UClass*>(ResolveObject());
 }
 
 FStringClassReference FStringClassReference::GetOrCreateIDForClass(const UClass *InClass)
diff --git a/Engine/Source/Runtime/CoreUObject/Private/Serialization/ArchiveShowReferences.cpp b/Engine/Source/Runtime/CoreUObject/Private/Serialization/ArchiveShowReferences.cpp
--- a/Engine/Source/Runtime/CoreUObject/Private/Serialization/ArchiveShowReferences.cpp
+++ b/Engine/Source/Runtime/CoreUObject/Private/Serialization/ArchiveShowReferences.cpp
@@ -43,9 +43,9 @@ FArchiveShowReferences::FArchiveShowReferences( FOutputDevice& inOutputAr, UObje
 	}
 
 	// similarly, if the object is a class, they all obviously reference their parent classes
-	if ( SourceObject->IsA(UClass::StaticClass()) )
+	if ( UClass* SourceObjectClass = dynamic_cast<UClass*>(SourceObject) )
 	{
-		for ( UClass* ParentClass = Cast<UClass>(SourceObject)->GetSuperClass(); ParentClass; ParentClass = ParentClass->GetSuperClass() )
+		for ( UClass* ParentClass = SourceObjectClass->GetSuperClass(); ParentClass; ParentClass = ParentClass->GetSuperClass() )
 		{
 			Found.AddUnique( ParentClass );
 		}
diff --git a/Engine/Source/Runtime/CoreUObject/Private/Serialization/AsyncLoading.cpp b/Engine/Source/Runtime/CoreUObject/Private/Serialization/AsyncLoading.cpp
--- a/Engine/Source/Runtime/CoreUObject/Private/Serialization/AsyncLoading.cpp
+++ b/Engine/Source/Runtime/CoreUObject/Private/Serialization/AsyncLoading.cpp
@@ -496,7 +496,7 @@ EAsyncPackageState::Type FAsyncPackage::LoadImports()
 
 		// Handle circular dependencies - try to find existing packages.
 		FString ImportPackageName(Import->ObjectName.ToString());
-		UPackage* ExistingPackage = Cast<UPackage>(StaticFindObjectFast(UPackage::StaticClass(), NULL, Import->ObjectName, true));
+		UPackage* ExistingPackage = dynamic_cast<UPackage*>(StaticFindObjectFast(UPackage::StaticClass(), NULL, Import->ObjectName, true));
 		if (ExistingPackage && !(ExistingPackage->PackageFlags & PKG_CompiledIn) && !ExistingPackage->bHasBeenFullyLoaded)//!ExistingPackage->HasAnyFlags(RF_WasLoaded))
 		{
 			// The import package already exists. Check if it's currently being streamed as well. If so, make sure
diff --git a/Engine/Source/Runtime/CoreUObject/Private/Serialization/BulkData.cpp b/Engine/Source/Runtime/CoreUObject/Private/Serialization/BulkData.cpp
--- a/Engine/Source/Runtime/CoreUObject/Private/Serialization/BulkData.cpp
+++ b/Engine/Source/Runtime/CoreUObject/Private/Serialization/BulkData.cpp
@@ -715,7 +715,7 @@ void FUntypedBulkData::Serialize( FArchive& Ar, UObject* Owner, int32 Idx )
 		else if( Ar.IsSaving() )
 		{
 			// check if we save the package compressed
-			UPackage* Pkg = Cast<UPackage>(Owner != NULL ? Owner->GetOutermost() : NULL);
+			UPackage* Pkg = Owner ? dynamic_cast<UPackage*>(Owner->GetOutermost()) : nullptr;
 			if (Pkg && !!(Pkg->PackageFlags & PKG_StoreCompressed) )
 			{
 				ECompressionFlags BaseCompressionMethod = COMPRESS_Default;
@@ -750,7 +750,7 @@ void FUntypedBulkData::Serialize( FArchive& Ar, UObject* Owner, int32 Idx )
 			Ar << BulkDataOffsetInFile;
 
 				// try to get the linkersave object
-			ULinkerSave* LinkerSave = Cast<ULinkerSave>(Ar.GetLinker());
+			ULinkerSave* LinkerSave = dynamic_cast<ULinkerSave*>(Ar.GetLinker());
 
 			// determine whether we are going to store the payload inline or not.
 			bool bStoreInline = !!(BulkDataFlags&BULKDATA_ForceInlinePayload) || LinkerSave == NULL;
diff --git a/Engine/Source/Runtime/CoreUObject/Private/UObject/Class.cpp b/Engine/Source/Runtime/CoreUObject/Private/UObject/Class.cpp
--- a/Engine/Source/Runtime/CoreUObject/Private/UObject/Class.cpp
+++ b/Engine/Source/Runtime/CoreUObject/Private/UObject/Class.cpp
@@ -144,7 +144,7 @@ UClass* UField::GetOwnerClass() const
 
 	while ((TestObject != NULL) && (OwnerClass == NULL))
 	{
-		OwnerClass = Cast<UClass>(TestObject);
+		OwnerClass = dynamic_cast<UClass*>(TestObject);
 		TestObject = TestObject->GetOuter();
 	}
 
@@ -153,9 +153,19 @@ UClass* UField::GetOwnerClass() const
 
 UStruct* UField::GetOwnerStruct() const
 {
-	const UObject* Obj;
-	for ( Obj=this; Obj && !Obj->IsA(UStruct::StaticClass()); Obj=Obj->GetOuter() );
-	return (UStruct*)Obj;
+	const UObject* Obj = this;
+	do
+	{
+		if (const UStruct* Result = dynamic_cast<const UStruct*>(Obj))
+		{
+			return const_cast<UStruct*>(Result);
+		}
+
+		Obj = Obj->GetOuter();
+	}
+	while (Obj);
+
+	return nullptr;
 }
 
 void UField::Bind()
@@ -186,7 +196,7 @@ struct FDisplayNameHelper
 	static FString Get(const UObject& Object)
 	{
 		FString Name = Object.GetName();
-		const UClass* Class = Cast<UClass>(&Object);
+		const UClass* Class = dynamic_cast<const UClass*>(&Object);
 		if (Class && !Class->HasAnyClassFlags(CLASS_Native))
 		{
 			Name.RemoveFromEnd(TEXT("_C"));
@@ -471,7 +481,7 @@ void UStruct::Link(FArchive& Ar, bool bRelinkExistingProperties)
 					break;
 				}
 
-				if (UProperty* Property = Cast<UProperty>(Field))
+				if (UProperty* Property = dynamic_cast<UProperty*>(Field))
 				{
 #if !WITH_EDITORONLY_DATA
 					// If we don't have the editor, make sure we aren't trying to link properties that are editor only.
@@ -506,10 +516,10 @@ void UStruct::Link(FArchive& Ar, bool bRelinkExistingProperties)
 			// check for internal struct recursion via arrays
 			for (UField* Field = Children; Field; Field = Field->Next)
 			{
-				UArrayProperty* ArrayProp = Cast<UArrayProperty>(Field);
+				UArrayProperty* ArrayProp = dynamic_cast<UArrayProperty*>(Field);
 				if (ArrayProp != NULL)
 				{
-					UStructProperty* StructProp = Cast<UStructProperty>(ArrayProp->Inner);
+					UStructProperty* StructProp = dynamic_cast<UStructProperty*>(ArrayProp->Inner);
 					if (StructProp != NULL && StructProp->Struct == this)
 					{
 						//we won't support this, too complicated
@@ -517,12 +527,13 @@ void UStruct::Link(FArchive& Ar, bool bRelinkExistingProperties)
 					}
 				}
 			}
-			UScriptStruct* ScriptStruct = CastChecked<UScriptStruct>(this);
-			ScriptStruct->PrepareCppStructOps();
-			UScriptStruct::ICppStructOps* CppStructOps = ScriptStruct->GetCppStructOps();
-			if (CppStructOps)
+
+			UScriptStruct& ScriptStruct = dynamic_cast<UScriptStruct&>(*this);
+			ScriptStruct.PrepareCppStructOps();
+
+			if (auto* CppStructOps = ScriptStruct.GetCppStructOps())
 			{
-				if (!ScriptStruct->InheritedCppStructOps())
+				if (!ScriptStruct.InheritedCppStructOps())
 				{
 					MinAlignment = CppStructOps->GetAlignment();
 					PropertiesSize = CppStructOps->GetSize();
@@ -540,7 +551,7 @@ void UStruct::Link(FArchive& Ar, bool bRelinkExistingProperties)
 	{
 		for (UField* Field = Children; (Field != NULL) && (Field->GetOuter() == this); Field = Field->Next)
 		{
-			if (UProperty* Property = Cast<UProperty>(Field))
+			if (UProperty* Property = dynamic_cast<UProperty*>(Field))
 			{
 				Property->LinkWithoutChangingOffset(Ar);
 			}
@@ -591,7 +602,7 @@ void UStruct::Link(FArchive& Ar, bool bRelinkExistingProperties)
 
 				for( UField* Field=Children; Field && Field->GetOuter()==this; Field=Field->Next )
 				{
-					UProperty* Property = Cast<UProperty>( Field );
+					UProperty* Property = dynamic_cast<UProperty*>( Field );
 					check(Property);
 					ColorComponentEntries[ColorComponentIndex++] = Property;
 				}
@@ -721,8 +732,8 @@ void UStruct::SerializeTaggedProperties(FArchive& Ar, uint8* Data, UStruct* Defa
 
 	check(Ar.IsLoading() || Ar.IsSaving());
 
-	UClass* DefaultsClass = Cast<UClass>(DefaultsStruct);
-	UScriptStruct* DefaultsScriptStruct = Cast<UScriptStruct>(DefaultsStruct);
+	UClass* DefaultsClass = dynamic_cast<UClass*>(DefaultsStruct);
+	UScriptStruct* DefaultsScriptStruct = dynamic_cast<UScriptStruct*>(DefaultsStruct);
 
 	if( Ar.IsLoading() )
 	{
@@ -871,7 +882,7 @@ void UStruct::SerializeTaggedProperties(FArchive& Ar, uint8* Data, UStruct* Defa
 			{
 				UE_LOG(LogClass, Warning, TEXT("Array bounds in %s of %s: %i/%i for package:  %s"), *Tag.Name.ToString(), *GetName(), Tag.ArrayIndex, Property->ArrayDim, *Ar.GetArchiveName() );
 			}
-			else if( Tag.Type==NAME_StrProperty && Cast<UNameProperty>(Property) )
+			else if( Tag.Type==NAME_StrProperty && dynamic_cast<UNameProperty*>(Property) )
 			{ 
 				FString str;  
 				Ar << str;
@@ -879,7 +890,7 @@ void UStruct::SerializeTaggedProperties(FArchive& Ar, uint8* Data, UStruct* Defa
 				AdvanceProperty = true;
 				continue; 
 			}
-			else if( Tag.Type==NAME_StrProperty && Cast<UTextProperty>(Property) ) // Convert serialized string to text.
+			else if( Tag.Type==NAME_StrProperty && dynamic_cast<UTextProperty*>(Property) ) // Convert serialized string to text.
 			{ 
 				FString str;
 				Ar << str;
@@ -889,7 +900,7 @@ void UStruct::SerializeTaggedProperties(FArchive& Ar, uint8* Data, UStruct* Defa
 				AdvanceProperty = true;
 				continue; 
 			}
-			else if( Tag.Type==NAME_TextProperty && Cast<UStrProperty>(Property) ) // Convert serialized text to string.
+			else if( Tag.Type==NAME_TextProperty && dynamic_cast<UStrProperty*>(Property) ) // Convert serialized text to string.
 			{ 
 				FText Text;  
 				Ar << Text;
@@ -1005,9 +1016,9 @@ void UStruct::SerializeTaggedProperties(FArchive& Ar, uint8* Data, UStruct* Defa
 				AdvanceProperty = true;
 				continue; 
 			}
-			else if( Cast<UStructProperty>(Property) && Cast<UStructProperty>(Property)->Struct && (Tag.Type != Property->GetID() || (Tag.Type == NAME_StructProperty && Tag.StructName != Cast<UStructProperty>(Property)->Struct->GetFName())) && (Cast<UStructProperty>(Property)->Struct->StructFlags & STRUCT_SerializeFromMismatchedTag))
+			else if( dynamic_cast<UStructProperty*>(Property) && dynamic_cast<UStructProperty*>(Property)->Struct && (Tag.Type != Property->GetID() || (Tag.Type == NAME_StructProperty && Tag.StructName != dynamic_cast<UStructProperty*>(Property)->Struct->GetFName())) && (dynamic_cast<UStructProperty*>(Property)->Struct->StructFlags & STRUCT_SerializeFromMismatchedTag))
 			{
-				UScriptStruct::ICppStructOps* CppStructOps = Cast<UStructProperty>(Property)->Struct->GetCppStructOps();
+				UScriptStruct::ICppStructOps* CppStructOps = dynamic_cast<UStructProperty*>(Property)->Struct->GetCppStructOps();
 				check(CppStructOps && CppStructOps->HasSerializeFromMismatchedTag()); // else should not have STRUCT_SerializeFromMismatchedTag
 				void* DestAddress = Property->ContainerPtrToValuePtr<void>(Data, Tag.ArrayIndex);  
 				if (CppStructOps->SerializeFromMismatchedTag(Tag, Ar, DestAddress))
@@ -1024,9 +1035,9 @@ void UStruct::SerializeTaggedProperties(FArchive& Ar, uint8* Data, UStruct* Defa
 			{
 				UE_LOG(LogClass, Warning, TEXT("Type mismatch in %s of %s - Previous (%s) Current(%s) for package:  %s"), *Tag.Name.ToString(), *GetName(), *Tag.Type.ToString(), *Property->GetID().ToString(), *Ar.GetArchiveName() );
 			}
-			else if( Tag.Type == NAME_ArrayProperty && Tag.InnerType != NAME_None && Tag.InnerType != CastChecked<UArrayProperty>(Property)->Inner->GetID() )
+			else if( Tag.Type == NAME_ArrayProperty && Tag.InnerType != NAME_None && Tag.InnerType != dynamic_cast<UArrayProperty&>(*Property).Inner->GetID() )
 			{
-				UArrayProperty* ArrayProperty = Cast<UArrayProperty>(Property);
+				UArrayProperty* ArrayProperty = dynamic_cast<UArrayProperty*>(Property);
 				void* ArrayPropertyData = ArrayProperty->ContainerPtrToValuePtr<void>(Data);
 
 				int32 ElementCount = 0;
@@ -1035,7 +1046,7 @@ void UStruct::SerializeTaggedProperties(FArchive& Ar, uint8* Data, UStruct* Defa
 				FScriptArrayHelper ScriptArrayHelper(ArrayProperty, ArrayPropertyData);
 				ScriptArrayHelper.EmptyAndAddValues(ElementCount);
 
-				if( Tag.InnerType==NAME_StrProperty && Cast<UTextProperty>(ArrayProperty->Inner) ) // Convert serialized string to text.
+				if( Tag.InnerType==NAME_StrProperty && dynamic_cast<UTextProperty*>(ArrayProperty->Inner) ) // Convert serialized string to text.
 				{ 
 					for(int32 i = 0; i < ElementCount; ++i)
 					{
@@ -1048,7 +1059,7 @@ void UStruct::SerializeTaggedProperties(FArchive& Ar, uint8* Data, UStruct* Defa
 					}
 					continue;
 				}
-				else if( Tag.InnerType==NAME_TextProperty && Cast<UStrProperty>(ArrayProperty->Inner) ) // Convert serialized text to string.
+				else if( Tag.InnerType==NAME_TextProperty && dynamic_cast<UStrProperty*>(ArrayProperty->Inner) ) // Convert serialized text to string.
 				{ 
 					for(int32 i = 0; i < ElementCount; ++i)
 					{
@@ -1190,9 +1201,9 @@ void UStruct::SerializeTaggedProperties(FArchive& Ar, uint8* Data, UStruct* Defa
 				PropertyName = Property->GetFName();
 				for( int32 Idx=0; Idx<Property->ArrayDim; Idx++ )
 				{
-					uint8* DataPtr = Property->ContainerPtrToValuePtr<uint8>(Data, Idx);
+					uint8* DataPtr      = Property->ContainerPtrToValuePtr           <uint8>(Data, Idx);
 					uint8* DefaultValue = Property->ContainerPtrToValuePtrForDefaults<uint8>(DefaultsStruct, Defaults, Idx);
-					if( (!IsA(UClass::StaticClass()) && !Defaults) || !Ar.DoDelta() || 
+					if( (!dynamic_cast<const UClass*>(this) && !Defaults) || !Ar.DoDelta() || 
 						!Property->Identical( DataPtr, DefaultValue, Ar.GetPortFlags()) || Ar.IsTransacting() )
 					{
 						if (bUseAtomicSerialization)
@@ -1404,7 +1415,7 @@ void UStruct::Serialize( FArchive& Ar )
 		}	
 	}
 
-	if (Ar.IsLoading() && !Cast<UClass>(this)) // classes are linked in the UClass serializer, which just called me
+	if (Ar.IsLoading() && !dynamic_cast<UClass*>(this)) // classes are linked in the UClass serializer, which just called me
 	{
 		// Link the properties.
 		Link( Ar, true );
@@ -1711,10 +1722,10 @@ bool FindConstructorUninitialized(UStruct* BaseClass,uint8* Data,uint8* Defaults
 		int32 Size = P->GetSize();
 		bool bProblem = false;
 		check(Size);
-		UBoolProperty* PB = Cast<UBoolProperty>(P);
-		UStructProperty* PS = Cast<UStructProperty>(P);
-		UStrProperty* PStr = Cast<UStrProperty>(P);
-		UArrayProperty* PArray = Cast<UArrayProperty>(P);
+		UBoolProperty*   PB     = dynamic_cast<UBoolProperty  *>(P);
+		UStructProperty* PS     = dynamic_cast<UStructProperty*>(P);
+		UStrProperty*    PStr   = dynamic_cast<UStrProperty   *>(P);
+		UArrayProperty*  PArray = dynamic_cast<UArrayProperty *>(P);
 		if(PStr)
 		{
 			// string that actually have data would be false positives, since they would point to the same string, but actually be different pointers
@@ -1898,7 +1909,7 @@ void UScriptStruct::PrepareCppStructOps()
 	bCppStructOpsFromBaseClass = false;
 	if (!CppStructOps)
 	{
-		UScriptStruct* Base = Cast<UScriptStruct>(GetSuperStruct());
+		UScriptStruct* Base = dynamic_cast<UScriptStruct*>(GetSuperStruct());
 		if (Base)
 		{
 			Base->PrepareCppStructOps();
@@ -1909,7 +1920,7 @@ void UScriptStruct::PrepareCppStructOps()
 #if !(UE_BUILD_SHIPPING || UE_BUILD_TEST)
 	if (!CppStructOps)
 	{
-		UScriptStruct* Base = Cast<UScriptStruct>(GetSuperStruct());
+		UScriptStruct* Base = dynamic_cast<UScriptStruct*>(GetSuperStruct());
 
 		while (Base)
 		{
@@ -1918,7 +1929,7 @@ void UScriptStruct::PrepareCppStructOps()
 				UE_LOG(LogClass, Fatal,TEXT("Couldn't bind to native BASE struct %s %s."),*GetName(),*Base->GetName());
 				break;
 			}
-			Base = Cast<UScriptStruct>(Base->GetSuperStruct());
+			Base = dynamic_cast<UScriptStruct*>(Base->GetSuperStruct());
 		}
 	}
 #endif
@@ -2801,7 +2812,7 @@ void UClass::Link(FArchive& Ar, bool bRelinkExistingProperties)
 		{
 			UProperty* P;
 			UFunction* F;
-			if( (P=Cast<UProperty>(*It))!=NULL )
+			if( (P=dynamic_cast<UProperty*>(*It))!=NULL )
 			{
 				if ( P->PropertyFlags & CPF_Net )
 				{
@@ -2813,7 +2824,7 @@ void UClass::Link(FArchive& Ar, bool bRelinkExistingProperties)
 					}
 				}
 			}
-			else if( (F=Cast<UFunction>(*It))!=NULL )
+			else if( (F=dynamic_cast<UFunction*>(*It))!=NULL )
 			{
 				check(!F->GetSuperFunction() || (F->GetSuperFunction()->FunctionFlags&FUNC_NetFuncFlags) == (F->FunctionFlags&FUNC_NetFuncFlags));
 				if( (F->FunctionFlags&FUNC_Net) && !F->GetSuperFunction() )
@@ -3241,7 +3252,7 @@ bool UClass::IsFunctionImplementedInBlueprint(FName InFunctionName) const
 
 bool UClass::HasProperty(UProperty* InProperty) const
 {
-	if ( UClass* PropertiesClass = Cast<UClass>(InProperty->GetOuter()) )
+	if ( UClass* PropertiesClass = dynamic_cast<UClass*>(InProperty->GetOuter()) )
 	{
 		return PropertiesClass->FindNearestCommonBaseClass(this) != nullptr;
 	}
@@ -3407,7 +3418,7 @@ bool UClass::HotReloadPrivateStaticClass(
 				*(void**)Target = NewVTable;
 				Count++;
 			}
-			else if (Cast<UClass>(Target))
+			else if (dynamic_cast<UClass*>(Target))
 			{
 				UClass *Class = CastChecked<UClass>(Target);
 				if (Class->ClassConstructor == OldClassConstructor)
@@ -3660,7 +3671,7 @@ void UFunction::InitializeDerivedMembers()
 	ParmsSize = 0;
 	ReturnValueOffset = MAX_uint16;
 
-	for (UProperty* Property = Cast<UProperty>(Children); Property; Property = Cast<UProperty>(Property->Next))
+	for (UProperty* Property = dynamic_cast<UProperty*>(Children); Property; Property = dynamic_cast<UProperty*>(Property->Next))
 	{
 		if (Property->PropertyFlags & CPF_Parm)
 		{
diff --git a/Engine/Source/Runtime/CoreUObject/Private/UObject/EditPropertyChain.cpp b/Engine/Source/Runtime/CoreUObject/Private/UObject/EditPropertyChain.cpp
--- a/Engine/Source/Runtime/CoreUObject/Private/UObject/EditPropertyChain.cpp
+++ b/Engine/Source/Runtime/CoreUObject/Private/UObject/EditPropertyChain.cpp
@@ -105,7 +105,7 @@ void FEditPropertyChain::SetListSize( int32 NewListSize )
 			for ( TIterator It(PropertyNode); It; --It )
 			{
 				// if we've found the member property, we can stop here
-				if ( Cast<UClass>(It->GetOuter()) != NULL )
+				if ( dynamic_cast<UClass*>(It->GetOuter()) != NULL )
 				{
 					PropertyNode = It.GetNode();
 					break;
diff --git a/Engine/Source/Runtime/CoreUObject/Private/UObject/GarbageCollection.cpp b/Engine/Source/Runtime/CoreUObject/Private/UObject/GarbageCollection.cpp
--- a/Engine/Source/Runtime/CoreUObject/Private/UObject/GarbageCollection.cpp
+++ b/Engine/Source/Runtime/CoreUObject/Private/UObject/GarbageCollection.cpp
@@ -461,7 +461,7 @@ public:
 			}
 
 			// Assemble token stream for UClass objects. This is only done once for each class.
-			if (UClass* Class = Cast<UClass>(Object))
+			if (UClass* Class = dynamic_cast<UClass*>(Object))
 			{
 				if (!Class->HasAnyClassFlags(CLASS_TokenStreamAssembled))
 				{
diff --git a/Engine/Source/Runtime/CoreUObject/Private/UObject/Linker.cpp b/Engine/Source/Runtime/CoreUObject/Private/UObject/Linker.cpp
--- a/Engine/Source/Runtime/CoreUObject/Private/UObject/Linker.cpp
+++ b/Engine/Source/Runtime/CoreUObject/Private/UObject/Linker.cpp
@@ -221,7 +221,7 @@ void ULinker::AddReferencedObjects(UObject* InThis, FReferenceCollector& Collect
  */
 FString ULinker::GetImportPathName(int32 ImportIndex)
 {
-	ULinkerLoad* Loader = Cast<ULinkerLoad>(this);
+	ULinkerLoad* Loader = dynamic_cast<ULinkerLoad*>(this);
 
 	FString Result;
 	for (FPackageIndex LinkerIndex = FPackageIndex::FromImport(ImportIndex); !LinkerIndex.IsNull();)
@@ -268,7 +268,7 @@ FString ULinker::GetImportPathName(int32 ImportIndex)
 FString ULinker::GetExportPathName(int32 ExportIndex, const TCHAR* FakeRoot,bool bResolveForcedExports/*=false*/)
 {
 	FString Result;
-	ULinkerLoad* Loader = Cast<ULinkerLoad>(this);
+	ULinkerLoad* Loader = dynamic_cast<ULinkerLoad*>(this);
 
 	bool bForcedExport = false;
 	for ( FPackageIndex LinkerIndex = FPackageIndex::FromExport(ExportIndex); !LinkerIndex.IsNull(); LinkerIndex = Exp(LinkerIndex).OuterIndex )
@@ -726,7 +726,7 @@ ULinkerLoad* GetPackageLinker
  */
 void ResetLoadersForSave(UObject* InOuter, const TCHAR *Filename)
 {
-	UPackage* Package = Cast<UPackage>(InOuter);
+	UPackage* Package = dynamic_cast<UPackage*>(InOuter);
 	// If we have a loader for the package, unload it to prevent conflicts if we are resaving to the same filename
 	ULinkerLoad* Loader = ULinkerLoad::FindExistingLinkerForPackage(Package);
 	// This is the loader corresponding to the package we're saving.
diff --git a/Engine/Source/Runtime/CoreUObject/Private/UObject/LinkerLoad.cpp b/Engine/Source/Runtime/CoreUObject/Private/UObject/LinkerLoad.cpp
--- a/Engine/Source/Runtime/CoreUObject/Private/UObject/LinkerLoad.cpp
+++ b/Engine/Source/Runtime/CoreUObject/Private/UObject/LinkerLoad.cpp
@@ -1977,7 +1977,7 @@ ULinkerLoad::EVerifyResult ULinkerLoad::VerifyImport(int32 ImportIndex)
 			WarningAppend += LOCTEXT("LoadWarningSuffix_redirection", " [redirection]").ToString();
 
 			// Create the redirector (no serialization yet)
-			UObjectRedirector* Redir = Cast<UObjectRedirector>(Import.SourceLinker->CreateExport(Import.SourceIndex));
+			UObjectRedirector* Redir = dynamic_cast<UObjectRedirector*>(Import.SourceLinker->CreateExport(Import.SourceIndex));
 			// this should probably never fail, but just in case
 			if (!Redir)
 			{
@@ -2196,11 +2196,13 @@ bool ULinkerLoad::VerifyImportInner(const int32 ImportIndex, FString& WarningSuf
 			{
 				// for loop does what we need
 			}
-			 if (Cast<UPackage>(Top->XObject) && (Cast<UPackage>(Top->XObject)->PackageFlags & PKG_CompiledIn))
-			 {
-				 // this is an import to a compiled in thing, just search for it in the package
-				 TmpPkg = Cast<UPackage>(Top->XObject);
-			 }
+
+			auto* Package = dynamic_cast<UPackage*>(Top->XObject);
+			if (Package && (Package->PackageFlags & PKG_CompiledIn))
+			{
+				// this is an import to a compiled in thing, just search for it in the package
+				TmpPkg = Package;
+			}
 		}
 
 		// Copy the SourceLinker from the FObjectImport for our Outer if the SourceLinker hasn't been set yet,
@@ -2442,7 +2444,7 @@ bool ULinkerLoad::VerifyImportInner(const int32 ImportIndex, FString& WarningSuf
 UObject* ULinkerLoad::CreateExportAndPreload(int32 ExportIndex, bool bForcePreload /* = false */)
 {
 	UObject *Object = CreateExport(ExportIndex);
-	if (Object && (bForcePreload || (Cast<UClass>(Object) != NULL) || Object->IsTemplate() || (Cast<UObjectRedirector>(Object) != NULL)))
+	if (Object && (bForcePreload || dynamic_cast<UClass*>(Object) || Object->IsTemplate() || dynamic_cast<UObjectRedirector*>(Object)))
 	{
 		Preload(Object);
 	}
@@ -2567,7 +2569,7 @@ int32 ULinkerLoad::FindExportIndex( FName ClassName, FName ClassPackage, FName O
 
 		if(Export.ObjectName == ObjectName && (ExportOuterIndex.IsImport() || Export.OuterIndex == ExportOuterIndex)) // this is very not legit to be passing INDEX_NONE into this function to mean "ignore"
 		{
-			UClass*	ExportClass = Cast<UClass>(IndexToObject(Export.ClassIndex));
+			UClass*	ExportClass = dynamic_cast<UClass*>(IndexToObject(Export.ClassIndex));
 
 			// See if this export's class inherits from the requested class.
 			for(UClass* ParentClass = ExportClass;ParentClass;ParentClass = ParentClass->GetSuperClass())
@@ -2718,12 +2720,12 @@ void ULinkerLoad::Preload( UObject* Object )
 			UClass* Cls = NULL;
 
 			// If this is a struct, make sure that its parent struct is completely loaded
-			if(	Object->IsA(UStruct::StaticClass()) )
+			if( UStruct* Struct = dynamic_cast<UStruct*>(Object) )
 			{
-				Cls = Cast<UClass>(Object);
-				if( ((UStruct*)Object)->GetSuperStruct() )
+				Cls = dynamic_cast<UClass*>(Object);
+				if( Struct->GetSuperStruct() )
 				{
-					Preload( ((UStruct*)Object)->GetSuperStruct() );
+					Preload( Struct->GetSuperStruct() );
 				}
 			}
 
@@ -2933,7 +2935,7 @@ UObject* ULinkerLoad::CreateExport( int32 Index )
 		{
 			LoadClass = UClass::StaticClass();
 		}
-		UObjectRedirector* LoadClassRedirector = Cast<UObjectRedirector>(LoadClass);
+		UObjectRedirector* LoadClassRedirector = dynamic_cast<UObjectRedirector*>(LoadClass);
 		if( LoadClassRedirector)
 		{
 			// mark this export as unloadable (so that other exports that
@@ -2947,7 +2949,7 @@ UObject* ULinkerLoad::CreateExport( int32 Index )
 		}
 
 		check(LoadClass);
-		check(Cast<UClass>(LoadClass) != NULL);
+		check(dynamic_cast<UClass*>(LoadClass) != NULL);
 
 		// Check for a valid superstruct while there is still time to safely bail, if this export has one
 		if( !Export.SuperIndex.IsNull() )
@@ -2977,7 +2979,7 @@ UObject* ULinkerLoad::CreateExport( int32 Index )
 					&& SuperStruct->GetLinker()
 					&& Export.SuperIndex.IsImport())
 				{
-					const UClass* AsClass = Cast<UClass>(SuperStruct);
+					const UClass* AsClass = dynamic_cast<UClass*>(SuperStruct);
 					if (AsClass && !AsClass->ClassDefaultObject)
 					{
 						SuperStruct->SetFlags(RF_NeedLoad);
@@ -3065,7 +3067,7 @@ UObject* ULinkerLoad::CreateExport( int32 Index )
 			return Export.Object;
 		}
 
-		UObjectRedirector* ParentRedirector = Cast<UObjectRedirector>(ThisParent);
+		UObjectRedirector* ParentRedirector = dynamic_cast<UObjectRedirector*>(ThisParent);
 		if( ThisParent == NULL || ParentRedirector)
 		{
 			// mark this export as unloadable (so that other exports that
@@ -3245,18 +3247,16 @@ UObject* ULinkerLoad::CreateExport( int32 Index )
 		if ( Export.Object != NULL )
 		{
 			// If it's a struct or class, set its parent.
-			if( Export.Object->IsA(UStruct::StaticClass()) )
+			if( UStruct* Struct = dynamic_cast<UStruct*>(Export.Object) )
 			{
 				if ( !Export.SuperIndex.IsNull() )
 				{
-					((UStruct*)Export.Object)->SetSuperStruct( (UStruct*)IndexToObject( Export.SuperIndex ) );
+					Struct->SetSuperStruct( (UStruct*)IndexToObject( Export.SuperIndex ) );
 				}
 
 				// If it's a class, bind it to C++.
-				if( Export.Object->IsA( UClass::StaticClass() ) )
+				if( UClass* ClassObject = dynamic_cast<UClass*>(Export.Object) )
 				{
-					UClass* ClassObject = static_cast<UClass*>(Export.Object);
-
 #if WITH_EDITOR
 					// Before we serialize the class, begin a scoped class dependency gather to create a list of other classes that may need to be recompiled
 					FScopedClassDependencyGather DependencyHelper(ClassObject);
@@ -3390,7 +3390,7 @@ UObject* ULinkerLoad::CreateImport( int32 Index )
 #if WITH_EDITOR
 				if( GIsEditor )
 				{
-					UObjectRedirector* Redirector = Cast<UObjectRedirector>(Import.XObject);
+					UObjectRedirector* Redirector = dynamic_cast<UObjectRedirector*>(Import.XObject);
 					if( Redirector )
 					{
 						Import.XObject = Redirector->DestinationObject;
diff --git a/Engine/Source/Runtime/CoreUObject/Private/UObject/LinkerSave.cpp b/Engine/Source/Runtime/CoreUObject/Private/UObject/LinkerSave.cpp
--- a/Engine/Source/Runtime/CoreUObject/Private/UObject/LinkerSave.cpp
+++ b/Engine/Source/Runtime/CoreUObject/Private/UObject/LinkerSave.cpp
@@ -24,13 +24,14 @@ ULinkerSave::ULinkerSave( const class FPostConstructInitializeProperties& PCIP,
 			UE_LOG(LogLinker, Fatal, TEXT("%s"), *FString::Printf( TEXT("Error opening file '%s'."), InFilename ) );
 		}
 
+		UPackage* Package = dynamic_cast<UPackage*>(LinkerRoot);
+
 		// Set main summary info.
 		Summary.Tag           = PACKAGE_FILE_TAG;
 		Summary.SetFileVersions( (int32)VER_LAST_ENGINE_UE3, GPackageFileUE4Version, GPackageFileLicenseeUE4Version, bInSaveUnversioned );
 		Summary.EngineVersion =	GEngineVersion;
-		Summary.PackageFlags  = Cast<UPackage>(LinkerRoot) ? Cast<UPackage>(LinkerRoot)->PackageFlags : 0;
+		Summary.PackageFlags  = Package ? Package->PackageFlags : 0;
 
-		UPackage *Package = Cast<UPackage>(LinkerRoot);
 		if (Package)
 		{
 			Summary.FolderName = Package->GetFolderName().ToString();
@@ -56,13 +57,14 @@ ULinkerSave::ULinkerSave(const class FPostConstructInitializeProperties& PCIP, U
 		Saver = new FBufferArchive( false, InParent->FileName );
 		check(Saver);
 
+		UPackage* Package = dynamic_cast<UPackage*>(LinkerRoot);
+
 		// Set main summary info.
 		Summary.Tag           = PACKAGE_FILE_TAG;
 		Summary.SetFileVersions( (int32)VER_LAST_ENGINE_UE3, GPackageFileUE4Version, GPackageFileLicenseeUE4Version, bInSaveUnversioned );
 		Summary.EngineVersion =	GEngineVersion;
-		Summary.PackageFlags  = Cast<UPackage>(LinkerRoot) ? Cast<UPackage>(LinkerRoot)->PackageFlags : 0;
+		Summary.PackageFlags  = Package ? Package->PackageFlags : 0;
 
-		UPackage *Package = Cast<UPackage>(LinkerRoot);
 		if (Package)
 		{
 			Summary.FolderName = Package->GetFolderName().ToString();
diff --git a/Engine/Source/Runtime/CoreUObject/Private/UObject/Obj.cpp b/Engine/Source/Runtime/CoreUObject/Private/UObject/Obj.cpp
--- a/Engine/Source/Runtime/CoreUObject/Private/UObject/Obj.cpp
+++ b/Engine/Source/Runtime/CoreUObject/Private/UObject/Obj.cpp
@@ -162,13 +162,13 @@ bool UObject::Rename( const TCHAR* InName, UObject* NewOuter, ERenameFlags Flags
 	if (bCreateRedirector)
 	{
 		// Look for an existing redirector with the same name/class/outer in the old package.
-		UObjectRedirector* Redirector = Cast<UObjectRedirector>(StaticFindObject(UObjectRedirector::StaticClass(), OldOuter, *OldName.ToString(), /*bExactClass=*/ true));
+		UObjectRedirector* Redirector = FindObject<UObjectRedirector>(OldOuter, *OldName.ToString(), /*bExactClass=*/ true);
 
 		// If it does not exist, create it.
 		if ( Redirector == NULL )
 		{
 			// create a UObjectRedirector with the same name as the old object we are redirecting
-			Redirector = Cast<UObjectRedirector>(StaticConstructObject(UObjectRedirector::StaticClass(), OldOuter, OldName, RF_Standalone | RF_Public));
+			Redirector = ConstructObject<UObjectRedirector>(UObjectRedirector::StaticClass(), OldOuter, OldName, RF_Standalone | RF_Public);
 		}
 
 		// point the redirector object to this object
@@ -1123,7 +1123,7 @@ void UObject::FAssetRegistryTag::GetAssetRegistryTagsFromSearchableProperties(co
 			else if ( Class->IsChildOf(UByteProperty::StaticClass()) )
 			{
 				// bytes are numerical, enums are alphabetical
-				UByteProperty* ByteProp = Cast<UByteProperty>(*FieldIt);
+				UByteProperty* ByteProp = dynamic_cast<UByteProperty*>(*FieldIt);
 				if ( ByteProp->Enum )
 				{
 					TagType = FAssetRegistryTag::TT_Alphabetical;
@@ -1168,11 +1168,10 @@ bool UObject::IsAsset () const
 	if ( bHasValidObjectFlags )
 	{
 		// Don't count objects embedded in other objects (e.g. font textures, sequences, material expressions)
-		UObject* LocalOuter = GetOuter();
-		if ( LocalOuter->IsA(UPackage::StaticClass()) )
+		if ( UPackage* LocalOuterPackage = dynamic_cast<UPackage*>(GetOuter()) )
 		{
 			// Also exclude any objects found in the transient package.
-			return LocalOuter != GetTransientPackage();
+			return LocalOuterPackage != GetTransientPackage();
 		}
 	}
 
@@ -1186,8 +1185,9 @@ bool UObject::IsSafeForRootSet() const
 		return false;
 	}
 
-	const ULinkerLoad* LinkerLoad = Cast<const ULinkerLoad>(this);
-	// Exclude linkers from root set if we're using seekfree loading		
+	const ULinkerLoad* LinkerLoad = dynamic_cast<const ULinkerLoad*>(this);
+
+	// Exclude linkers from root set if we're using seekfree loading
 	if( !HasAnyFlags(RF_PendingKill)
 		&& ( !FPlatformProperties::RequiresCookedData() || LinkerLoad == NULL || LinkerLoad->HasAnyFlags(RF_ClassDefaultObject) ) )
 	{
@@ -1473,7 +1473,7 @@ void UObject::LoadConfig( UClass* ConfigClass/*=NULL*/, const TCHAR* InFilename/
 #endif // #if WITH_EDITOR
 
 		UE_LOG(LogConfig, Verbose, TEXT("   Loading value for %s from [%s]"), *Key, *ClassSection);
-		UArrayProperty* Array = Cast<UArrayProperty>( Property );
+		UArrayProperty* Array = dynamic_cast<UArrayProperty*>( Property );
 		if( Array == NULL )
 		{
 			for( int32 i=0; i<Property->ArrayDim; i++ )
@@ -1659,7 +1659,7 @@ void UObject::SaveConfig( uint64 Flags, const TCHAR* InFilename, FConfigCacheIni
 			const bool bShouldCheckIfIdenticalBeforeAdding = !GetClass()->HasAnyClassFlags(CLASS_ConfigDoNotCheckDefaults) && !bPerObject && bIsPropertyInherited;
 			UObject* SuperClassDefaultObject = GetClass()->GetSuperClass()->GetDefaultObject();
 
-			UArrayProperty* Array   = Cast<UArrayProperty>( Property );
+			UArrayProperty* Array   = dynamic_cast<UArrayProperty*>( Property );
 			if( Array )
 			{
 				if ( !bShouldCheckIfIdenticalBeforeAdding || !Property->Identical_InContainer(this, SuperClassDefaultObject) )
@@ -2044,7 +2044,7 @@ UObject* UObject::CreateArchetype( const TCHAR* ArchetypeName, UObject* Archetyp
 	EObjectFlags ArchetypeObjectFlags = RF_Public | RF_ArchetypeObject;
 	
 	// Archetypes residing directly in packages need to be marked RF_Standalone
-	if( ArchetypeOuter->IsA(UPackage::StaticClass()) )
+	if( dynamic_cast<UPackage*>(ArchetypeOuter) )
 	{
 		ArchetypeObjectFlags |= RF_Standalone;
 	}
@@ -2160,15 +2160,15 @@ static void PrivateRecursiveDumpFlags(UStruct* Struct, void* Data, FOutputDevice
 			for( int32 i=0; i<It->ArrayDim; i++ )
 			{
 				uint8* Value = It->ContainerPtrToValuePtr<uint8>(Data, i);
-				UObjectPropertyBase* Prop = Cast<UObjectPropertyBase>(*It);
+				UObjectPropertyBase* Prop = dynamic_cast<UObjectPropertyBase*>(*It);
 				if(Prop)
 				{
 					UObject* Obj = Prop->GetObjectPropertyValue(Value);
 					PrivateDumpObjectFlags( Obj, Ar );
 				}
-				else if( Cast<UStructProperty>(*It) )
+				else if( UStructProperty* StructProperty = dynamic_cast<UStructProperty*>(*It) )
 				{
-					PrivateRecursiveDumpFlags( ((UStructProperty*)*It)->Struct, Value, Ar );
+					PrivateRecursiveDumpFlags( StructProperty->Struct, Value, Ar );
 				}
 			}
 		}
@@ -2505,21 +2505,21 @@ bool StaticExec( UWorld* InWorld, const TCHAR* Cmd, FOutputDevice& Ar )
 					if (bResult)
 					{
 						FString ExtraInfo;
-						if (It->GetClass()->ClassCastFlags & CASTCLASS_UStructProperty)
+						if (auto* StructProperty = dynamic_cast<UStructProperty*>(It->GetClass()))
 						{
-							ExtraInfo = *Cast<UStructProperty>(*It)->Struct->GetName();
+							ExtraInfo = *StructProperty->Struct->GetName();
 						}
-						else if (It->GetClass()->ClassCastFlags & CASTCLASS_UClassProperty)
+						else if (auto* ClassProperty = dynamic_cast<UClassProperty*>(It->GetClass()))
 						{
-							ExtraInfo = FString::Printf(TEXT("class<%s>"), *Cast<UClassProperty>(*It)->MetaClass->GetName());
+							ExtraInfo = FString::Printf(TEXT("class<%s>"), *ClassProperty->MetaClass->GetName());
 						}
-						else if (It->GetClass()->ClassCastFlags & CASTCLASS_UAssetClassProperty)
+						else if (auto* AssetClassProperty = dynamic_cast<UAssetClassProperty*>(It->GetClass()))
 						{
-							ExtraInfo = FString::Printf(TEXT("AssetSubclassOf<%s>"), *Cast<UAssetClassProperty>(*It)->MetaClass->GetName());
+							ExtraInfo = FString::Printf(TEXT("AssetSubclassOf<%s>"), *AssetClassProperty->MetaClass->GetName());
 						}
-						else if (It->GetClass()->ClassCastFlags & CASTCLASS_UObjectPropertyBase)
+						else if (auto* ObjectPropertyBase = dynamic_cast<UObjectPropertyBase*>(It->GetClass()))
 						{
-							ExtraInfo = *Cast<UObjectPropertyBase>(*It)->PropertyClass->GetName();
+							ExtraInfo = *ObjectPropertyBase->PropertyClass->GetName();
 						}
 						else
 						{
@@ -2612,7 +2612,7 @@ bool StaticExec( UWorld* InWorld, const TCHAR* Cmd, FOutputDevice& Ar )
 									}
 									continue;
 								}
-								if ( Property->ArrayDim > 1 || Cast<UArrayProperty>(Property) != NULL )
+								if ( Property->ArrayDim > 1 || dynamic_cast<UArrayProperty*>(Property) != NULL )
 								{
 									uint8* BaseData = Property->ContainerPtrToValuePtr<uint8>(CurrentObject);
 									Ar.Logf(TEXT("%i) %s.%s ="), cnt++, *CurrentObject->GetFullName(), *Property->GetName());
@@ -2622,7 +2622,7 @@ bool StaticExec( UWorld* InWorld, const TCHAR* Cmd, FOutputDevice& Ar )
 									UProperty* ExportProperty = Property;
 									if ( Property->ArrayDim == 1 )
 									{
-										UArrayProperty* ArrayProp = Cast<UArrayProperty>(Property);
+										UArrayProperty* ArrayProp = dynamic_cast<UArrayProperty*>(Property);
 										FScriptArrayHelper ArrayHelper(ArrayProp, BaseData);
 
 										BaseData = ArrayHelper.GetRawPtr();
diff --git a/Engine/Source/Runtime/CoreUObject/Private/UObject/ObjectRedirector.cpp b/Engine/Source/Runtime/CoreUObject/Private/UObject/ObjectRedirector.cpp
--- a/Engine/Source/Runtime/CoreUObject/Private/UObject/ObjectRedirector.cpp
+++ b/Engine/Source/Runtime/CoreUObject/Private/UObject/ObjectRedirector.cpp
@@ -133,7 +133,7 @@ FScopedRedirectorCatcher::~FScopedRedirectorCatcher()
 void FScopedRedirectorCatcher::OnRedirectorFollowed( const FString& InString, UObject* InObject)
 {
 	// this needs to be the redirector
-	check(InObject->IsA(UObjectRedirector::StaticClass()));
+	check(dynamic_cast<UObjectRedirector*>(InObject));
 
 	// if the path of the redirector was the same as the path to the object constant
 	// being compiled, then the script code has a text reference to a redirector, 
diff --git a/Engine/Source/Runtime/CoreUObject/Private/UObject/Package.cpp b/Engine/Source/Runtime/CoreUObject/Private/UObject/Package.cpp
--- a/Engine/Source/Runtime/CoreUObject/Private/UObject/Package.cpp
+++ b/Engine/Source/Runtime/CoreUObject/Private/UObject/Package.cpp
@@ -108,7 +108,7 @@ UMetaData* UPackage::GetMetaData()
 	// If there is no MetaData, try to find it.
 	if (MetaData == NULL)
 	{
-		MetaData = Cast<UMetaData>(StaticFindObjectFast(UMetaData::StaticClass(), this, FName(NAME_PackageMetaData)));
+		MetaData = FindObjectFast<UMetaData>(this, FName(NAME_PackageMetaData));
 
 		// If MetaData is NULL then it wasn't loaded by linker, so we have to create it.
 		if(MetaData == NULL)
diff --git a/Engine/Source/Runtime/CoreUObject/Private/UObject/Property.cpp b/Engine/Source/Runtime/CoreUObject/Private/UObject/Property.cpp
--- a/Engine/Source/Runtime/CoreUObject/Private/UObject/Property.cpp
+++ b/Engine/Source/Runtime/CoreUObject/Private/UObject/Property.cpp
@@ -479,17 +479,19 @@ FString UProperty::GetCPPMacroType( FString& ExtendedTypeText ) const
 void UProperty::ExportCppDeclaration( FOutputDevice& Out, EExportedDeclaration::Type DeclarationType, const TCHAR* ArrayDimOverride ) const
 {
 	const bool bIsParameter     = DeclarationType == EExportedDeclaration::Parameter;
-	const bool bIsInterfaceProp = IsA(UInterfaceProperty::StaticClass());
+	const bool bIsInterfaceProp = dynamic_cast<const UInterfaceProperty*>(this) != nullptr;
 
 	// export the property type text (e.g. FString; int32; TArray, etc.)
 	FString ExtendedTypeText;
 	FString TypeText = GetCPPType(&ExtendedTypeText, bIsParameter ? CPPF_ArgumentOrReturnValue : 0);
 
-	if (!Cast<const UBoolProperty>(this)) // can't have const bitfields because then we cannot determine their offset and mask from the compiler
+	if (!dynamic_cast<const UBoolProperty*>(this)) // can't have const bitfields because then we cannot determine their offset and mask from the compiler
 	{
+		const UObjectProperty* ObjectProp = dynamic_cast<const UObjectProperty*>(this);
+
 		// export 'const' for parameters
 		const bool bIsConstParam   = bIsParameter && (HasAnyPropertyFlags(CPF_ConstParm) || (bIsInterfaceProp && !HasAllPropertyFlags(CPF_OutParm)));
-		const bool bIsOnConstClass = IsA(UObjectProperty::StaticClass()) && ((UObjectProperty*)this)->PropertyClass && ((UObjectProperty*)this)->PropertyClass->HasAnyClassFlags(CLASS_Const);
+		const bool bIsOnConstClass = ObjectProp && ObjectProp->PropertyClass && ObjectProp->PropertyClass->HasAnyClassFlags(CLASS_Const);
 		if (bIsConstParam || bIsOnConstClass)
 		{
 			TypeText = FString::Printf(TEXT("const %s"), *TypeText);
@@ -497,7 +499,7 @@ void UProperty::ExportCppDeclaration( FOutputDevice& Out, EExportedDeclaration::
 
 		if (DeclarationType == EExportedDeclaration::Member)
 		{
-			UClass* MyClass = Cast<UClass>(GetOuter());
+			UClass* MyClass = dynamic_cast<UClass*>(GetOuter());
 			if (MyClass && MyClass->HasAnyClassFlags(CLASS_Const))
 			{
 				ExtendedTypeText += TEXT(" const");
@@ -518,7 +520,7 @@ void UProperty::ExportCppDeclaration( FOutputDevice& Out, EExportedDeclaration::
 		}
 	}
 
-	if( IsA(UBoolProperty::StaticClass()) )
+	if( dynamic_cast<const UBoolProperty*>(this) )
 	{
 		// if this is a member variable, export it as a bitfield
 		if( ArrayDim==1 && DeclarationType == EExportedDeclaration::Member )
diff --git a/Engine/Source/Runtime/CoreUObject/Private/UObject/PropertyArray.cpp b/Engine/Source/Runtime/CoreUObject/Private/UObject/PropertyArray.cpp
--- a/Engine/Source/Runtime/CoreUObject/Private/UObject/PropertyArray.cpp
+++ b/Engine/Source/Runtime/CoreUObject/Private/UObject/PropertyArray.cpp
@@ -132,7 +132,7 @@ void UArrayProperty::ExportTextItem( FString& ValueStr, const void* PropertyValu
 	FScriptArrayHelper DefaultArrayHelper(this, DefaultValue);
 
 	uint8* StructDefaults = NULL;
-	UStructProperty* StructProperty = Cast<UStructProperty>(Inner);
+	UStructProperty* StructProperty = dynamic_cast<UStructProperty*>(Inner);
 	if ( StructProperty != NULL )
 	{
 		checkSlow(StructProperty->Struct);
diff --git a/Engine/Source/Runtime/CoreUObject/Private/UObject/PropertyAssetClass.cpp b/Engine/Source/Runtime/CoreUObject/Private/UObject/PropertyAssetClass.cpp
--- a/Engine/Source/Runtime/CoreUObject/Private/UObject/PropertyAssetClass.cpp
+++ b/Engine/Source/Runtime/CoreUObject/Private/UObject/PropertyAssetClass.cpp
@@ -26,7 +26,7 @@ void UAssetClassProperty::Serialize( FArchive& Ar )
 		// If we failed to load the MetaClass and we're not a CDO, that means we relied on a class that has been removed or doesn't exist.
 		// The most likely cause for this is either an incomplete recompile, or if content was migrated between games that had native class dependencies
 		// that do not exist in this game.  We allow blueprint classes to continue, because compile on load will error out, and stub the class that was using it
-		UClass* TestClass = Cast<UClass>(GetOwnerStruct());
+		UClass* TestClass = dynamic_cast<UClass*>(GetOwnerStruct());
 		if( TestClass && TestClass->HasAllClassFlags(CLASS_Native) && !TestClass->HasAllClassFlags(CLASS_NewerVersionExists) && (TestClass->GetOutermost() != GetTransientPackage()) )
 		{
 			checkf(false, TEXT("Class property tried to serialize a missing class.  Did you remove a native class and not fully recompile?"));
diff --git a/Engine/Source/Runtime/CoreUObject/Private/UObject/PropertyBaseObject.cpp b/Engine/Source/Runtime/CoreUObject/Private/UObject/PropertyBaseObject.cpp
--- a/Engine/Source/Runtime/CoreUObject/Private/UObject/PropertyBaseObject.cpp
+++ b/Engine/Source/Runtime/CoreUObject/Private/UObject/PropertyBaseObject.cpp
@@ -349,7 +349,7 @@ UObject* UObjectPropertyBase::FindImportedObject( const UProperty* Property, UOb
 	// if we found an object, and we have a parent, make sure we are in the same package if the found object is private, unless it's a cross level property
 	if (Result && !Result->HasAnyFlags(RF_Public) && OwnerObject && Result->GetOutermost() != OwnerObject->GetOutermost())
 	{
-		const UObjectPropertyBase* ObjectProperty = Cast<const UObjectPropertyBase>(Property);
+		const UObjectPropertyBase* ObjectProperty = dynamic_cast<const UObjectPropertyBase*>(Property);
 		if ( !ObjectProperty || !ObjectProperty->AllowCrossLevel())
 		{
 			UE_LOG(LogProperty, Warning, TEXT("Illegal TEXT reference to a private object in external package (%s) from referencer (%s).  Import failed..."), *Result->GetFullName(), *OwnerObject->GetFullName());
diff --git a/Engine/Source/Runtime/CoreUObject/Private/UObject/PropertyClass.cpp b/Engine/Source/Runtime/CoreUObject/Private/UObject/PropertyClass.cpp
--- a/Engine/Source/Runtime/CoreUObject/Private/UObject/PropertyClass.cpp
+++ b/Engine/Source/Runtime/CoreUObject/Private/UObject/PropertyClass.cpp
@@ -17,7 +17,7 @@ void UClassProperty::Serialize( FArchive& Ar )
 		// If we failed to load the MetaClass and we're not a CDO, that means we relied on a class that has been removed or doesn't exist.
 		// The most likely cause for this is either an incomplete recompile, or if content was migrated between games that had native class dependencies
 		// that do not exist in this game.  We allow blueprint classes to continue, because compile on load will error out, and stub the class that was using it
-		UClass* TestClass = Cast<UClass>(GetOwnerStruct());
+		UClass* TestClass = dynamic_cast<UClass*>(GetOwnerStruct());
 		if( TestClass && TestClass->HasAllClassFlags(CLASS_Native) && !TestClass->HasAllClassFlags(CLASS_NewerVersionExists) && (TestClass->GetOutermost() != GetTransientPackage()) )
 		{
 			checkf(false, TEXT("Class property tried to serialize a missing class.  Did you remove a native class and not fully recompile?"));
@@ -37,7 +37,7 @@ const TCHAR* UClassProperty::ImportText_Internal( const TCHAR* Buffer, void* Dat
 	{
 		// Validate metaclass.
 		UClass* C = (UClass*)GetObjectPropertyValue(Data);
-		if ((C != NULL) && ((Cast<UClass>(C) == NULL) || !C->IsChildOf(MetaClass)))
+		if (C && (!dynamic_cast<UClass*>(C) || !C->IsChildOf(MetaClass)))
 		{
 			// the object we imported doesn't implement our interface class
 			ErrorText->Logf(TEXT("Invalid object '%s' specified for property '%s'"), *C->GetFullName(), *GetName());
diff --git a/Engine/Source/Runtime/CoreUObject/Private/UObject/PropertyInterface.cpp b/Engine/Source/Runtime/CoreUObject/Private/UObject/PropertyInterface.cpp
--- a/Engine/Source/Runtime/CoreUObject/Private/UObject/PropertyInterface.cpp
+++ b/Engine/Source/Runtime/CoreUObject/Private/UObject/PropertyInterface.cpp
@@ -189,7 +189,7 @@ void UInterfaceProperty::Serialize( FArchive& Ar )
 		// If we failed to load the InterfaceClass and we're not a CDO, that means we relied on a class that has been removed or doesn't exist.
 		// The most likely cause for this is either an incomplete recompile, or if content was migrated between games that had native class dependencies
 		// that do not exist in this game.  We allow blueprint classes to continue, because compile-on-load will error out, and stub the class that was using it
-		UClass* TestClass = Cast<UClass>(GetOwnerStruct());
+		UClass* TestClass = dynamic_cast<UClass*>(GetOwnerStruct());
 		if( TestClass && TestClass->HasAllClassFlags(CLASS_Native) && !TestClass->HasAllClassFlags(CLASS_NewerVersionExists) && (TestClass->GetOutermost() != GetTransientPackage()) )
 		{
 			checkf(false, TEXT("Interface property tried to serialize a missing interface.  Did you remove a native class and not fully recompile?"));
diff --git a/Engine/Source/Runtime/CoreUObject/Private/UObject/SavePackage.cpp b/Engine/Source/Runtime/CoreUObject/Private/UObject/SavePackage.cpp
--- a/Engine/Source/Runtime/CoreUObject/Private/UObject/SavePackage.cpp
+++ b/Engine/Source/Runtime/CoreUObject/Private/UObject/SavePackage.cpp
@@ -272,7 +272,7 @@ public:
 	 * 
 	 * @param	InOuter		the package to save
 	 */
-	FArchiveSaveTagExports( UObject* InOuter )
+	FArchiveSaveTagExports( UPackage* InOuter )
 	: Outer(InOuter)
 	{
 		ArIsSaving				= true;
@@ -288,7 +288,7 @@ public:
 	 * Package we're currently saving.  Only objects contained
 	 * within this package will be tagged for serialization.
 	 */
-	UObject* Outer;
+	UPackage* Outer;
 
 	/**
 	 * Returns the name of the Archive.  Useful for getting the name of the package a struct or object
@@ -313,14 +313,15 @@ private:
  **/
 FString FArchiveSaveTagExports::GetArchiveName() const
 {
-	return Outer != NULL
+	return Outer != nullptr
 		? *FString::Printf(TEXT("SaveTagExports (%s)"), *Outer->GetName())
 		: TEXT("SaveTagExports");
 }
 
 FArchive& FArchiveSaveTagExports::operator<<( UObject*& Obj )
 {
-	CheckObjectPriorToSave(Obj, CastChecked<UPackage>(Outer));
+	check(Outer);
+	CheckObjectPriorToSave(Obj, Outer);
 	if( Obj && Obj->IsIn(Outer) && !Obj->HasAnyFlags(RF_Transient) && !Obj->HasAnyMarks(OBJECTMARK_TagExp) )
 	{
 #if 0
@@ -590,14 +591,14 @@ FArchive& FArchiveSaveTagImports::operator<<( UObject*& Obj )
 		if( !Obj->HasAnyFlags(RF_Transient) || Obj->HasAllFlags(RF_Native) )
 		{
 			// remember it as a dependency, unless it's a top level package or native
-			bool bIsTopLevelPackage = Obj->GetOuter() == NULL && Obj->IsA(UPackage::StaticClass());
+			bool bIsTopLevelPackage = Obj->GetOuter() == NULL && dynamic_cast<UPackage*>(Obj);
 			bool bIsNative = Obj->HasAnyFlags(RF_Native);
 			UObject* Outer = Obj->GetOuter();
 			
 			// go up looking for native classes
 			while (!bIsNative && Outer)
 			{
-				if (Outer->IsA(UClass::StaticClass()) && Outer->HasAnyFlags(RF_Native))
+				if (dynamic_cast<UClass*>(Outer) && Outer->HasAnyFlags(RF_Native))
 				{
 					bIsNative = true;
 				}
@@ -1529,31 +1530,28 @@ class FExportReferenceSorter : public FArchiveUObject
 			const int32 PreviousReferencedObjectCount = ReferencedObjects.Num();
 			const int32 PreviousInsertIndex = CurrentInsertIndex;
 
-			if ( RequiredObject->IsA(UStruct::StaticClass()) )
+			if (UStruct* RequiredObjectStruct = dynamic_cast<UStruct*>(RequiredObject))
 			{
 				// if this is a struct/class/function/state, it may have a super that needs to be processed first
-				ProcessStruct((UStruct*)RequiredObject);
+				ProcessStruct(RequiredObjectStruct);
+			}
+			else if ( bProcessObject )
+			{
+				// this means that RequiredObject is being force-loaded by the referencing object, rather than simply referenced
+				ProcessObject(RequiredObject);
 			}
 			else
 			{
-				if ( bProcessObject )
+				// only the object's class and archetype are force-loaded, so only those objects need to be in the list before
+				// whatever object was referencing RequiredObject
+				if ( ProcessedObjects.Find(RequiredObject->GetOuter()) == INDEX_NONE )
 				{
-					// this means that RequiredObject is being force-loaded by the referencing object, rather than simply referenced
-					ProcessObject(RequiredObject);
+					HandleDependency(RequiredObject->GetOuter());
 				}
-				else
-				{
-					// only the object's class and archetype are force-loaded, so only those objects need to be in the list before
-					// whatever object was referencing RequiredObject
-					if ( ProcessedObjects.Find(RequiredObject->GetOuter()) == INDEX_NONE )
-					{
-						HandleDependency(RequiredObject->GetOuter());
-					}
 
-					// class is needed before archetype, but we need to process these in reverse order because we are inserting into the list.
-					ProcessObject(RequiredObject->GetArchetype());
-					ProcessStruct(RequiredObject->GetClass());
-				}
+				// class is needed before archetype, but we need to process these in reverse order because we are inserting into the list.
+				ProcessObject(RequiredObject->GetArchetype());
+				ProcessStruct(RequiredObject->GetClass());
 			}
 
 			// InsertIndexOffset is the amount the CurrentInsertIndex was incremented during the serialization of SuperField; we need to
@@ -1608,7 +1606,7 @@ public:
 			}
 
 			// UObjectRedirectors are always force-loaded as the loading code needs immediate access to the object pointed to by the Redirector
-			UObjectRedirector* Redirector = Cast<UObjectRedirector>(Object);
+			UObjectRedirector* Redirector = dynamic_cast<UObjectRedirector*>(Object);
 			if ( Redirector != NULL && Redirector->DestinationObject != NULL )
 			{
 				// the Redirector does not force-load the destination object, so we only need its class and archetype.
@@ -1682,23 +1680,23 @@ public:
 				// of exports before the class, so that when CreateExport is called for this object reference we don't have to seek.
 				// Note that in the non-UField case, we don't actually need the object itself to appear before the referencing object/class because it won't
 				// be force-loaded (thus we don't need to add the referenced object to the ReferencedObject list)
-				if ( Object->IsA(UField::StaticClass()) )
+				if (dynamic_cast<UField*>(Object))
 				{
 					// when field processing is enabled, ignore any referenced classes since a class's class and CDO are both intrinsic and
 					// attempting to deal with them here will only cause problems
-					if ( !bIgnoreFieldReferences && !Object->IsA(UClass::StaticClass()) )
+					if ( !bIgnoreFieldReferences && !dynamic_cast<UClass*>(Object) )
 					{
 						if ( CurrentClass == NULL || Object->GetOuter() != CurrentClass )
 						{
-							if ( Object->IsA(UStruct::StaticClass()) )
+							if ( UStruct* StructObject = dynamic_cast<UStruct*>(Object) )
 							{
 								// if this is a struct/class/function/state, it may have a super that needs to be processed first (Preload force-loads UStruct::SuperField)
-								ProcessStruct((UStruct*)Object);
+								ProcessStruct(StructObject);
 							}
 							else
 							{
 								// byte properties that are enum references need their enums loaded first so that config importing works
-								UByteProperty* ByteProp = Cast<UByteProperty>(Object);
+								UByteProperty* ByteProp = dynamic_cast<UByteProperty*>(Object);
 								if (ByteProp != NULL && ByteProp->Enum != NULL)
 								{
 									HandleDependency(ByteProp->Enum, true);
@@ -1761,7 +1759,7 @@ public:
 					HandleDependency(Object->GetArchetype(), true);
 
 					// UObjectRedirectors are always force-loaded as the loading code needs immediate access to the object pointed to by the Redirector
-					UObjectRedirector* Redirector = Cast<UObjectRedirector>(Object);
+					UObjectRedirector* Redirector = dynamic_cast<UObjectRedirector*>(Object);
 					if ( Redirector != NULL && Redirector->DestinationObject != NULL )
 					{
 						// the Redirector does not force-load the destination object, so we only need its class and archetype.
@@ -1833,7 +1831,7 @@ public:
 
 					// invoke the serialize operator rather than calling Serialize directly so that the object is handled correctly (i.e. if it is a struct, then we should
 					// call ProcessStruct, etc. and all this logic is already contained in the serialization operator)
-					if ( Cast<UClass>(StructObject) == NULL )
+					if (!dynamic_cast<UClass*>(StructObject))
 					{
 						// before processing the Children reference, set the CurrentClass to the class which contains this StructObject so that we
 						// don't inadvertently serialize other fields of the owning class too early.
@@ -1850,7 +1848,7 @@ public:
 				// Preload will force-load the class default object when called on a UClass object, so make sure that the CDO is always immediately after its class
 				// in the export list; we can't resolve this circular reference, but hopefully we the CDO will fit into the same memory block as the class during 
 				// seek-free loading.
-				UClass* ClassObject = Cast<UClass>(StructObject);
+				UClass* ClassObject = dynamic_cast<UClass*>(StructObject);
 				if ( ClassObject != NULL )
 				{
 					UObject* CDO = ClassObject->GetDefaultObject();
@@ -1970,10 +1968,10 @@ struct FObjectExportSeekFreeSorter
 		for( int32 ExportIndex=FirstSortIndex; ExportIndex<Linker->ExportMap.Num(); ExportIndex++ )
 		{
 			const FObjectExport& Export = Linker->ExportMap[ExportIndex];
-			if( Export.Object && Export.Object->IsA(UClass::StaticClass()) )
+			if( UClass* ExportObjectClass = dynamic_cast<UClass*>(Export.Object) )
 			{
 				SortArchive.Clear();
-				SortArchive.ProcessStruct((UClass*)Export.Object);
+				SortArchive.ProcessStruct(ExportObjectClass);
 #if EXPORT_SORTING_DETAILED_LOGGING
 				TArray<UObject*> ReferencedObjects;
 				SortArchive.GetExportList(ReferencedObjects, bRetrieveInitialReferences);
@@ -2172,7 +2170,7 @@ static bool ValidateConformCompatibility(UPackage* NewPackage, ULinkerLoad* OldL
 	for (int32 i = 0; i < OldLinker->ExportMap.Num(); i++)
 	{
 		UClass* NewClass = (UClass*)StaticFindObjectFast(UClass::StaticClass(), NewPackage, OldLinker->ExportMap[i].ObjectName, true, false);
-		UClass* OldClass = Cast<UClass>(OldLinker->Create(UClass::StaticClass(), OldLinker->ExportMap[i].ObjectName, OldLinker->LinkerRoot, LOAD_None, false));
+		UClass* OldClass = static_cast<UClass*>(OldLinker->Create(UClass::StaticClass(), OldLinker->ExportMap[i].ObjectName, OldLinker->LinkerRoot, LOAD_None, false));
 		if (OldClass != NULL && NewClass != NULL && OldClass->HasAnyFlags(RF_Native) && NewClass->HasAnyFlags(RF_Native))
 		{
 			OldClass->ClassConstructor = NewClass->ClassConstructor;
@@ -2189,11 +2187,11 @@ static bool ValidateConformCompatibility(UPackage* NewPackage, ULinkerLoad* OldL
 		{
 			// load the object so we can analyze it
 			BeginLoad();
-			UClass* OldClass = Cast<UClass>(OldLinker->Create(UClass::StaticClass(), OldLinker->ExportMap[i].ObjectName, OldLinker->LinkerRoot, LOAD_None, false));
+			UClass* OldClass = static_cast<UClass*>(OldLinker->Create(UClass::StaticClass(), OldLinker->ExportMap[i].ObjectName, OldLinker->LinkerRoot, LOAD_None, false));
 			EndLoad();
 			if (OldClass != NULL)
 			{
-				UClass* NewClass = Cast<UClass>(StaticFindObjectFast(UClass::StaticClass(), NewPackage, OldClass->GetFName(), true, false));
+				UClass* NewClass = FindObjectFast<UClass>(NewPackage, OldClass->GetFName(), true, false);
 				if (NewClass != NULL)
 				{
 					for (TFieldIterator<UField> OldFieldIt(OldClass,EFieldIteratorFlags::ExcludeSuper); OldFieldIt; ++OldFieldIt)
@@ -2202,8 +2200,8 @@ static bool ValidateConformCompatibility(UPackage* NewPackage, ULinkerLoad* OldL
 						{
 							if (OldFieldIt->GetFName() == NewFieldIt->GetFName())
 							{
-								UProperty* OldProp = Cast<UProperty>(*OldFieldIt);
-								UProperty* NewProp = Cast<UProperty>(*NewFieldIt);
+								UProperty* OldProp = dynamic_cast<UProperty*>(*OldFieldIt);
+								UProperty* NewProp = dynamic_cast<UProperty*>(*NewFieldIt);
 								if (OldProp != NULL && NewProp != NULL)
 								{
 									if ((OldProp->PropertyFlags & CPF_Net) != (NewProp->PropertyFlags & CPF_Net))
@@ -2214,8 +2212,8 @@ static bool ValidateConformCompatibility(UPackage* NewPackage, ULinkerLoad* OldL
 								}
 								else
 								{
-									UFunction* OldFunc = Cast<UFunction>(*OldFieldIt);
-									UFunction* NewFunc = Cast<UFunction>(*NewFieldIt);
+									UFunction* OldFunc = dynamic_cast<UFunction*>(*OldFieldIt);
+									UFunction* NewFunc = dynamic_cast<UFunction*>(*NewFieldIt);
 									if (OldFunc != NULL && NewFunc != NULL)
 									{
 										if ((OldFunc->FunctionFlags & (FUNC_Net | FUNC_NetServer | FUNC_NetClient)) != (NewFunc->FunctionFlags & (FUNC_Net | FUNC_NetServer | FUNC_NetClient)))
@@ -2774,7 +2772,7 @@ bool UPackage::SavePackage( UPackage* InOuter, UObject* Base, EObjectFlags TopLe
 							}
 
 							// See whether the object we are referencing is in another map package.
-							UPackage* ObjPackage = Cast<UPackage>(Obj->GetOutermost());
+							UPackage* ObjPackage = dynamic_cast<UPackage*>(Obj->GetOutermost());
 							if( ObjPackage && ObjPackage->ContainsMap() )
 							{
 								if ( ObjPackage != Obj && Obj->GetFName() != NAME_PersistentLevel && Obj->GetClass()->GetFName() != WorldClassName )
@@ -3138,7 +3136,7 @@ bool UPackage::SavePackage( UPackage* InOuter, UObject* Base, EObjectFlags TopLe
 					{
 						Linker->ObjectIndicesMap.Add(Object, FPackageIndex::FromExport(i));
 
-						UPackage* Package = Cast<UPackage>(Object);
+						UPackage* Package = dynamic_cast<UPackage*>(Object);
 						if (Package != NULL)
 						{
 							Linker->ExportMap[i].PackageFlags = Package->PackageFlags;
@@ -3314,9 +3312,8 @@ bool UPackage::SavePackage( UPackage* InOuter, UObject* Base, EObjectFlags TopLe
 						}
 
 						// Set the parent index, if this export represents a UStruct-derived object
-						if( Export.Object->IsA(UStruct::StaticClass()) )
+						if( UStruct* Struct = dynamic_cast<UStruct*>(Export.Object) )
 						{
-							UStruct* Struct = (UStruct*)Export.Object;
 							if (Struct->GetSuperStruct() != NULL)
 							{
 								Export.SuperIndex = Linker->MapObject(Struct->GetSuperStruct());
diff --git a/Engine/Source/Runtime/CoreUObject/Private/UObject/ScriptCore.cpp b/Engine/Source/Runtime/CoreUObject/Private/UObject/ScriptCore.cpp
--- a/Engine/Source/Runtime/CoreUObject/Private/UObject/ScriptCore.cpp
+++ b/Engine/Source/Runtime/CoreUObject/Private/UObject/ScriptCore.cpp
@@ -667,7 +667,7 @@ bool UObject::CallFunctionByNameWithArguments(const TCHAR* Str, FOutputDevice& A
 		LastParameter = *It;
 	}
 
-	UStrProperty* LastStringParameter = Cast<UStrProperty>(LastParameter);
+	UStrProperty* LastStringParameter = dynamic_cast<UStrProperty*>(LastParameter);
 
 
 	// Parse all function parameters.
@@ -681,7 +681,7 @@ bool UObject::CallFunctionByNameWithArguments(const TCHAR* Str, FOutputDevice& A
 		UProperty* propertyParam = *It;
 		if (NumParamsEvaluated == 0 && Executor)
 		{
-			UObjectPropertyBase* Op = Cast<UObjectPropertyBase>(*It);
+			UObjectPropertyBase* Op = dynamic_cast<UObjectPropertyBase*>(*It);
 			if( Op && Executor->IsA(Op->PropertyClass) )
 			{
 				// First parameter is implicit reference to object executing the command.
@@ -1189,7 +1189,7 @@ IMPLEMENT_VM_FUNCTION( EX_PopExecutionFlowIfNot, execPopExecutionFlowIfNot );
 
 void UObject::execLet( FFrame& Stack, RESULT_DECL )
 {
-	checkSlow(!IsA(UBoolProperty::StaticClass()));
+	checkSlow(!dynamic_cast<UBoolProperty*>(this));
 
 	// Get variable address.
 	Stack.MostRecentPropertyAddress = NULL;
@@ -1224,13 +1224,13 @@ void UObject::execLetObj( FFrame& Stack, RESULT_DECL )
 	}
 
 	void* ObjAddr = Stack.MostRecentPropertyAddress;
-	UObjectPropertyBase* ObjectProperty = Cast<UObjectPropertyBase>(Stack.MostRecentProperty);
+	UObjectPropertyBase* ObjectProperty = dynamic_cast<UObjectPropertyBase*>(Stack.MostRecentProperty);
 	if (ObjectProperty == NULL)
 	{
 		UArrayProperty* ArrayProp = ExactCast<UArrayProperty>(Stack.MostRecentProperty);
 		if (ArrayProp != NULL)
 		{
-			ObjectProperty = Cast<UObjectPropertyBase>(ArrayProp->Inner);
+			ObjectProperty = dynamic_cast<UObjectPropertyBase*>(ArrayProp->Inner);
 		}
 	}
 
@@ -1259,13 +1259,13 @@ void UObject::execLetWeakObjPtr( FFrame& Stack, RESULT_DECL )
 	}
 
 	void* ObjAddr = Stack.MostRecentPropertyAddress;
-	UObjectPropertyBase* ObjectProperty = Cast<UObjectPropertyBase>(Stack.MostRecentProperty);
+	UObjectPropertyBase* ObjectProperty = dynamic_cast<UObjectPropertyBase*>(Stack.MostRecentProperty);
 	if (ObjectProperty == NULL)
 	{
 		UArrayProperty* ArrayProp = ExactCast<UArrayProperty>(Stack.MostRecentProperty);
 		if (ArrayProp != NULL)
 		{
-			ObjectProperty = Cast<UObjectPropertyBase>(ArrayProp->Inner);
+			ObjectProperty = dynamic_cast<UObjectPropertyBase*>(ArrayProp->Inner);
 		}
 	}
 	
@@ -1319,7 +1319,7 @@ void UObject::execLetBool( FFrame& Stack, RESULT_DECL )
 	Stack.Step( Stack.Object, &NewValue );
 	if( BoolAddr )
 	{
-		checkSlow(BoolProperty->IsA(UBoolProperty::StaticClass()));
+		checkSlow(dynamic_cast<UBoolProperty*>(BoolProperty));
 		BoolProperty->SetPropertyValue( BoolAddr, NewValue );
 	}
 }
@@ -1745,7 +1745,7 @@ void UObject::execStructConst( FFrame& Stack, RESULT_DECL )
 	for( UProperty* StructProp = ScriptStruct->PropertyLink; StructProp; StructProp = StructProp->PropertyLinkNext )
 	{
 		// Const struct arrays aren't supported yet
-		if( StructProp->IsA(UArrayProperty::StaticClass()) )
+		if ( dynamic_cast<UArrayProperty*>(StructProp) )
 		{
 			continue;
 		}
@@ -1867,9 +1867,10 @@ void UObject::execMetaCast( FFrame& Stack, RESULT_DECL )
 	UClass* MetaClass = (UClass*)Stack.ReadObject();
 
 	// Compile actor expression.
-	UObject* Castee=NULL;
+	UObject* Castee = nullptr;
 	Stack.Step( Stack.Object, &Castee );
-	*(UObject**)Result = (Castee && Castee->IsA(UClass::StaticClass()) && ((UClass*)Castee)->IsChildOf(MetaClass)) ? Castee : NULL;
+	UClass* CasteeClass = dynamic_cast<UClass*>(Castee);
+	*(UObject**)Result = (CasteeClass && CasteeClass->IsChildOf(MetaClass)) ? Castee : nullptr;
 }
 IMPLEMENT_VM_FUNCTION( EX_MetaCast, execMetaCast );
 
@@ -1907,7 +1908,7 @@ void UObject::execObjectToInterface( FFrame& Stack, RESULT_DECL )
 	FScriptInterface& InterfaceValue = *(FScriptInterface*)Result;
 
 	// read the interface class off the stack
-	UClass* InterfaceClass = Cast<UClass>(Stack.ReadObject());
+	UClass* InterfaceClass = dynamic_cast<UClass*>(Stack.ReadObject());
 	checkSlow(InterfaceClass != NULL);
 
 	// read the object off the stack
@@ -1933,7 +1934,7 @@ void UObject::execInterfaceToInterface( FFrame& Stack, RESULT_DECL )
 	FScriptInterface& CastResult = *(FScriptInterface*)Result;
 
 	// read the interface class off the stack
-	UClass* ClassToCastTo = Cast<UClass>(Stack.ReadObject());
+	UClass* ClassToCastTo = dynamic_cast<UClass*>(Stack.ReadObject());
 	checkSlow(ClassToCastTo != NULL);
 	checkSlow(ClassToCastTo->HasAnyClassFlags(CLASS_Interface));
 
diff --git a/Engine/Source/Runtime/CoreUObject/Private/UObject/TextureAllocations.cpp b/Engine/Source/Runtime/CoreUObject/Private/UObject/TextureAllocations.cpp
--- a/Engine/Source/Runtime/CoreUObject/Private/UObject/TextureAllocations.cpp
+++ b/Engine/Source/Runtime/CoreUObject/Private/UObject/TextureAllocations.cpp
@@ -193,16 +193,18 @@ FArchive& operator<<( FArchive& Ar, FTextureAllocations& TextureAllocations )
 				{
 					UObject* Object = TagExpObjects(Index);
 					check(Object->HasAnyMarks(OBJECTMARK_TagExp));
-					if ( !Object->HasAnyFlags(RF_ClassDefaultObject) && Object->IsA(UTexture2D::StaticClass()) )
+					if ( !Object->HasAnyFlags(RF_ClassDefaultObject) )
 					{
-						UTexture2D* Texture2D = Cast<UTexture2D>(Object);
-						int32 PreAllocateSizeX = 0;
-						int32 PreAllocateSizeY = 0;
-						int32 PreAllocateNumMips = 0;
-						uint32 TexCreateFlags = 0;
-						if ( Texture2D->GetResourceMemSettings(Texture2D->FirstResourceMemMip, PreAllocateSizeX, PreAllocateSizeY, PreAllocateNumMips, TexCreateFlags ) )
+						if (UTexture2D* Texture2D = dynamic_cast<UTexture2D*>(Object))
 						{
-							TextureAllocations.AddResourceMemInfo(PreAllocateSizeX, PreAllocateSizeY, PreAllocateNumMips, Texture2D->Format, TexCreateFlags );
+							int32 PreAllocateSizeX = 0;
+							int32 PreAllocateSizeY = 0;
+							int32 PreAllocateNumMips = 0;
+							uint32 TexCreateFlags = 0;
+							if ( Texture2D->GetResourceMemSettings(Texture2D->FirstResourceMemMip, PreAllocateSizeX, PreAllocateSizeY, PreAllocateNumMips, TexCreateFlags ) )
+							{
+								TextureAllocations.AddResourceMemInfo(PreAllocateSizeX, PreAllocateSizeY, PreAllocateNumMips, Texture2D->Format, TexCreateFlags );
+							}
 						}
 					}
 				}
@@ -223,7 +225,7 @@ FArchive& operator<<( FArchive& Ar, FTextureAllocations& TextureAllocations )
 				NumExportTexturesAdded = 0;
 				for ( int32 ExportIndex=0; ExportIndex < Linker->ExportMap.Num(); ++ExportIndex )
 				{
-					UTexture2D* Texture2D = Cast<UTexture2D>(Linker->ExportMap(ExportIndex).Object);
+					UTexture2D* Texture2D = dynamic_cast<UTexture2D*>(Linker->ExportMap(ExportIndex).Object);
 					if ( Texture2D && !Texture2D->HasAnyFlags(RF_ClassDefaultObject) )
 					{
 						NumExportTexturesTotal++;
diff --git a/Engine/Source/Runtime/CoreUObject/Private/UObject/UObjectBase.cpp b/Engine/Source/Runtime/CoreUObject/Private/UObject/UObjectBase.cpp
--- a/Engine/Source/Runtime/CoreUObject/Private/UObject/UObjectBase.cpp
+++ b/Engine/Source/Runtime/CoreUObject/Private/UObject/UObjectBase.cpp
@@ -449,8 +449,13 @@ struct FStructOrEnumCompiledInfo : FFieldCompiledInInfo
 {
 	FStructOrEnumCompiledInfo(SIZE_T InClassSize, uint32 InCrc)
 		: FFieldCompiledInInfo(InClassSize, InCrc)
-	{}
-	virtual UField* Register() const { return NULL; }
+	{
+	}
+
+	virtual UClass* Register() const
+	{
+		return nullptr;
+	}
 };
 
 /** Registered struct info (including size and reflection info) */
@@ -627,13 +632,13 @@ void UClassCompiledInDefer(FFieldCompiledInInfo* ClassInfo, const TCHAR* Name, S
 			GetObjectsWithOuter(ExistingClass, ClassSubobjects);
 			for (auto ClassSubobject : ClassSubobjects)
 			{
-				if (auto Enum = Cast<UEnum>(ClassSubobject))
+				if (auto Enum = dynamic_cast<UEnum*>(ClassSubobject))
 				{
 					Enum->RemoveNamesFromMasterList();
 				}
 			}
 		}
-		ClassInfo->OldField = ExistingClass;
+		ClassInfo->OldClass = ExistingClass;
 		GetHotReloadClasses().Add(ClassInfo);
 
 		*ExistingClassInfo = ClassInfo;
@@ -674,17 +679,18 @@ void UClassRegisterAllCompiledInClasses()
 void UClassReplaceHotReloadClasses()
 {
 	if (FCoreUObjectDelegates::ReplaceHotReloadClassDelegate.IsBound())
-	{	
+	{
 		for (auto Class : GetHotReloadClasses())
 		{
+			check(Class->OldClass);
+
+			UClass* RegisteredClass = nullptr;
 			if (Class->bHasChanged)
 			{
-				FCoreUObjectDelegates::ReplaceHotReloadClassDelegate.Execute(CastChecked<UClass>(Class->OldField), CastChecked<UClass>(Class->Register()));
-			}
-			else
-			{
-				FCoreUObjectDelegates::ReplaceHotReloadClassDelegate.Execute(CastChecked<UClass>(Class->OldField), NULL);
+				RegisteredClass = Class->Register();
 			}
+
+			FCoreUObjectDelegates::ReplaceHotReloadClassDelegate.Execute(Class->OldClass, RegisteredClass);
 		}
 	}
 	GetHotReloadClasses().Empty();
diff --git a/Engine/Source/Runtime/CoreUObject/Private/UObject/UObjectBaseUtility.cpp b/Engine/Source/Runtime/CoreUObject/Private/UObject/UObjectBaseUtility.cpp
--- a/Engine/Source/Runtime/CoreUObject/Private/UObject/UObjectBaseUtility.cpp
+++ b/Engine/Source/Runtime/CoreUObject/Private/UObject/UObjectBaseUtility.cpp
@@ -110,9 +110,18 @@ FString UObjectBaseUtility::GetFullGroupName( bool bStartWithOuter ) const
  */
 UPackage* UObjectBaseUtility::GetOutermost() const
 {
-	UObject* Top;
-	for( Top = (UObject*)this ; Top && Top->GetOuter() ; Top = Top->GetOuter() );
-	return CastChecked<UPackage>(Top);
+	UObject* Top = (UObject*)this;
+	for (;;)
+	{
+		UObject* Outer = Top->GetOuter();
+		if (!Outer)
+		{
+			UPackage* Result = dynamic_cast<UPackage*>(Top);
+			check(Result);
+			return Result;
+		}
+		Top = Outer;
+	}
 }
 
 /** 
diff --git a/Engine/Source/Runtime/CoreUObject/Private/UObject/UObjectGlobals.cpp b/Engine/Source/Runtime/CoreUObject/Private/UObject/UObjectGlobals.cpp
--- a/Engine/Source/Runtime/CoreUObject/Private/UObject/UObjectGlobals.cpp
+++ b/Engine/Source/Runtime/CoreUObject/Private/UObject/UObjectGlobals.cpp
@@ -569,7 +569,7 @@ bool ResolveName( UObject*& InPackage, FString& InOutName, bool Create, bool Thr
 		{
 			if (!ScriptPackageName)
 			{
-				InPackage = LoadPackage(Cast<UPackage>(InPackage), *PartialName, 0);
+				InPackage = LoadPackage(dynamic_cast<UPackage*>(InPackage), *PartialName, 0);
 			}
 			if (!InPackage)
 			{
@@ -692,7 +692,7 @@ UObject* StaticLoadObject(UClass* ObjectClass, UObject* InOuter, const TCHAR* In
 			}
 
 			// If the object was not found, check for a redirector and follow it if the class matches
-			UObjectRedirector* Redirector = Cast<UObjectRedirector>(StaticFindObjectFast(UObjectRedirector::StaticClass(), InOuter, *StrName));
+			UObjectRedirector* Redirector = FindObjectFast<UObjectRedirector>(InOuter, *StrName);
 			if ( Redirector && Redirector->DestinationObject && Redirector->DestinationObject->IsA(ObjectClass) )
 			{
 				return Redirector->DestinationObject;
@@ -1455,7 +1455,7 @@ UObject* StaticDuplicateObjectEx( FObjectDuplicationParameters& Parameters )
 	
 	if (GIsDuplicatingClassForReinstancing)
 	{
-		FBlueprintSupport::DuplicateAllFields(Cast<UStruct>(Parameters.SourceObject), Writer);
+		FBlueprintSupport::DuplicateAllFields(dynamic_cast<UStruct*>(Parameters.SourceObject), Writer);
 	}
 
 	InstanceGraph.SetDestinationRoot( DupRootObject );
@@ -2433,7 +2433,7 @@ private:
 					}
 					if ( InReferencingObject != NULL )
 					{
-						CurrentReferenceInfo->ReferencingProperties.AddUnique(Cast<const UProperty>(InReferencingObject));
+						CurrentReferenceInfo->ReferencingProperties.AddUnique(dynamic_cast<const UProperty*>(InReferencingObject));
 					}
 					CurrentReferenceInfo->TotalReferences++;
 				}
diff --git a/Engine/Source/Runtime/CoreUObject/Public/Serialization/ArchiveUObject.h b/Engine/Source/Runtime/CoreUObject/Public/Serialization/ArchiveUObject.h
--- a/Engine/Source/Runtime/CoreUObject/Public/Serialization/ArchiveUObject.h
+++ b/Engine/Source/Runtime/CoreUObject/Public/Serialization/ArchiveUObject.h
@@ -488,7 +488,7 @@ public:
 		for (FObjectIterator It; It; ++It)
 		{
 			UObject* PotentialReferencer = *It;
-			if (!TargetObjects.Contains(Cast<T>(PotentialReferencer))
+			if (!TargetObjects.Contains(dynamic_cast<T*>(PotentialReferencer))
 			&&	(PackageToCheck == NULL || PotentialReferencer->IsIn(PackageToCheck))
 			&&	(!bIgnoreTemplates || !PotentialReferencer->IsTemplate()) )
 			{
diff --git a/Engine/Source/Runtime/CoreUObject/Public/Templates/Casts.h b/Engine/Source/Runtime/CoreUObject/Public/Templates/Casts.h
--- a/Engine/Source/Runtime/CoreUObject/Public/Templates/Casts.h
+++ b/Engine/Source/Runtime/CoreUObject/Public/Templates/Casts.h
@@ -2,136 +2,213 @@
 
 #pragma once
 
-// it doesn't seem worthwhile to risk changing Cast<>/IsA() to work with interfaces at this time,
-// so we'll just make sure it fails... spectacularly
-#define CATCH_UNSUPPORTED_INTERFACECAST(TheType) static_assert((TheType::StaticClassFlags & CLASS_Interface) == 0, "Casting to interface type " #TheType " is unsupported.")
-
 COREUOBJECT_API void CastLogError(const TCHAR* FromType, const TCHAR* ToType);
 
-// Dynamically cast an object type-safely.
-template< class T > T* Cast( UObject* Src )
+/**
+ * Metafunction which detects whether or not a class is an IInterface.  Rules:
+ *
+ * 1. A UObject is not an IInterface.
+ * 2. A type without a UClassType typedef member is not an IInterface.
+ * 3. A type whose UClassType::StaticClassFlags does not have CLASS_Interface set is not an IInterface.
+ *
+ * Otherwise, assume it's an IInterface.
+ */
+template <typename T, bool bIsAUObject_IMPL = CanConvertPointerFromTo<T, UObject>::Result>
+struct TIsIInterface
 {
-	CATCH_UNSUPPORTED_INTERFACECAST(T);
+	enum { Value = false };
+};
 
-	return Src && Src->IsA<T>() ? (T*)Src : NULL;
-}
-template< class T > FORCEINLINE T* ExactCast( UObject* Src )
+template <typename T>
+struct TIsIInterface<T, false>
 {
-	return Src && (Src->GetClass() == T::StaticClass()) ? (T*)Src : NULL;
-}
+	template <typename U> static char (&Resolve(typename U::UClassType*))[(U::UClassType::StaticClassFlags & CLASS_Interface) ? 2 : 1];
+	template <typename U> static char (&Resolve(...))[1];
+
+	enum { Value = sizeof(Resolve<T>(0)) - 1 };
+};
 
-template< class T > T* CastChecked( UObject* Src, ECastCheckedType::Type CheckType = ECastCheckedType::NullChecked )
+template <typename T>
+struct TIsCastable
 {
-	CATCH_UNSUPPORTED_INTERFACECAST(T);
-	
-#if !(UE_BUILD_SHIPPING || UE_BUILD_TEST)
-	if( (!Src && CheckType == ECastCheckedType::NullChecked) || (Src && !Src->template IsA<T>() ) )
-	{
-		CastLogError(Src ? *Src->GetFullName() : TEXT("NULL"), *T::StaticClass()->GetName());
-	}
-#endif // !(UE_BUILD_SHIPPING || UE_BUILD_TEST)
-	return (T*)Src;
-}
+	// It's from-castable if it's an interface or a UObject-derived type
+	enum { Value = TIsIInterface<T>::Value || CanConvertPointerFromTo<T, const volatile UObject>::Result };
+};
 
-// auto weak versions
+template <typename T>
+struct TIsCastableToPointer : TOr<TIsVoidType<T>, TIsCastable<T>>
+{
+	// It's to-pointer-castable if it's from-castable or void
+};
 
-// Dynamically cast an object type-safely.
-template< class T, class U > T* Cast( const TAutoWeakObjectPtr<U>& Src )
+template <typename T>
+FORCEINLINE typename TEnableIf<TIsIInterface<T>::Value, FString>::Type GetTypeName()
 {
-	return Cast<T>(Src.Get());
+	return T::UClassType::StaticClass()->GetName();
 }
 
-template< class T, class U > FORCEINLINE T* ExactCast( const TAutoWeakObjectPtr<U>& Src )
+template <typename T>
+FORCEINLINE typename TEnableIf<!TIsIInterface<T>::Value, FString>::Type GetTypeName()
 {
-	return ExactCast<T>(Src.Get());
+	return T::StaticClass()->GetName();
 }
 
-template< class T, class U > T* CastChecked( const TAutoWeakObjectPtr<U>& Src, ECastCheckedType::Type CheckType = ECastCheckedType::NullChecked )
+enum class ECastType
 {
-	return CastChecked<T>(Src.Get(), CheckType);
-}
+	UObjectToUObject,
+	InterfaceToUObject,
+	UObjectToInterface,
+	InterfaceToInterface,
+	FromCastFlags
+};
 
-// FSubobjectPtr versions
+template <typename Type>
+struct TCastFlags
+{
+	static const uint64 Value = CASTCLASS_None;
+};
 
-template< class T > T* Cast( const FSubobjectPtr& Src )
+template <typename From, typename To, bool bFromInterface = TIsIInterface<From>::Value, bool bToInterface = TIsIInterface<To>::Value, uint64 CastClass = TCastFlags<To>::Value>
+struct TGetCastType
 {
-	return Cast<T>(Src.Get());
-}
+	static const ECastType Value = ECastType::FromCastFlags;
+};
+
+template <typename From, typename To                  > struct TGetCastType<From, To, false, false, CASTCLASS_None> { static const ECastType Value = ECastType::UObjectToUObject;     };
+template <typename From, typename To                  > struct TGetCastType<From, To, false, true , CASTCLASS_None> { static const ECastType Value = ECastType::UObjectToInterface;   };
+template <typename From, typename To, uint64 CastClass> struct TGetCastType<From, To, true,  false, CastClass     > { static const ECastType Value = ECastType::InterfaceToUObject;   };
+template <typename From, typename To, uint64 CastClass> struct TGetCastType<From, To, true,  true , CastClass     > { static const ECastType Value = ECastType::InterfaceToInterface; };
+
+template <typename To, typename From>
+To* Cast(From* Src);
 
-template< class T > FORCEINLINE T* ExactCast( const FSubobjectPtr& Src )
+template <typename From, typename To, ECastType CastType = TGetCastType<From, To>::Value>
+struct TCastImpl
 {
-	return ExactCast<T>(Src.Get());
-}
+	// This is the cast flags implementation
+	FORCEINLINE static To* DoCast( UObject* Src )
+	{
+		return Src->GetClass()->HasAnyCastFlag(TCastFlags<To>::Value) ? (To*)Src : nullptr;
+	}
+};
 
-template< class T > T* CastChecked( const FSubobjectPtr& Src, ECastCheckedType::Type CheckType = ECastCheckedType::NullChecked )
+template <typename From, typename To>
+struct TCastImpl<From, To, ECastType::UObjectToUObject>
 {
-	return CastChecked<T>(Src.Get(), CheckType);
-}
+	FORCEINLINE static To* DoCast( UObject* Src )
+	{
+		return Src->IsA<To>() ? (To*)Src : nullptr;
+	}
+};
 
+template <typename From, typename To>
+struct TCastImpl<From, To, ECastType::InterfaceToUObject>
+{
+	FORCEINLINE static To* DoCast( From* Src )
+	{
+		return Cast<To>(Src->__GetUObject());
+	}
+};
 
-/**
- * Dynamically casts an object pointer to an interface pointer.  InterfaceType must correspond to a native interface class.
- *
- * @param	Src		the object to be casted
- *
- * @return	a pointer to the interface type specified or NULL if the object doesn't support for specified interface
- */
-template<typename InterfaceType>
-FORCEINLINE InterfaceType* InterfaceCast( UObject* Src )
+template <typename From, typename To>
+struct TCastImpl<From, To, ECastType::UObjectToInterface>
 {
-	return Src != NULL
-		? static_cast<InterfaceType*>(Src->GetInterfaceAddress(InterfaceType::UClassType::StaticClass()))
-		: NULL;
-}
-template<typename InterfaceType>
-TScriptInterface<InterfaceType> ScriptInterfaceCast( UObject* Src )
+	FORCEINLINE static To* DoCast( UObject* Src )
+	{
+		return (To*)Src->GetInterfaceAddress(To::UClassType::StaticClass());
+	}
+};
+
+template <typename From, typename To>
+struct TCastImpl<From, To, ECastType::InterfaceToInterface>
 {
-	TScriptInterface<InterfaceType> Result;
-	InterfaceType* ISrc = InterfaceCast<InterfaceType>(Src);
-	if ( ISrc != NULL )
+	FORCEINLINE static To* DoCast( From* Src )
 	{
-		Result.SetObject(Src);
-		Result.SetInterface(ISrc);
+		return Cast<To>(Src->__GetUObject());
 	}
-	return Result;
+};
+
+// Dynamically cast an object type-safely.
+template <typename To, typename From>
+FORCEINLINE To* Cast(From* Src)
+{
+	return Src ? TCastImpl<From, To>::DoCast(Src) : nullptr;
 }
 
-// Const versions of the casts
-template< class T > FORCEINLINE const T* Cast         ( const UObject* Src                                                                   ) { return Cast         <T>(const_cast<UObject*>(Src)); }
-template< class T > FORCEINLINE const T* ExactCast    ( const UObject* Src                                                                   ) { return ExactCast    <T>(const_cast<UObject*>(Src)); }
-template< class T > FORCEINLINE const T* CastChecked  ( const UObject* Src, ECastCheckedType::Type CheckType = ECastCheckedType::NullChecked ) { return CastChecked  <T>(const_cast<UObject*>(Src), CheckType); }
-template< class T > FORCEINLINE const T* InterfaceCast( const UObject* Src                                                                   ) { return InterfaceCast<T>(const_cast<UObject*>(Src)); }
+template< class T >
+FORCEINLINE T* ExactCast( UObject* Src )
+{
+	return Src && (Src->GetClass() == T::StaticClass()) ? (T*)Src : nullptr;
+}
 
-// specializations of cast methods
-#if (UE_BUILD_SHIPPING || UE_BUILD_TEST)
+template <typename To, typename From>
+To* CastChecked(From* Src, ECastCheckedType::Type CheckType = ECastCheckedType::NullChecked)
+{
+	To* Result = Cast<To>(Src);
 
-#define DECLARE_CAST_BY_FLAG(ClassName) \
-class ClassName; \
-template<> FORCEINLINE ClassName* Cast( UObject* Src ) \
-{ \
-	return Src && Src->GetClass()->HasAnyCastFlag(CASTCLASS_##ClassName) ? (ClassName*)Src : NULL; \
-} \
-template<> FORCEINLINE ClassName* CastChecked( UObject* Src, ECastCheckedType::Type CheckType ) \
-{ \
-	return (ClassName*)Src; \
+	#if !(UE_BUILD_SHIPPING || UE_BUILD_TEST)
+
+		if (Src)
+		{
+			if (!Result)
+			{
+				CastLogError(*Cast<UObject>(Src)->GetFullName(), *GetTypeName<To>());
+			}
+		}
+		else if (CheckType == ECastCheckedType::NullChecked)
+		{
+			CastLogError(TEXT("nullptr"), *GetTypeName<To>());
+		}
+
+	#endif
+
+	return Result;
 }
 
-#else
+template <typename InterfaceType>
+DEPRECATED(4.5, "InterfaceCast is deprecated, use Cast instead.")
+FORCEINLINE InterfaceType* InterfaceCast(UObject* Src)
+{
+	return Cast<InterfaceType>(Src);
+}
 
-#define DECLARE_CAST_BY_FLAG(ClassName) \
-class ClassName; \
-template<> FORCEINLINE ClassName* Cast( UObject* Src ) \
-{ \
-	return Src && Src->GetClass()->HasAnyCastFlag(CASTCLASS_##ClassName) ? (ClassName*)Src : NULL; \
-} \
-template<> FORCEINLINE ClassName* CastChecked( UObject* Src, ECastCheckedType::Type CheckType ) \
-{ \
-	if( (!Src && CheckType == ECastCheckedType::NullChecked) || (Src && !Src->GetClass()->HasAnyCastFlag(CASTCLASS_##ClassName) ) ) \
-		CastLogError(Src ? *Src->GetFullName() : TEXT("NULL"), TEXT(#ClassName)); \
-	return (ClassName*)Src; \
+template <typename InterfaceType>
+DEPRECATED(4.5, "InterfaceCast is deprecated, use Cast instead.")
+FORCEINLINE InterfaceType* InterfaceCast(const UObject* Src)
+{
+	return Cast<InterfaceType>(const_cast<UObject*>(Src));
 }
 
-#endif
+// auto weak versions
+template< class T, class U > FORCEINLINE T* Cast       ( const TAutoWeakObjectPtr<U>& Src                                                                   ) { return Cast       <T>(Src.Get()); }
+template< class T, class U > FORCEINLINE T* ExactCast  ( const TAutoWeakObjectPtr<U>& Src                                                                   ) { return ExactCast  <T>(Src.Get()); }
+template< class T, class U > FORCEINLINE T* CastChecked( const TAutoWeakObjectPtr<U>& Src, ECastCheckedType::Type CheckType = ECastCheckedType::NullChecked ) { return CastChecked<T>(Src.Get(), CheckType); }
+
+// FSubobjectPtr versions
+template< class T > FORCEINLINE T* Cast       ( const FSubobjectPtr& Src                                                                   ) { return Cast       <T>(Src.Get()); }
+template< class T > FORCEINLINE T* ExactCast  ( const FSubobjectPtr& Src                                                                   ) { return ExactCast  <T>(Src.Get()); }
+template< class T > FORCEINLINE T* CastChecked( const FSubobjectPtr& Src, ECastCheckedType::Type CheckType = ECastCheckedType::NullChecked ) { return CastChecked<T>(Src.Get(), CheckType); }
+
+// TSubclassOf versions
+template< class T, class U > FORCEINLINE T* Cast       ( const TSubclassOf<U>& Src                                                                   ) { return Cast       <T>(*Src); }
+template< class T, class U > FORCEINLINE T* CastChecked( const TSubclassOf<U>& Src, ECastCheckedType::Type CheckType = ECastCheckedType::NullChecked ) { return CastChecked<T>(*Src, CheckType); }
+
+// Const versions of the casts
+template< class T, class U > FORCEINLINE const T* Cast       ( const U      * Src                                                                   ) { return Cast       <T>(const_cast<U      *>(Src)); }
+template< class T          > FORCEINLINE const T* ExactCast  ( const UObject* Src                                                                   ) { return ExactCast  <T>(const_cast<UObject*>(Src)); }
+template< class T, class U > FORCEINLINE const T* CastChecked( const U      * Src, ECastCheckedType::Type CheckType = ECastCheckedType::NullChecked ) { return CastChecked<T>(const_cast<U      *>(Src), CheckType); }
+
+#define DECLARE_CAST_BY_FLAG_FWD(ClassName) class ClassName;
+#define DECLARE_CAST_BY_FLAG_CAST(ClassName) \
+	template <> \
+	struct TCastFlags<ClassName> \
+	{ \
+		static const uint64 Value = CASTCLASS_##ClassName; \
+	};
+
+#define DECLARE_CAST_BY_FLAG(ClassName) \
+	DECLARE_CAST_BY_FLAG_FWD(ClassName) \
+	DECLARE_CAST_BY_FLAG_CAST(ClassName)
 
 DECLARE_CAST_BY_FLAG(UField)
 DECLARE_CAST_BY_FLAG(UEnum)
@@ -160,3 +237,52 @@ DECLARE_CAST_BY_FLAG(UTextProperty)
 DECLARE_CAST_BY_FLAG(UArrayProperty)
 DECLARE_CAST_BY_FLAG(UDelegateProperty)
 DECLARE_CAST_BY_FLAG(UMulticastDelegateProperty)
+
+#undef DECLARE_CAST_BY_FLAG
+#undef DECLARE_CAST_BY_FLAG_CASTCHECKED
+#undef DECLARE_CAST_BY_FLAG_CAST
+#undef DECLARE_CAST_BY_FLAG_FWD
+
+template <typename To, typename From>
+FORCEINLINE typename TEnableIf<TAnd<TIsPointerType<To>, TAnd<TIsCastableToPointer<typename TRemovePointer<To>::Type>, TIsCastable<From>>>::Value, To>::Type __DynamicCast(From* Arg)
+{
+	typedef typename TRemovePointer<To  >::Type ToValueType;
+	typedef typename TRemovePointer<From>::Type FromValueType;
+
+	// Casting away const/volatile
+	static_assert(!TLosesQualifiersFromTo<FromValueType, ToValueType>::Value, "Conversion loses qualifiers");
+
+	// If we're casting to void, cast to UObject instead and let it implicitly cast to void
+	return Cast<typename TChooseClass<TIsVoidType<ToValueType>::Value, UObject, ToValueType>::Result>(Arg);
+}
+
+template <typename To, typename From>
+FORCEINLINE typename TEnableIf<!TAnd<TIsPointerType<To>, TAnd<TIsCastableToPointer<typename TRemovePointer<To>::Type>, TIsCastable<From>>>::Value, To>::Type __DynamicCast(From* Arg)
+{
+	return dynamic_cast<To>(Arg);
+}
+
+template <typename To, typename From>
+FORCEINLINE typename TEnableIf<TAnd<TIsCastable<typename TRemoveReference<To>::Type>, TIsCastable<typename TRemoveReference<From>::Type>>::Value, To>::Type __DynamicCast(From&& Arg)
+{
+	typedef typename TRemoveReference<From>::Type FromValueType;
+	typedef typename TRemoveReference<To  >::Type ToValueType;
+
+	// Casting away const/volatile
+	static_assert(!TLosesQualifiersFromTo<FromValueType, ToValueType>::Value, "Conversion loses qualifiers");
+
+	// T&& can only be cast to U&&
+	// http://en.cppreference.com/w/cpp/language/dynamic_cast
+	static_assert(TOr<TIsLValueReferenceType<From>, TIsRValueReferenceType<To>>::Value, "Cannot dynamic_cast from an rvalue to a non-rvalue reference");
+
+	return Forward<To>(*CastChecked<typename TRemoveReference<To>::Type>(&Arg));
+}
+
+template <typename To, typename From>
+FORCEINLINE typename TEnableIf<!TAnd<TIsCastable<typename TRemoveReference<To>::Type>, TIsCastable<typename TRemoveReference<From>::Type>>::Value, To>::Type __DynamicCast(From&& Arg)
+{
+	// This may fail when dynamic_casting rvalue references due to patchy compiler support
+	return dynamic_cast<To>(Arg);
+}
+
+#define dynamic_cast __DynamicCast
diff --git a/Engine/Source/Runtime/CoreUObject/Public/UObject/AssetPtr.h b/Engine/Source/Runtime/CoreUObject/Public/UObject/AssetPtr.h
--- a/Engine/Source/Runtime/CoreUObject/Public/UObject/AssetPtr.h
+++ b/Engine/Source/Runtime/CoreUObject/Public/UObject/AssetPtr.h
@@ -206,7 +206,7 @@ public:
 	 */
 	FORCEINLINE T *Get() const
 	{
-		return Cast<T>(FAssetPtr::Get());
+		return dynamic_cast<T*>(FAssetPtr::Get());
 	}
 
 	/**  
@@ -449,7 +449,7 @@ public:
 	 */
 	FORCEINLINE UClass *Get() const
 	{
-		UClass* Class = Cast<UClass>(FAssetPtr::Get());
+		UClass* Class = dynamic_cast<UClass*>(FAssetPtr::Get());
 		if (!Class || !Class->IsChildOf(TClass::StaticClass()))
 		{
 			return NULL;
diff --git a/Engine/Source/Runtime/CoreUObject/Public/UObject/Class.h b/Engine/Source/Runtime/CoreUObject/Public/UObject/Class.h
--- a/Engine/Source/Runtime/CoreUObject/Public/UObject/Class.h
+++ b/Engine/Source/Runtime/CoreUObject/Public/UObject/Class.h
@@ -1723,7 +1723,7 @@ public:
 	// Class flags; See EClassFlags for more information
 	uint32 ClassFlags;
 
-	// Cast flags used to accelerate Cast<T> on objects of this type for common T
+	// Cast flags used to accelerate dynamic_cast<T*> on objects of this type for common T
 	EClassCastFlags ClassCastFlags;
 
 	// Class pseudo-unique counter; used to accelerate unique instance name generation
diff --git a/Engine/Source/Runtime/CoreUObject/Public/UObject/ConstructorHelpers.h b/Engine/Source/Runtime/CoreUObject/Public/UObject/ConstructorHelpers.h
--- a/Engine/Source/Runtime/CoreUObject/Public/UObject/ConstructorHelpers.h
+++ b/Engine/Source/Runtime/CoreUObject/Public/UObject/ConstructorHelpers.h
@@ -28,7 +28,7 @@ namespace ConstructorHelpersInternal
 
 		UClass* Class = T::StaticClass();
 		Class->GetDefaultObject(); // force the CDO to be created if it hasn't already
-		T* ObjectPtr = Cast<T>(StaticLoadObject(T::StaticClass(), NULL, *PathName));
+		T* ObjectPtr = LoadObject<T>(NULL, *PathName);
 		return ObjectPtr;
 	}
 
diff --git a/Engine/Source/Runtime/CoreUObject/Public/UObject/LazyObjectPtr.h b/Engine/Source/Runtime/CoreUObject/Public/UObject/LazyObjectPtr.h
--- a/Engine/Source/Runtime/CoreUObject/Public/UObject/LazyObjectPtr.h
+++ b/Engine/Source/Runtime/CoreUObject/Public/UObject/LazyObjectPtr.h
@@ -326,7 +326,7 @@ public:
 	 */
 	FORCEINLINE T *Get() const
 	{
-		return Cast<T>(FLazyObjectPtr::Get());
+		return dynamic_cast<T*>(FLazyObjectPtr::Get());
 	}
 
 	/**  
diff --git a/Engine/Source/Runtime/CoreUObject/Public/UObject/ObjectBase.h b/Engine/Source/Runtime/CoreUObject/Public/UObject/ObjectBase.h
--- a/Engine/Source/Runtime/CoreUObject/Public/UObject/ObjectBase.h
+++ b/Engine/Source/Runtime/CoreUObject/Public/UObject/ObjectBase.h
@@ -524,7 +524,7 @@ namespace UC
 		NotBlueprintable,
 
 		/// This keyword indicates that the class should be accessible outside of it's module, but does not need all methods exported.
-		/// It exports only the autogenerated methods required for Cast<>, etc... to work.
+		/// It exports only the autogenerated methods required for dynamic_cast<>, etc... to work.
 		MinimalAPI,
 
 		/// Prevents automatic generation of the constructor declaration.
@@ -623,7 +623,7 @@ namespace UI
 		DependsOn,
 
 		/// This keyword indicates that the interface should be accessible outside of it's module, but does not need all methods exported.
-		/// It exports only the autogenerated methods required for Cast<>, etc... to work.
+		/// It exports only the autogenerated methods required for dynamic_cast<>, etc... to work.
 		MinimalAPI,
 
 		/// Exposes this interface as an acceptable base class for creating blueprints.  The default is NotBlueprintable, unless inherited otherwise. This is inherited by subclasses.
diff --git a/Engine/Source/Runtime/CoreUObject/Public/UObject/PropertyTag.h b/Engine/Source/Runtime/CoreUObject/Public/UObject/PropertyTag.h
--- a/Engine/Source/Runtime/CoreUObject/Public/UObject/PropertyTag.h
+++ b/Engine/Source/Runtime/CoreUObject/Public/UObject/PropertyTag.h
@@ -27,40 +27,36 @@ struct FPropertyTag
 	FPropertyTag()
 	{}
 	FPropertyTag( FArchive& InSaveAr, UProperty* Property, int32 InIndex, uint8* Value, uint8* Defaults )
-	:	Type		(Property->GetID())
-	, BoolVal	(0)
-	,	Name		(Property->GetFName())
-	,	StructName	(NAME_None)
-	,	EnumName	(NAME_None)
-	, InnerType	(NAME_None)
-	,	Size		(0)
-	,	ArrayIndex	(InIndex)
-	,	SizeOffset	(INDEX_NONE)
+		: Type      (Property->GetID())
+		, BoolVal   (0)
+		, Name      (Property->GetFName())
+		, StructName(NAME_None)
+		, EnumName  (NAME_None)
+		, InnerType (NAME_None)
+		, Size      (0)
+		, ArrayIndex(InIndex)
+		, SizeOffset(INDEX_NONE)
 	{
 		if (Property)
 		{
 			// Handle structs.
-			if (Property->IsA(UStructProperty::StaticClass()))
+			if (UStructProperty* StructProperty = Cast<UStructProperty>(Property))
 			{
-				UStructProperty* StructProperty = CastChecked<UStructProperty>(Property);
 				StructName = StructProperty->Struct->GetFName();
 			}
-			else if (Property->GetClass() == UByteProperty::StaticClass())
+			else if (UByteProperty* ByteProp = Cast<UByteProperty>(Property))
 			{
-				UByteProperty* ByteProp = CastChecked<UByteProperty>(Property);
 				if (ByteProp->Enum != NULL)
 				{
 					EnumName = ByteProp->Enum->GetFName();
 				}
 			}
-			else if (Property->IsA(UArrayProperty::StaticClass()))
+			else if (UArrayProperty* ArrayProp = Cast<UArrayProperty>(Property))
 			{
-				UArrayProperty* ArrayProp = CastChecked<UArrayProperty>(Property);
 				InnerType = ArrayProp->Inner->GetID();
 			}
-			else if (Property->IsA(UBoolProperty::StaticClass()))
+			else if (UBoolProperty* Bool = Cast<UBoolProperty>(Property))
 			{
-				UBoolProperty* Bool = CastChecked<UBoolProperty>(Property);
 				BoolVal = Bool->GetPropertyValue(Value);
 			}
 		}
diff --git a/Engine/Source/Runtime/CoreUObject/Public/UObject/Stack.h b/Engine/Source/Runtime/CoreUObject/Public/UObject/Stack.h
--- a/Engine/Source/Runtime/CoreUObject/Public/UObject/Stack.h
+++ b/Engine/Source/Runtime/CoreUObject/Public/UObject/Stack.h
@@ -279,26 +279,17 @@ inline VariableSizeType FFrame::ReadVariableSize( UField** ExpressionField/*=NUL
 
 	if ( Field != NULL )
 	{
-		UProperty* Property = Cast<UProperty>(Field);
-		if ( Property != NULL )
+		if (UProperty* Property = dynamic_cast<UProperty*>(Field))
 		{
 			Result = Property->GetSize();
 		}
-		else
+		else if (UEnum* ExplicitEnumValue = dynamic_cast<UEnum*>(Field))
 		{
-			UEnum* ExplicitEnumValue = Cast<UEnum>(Field);
-			if ( ExplicitEnumValue != NULL )
-			{
-				Result = 1;
-			}
-			else
-			{
-				UFunction* FunctionRef = Cast<UFunction>(Field);
-				if ( FunctionRef != NULL )
-				{
-					Result = sizeof(ScriptPointerType);
-				}
-			}
+			Result = 1;
+		}
+		else if (UFunction* FunctionRef = dynamic_cast<UFunction*>(Field))
+		{
+			Result = sizeof(ScriptPointerType);
 		}
 	}
 	else
@@ -364,7 +355,7 @@ inline VariableSizeType FFrame::ReadVariableSize( UField** ExpressionField/*=NUL
 
 	if ( ExpressionField != NULL )
 	{
-		*ExpressionField = Cast<UField>(Field);
+		*ExpressionField = dynamic_cast<UField*>(Field);
 	}
 
 	return Result;
@@ -397,7 +388,7 @@ FORCEINLINE_DEBUGGABLE void FFrame::StepCompiledIn(void*const Result)
 	}
 	else
 	{
-		checkSlow(Cast<TProperty>(PropertyChainForCompiledIn) && Cast<UProperty>(PropertyChainForCompiledIn));
+		checkSlow(dynamic_cast<TProperty*>(PropertyChainForCompiledIn) && dynamic_cast<UProperty*>(PropertyChainForCompiledIn));
 		TProperty* Property = (TProperty*)PropertyChainForCompiledIn;
 		PropertyChainForCompiledIn = Property->Next;
 		StepExplicitProperty(Result, Property);
@@ -415,7 +406,7 @@ FORCEINLINE_DEBUGGABLE TNativeType& FFrame::StepCompiledInRef(void*const Tempora
 	}
 	else
 	{
-		checkSlow(Cast<TProperty>(PropertyChainForCompiledIn) && Cast<UProperty>(PropertyChainForCompiledIn));
+		checkSlow(dynamic_cast<TProperty*>(PropertyChainForCompiledIn) && dynamic_cast<UProperty*>(PropertyChainForCompiledIn));
 		TProperty* Property = (TProperty*)PropertyChainForCompiledIn;
 		PropertyChainForCompiledIn = Property->Next;
 		StepExplicitProperty(TemporaryBuffer, Property);
diff --git a/Engine/Source/Runtime/CoreUObject/Public/UObject/UObjectBase.h b/Engine/Source/Runtime/CoreUObject/Public/UObject/UObjectBase.h
--- a/Engine/Source/Runtime/CoreUObject/Public/UObject/UObjectBase.h
+++ b/Engine/Source/Runtime/CoreUObject/Public/UObject/UObjectBase.h
@@ -256,20 +256,20 @@ struct FFieldCompiledInInfo
 	FFieldCompiledInInfo(SIZE_T InClassSize, uint32 InCrc)
 		: Size(InClassSize)
 		, Crc(InCrc)
-		, OldField(NULL)
+		, OldClass(NULL)
 		, bHasChanged(false)
 	{
 	}
 
 	/** Registers the native class (constructs a UClass object) */
-	virtual UField* Register() const = 0;
+	virtual UClass* Register() const = 0;
 
 	/** Size of the class */
 	SIZE_T Size;
 	/** CRC of the generated code for this class */
 	uint32 Crc;
 	/** Old UClass object */
-	UField* OldField;
+	UClass* OldClass;
 	/** True if this class has changed after hot-reload (or new class) */
 	bool bHasChanged;
 };
@@ -290,7 +290,7 @@ struct TClassCompiledInDefer : public FFieldCompiledInInfo
 	{
 		UClassCompiledInDefer(this, InName, InClassSize, InCrc);
 	}
-	virtual UField* Register() const override
+	virtual UClass* Register() const override
 	{
 		return TClass::StaticClass();
 	}
diff --git a/Engine/Source/Runtime/CoreUObject/Public/UObject/UnrealType.h b/Engine/Source/Runtime/CoreUObject/Public/UObject/UnrealType.h
--- a/Engine/Source/Runtime/CoreUObject/Public/UObject/UnrealType.h
+++ b/Engine/Source/Runtime/CoreUObject/Public/UObject/UnrealType.h
@@ -738,8 +738,8 @@ public:
 	 */
 	UProperty* GetOwnerProperty()
 	{
-		UProperty* Result=this;
-		for ( UProperty* PropBase=Cast<UProperty>(GetOuter()); PropBase; PropBase=Cast<UProperty>(PropBase->GetOuter()) )
+		UProperty* Result = this;
+		for (UProperty* PropBase = dynamic_cast<UProperty*>(GetOuter()); PropBase; PropBase = dynamic_cast<UProperty*>(PropBase->GetOuter()))
 		{
 			Result = PropBase;
 		}
@@ -3023,7 +3023,7 @@ T* FindFieldChecked( const UStruct* Scope, FName FieldName )
 	if ( FieldName != NAME_None && Scope != NULL )
 	{
 		const UStruct* InitialScope = Scope;
-		for ( ; Scope != NULL; Scope = Cast<const UStruct>(Scope->GetOuter()) )
+		for ( ; Scope != NULL; Scope = dynamic_cast<const UStruct*>(Scope->GetOuter()) )
 		{
 			for ( TFieldIterator<T> It(Scope); It; ++It )
 			{
@@ -3046,13 +3046,13 @@ T* FindFieldChecked( const UStruct* Scope, FName FieldName )
 template<typename T>
 T* SmartCastProperty( UProperty* Src )
 {
-	T* Result = Cast<T>(Src);
+	T* Result = dynamic_cast<T*>(Src);
 	if ( Result == NULL )
 	{
-		UArrayProperty* ArrayProp = Cast<UArrayProperty>(Src);
+		UArrayProperty* ArrayProp = dynamic_cast<UArrayProperty*>(Src);
 		if ( ArrayProp != NULL )
 		{
-			Result = Cast<T>(ArrayProp->Inner);
+			Result = dynamic_cast<T*>(ArrayProp->Inner);
 		}
 	}
 	return Result;
diff --git a/Engine/Source/Runtime/Engine/Classes/Engine/BlendableInterface.h b/Engine/Source/Runtime/Engine/Classes/Engine/BlendableInterface.h
--- a/Engine/Source/Runtime/Engine/Classes/Engine/BlendableInterface.h
+++ b/Engine/Source/Runtime/Engine/Classes/Engine/BlendableInterface.h
@@ -25,7 +25,7 @@ enum EBlendableLocation
 	BL_MAX,
 };
 
-/** Dummy class needed to support InterfaceCast<IBlendableInterface>(Object) */
+/** Dummy class needed to support Cast<IBlendableInterface>(Object) */
 UINTERFACE()
 class UBlendableInterface : public UInterface
 {
diff --git a/Engine/Source/Runtime/Engine/Classes/Engine/World.h b/Engine/Source/Runtime/Engine/Classes/Engine/World.h
--- a/Engine/Source/Runtime/Engine/Classes/Engine/World.h
+++ b/Engine/Source/Runtime/Engine/Classes/Engine/World.h
@@ -2059,35 +2059,35 @@ public:
 	 * @param	bShouldModifyLevel		If true, Modify() the level before removing the actor if in the editor.
 	 */
 	void RemoveActor( AActor* Actor, bool bShouldModifyLevel );
-	
-	
+
 	AActor* SpawnActor( UClass* Class, FVector const* Location=NULL, FRotator const* Rotation=NULL, const FActorSpawnParameters& SpawnParameters = FActorSpawnParameters() );
 
 	// Templated version of SpawnActor that allows you to specify a class type
 	template< class T >
 	T* SpawnActor( const FActorSpawnParameters& SpawnParameters = FActorSpawnParameters() )
 	{
-		return CastChecked<T>(SpawnActor(T::StaticClass(), NULL, NULL, SpawnParameters),ECastCheckedType::NullAllowed);
-	};
+		return static_cast<T*>(SpawnActor(T::StaticClass(), NULL, NULL, SpawnParameters));
+	}
 
-	
 	template< class T >
 	T* SpawnActor( FVector const& Location, FRotator const& Rotation, const FActorSpawnParameters& SpawnParameters = FActorSpawnParameters() )
 	{
-		return CastChecked<T>(SpawnActor(T::StaticClass(), &Location, &Rotation, SpawnParameters),ECastCheckedType::NullAllowed);
-	};
+		return static_cast<T*>(SpawnActor(T::StaticClass(), &Location, &Rotation, SpawnParameters));
+	}
 	
 	template< class T >
 	T* SpawnActor( UClass* Class, const FActorSpawnParameters& SpawnParameters = FActorSpawnParameters() )
 	{
-		return CastChecked<T>(SpawnActor(Class, NULL, NULL, SpawnParameters),ECastCheckedType::NullAllowed);
-	};
+		check(Class->IsChildOf(T::StaticClass()));
+		return static_cast<T*>(SpawnActor(Class, NULL, NULL, SpawnParameters));
+	}
 
 	template< class T >
 	T* SpawnActor( UClass* Class, FVector const& Location, FRotator const& Rotation, const FActorSpawnParameters& SpawnParameters = FActorSpawnParameters() )
 	{
-		return CastChecked<T>(SpawnActor(Class, &Location, &Rotation, SpawnParameters),ECastCheckedType::NullAllowed);
-	};
+		check(Class->IsChildOf(T::StaticClass()));
+		return static_cast<T*>(SpawnActor(Class, &Location, &Rotation, SpawnParameters));
+	}
 
 
 	/**
@@ -2115,7 +2115,7 @@ public:
 		SpawnInfo.Instigator = Instigator;
 		SpawnInfo.bDeferConstruction = true;
 		return (Class != NULL) ? Cast<T>(SpawnActor(Class, &Location, &Rotation, SpawnInfo )) : NULL;
-	};
+	}
 
 	/** Returns the current GameMode instance, valid only on server. */
 	template< class T >
diff --git a/Engine/Source/Runtime/Engine/Classes/GameFramework/NavMovementComponent.h b/Engine/Source/Runtime/Engine/Classes/GameFramework/NavMovementComponent.h
--- a/Engine/Source/Runtime/Engine/Classes/GameFramework/NavMovementComponent.h
+++ b/Engine/Source/Runtime/Engine/Classes/GameFramework/NavMovementComponent.h
@@ -64,7 +64,7 @@ public:
 	/** @returns based location of controlled actor */
 	virtual FBasedPosition GetActorFeetLocationBased() const;
 	/** @returns navigation location of controlled actor */
-	FORCEINLINE FVector GetActorNavLocation() const { INavAgentInterface* MyOwner = InterfaceCast<INavAgentInterface>(GetOwner()); return MyOwner ? MyOwner->GetNavAgentLocation() : FVector::ZeroVector; }
+	FORCEINLINE FVector GetActorNavLocation() const { INavAgentInterface* MyOwner = Cast<INavAgentInterface>(GetOwner()); return MyOwner ? MyOwner->GetNavAgentLocation() : FVector::ZeroVector; }
 
 	/** request direct movement */
 	virtual void RequestDirectMove(const FVector& MoveVelocity, bool bForceMaxSpeed);
diff --git a/Engine/Source/Runtime/Engine/Private/AI/Navigation/NavCollision.cpp b/Engine/Source/Runtime/Engine/Private/AI/Navigation/NavCollision.cpp
--- a/Engine/Source/Runtime/Engine/Private/AI/Navigation/NavCollision.cpp
+++ b/Engine/Source/Runtime/Engine/Private/AI/Navigation/NavCollision.cpp
@@ -102,7 +102,7 @@ FDerivedDataNavCollisionCooker::FDerivedDataNavCollisionCooker(FName InFormat, U
 	check(NavCollisionInstance != NULL);
 	CollisionDataProvider = NavCollisionInstance->GetOuter();
 	DataGuid = NavCollisionInstance->GetGuid();
-	IInterface_CollisionDataProvider* CDP = InterfaceCast<IInterface_CollisionDataProvider>(CollisionDataProvider);
+	IInterface_CollisionDataProvider* CDP = Cast<IInterface_CollisionDataProvider>(CollisionDataProvider);
 	if (CDP)
 	{
 		CDP->GetMeshId(MeshId);
diff --git a/Engine/Source/Runtime/Engine/Private/AI/Navigation/NavLinkRenderingComponent.cpp b/Engine/Source/Runtime/Engine/Private/AI/Navigation/NavLinkRenderingComponent.cpp
--- a/Engine/Source/Runtime/Engine/Private/AI/Navigation/NavLinkRenderingComponent.cpp
+++ b/Engine/Source/Runtime/Engine/Private/AI/Navigation/NavLinkRenderingComponent.cpp
@@ -26,7 +26,7 @@ UNavLinkRenderingComponent::UNavLinkRenderingComponent(const class FPostConstruc
 FBoxSphereBounds UNavLinkRenderingComponent::CalcBounds(const FTransform & LocalToWorld) const
 {
 	AActor* LinkOwnerActor = Cast<AActor>(GetOwner());
-	INavLinkHostInterface* LinkOwnerHost = InterfaceCast<INavLinkHostInterface>(GetOwner());
+	INavLinkHostInterface* LinkOwnerHost = Cast<INavLinkHostInterface>(GetOwner());
 
 	if (LinkOwnerActor != NULL && LinkOwnerHost != NULL)
 	{
@@ -78,7 +78,7 @@ FNavLinkRenderingProxy::FNavLinkRenderingProxy(const UPrimitiveComponent* InComp
 	: FPrimitiveSceneProxy(InComponent)
 {
 	LinkOwnerActor = Cast<AActor>(InComponent->GetOwner());
-	LinkOwnerHost = InterfaceCast<INavLinkHostInterface>(InComponent->GetOwner());
+	LinkOwnerHost = Cast<INavLinkHostInterface>(InComponent->GetOwner());
 
 	if (LinkOwnerActor != NULL && LinkOwnerHost != NULL)
 	{
diff --git a/Engine/Source/Runtime/Engine/Private/AI/Navigation/NavigationGraph.cpp b/Engine/Source/Runtime/Engine/Private/AI/Navigation/NavigationGraph.cpp
--- a/Engine/Source/Runtime/Engine/Private/AI/Navigation/NavigationGraph.cpp
+++ b/Engine/Source/Runtime/Engine/Private/AI/Navigation/NavigationGraph.cpp
@@ -70,7 +70,7 @@ ANavigationData* ANavigationGraph::CreateNavigationInstances(UNavigationSystem*
 	FActorIterator It(NavSys->GetWorld());
 	for(; It; ++It )
 	{
-		if (InterfaceCast<INavNodeInterface>(*It) != NULL)
+		if (Cast<INavNodeInterface>(*It) != NULL)
 		{
 			bCreateNavigation = true;
 			break;
diff --git a/Engine/Source/Runtime/Engine/Private/AI/Navigation/NavigationPath.cpp b/Engine/Source/Runtime/Engine/Private/AI/Navigation/NavigationPath.cpp
--- a/Engine/Source/Runtime/Engine/Private/AI/Navigation/NavigationPath.cpp
+++ b/Engine/Source/Runtime/Engine/Private/AI/Navigation/NavigationPath.cpp
@@ -77,7 +77,7 @@ void FNavigationPath::SetGoalActorObservation(const AActor& ActorToObserve, floa
 	const bool RegisterForPathUpdates = GoalActor.IsValid() == false;	
 	GoalActor = &ActorToObserve;
 	checkSlow(GoalActor.IsValid());
-	GoalActorAsNavAgent = InterfaceCast<INavAgentInterface>(&ActorToObserve);
+	GoalActorAsNavAgent = Cast<INavAgentInterface>(&ActorToObserve);
 	GoalActorLocationTetherDistanceSq = FMath::Square(TetherDistance);
 	UpdateLastRepathGoalLocation();
 
@@ -87,7 +87,7 @@ void FNavigationPath::SetGoalActorObservation(const AActor& ActorToObserve, floa
 void FNavigationPath::SetSourceActor(const AActor& InSourceActor)
 {
 	SourceActor = &InSourceActor;
-	SourceActorAsNavAgent = InterfaceCast<INavAgentInterface>(&InSourceActor);
+	SourceActorAsNavAgent = Cast<INavAgentInterface>(&InSourceActor);
 }
 
 void FNavigationPath::UpdateLastRepathGoalLocation()
diff --git a/Engine/Source/Runtime/Engine/Private/AI/Navigation/NavigationSystem.cpp b/Engine/Source/Runtime/Engine/Private/AI/Navigation/NavigationSystem.cpp
--- a/Engine/Source/Runtime/Engine/Private/AI/Navigation/NavigationSystem.cpp
+++ b/Engine/Source/Runtime/Engine/Private/AI/Navigation/NavigationSystem.cpp
@@ -908,7 +908,7 @@ UNavigationPath* UNavigationSystem::FindPathToActorSynchronously(UObject* WorldC
 		return NULL; 
 	}
 
-	INavAgentInterface* NavAgent = InterfaceCast<INavAgentInterface>(GoalActor);
+	INavAgentInterface* NavAgent = Cast<INavAgentInterface>(GoalActor);
 	UNavigationPath* GeneratedPath = FindPathToLocationSynchronously(WorldContext, PathStart, NavAgent ? NavAgent->GetNavAgentLocation() : GoalActor->GetActorLocation(), PathfindingContext, FilterClass);
 	if (GeneratedPath != NULL && GeneratedPath->GetPath().IsValid() == true)
 	{
@@ -944,7 +944,7 @@ UNavigationPath* UNavigationSystem::FindPathToLocationSynchronously(UObject* Wor
 
 		if (PathfindingContext != NULL)
 		{
-			INavAgentInterface* NavAgent = InterfaceCast<INavAgentInterface>(PathfindingContext);
+			INavAgentInterface* NavAgent = Cast<INavAgentInterface>(PathfindingContext);
 			
 			if (NavAgent != NULL)
 			{
@@ -1004,7 +1004,7 @@ bool UNavigationSystem::NavigationRaycast(UObject* WorldContextObject, const FVe
 
 		// figure out which navigation data to use
 		const ANavigationData* NavData = NULL;
-		INavAgentInterface* MyNavAgent = InterfaceCast<INavAgentInterface>(Querier);
+		INavAgentInterface* MyNavAgent = Cast<INavAgentInterface>(Querier);
 		if (MyNavAgent)
 		{
 			const FNavAgentProperties* AgentProps = MyNavAgent->GetNavAgentProperties();
@@ -1193,7 +1193,7 @@ bool UNavigationSystem::IsThereAnywhereToBuildNavigation() const
 
 bool UNavigationSystem::IsNavigationRelevant(const AActor* TestActor) const
 {
-	const INavRelevantInterface* NavInterface = InterfaceCast<const INavRelevantInterface>(TestActor);
+	const INavRelevantInterface* NavInterface = Cast<const INavRelevantInterface>(TestActor);
 	if (NavInterface && NavInterface->IsNavigationRelevant())
 	{
 		return true;
@@ -1207,7 +1207,7 @@ bool UNavigationSystem::IsNavigationRelevant(const AActor* TestActor) const
 
 	for (int32 Idx = 0; Idx < Components.Num(); Idx++)
 	{
-		NavInterface = InterfaceCast<const INavRelevantInterface>(Components[Idx]);
+		NavInterface = Cast<const INavRelevantInterface>(Components[Idx]);
 		if (NavInterface && NavInterface->IsNavigationRelevant())
 		{
 			return true;
@@ -2032,7 +2032,7 @@ void UNavigationSystem::UpdateNavOctree(class AActor* Actor)
 {
 	SCOPE_CYCLE_COUNTER(STAT_DebugNavOctree);
 
-	INavRelevantInterface* NavElement = InterfaceCast<INavRelevantInterface>(Actor);
+	INavRelevantInterface* NavElement = Cast<INavRelevantInterface>(Actor);
 	if (NavElement)
 	{
 		UNavigationSystem* NavSys = Actor ? UNavigationSystem::GetCurrent(Actor->GetWorld()) : NULL;
@@ -2054,7 +2054,7 @@ void UNavigationSystem::UpdateNavOctree(class UActorComponent* Comp)
 		return;
 	}
 
-	INavRelevantInterface* NavElement = InterfaceCast<INavRelevantInterface>(Comp);
+	INavRelevantInterface* NavElement = Cast<INavRelevantInterface>(Comp);
 	if (NavElement)
 	{
 		AActor* OwnerActor = Comp ? Comp->GetOwner() : NULL;
@@ -2138,7 +2138,7 @@ void UNavigationSystem::UpdateNavOctreeElement(UObject* ElementOwner, class INav
 
 void UNavigationSystem::UpdateNavOctreeParentChain(UObject* ElementOwner)
 {
-	INavRelevantInterface* ElementInterface = InterfaceCast<INavRelevantInterface>(ElementOwner);
+	INavRelevantInterface* ElementInterface = Cast<INavRelevantInterface>(ElementOwner);
 	const int32 UpdateFlags = OctreeUpdate_ParentChain | OctreeUpdate_Refresh;
 
 	TArray<FWeakObjectPtr> ChildNodes;
@@ -2158,7 +2158,7 @@ void UNavigationSystem::UpdateNavOctreeParentChain(UObject* ElementOwner)
 		if (ChildNodes[Idx].IsValid())
 		{
 			UObject* ChildNodeOb = ChildNodes[Idx].Get();
-			ChildNavInterfaces[Idx] = InterfaceCast<INavRelevantInterface>(ChildNodeOb);
+			ChildNavInterfaces[Idx] = Cast<INavRelevantInterface>(ChildNodeOb);
 			UnregisterNavOctreeElement(ChildNodeOb, ChildNavInterfaces[Idx], UpdateFlags);
 		}
 	}
@@ -2178,7 +2178,7 @@ void UNavigationSystem::UpdateNavOctreeParentChain(UObject* ElementOwner)
 void UNavigationSystem::OnComponentRegistered(UActorComponent* Comp)
 {
 	SCOPE_CYCLE_COUNTER(STAT_DebugNavOctree);
-	INavRelevantInterface* NavInterface = InterfaceCast<INavRelevantInterface>(Comp);
+	INavRelevantInterface* NavInterface = Cast<INavRelevantInterface>(Comp);
 	if (NavInterface)
 	{
 		AActor* OwnerActor = Comp ? Comp->GetOwner() : NULL;
@@ -2196,7 +2196,7 @@ void UNavigationSystem::OnComponentRegistered(UActorComponent* Comp)
 void UNavigationSystem::OnComponentUnregistered(UActorComponent* Comp)
 {
 	SCOPE_CYCLE_COUNTER(STAT_DebugNavOctree);
-	INavRelevantInterface* NavInterface = InterfaceCast<INavRelevantInterface>(Comp);
+	INavRelevantInterface* NavInterface = Cast<INavRelevantInterface>(Comp);
 	if (NavInterface)
 	{
 		AActor* OwnerActor = Comp ? Comp->GetOwner() : NULL;
@@ -2216,7 +2216,7 @@ void UNavigationSystem::OnComponentUnregistered(UActorComponent* Comp)
 void UNavigationSystem::OnActorRegistered(class AActor* Actor)
 {
 	SCOPE_CYCLE_COUNTER(STAT_DebugNavOctree);
-	INavRelevantInterface* NavInterface = InterfaceCast<INavRelevantInterface>(Actor);
+	INavRelevantInterface* NavInterface = Cast<INavRelevantInterface>(Actor);
 	if (NavInterface)
 	{
 		UNavigationSystem* NavSys = UNavigationSystem::GetCurrent(Actor->GetWorld());
@@ -2230,7 +2230,7 @@ void UNavigationSystem::OnActorRegistered(class AActor* Actor)
 void UNavigationSystem::OnActorUnregistered(class AActor* Actor)
 {
 	SCOPE_CYCLE_COUNTER(STAT_DebugNavOctree);
-	INavRelevantInterface* NavInterface = InterfaceCast<INavRelevantInterface>(Actor);
+	INavRelevantInterface* NavInterface = Cast<INavRelevantInterface>(Actor);
 	if (NavInterface)
 	{
 		UNavigationSystem* NavSys = UNavigationSystem::GetCurrent(Actor->GetWorld());
diff --git a/Engine/Source/Runtime/Engine/Private/Character.cpp b/Engine/Source/Runtime/Engine/Private/Character.cpp
--- a/Engine/Source/Runtime/Engine/Private/Character.cpp
+++ b/Engine/Source/Runtime/Engine/Private/Character.cpp
@@ -947,7 +947,7 @@ void ACharacter::OnRep_ReplicatedBasedMovement()
 		// When position or base changes, movement mode will need to be updated. This assumes rotation changes don't affect that.
 		CharacterMovement->bJustTeleported |= (bBaseChanged || GetActorLocation() != OldLocation);
 
-		INetworkPredictionInterface* PredictionInterface = InterfaceCast<INetworkPredictionInterface>(GetMovementComponent());
+		INetworkPredictionInterface* PredictionInterface = Cast<INetworkPredictionInterface>(GetMovementComponent());
 		if (PredictionInterface)
 		{
 			PredictionInterface->SmoothCorrection(OldLocation);
@@ -1029,7 +1029,7 @@ void ACharacter::SimulatedRootMotionPositionFixup(float DeltaSeconds)
 							CharacterMovement->SimulateRootMotion(DeltaTime, LocalRootMotionTransform);
 
 							// After movement correction, smooth out error in position if any.
-							INetworkPredictionInterface* PredictionInterface = InterfaceCast<INetworkPredictionInterface>(GetMovementComponent());
+							INetworkPredictionInterface* PredictionInterface = Cast<INetworkPredictionInterface>(GetMovementComponent());
 							if (PredictionInterface)
 							{
 								PredictionInterface->SmoothCorrection(OldLocation);
@@ -1182,7 +1182,7 @@ void ACharacter::PostNetReceiveLocationAndRotation()
 			const FVector OldLocation = GetActorLocation();
 			UpdateSimulatedPosition(ReplicatedMovement.Location, ReplicatedMovement.Rotation);
 
-			INetworkPredictionInterface* PredictionInterface = InterfaceCast<INetworkPredictionInterface>(GetMovementComponent());
+			INetworkPredictionInterface* PredictionInterface = Cast<INetworkPredictionInterface>(GetMovementComponent());
 			if (PredictionInterface)
 			{
 				PredictionInterface->SmoothCorrection(OldLocation);
diff --git a/Engine/Source/Runtime/Engine/Private/Interpolation.cpp b/Engine/Source/Runtime/Engine/Private/Interpolation.cpp
--- a/Engine/Source/Runtime/Engine/Private/Interpolation.cpp
+++ b/Engine/Source/Runtime/Engine/Private/Interpolation.cpp
@@ -2199,7 +2199,7 @@ void UInterpGroup::UpdateAnimWeights(float NewPosition, class UInterpGroupInst*
 	}
 
 	// Find Anim Interface. If not return. When initialize, it will print error
-	IMatineeAnimInterface * IMAI = InterfaceCast<IMatineeAnimInterface>(Actor);
+	IMatineeAnimInterface * IMAI = Cast<IMatineeAnimInterface>(Actor);
 	if (!IMAI)
 	{
 		return;
@@ -2596,7 +2596,7 @@ void UInterpGroupInst::InitGroupInst(UInterpGroup* InGroup, AActor* InGroupActor
 	bool bHasAnimTrack = Group->HasAnimControlTrack();
 	if (bHasAnimTrack && GroupActor != NULL && !GroupActor->IsPendingKill())
 	{
-		IMatineeAnimInterface * IMAI = InterfaceCast<IMatineeAnimInterface>(GroupActor);
+		IMatineeAnimInterface * IMAI = Cast<IMatineeAnimInterface>(GroupActor);
 		if (IMAI)
 		{
 			// If in the editor and we haven't started playing, this should be Matinee! Bit yuck...
@@ -2634,7 +2634,7 @@ void UInterpGroupInst::TermGroupInst(bool bDeleteTrackInst)
 	bool bHasAnimTrack = Group->HasAnimControlTrack();
 	if (GroupActor != NULL && !GroupActor->IsPendingKill())
 	{
-		IMatineeAnimInterface * IMAI = InterfaceCast<IMatineeAnimInterface>(GroupActor);
+		IMatineeAnimInterface * IMAI = Cast<IMatineeAnimInterface>(GroupActor);
 		if (IMAI)
 		{
 			// If in the editor and we haven't started playing, this should be Matinee!
@@ -7329,7 +7329,7 @@ void UInterpTrackAnimControl::PreviewUpdateTrack(float NewPosition, class UInter
 	{
 		// if we're going backward or if not @ the first frame of the animation
 		bool bFireNotifier = !bSkipAnimNotifiers && (TimeElapsed < 0.f || !bResetTime) ;
-		IMatineeAnimInterface * IMAI = InterfaceCast<IMatineeAnimInterface>(Actor);
+		IMatineeAnimInterface * IMAI = Cast<IMatineeAnimInterface>(Actor);
 		if (IMAI)
 		{
 			IMAI->PreviewSetAnimPosition(SlotName, ChannelIndex, NewAnimSeq, NewAnimPosition, bNewLooping, bFireNotifier, TimeElapsed);
@@ -7349,7 +7349,7 @@ void UInterpTrackAnimControl::UpdateTrack(float NewPosition, UInterpTrackInst* T
 	//UE_LOG(LogMatinee, Log, TEXT("UInterpTrackAnimControl::UpdateTrack, NewPosition: %3.2f"), NewPosition);
 
 	UInterpTrackInstAnimControl* AnimInst = CastChecked<UInterpTrackInstAnimControl>(TrInst);
-	IMatineeAnimInterface * IMAI = InterfaceCast<IMatineeAnimInterface>(Actor);
+	IMatineeAnimInterface * IMAI = Cast<IMatineeAnimInterface>(Actor);
 	if (!IMAI)
 	{
 		// Actor does not support Matinee Anim Interface
diff --git a/Engine/Source/Runtime/Engine/Private/Level.cpp b/Engine/Source/Runtime/Engine/Private/Level.cpp
--- a/Engine/Source/Runtime/Engine/Private/Level.cpp
+++ b/Engine/Source/Runtime/Engine/Private/Level.cpp
@@ -1799,7 +1799,7 @@ void ULevel::TickRuntimeMovieScenePlayers( const float DeltaSeconds )
 		check( OwningWorld->IsGameWorld() );
 
 		IRuntimeMovieScenePlayerInterface* RuntimeMovieScenePlayer =
-			InterfaceCast< IRuntimeMovieScenePlayerInterface >( ActiveRuntimeMovieScenePlayers[ CurPlayerIndex ] );
+			Cast< IRuntimeMovieScenePlayerInterface >( ActiveRuntimeMovieScenePlayers[ CurPlayerIndex ] );
 		check( RuntimeMovieScenePlayer != NULL );
 
 		// @todo sequencer runtime: Support expiring instances of RuntimeMovieScenePlayers that have finished playing
diff --git a/Engine/Source/Runtime/Engine/Private/Pawn.cpp b/Engine/Source/Runtime/Engine/Private/Pawn.cpp
--- a/Engine/Source/Runtime/Engine/Private/Pawn.cpp
+++ b/Engine/Source/Runtime/Engine/Private/Pawn.cpp
@@ -987,7 +987,7 @@ void APawn::PostNetReceiveLocationAndRotation()
 		TeleportTo( ReplicatedMovement.Location, ReplicatedMovement.Rotation, false, true );
 		// SetActorLocationAndRotation(ReplicatedMovement.Location, ReplicatedMovement.Rotation); <-- preferred, but awaiting answer to question about UpdateNavOctree() missing in SceneComponent::MoveComponent
 
-		INetworkPredictionInterface* PredictionInterface = InterfaceCast<INetworkPredictionInterface>(GetMovementComponent());
+		INetworkPredictionInterface* PredictionInterface = Cast<INetworkPredictionInterface>(GetMovementComponent());
 		if (PredictionInterface)
 		{
 			PredictionInterface->SmoothCorrection(OldLocation);
diff --git a/Engine/Source/Runtime/Engine/Private/PhysicsEngine/BodySetup.cpp b/Engine/Source/Runtime/Engine/Private/PhysicsEngine/BodySetup.cpp
--- a/Engine/Source/Runtime/Engine/Private/PhysicsEngine/BodySetup.cpp
+++ b/Engine/Source/Runtime/Engine/Private/PhysicsEngine/BodySetup.cpp
@@ -804,7 +804,7 @@ FByteBulkData* UBodySetup::GetCookedData(FName Format)
 		return NULL;
 	}
 
-	IInterface_CollisionDataProvider* CDP = InterfaceCast<IInterface_CollisionDataProvider>(GetOuter());
+	IInterface_CollisionDataProvider* CDP = Cast<IInterface_CollisionDataProvider>(GetOuter());
 
 	// If there is nothing to cook or if we are reading data from a cooked package for an asset with no collision, 
 	// we want to return here
diff --git a/Engine/Source/Runtime/Engine/Private/PhysicsEngine/PhysDerivedData.cpp b/Engine/Source/Runtime/Engine/Private/PhysicsEngine/PhysDerivedData.cpp
--- a/Engine/Source/Runtime/Engine/Private/PhysicsEngine/PhysDerivedData.cpp
+++ b/Engine/Source/Runtime/Engine/Private/PhysicsEngine/PhysDerivedData.cpp
@@ -18,7 +18,7 @@ FDerivedDataPhysXCooker::FDerivedDataPhysXCooker( FName InFormat, UBodySetup* In
 	DataGuid = BodySetup->BodySetupGuid;
 	bGenerateNormalMesh = BodySetup->bGenerateNonMirroredCollision;
 	bGenerateMirroredMesh = BodySetup->bGenerateMirroredCollision;
-	IInterface_CollisionDataProvider* CDP = InterfaceCast<IInterface_CollisionDataProvider>(CollisionDataProvider);
+	IInterface_CollisionDataProvider* CDP = Cast<IInterface_CollisionDataProvider>(CollisionDataProvider);
 	if (CDP)
 	{
 		CDP->GetMeshId(MeshId);
@@ -155,14 +155,14 @@ bool FDerivedDataPhysXCooker::ShouldGenerateTriMeshData(bool InUseAllTriData)
 {
 	check(Cooker != NULL);
 
-	IInterface_CollisionDataProvider* CDP = InterfaceCast<IInterface_CollisionDataProvider>(CollisionDataProvider);
+	IInterface_CollisionDataProvider* CDP = Cast<IInterface_CollisionDataProvider>(CollisionDataProvider);
 	const bool bPerformCook = ( CDP != NULL ) ? CDP->ContainsPhysicsTriMeshData(InUseAllTriData) : false;
 	return bPerformCook;
 }
 
 bool FDerivedDataPhysXCooker::ShouldGenerateNegXTriMeshData()
 {
-	IInterface_CollisionDataProvider* CDP = InterfaceCast<IInterface_CollisionDataProvider>(CollisionDataProvider);
+	IInterface_CollisionDataProvider* CDP = Cast<IInterface_CollisionDataProvider>(CollisionDataProvider);
 	const bool bWantsNegX = ( CDP != NULL ) ? CDP->WantsNegXTriMesh() : false;
 	return bWantsNegX;
 }
@@ -174,7 +174,7 @@ bool FDerivedDataPhysXCooker::BuildTriMesh( TArray<uint8>& OutData, bool bInMirr
 
 	bool bResult = false;
 	FTriMeshCollisionData TriangleMeshDesc;
-	IInterface_CollisionDataProvider* CDP = InterfaceCast<IInterface_CollisionDataProvider>(CollisionDataProvider);
+	IInterface_CollisionDataProvider* CDP = Cast<IInterface_CollisionDataProvider>(CollisionDataProvider);
 	check(CDP != NULL); // It's all been checked before getting into this function
 		
 	bool bHaveTriMeshData = CDP->GetPhysicsTriMeshData(&TriangleMeshDesc, InUseAllTriData);
diff --git a/Engine/Source/Runtime/Engine/Private/PlayerController.cpp b/Engine/Source/Runtime/Engine/Private/PlayerController.cpp
--- a/Engine/Source/Runtime/Engine/Private/PlayerController.cpp
+++ b/Engine/Source/Runtime/Engine/Private/PlayerController.cpp
@@ -759,7 +759,7 @@ void APlayerController::Possess(APawn* PawnToPossess)
 
 		GetPawn()->SetActorTickEnabled(true);
 
-		INetworkPredictionInterface* NetworkPredictionInterface = GetPawn() ? InterfaceCast<INetworkPredictionInterface>(GetPawn()->GetMovementComponent()) : NULL;
+		INetworkPredictionInterface* NetworkPredictionInterface = GetPawn() ? Cast<INetworkPredictionInterface>(GetPawn()->GetMovementComponent()) : NULL;
 		if (NetworkPredictionInterface)
 		{
 			NetworkPredictionInterface->ResetPredictionData_Server();
@@ -1478,7 +1478,7 @@ void APlayerController::SendClientAdjustment()
 	APawn* RemotePawn = GetPawnOrSpectator();
 	if (RemotePawn && (GetNetMode() < NM_Client) && (RemotePawn->GetRemoteRole() == ROLE_AutonomousProxy))
 	{
-		INetworkPredictionInterface* NetworkPredictionInterface = InterfaceCast<INetworkPredictionInterface>(RemotePawn->GetMovementComponent());
+		INetworkPredictionInterface* NetworkPredictionInterface = Cast<INetworkPredictionInterface>(RemotePawn->GetMovementComponent());
 		if (NetworkPredictionInterface)
 		{
 			NetworkPredictionInterface->SendClientAdjustment();
@@ -3701,7 +3701,7 @@ void APlayerController::TickActor( float DeltaSeconds, ELevelTick TickType, FAct
 		// skip updates if pawn lost autonomous proxy role (e.g. TurnOff() call)
 		if (GetPawn() && !GetPawn()->IsPendingKill() && GetPawn()->GetRemoteRole() == ROLE_AutonomousProxy)
 		{
-			INetworkPredictionInterface* NetworkPredictionInterface = InterfaceCast<INetworkPredictionInterface>(GetPawn()->GetMovementComponent());
+			INetworkPredictionInterface* NetworkPredictionInterface = Cast<INetworkPredictionInterface>(GetPawn()->GetMovementComponent());
 			if (NetworkPredictionInterface)
 			{
 				FNetworkPredictionData_Server* ServerData = NetworkPredictionInterface->GetPredictionData_Server();
diff --git a/Engine/Source/Runtime/Engine/Private/PostProcessVolume.cpp b/Engine/Source/Runtime/Engine/Private/PostProcessVolume.cpp
--- a/Engine/Source/Runtime/Engine/Private/PostProcessVolume.cpp
+++ b/Engine/Source/Runtime/Engine/Private/PostProcessVolume.cpp
@@ -38,7 +38,7 @@ void APostProcessVolume::PostEditChangeProperty(FPropertyChangedEvent& PropertyC
 			{
 				UObject* Obj = Settings.Blendables[i];
 
-				if(!InterfaceCast<IBlendableInterface>(Obj))
+				if(!Cast<IBlendableInterface>(Obj))
 				{
 					Settings.Blendables[i] = 0;
 				}
diff --git a/Engine/Source/Runtime/Engine/Private/Replication.cpp b/Engine/Source/Runtime/Engine/Private/Replication.cpp
--- a/Engine/Source/Runtime/Engine/Private/Replication.cpp
+++ b/Engine/Source/Runtime/Engine/Private/Replication.cpp
@@ -130,7 +130,7 @@ void AMatineeActor::PostNetReceive()
 				// notify actors that something has changed
 				for (int32 ActorIndex = 0; ActorIndex < ControlledActors.Num(); ++ActorIndex )
 				{
-					IMatineeInterface * IMI = InterfaceCast<IMatineeInterface>(ControlledActors[ActorIndex]);
+					IMatineeInterface * IMI = Cast<IMatineeInterface>(ControlledActors[ActorIndex]);
 					if (IMI)
 					{
 						IMI->InterpolationChanged(this);
@@ -164,7 +164,7 @@ void AMatineeActor::PostNetReceive()
 					// fire an event if we're really playing (and not just starting it up to do a position update)
 					if (bIsPlaying)
 					{
-						IMatineeInterface * IMI = InterfaceCast<IMatineeInterface>(Actor);
+						IMatineeInterface * IMI = Cast<IMatineeInterface>(Actor);
 						if (IMI)
 						{
 							IMI->InterpolationStarted(this);
@@ -221,7 +221,7 @@ void AMatineeActor::PostNetReceive()
 					// fire an event if we were really playing (and not just starting it up to do a position update)
 					if (SavedbIsPlaying)
 					{
-						IMatineeInterface * IMI = InterfaceCast<IMatineeInterface>(Actor);
+						IMatineeInterface * IMI = Cast<IMatineeInterface>(Actor);
 						if (IMI)
 						{
 							IMI->InterpolationFinished(this);
diff --git a/Engine/Source/Runtime/Engine/Private/SceneView.cpp b/Engine/Source/Runtime/Engine/Private/SceneView.cpp
--- a/Engine/Source/Runtime/Engine/Private/SceneView.cpp
+++ b/Engine/Source/Runtime/Engine/Private/SceneView.cpp
@@ -805,7 +805,7 @@ void FSceneView::OverridePostProcessSettings(const FPostProcessSettings& Src, fl
 				continue;
 			}
 
-			IBlendableInterface* BlendableInterface = InterfaceCast<IBlendableInterface>(Object);
+			IBlendableInterface* BlendableInterface = Cast<IBlendableInterface>(Object);
 			
 			if(!BlendableInterface)
 			{
@@ -817,7 +817,7 @@ void FSceneView::OverridePostProcessSettings(const FPostProcessSettings& Src, fl
 	}
 }
 
-/** Dummy class needed to support InterfaceCast<IBlendableInterface>(Object) */
+/** Dummy class needed to support Cast<IBlendableInterface>(Object) */
 UBlendableInterface::UBlendableInterface(const class FPostConstructInitializeProperties& PCIP) : Super(PCIP) 
 {
 }
diff --git a/Engine/Source/Runtime/GameplayAbilities/Private/AbilitySystemComponent.cpp b/Engine/Source/Runtime/GameplayAbilities/Private/AbilitySystemComponent.cpp
--- a/Engine/Source/Runtime/GameplayAbilities/Private/AbilitySystemComponent.cpp
+++ b/Engine/Source/Runtime/GameplayAbilities/Private/AbilitySystemComponent.cpp
@@ -135,7 +135,7 @@ bool UAbilitySystemComponent::AreGameplayEffectApplicationRequirementsSatisfied(
 		Instigator.GetOwnedGameplayTags(InstigatorTags);
 
 		FGameplayTagContainer TargetTags;
-		IGameplayTagAssetInterface* OwnerGTA = InterfaceCast<IGameplayTagAssetInterface>(AbilityActorInfo->Actor.Get());
+		IGameplayTagAssetInterface* OwnerGTA = Cast<IGameplayTagAssetInterface>(AbilityActorInfo->Actor.Get());
 		if (OwnerGTA)
 		{
 			OwnerGTA->GetOwnedGameplayTags(TargetTags);
@@ -494,7 +494,7 @@ void UAbilitySystemComponent::InvokeGameplayCueEvent(const FGameplayEffectSpec &
 		ABILITY_LOG(Warning, TEXT("InvokeGameplayCueEvent: %s"), *Spec.ToSimpleString());
 	}
 
-	IGameplayCueInterface* GameplayCueInterface = InterfaceCast<IGameplayCueInterface>(ActorOwner);
+	IGameplayCueInterface* GameplayCueInterface = Cast<IGameplayCueInterface>(ActorOwner);
 	if (!GameplayCueInterface)
 	{
 		ABILITY_LOG(Warning, TEXT("InvokeGameplayCueEvent %s on Actor %s that is not IGameplayCueInterface"), *Spec.ToSimpleString(), ActorOwner ? *ActorOwner->GetName() : TEXT("NULL"));
@@ -562,7 +562,7 @@ void UAbilitySystemComponent::RemoveGameplayCue(const FGameplayTag GameplayCueTa
 void UAbilitySystemComponent::InvokeGameplayCueEvent(const FGameplayTag GameplayCueTag, EGameplayCueEvent::Type EventType)
 {
 	AActor* ActorOwner = AbilityActorInfo->Actor.Get();
-	IGameplayCueInterface* GameplayCueInterface = InterfaceCast<IGameplayCueInterface>(ActorOwner);
+	IGameplayCueInterface* GameplayCueInterface = Cast<IGameplayCueInterface>(ActorOwner);
 	if (!GameplayCueInterface)
 	{
 		return;
diff --git a/Engine/Source/Runtime/GameplayAbilities/Private/AbilitySystemGlobals.cpp b/Engine/Source/Runtime/GameplayAbilities/Private/AbilitySystemGlobals.cpp
--- a/Engine/Source/Runtime/GameplayAbilities/Private/AbilitySystemGlobals.cpp
+++ b/Engine/Source/Runtime/GameplayAbilities/Private/AbilitySystemGlobals.cpp
@@ -96,7 +96,7 @@ UAbilitySystemComponent* UAbilitySystemGlobals::GetAbilitySystemComponentFromAct
 		return nullptr;
 	}
 
-	IAbilitySystemInterface* ASI = InterfaceCast<IAbilitySystemInterface>(Actor);
+	IAbilitySystemInterface* ASI = Cast<IAbilitySystemInterface>(Actor);
 	if (ASI)
 	{
 		return ASI->GetAbilitySystemComponent();
diff --git a/Engine/Source/Runtime/GameplayAbilities/Private/GameplayEffect.cpp b/Engine/Source/Runtime/GameplayAbilities/Private/GameplayEffect.cpp
--- a/Engine/Source/Runtime/GameplayAbilities/Private/GameplayEffect.cpp
+++ b/Engine/Source/Runtime/GameplayAbilities/Private/GameplayEffect.cpp
@@ -468,7 +468,7 @@ void FGameplayEffectInstigatorContext::AddInstigator(class AActor *InInstigator)
 	InstigatorAbilitySystemComponent = NULL;
 
 	// Cache off his AbilitySystemComponent.
-	IAbilitySystemInterface* AbilitySystemInterface = InterfaceCast<IAbilitySystemInterface>(Instigator);
+	IAbilitySystemInterface* AbilitySystemInterface = Cast<IAbilitySystemInterface>(Instigator);
 	if (AbilitySystemInterface)
 	{
 		InstigatorAbilitySystemComponent = AbilitySystemInterface->GetAbilitySystemComponent();
diff --git a/Engine/Source/Runtime/GameplayAbilities/Public/GameplayEffectTypes.h b/Engine/Source/Runtime/GameplayAbilities/Public/GameplayEffectTypes.h
--- a/Engine/Source/Runtime/GameplayAbilities/Public/GameplayEffectTypes.h
+++ b/Engine/Source/Runtime/GameplayAbilities/Public/GameplayEffectTypes.h
@@ -220,7 +220,7 @@ struct FGameplayEffectInstigatorContext
 
 	void GetOwnedGameplayTags(OUT FGameplayTagContainer &TagContainer)
 	{
-		IGameplayTagAssetInterface* TagInterface = InterfaceCast<IGameplayTagAssetInterface>(Instigator);
+		IGameplayTagAssetInterface* TagInterface = Cast<IGameplayTagAssetInterface>(Instigator);
 		if (TagInterface)
 		{
 			TagInterface->GetOwnedGameplayTags(TagContainer);
diff --git a/Engine/Source/Runtime/MovieSceneCoreTypes/Private/Instances/MovieSceneAnimationTrackInstance.cpp b/Engine/Source/Runtime/MovieSceneCoreTypes/Private/Instances/MovieSceneAnimationTrackInstance.cpp
--- a/Engine/Source/Runtime/MovieSceneCoreTypes/Private/Instances/MovieSceneAnimationTrackInstance.cpp
+++ b/Engine/Source/Runtime/MovieSceneCoreTypes/Private/Instances/MovieSceneAnimationTrackInstance.cpp
@@ -22,7 +22,7 @@ void FMovieSceneAnimationTrackInstance::Update( float Position, float LastPositi
 
 	for (int32 i = 0; i < RuntimeObjects.Num(); ++i)
 	{
-		IMatineeAnimInterface* AnimInterface = InterfaceCast<IMatineeAnimInterface>(RuntimeObjects[i]);
+		IMatineeAnimInterface* AnimInterface = Cast<IMatineeAnimInterface>(RuntimeObjects[i]);
 		if (AnimInterface)
 		{
 			UMovieSceneAnimationSection* AnimSection = Cast<UMovieSceneAnimationSection>(AnimationTrack->GetAnimSectionAtTime(Position));
