#docs - Cleaned up comments in Object.h (skeleton definitions for core object module)

- Added links to actual source files. Lots of times intellisense/VA will take you to the less helpful skeleton file rather than the actual source file.

diff --git a/Engine/Source/Runtime/CoreUObject/Classes/Object.h b/Engine/Source/Runtime/CoreUObject/Classes/Object.h
--- a/Engine/Source/Runtime/CoreUObject/Classes/Object.h
+++ b/Engine/Source/Runtime/CoreUObject/Classes/Object.h
@@ -5,12 +5,14 @@
 //=============================================================================
 // Unreal base structures.
 
-// Temporary mirrors of C++ structs
+// Temporary mirrors of C++ structs, used mainly as forward declarations for the core object module
+// to avoid including the full engine source. In most cases the full class definition is in another
+// file and is noted as such. More complete documentation will generally be found in those files.
 
 #pragma once
 #if !CPP      //noexport class
 
-// String search case used in UnrealString.h
+/** String search case used in UnrealString.h */
 UENUM()
 namespace ESearchCase
 {
@@ -21,7 +23,7 @@ namespace ESearchCase
 	};
 }
 
-// String search dir used in UnrealString.h
+/** String search dir used in UnrealString.h */
 UENUM()
 namespace ESearchDir
 {
@@ -32,7 +34,7 @@ namespace ESearchDir
 	};
 }
 
-// Generic axis enum (mirrored for native use in Axis.h).
+/** Generic axis enum (mirrored for native use in Axis.h). */
 UENUM()
 namespace EAxis
 {
@@ -208,7 +210,10 @@ struct FNetworkGUID
 	uint32 Value;
 };
 
-// A point or direction FVector in 3d space.
+/**
+ * A point or direction FVector in 3d space.
+ * The full C++ class is located here: Engine\Source\Runtime\Core\Public\Math\Vector.h
+ */
 USTRUCT(immutable, noexport, BlueprintType, meta=(HasNativeMake="Engine.KismetMathLibrary.MakeVector", HasNativeBreak="Engine.KismetMathLibrary.BreakVector"))
 struct FVector
 {
@@ -222,6 +227,10 @@ struct FVector
 	float Z;
 };
 
+/**
+* A 4-D homogeneous vector.
+* The full C++ class is located here: Engine\Source\Runtime\Core\Public\Math\Vector4.h
+*/
 USTRUCT(immutable, noexport)
 struct FVector4
 {
@@ -239,7 +248,10 @@ struct FVector4
 
 };
 
-
+/**
+* A point or direction FVector in 2d space.
+* The full C++ class is located here: Engine\Source\Runtime\Core\Public\Math\Vector2D.h
+*/
 USTRUCT(immutable, noexport, BlueprintType, meta=(HasNativeMake="Engine.KismetMathLibrary.MakeVector2D", HasNativeBreak="Engine.KismetMathLibrary.BreakVector2D"))
 struct FVector2D
 {
@@ -265,8 +277,10 @@ struct FTwoVectors
 };
 
 
-// A plane definition in 3D space.
-
+/**
+ * A plane definition in 3D space.
+ * The full C++ class is located here: Engine\Source\Runtime\Core\Public\Math\Plane.h
+ */
 USTRUCT(immutable, noexport)
 struct FPlane : public FVector
 {
@@ -276,8 +290,10 @@ struct FPlane : public FVector
 };
 
 
-// An orthogonal rotation in 3d space.
-
+/**
+ * An orthogonal rotation in 3d space.
+ * The full C++ class is located here: Engine\Source\Runtime\Core\Public\Math\Rotator.h
+ */
 USTRUCT(immutable, noexport, BlueprintType, meta=(HasNativeMake="Engine.KismetMathLibrary.MakeRot", HasNativeBreak="Engine.KismetMathLibrary.BreakRot"))
 struct FRotator
 {
@@ -293,8 +309,10 @@ struct FRotator
 };
 
 
-// Quaternion.
-
+/**
+ * Quaternion.
+ * The full C++ class is located here: Engine\Source\Runtime\Core\Public\Math\Quat.h
+ */
 USTRUCT(immutable, noexport)
 struct FQuat
 {
@@ -313,8 +331,10 @@ struct FQuat
 };
 
 
-// A packed normal.
-
+/**
+ * A packed normal.
+ * The full C++ class is located here: Engine\Source\Runtime\RenderCore\Public\PackedNormal.h
+ */
 USTRUCT(immutable, noexport)
 struct FPackedNormal
 {
@@ -335,8 +355,8 @@ struct FPackedNormal
 
 /**
  * Screen coordinates.
+ * The full C++ class is located here: Engine\Source\Runtime\Core\Public\Math\IntPoint.h
  */
-
 USTRUCT(immutable, noexport)
 struct FIntPoint
 {
@@ -349,9 +369,9 @@ struct FIntPoint
 };
 
 /**
- *  An integer vector in 3D space
+ * An integer vector in 3D space.
+ * The full C++ class is located here: Engine\Source\Runtime\Core\Public\Math\IntVector.h
  */
-
 USTRUCT(immutable, noexport)
 struct FIntVector
 {
@@ -366,8 +386,10 @@ struct FIntVector
 };
 
 
-// A Color.
-
+/**
+ * A Color (BGRA).
+ * The full C++ class is located here: Engine\Source\Runtime\Core\Public\Math\Color.h
+ */
 USTRUCT(immutable, noexport)
 struct FColor
 {
@@ -386,8 +408,10 @@ struct FColor
 };
 
 
-// A linear color.
-
+/**
+ * A linear color.
+ * The full C++ class is located here: Engine\Source\Runtime\Core\Public\Math\Color.h
+ */
 USTRUCT(immutable, noexport, BlueprintType)
 struct FLinearColor
 {
@@ -406,8 +430,10 @@ struct FLinearColor
 };
 
 
-// A bounding box.
-
+/**
+ * A bounding box.
+ * The full C++ class is located here: Engine\Source\Runtime\Core\Public\Math\Box.h
+ */
 USTRUCT(immutable, noexport)
 struct FBox
 {
@@ -423,8 +449,10 @@ struct FBox
 };
 
 
-// A bounding box and bounding sphere with the same origin.
-
+/**
+ * A bounding box and bounding sphere with the same origin.
+ * The full C++ class is located here : Engine\Source\Runtime\Core\Public\Math\BoxSphereBounds.h
+ */
 USTRUCT(noexport)
 struct FBoxSphereBounds
 {
@@ -441,6 +469,7 @@ struct FBoxSphereBounds
 
 /**
  * Structure for arbitrarily oriented boxes (i.e. not necessarily axis-aligned).
+ * The full C++ class is located here: Engine\Source\Runtime\Core\Public\Math\OrientedBox.h
  */
 USTRUCT(immutable, noexport)
 struct FOrientedBox
@@ -467,8 +496,10 @@ struct FOrientedBox
 	float ExtentZ;
 };
 
-// a 4x4 matrix.
-
+/*
+ * A 4x4 matrix.
+ * The full C++ class is located here: Engine\Source\Runtime\Core\Public\Math\Matrix.h
+ */
 USTRUCT(immutable, noexport)
 struct FMatrix
 {
@@ -647,8 +678,10 @@ struct FInterpCurveLinearColor
 };
 
 
-/** Transform definition. */
-
+/**
+ * Transform composed of Quat/Translation/Scale.
+ * The full C++ class is located here: Engine\Source\Runtime\Core\Public\Math\Transform.h
+ */
 USTRUCT(noexport, BlueprintType, meta=(HasNativeMake="Engine.KismetMathLibrary.MakeTransform", HasNativeBreak="Engine.KismetMathLibrary.BreakTransform"))
 struct FTransform
 {
@@ -664,8 +697,10 @@ struct FTransform
 };
 
 
-/** Thread-safe RNG. */
-
+/**
+ * Thread-safe RNG.
+ * The full C++ class is located here: Engine\Source\Runtime\Core\Public\Math\RandomStream.h
+ */
 USTRUCT(noexport, BlueprintType)
 struct FRandomStream
 {
@@ -792,8 +827,11 @@ struct FInt32Range
 };
 
 //=============================================================================
-// Object: The base class all objects.
-// This is a built-in Unreal class and it shouldn't be modified by mod authors 
+/**
+ * Object: The base class all objects.
+ * This is a built-in Unreal class and it shouldn't be modified by mod authors.
+ * The full C++ class is located here: Engine\Source\Runtime\CoreUObject\Public\UObject\UObject.h
+ */
 //=============================================================================
 
 UCLASS(abstract, noexport)
