Minor fixes to new casts.

* InterfaceCast deprecation version number changed to 4.6.
* Double-underscored identifiers replaced with names which aren't reserved.

diff --git a/Engine/Source/Programs/UnrealHeaderTool/Private/CodeGenerator.cpp b/Engine/Source/Programs/UnrealHeaderTool/Private/CodeGenerator.cpp
--- a/Engine/Source/Programs/UnrealHeaderTool/Private/CodeGenerator.cpp
+++ b/Engine/Source/Programs/UnrealHeaderTool/Private/CodeGenerator.cpp
@@ -1755,7 +1755,7 @@ void FNativeClassHeaderGenerator::ExportInterfaceClassDeclaration( FClass* Class
 		// to script VM functions
 		if (SuperClass->IsChildOf(UInterface::StaticClass()))
 		{
-			InterfaceBoilerplate.Logf(TEXT("\tvirtual UObject* __GetUObject() const = 0;\r\n"));
+			InterfaceBoilerplate.Logf(TEXT("\tvirtual UObject* _getUObject() const = 0;\r\n"));
 		}
 
 		// Replication, add in the declaration for GetLifetimeReplicatedProps() automatically if there are any net flagged properties
@@ -1918,7 +1918,7 @@ void FNativeClassHeaderGenerator::ExportClassHeaderInner(FClass* Class, bool bVa
 					ClassBoilerplate.Logf(TEXT("    static const TCHAR* StaticConfigName() {return TEXT(\"%s\");}\r\n\r\n"), *Class->ClassConfigName.ToString());
 				}
 
-				ClassBoilerplate.Logf(TEXT("%sUObject* __GetUObject() const { return const_cast<%s*>(this); }\r\n"), FCString::Spc(4), ClassCPPName);
+				ClassBoilerplate.Logf(TEXT("%sUObject* _getUObject() const { return const_cast<%s*>(this); }\r\n"), FCString::Spc(4), ClassCPPName);
 
 				// Replication, add in the declaration for GetLifetimeReplicatedProps() automatically if there are any net flagged properties
 				for( TFieldIterator<UProperty> It(Class,EFieldIteratorFlags::ExcludeSuper); It; ++It )
diff --git a/Engine/Source/Runtime/CoreUObject/Public/Templates/Casts.h b/Engine/Source/Runtime/CoreUObject/Public/Templates/Casts.h
--- a/Engine/Source/Runtime/CoreUObject/Public/Templates/Casts.h
+++ b/Engine/Source/Runtime/CoreUObject/Public/Templates/Casts.h
@@ -106,7 +106,7 @@ struct TCastImpl<From, To, ECastType::InterfaceToUObject>
 {
 	FORCEINLINE static To* DoCast( From* Src )
 	{
-		return Cast<To>(Src->__GetUObject());
+		return Cast<To>(Src->_getUObject());
 	}
 };
 
@@ -124,7 +124,7 @@ struct TCastImpl<From, To, ECastType::InterfaceToInterface>
 {
 	FORCEINLINE static To* DoCast( From* Src )
 	{
-		return Cast<To>(Src->__GetUObject());
+		return Cast<To>(Src->_getUObject());
 	}
 };
 
@@ -166,14 +166,14 @@ To* CastChecked(From* Src, ECastCheckedType::Type CheckType = ECastCheckedType::
 }
 
 template <typename InterfaceType>
-DEPRECATED(4.5, "InterfaceCast is deprecated, use Cast instead.")
+DEPRECATED(4.6, "InterfaceCast is deprecated, use Cast or dynamic_cast instead.")
 FORCEINLINE InterfaceType* InterfaceCast(UObject* Src)
 {
 	return Cast<InterfaceType>(Src);
 }
 
 template <typename InterfaceType>
-DEPRECATED(4.5, "InterfaceCast is deprecated, use Cast instead.")
+DEPRECATED(4.6, "InterfaceCast is deprecated, use Cast or dynamic_cast instead.")
 FORCEINLINE InterfaceType* InterfaceCast(const UObject* Src)
 {
 	return Cast<InterfaceType>(const_cast<UObject*>(Src));
@@ -243,46 +243,49 @@ DECLARE_CAST_BY_FLAG(UMulticastDelegateProperty)
 #undef DECLARE_CAST_BY_FLAG_CAST
 #undef DECLARE_CAST_BY_FLAG_FWD
 
-template <typename To, typename From>
-FORCEINLINE typename TEnableIf<TAnd<TIsPointerType<To>, TAnd<TIsCastableToPointer<typename TRemovePointer<To>::Type>, TIsCastable<From>>>::Value, To>::Type __DynamicCast(From* Arg)
+namespace UE4Casts_Private
 {
-	typedef typename TRemovePointer<To  >::Type ToValueType;
-	typedef typename TRemovePointer<From>::Type FromValueType;
+	template <typename To, typename From>
+	FORCEINLINE typename TEnableIf<TAnd<TIsPointerType<To>, TAnd<TIsCastableToPointer<typename TRemovePointer<To>::Type>, TIsCastable<From>>>::Value, To>::Type DynamicCast(From* Arg)
+	{
+		typedef typename TRemovePointer<To  >::Type ToValueType;
+		typedef typename TRemovePointer<From>::Type FromValueType;
 
-	// Casting away const/volatile
-	static_assert(!TLosesQualifiersFromTo<FromValueType, ToValueType>::Value, "Conversion loses qualifiers");
+		// Casting away const/volatile
+		static_assert(!TLosesQualifiersFromTo<FromValueType, ToValueType>::Value, "Conversion loses qualifiers");
 
-	// If we're casting to void, cast to UObject instead and let it implicitly cast to void
-	return Cast<typename TChooseClass<TIsVoidType<ToValueType>::Value, UObject, ToValueType>::Result>(Arg);
-}
+		// If we're casting to void, cast to UObject instead and let it implicitly cast to void
+		return Cast<typename TChooseClass<TIsVoidType<ToValueType>::Value, UObject, ToValueType>::Result>(Arg);
+	}
 
-template <typename To, typename From>
-FORCEINLINE typename TEnableIf<!TAnd<TIsPointerType<To>, TAnd<TIsCastableToPointer<typename TRemovePointer<To>::Type>, TIsCastable<From>>>::Value, To>::Type __DynamicCast(From* Arg)
-{
-	return dynamic_cast<To>(Arg);
-}
+	template <typename To, typename From>
+	FORCEINLINE typename TEnableIf<!TAnd<TIsPointerType<To>, TAnd<TIsCastableToPointer<typename TRemovePointer<To>::Type>, TIsCastable<From>>>::Value, To>::Type DynamicCast(From* Arg)
+	{
+		return dynamic_cast<To>(Arg);
+	}
 
-template <typename To, typename From>
-FORCEINLINE typename TEnableIf<TAnd<TIsCastable<typename TRemoveReference<To>::Type>, TIsCastable<typename TRemoveReference<From>::Type>>::Value, To>::Type __DynamicCast(From&& Arg)
-{
-	typedef typename TRemoveReference<From>::Type FromValueType;
-	typedef typename TRemoveReference<To  >::Type ToValueType;
+	template <typename To, typename From>
+	FORCEINLINE typename TEnableIf<TAnd<TIsCastable<typename TRemoveReference<To>::Type>, TIsCastable<typename TRemoveReference<From>::Type>>::Value, To>::Type DynamicCast(From&& Arg)
+	{
+		typedef typename TRemoveReference<From>::Type FromValueType;
+		typedef typename TRemoveReference<To  >::Type ToValueType;
 
-	// Casting away const/volatile
-	static_assert(!TLosesQualifiersFromTo<FromValueType, ToValueType>::Value, "Conversion loses qualifiers");
+		// Casting away const/volatile
+		static_assert(!TLosesQualifiersFromTo<FromValueType, ToValueType>::Value, "Conversion loses qualifiers");
 
-	// T&& can only be cast to U&&
-	// http://en.cppreference.com/w/cpp/language/dynamic_cast
-	static_assert(TOr<TIsLValueReferenceType<From>, TIsRValueReferenceType<To>>::Value, "Cannot dynamic_cast from an rvalue to a non-rvalue reference");
+		// T&& can only be cast to U&&
+		// http://en.cppreference.com/w/cpp/language/dynamic_cast
+		static_assert(TOr<TIsLValueReferenceType<From>, TIsRValueReferenceType<To>>::Value, "Cannot dynamic_cast from an rvalue to a non-rvalue reference");
 
-	return Forward<To>(*CastChecked<typename TRemoveReference<To>::Type>(&Arg));
-}
+		return Forward<To>(*CastChecked<typename TRemoveReference<To>::Type>(&Arg));
+	}
 
-template <typename To, typename From>
-FORCEINLINE typename TEnableIf<!TAnd<TIsCastable<typename TRemoveReference<To>::Type>, TIsCastable<typename TRemoveReference<From>::Type>>::Value, To>::Type __DynamicCast(From&& Arg)
-{
-	// This may fail when dynamic_casting rvalue references due to patchy compiler support
-	return dynamic_cast<To>(Arg);
+	template <typename To, typename From>
+	FORCEINLINE typename TEnableIf<!TAnd<TIsCastable<typename TRemoveReference<To>::Type>, TIsCastable<typename TRemoveReference<From>::Type>>::Value, To>::Type DynamicCast(From&& Arg)
+	{
+		// This may fail when dynamic_casting rvalue references due to patchy compiler support
+		return dynamic_cast<To>(Arg);
+	}
 }
 
-#define dynamic_cast __DynamicCast
+#define dynamic_cast UE4Casts_Private::DynamicCast
