Add hooks so game code can detect when demos start playing

diff --git a/Engine/Source/Runtime/CoreUObject/Private/UObject/UObjectGlobals.cpp b/Engine/Source/Runtime/CoreUObject/Private/UObject/UObjectGlobals.cpp
--- a/Engine/Source/Runtime/CoreUObject/Private/UObject/UObjectGlobals.cpp
+++ b/Engine/Source/Runtime/CoreUObject/Private/UObject/UObjectGlobals.cpp
@@ -70,6 +70,7 @@ FSimpleMulticastDelegate FCoreUObjectDelegates::PreGarbageCollect;
 FSimpleMulticastDelegate FCoreUObjectDelegates::PostGarbageCollect;
 FSimpleMulticastDelegate FCoreUObjectDelegates::PreLoadMap;
 FSimpleMulticastDelegate FCoreUObjectDelegates::PostLoadMap;
+FSimpleMulticastDelegate FCoreUObjectDelegates::PostDemoPlay;
 FCoreUObjectDelegates::FOnLoadObjectsOnTop FCoreUObjectDelegates::ShouldLoadOnTop;
 
 FCoreUObjectDelegates::FStringAssetReferenceLoaded FCoreUObjectDelegates::StringAssetReferenceLoaded;
diff --git a/Engine/Source/Runtime/CoreUObject/Public/UObject/UObjectGlobals.h b/Engine/Source/Runtime/CoreUObject/Public/UObject/UObjectGlobals.h
--- a/Engine/Source/Runtime/CoreUObject/Public/UObject/UObjectGlobals.h
+++ b/Engine/Source/Runtime/CoreUObject/Public/UObject/UObjectGlobals.h
@@ -1645,6 +1645,9 @@ struct COREUOBJECT_API FCoreUObjectDelegates
 	// Sent at the _successful_ end of LoadMap
 	static FSimpleMulticastDelegate PostLoadMap;
 
+	// Sent at the _successful_ end of LoadMap
+	static FSimpleMulticastDelegate PostDemoPlay;
+
 	// Called before garbage collection
 	static FSimpleMulticastDelegate PreGarbageCollect;
 
diff --git a/Engine/Source/Runtime/Engine/Private/World.cpp b/Engine/Source/Runtime/Engine/Private/World.cpp
--- a/Engine/Source/Runtime/Engine/Private/World.cpp
+++ b/Engine/Source/Runtime/Engine/Private/World.cpp
@@ -2730,6 +2730,8 @@ bool UWorld::HandleDemoPlayCommand( const TCHAR* Cmd, FOutputDevice& Ar, UWorld*
 			Ar.Logf( TEXT( "Demo playback failed: %s" ), *Error );
 			DestroyDemoNetDriver();
 		}
+
+		FCoreUObjectDelegates::PostDemoPlay.Broadcast();
 #endif
 	}
 	else
