Server functions passed by non-const reference result in bad code generation.

Fixed.

#ttp 338601

diff --git a/Engine/Source/Programs/UnrealHeaderTool/Private/HeaderParser.cpp b/Engine/Source/Programs/UnrealHeaderTool/Private/HeaderParser.cpp
--- a/Engine/Source/Programs/UnrealHeaderTool/Private/HeaderParser.cpp
+++ b/Engine/Source/Programs/UnrealHeaderTool/Private/HeaderParser.cpp
@@ -3253,7 +3253,11 @@ bool FHeaderParser::GetVarType
 			case EVariableCategory::ReplicatedParameter:
 			{
 				if (!(Flags & CPF_ConstParm))
+				{
 					FError::Throwf(TEXT("Replicated %s parameters cannot be passed by non-const reference"), VarType.Identifier);
+				}
+
+				Flags |= CPF_ReferenceParm;
 			}
 			break;
 
diff --git a/Engine/Source/Runtime/CoreUObject/Private/UObject/Property.cpp b/Engine/Source/Runtime/CoreUObject/Private/UObject/Property.cpp
--- a/Engine/Source/Runtime/CoreUObject/Private/UObject/Property.cpp
+++ b/Engine/Source/Runtime/CoreUObject/Private/UObject/Property.cpp
@@ -575,7 +575,7 @@ void UProperty::ExportCppDeclaration(FOutputDevice& Out, EExportedDeclaration::T
 				{
 					// export as a pointer if this is an optional out parm, reference if it's just an out parm, standard otherwise...
 					TCHAR ModifierString[2]={0,0};
-					if ( HasAnyPropertyFlags(CPF_OutParm) || bIsInterfaceProp )
+					if (HasAnyPropertyFlags(CPF_OutParm | CPF_ReferenceParm) || bIsInterfaceProp)
 					{
 						ModifierString[0] = TEXT('&');
 					}
