TTP 349085 Arrays don't migrate their values property between UMG previews and templates when migrating an empty array.  Turned out to be a bug in our text importer for array properties, we exported empty string for an empty array, but didn't accept empty strings as valid importing text.

diff --git a/Engine/Source/Runtime/CoreUObject/Private/UObject/PropertyArray.cpp b/Engine/Source/Runtime/CoreUObject/Private/UObject/PropertyArray.cpp
--- a/Engine/Source/Runtime/CoreUObject/Private/UObject/PropertyArray.cpp
+++ b/Engine/Source/Runtime/CoreUObject/Private/UObject/PropertyArray.cpp
@@ -196,12 +196,20 @@ const TCHAR* UArrayProperty::ImportText_Internal( const TCHAR* Buffer, void* Dat
 {
 	checkSlow(Inner);
 
-	if ( *Buffer++ != TCHAR('(') )
+	FScriptArrayHelper ArrayHelper(this, Data);
+
+	// If we export an empty array we export an empty string, so ensure that if we're passed an empty string
+	// we interpret it as an empty array.
+	if ( *Buffer == TCHAR('\0') )
 	{
+		ArrayHelper.EmptyValues();
 		return NULL;
 	}
 
-	FScriptArrayHelper ArrayHelper(this, Data);
+	if ( *Buffer++ != TCHAR('(') )
+	{
+		return NULL;
+	}
 
 	// only clear the array if we're not importing localized text
 	if ( (PortFlags&PPF_LocalizedOnly) == 0 )
