Blueprints/Core/UMG - Improving our reference finding code to not require reseeding a potentially large hashtable with entries many many times.  Significantly reduces compile times on blueprints that are used in many places in memory.

diff --git a/Engine/Source/Runtime/CoreUObject/Private/Serialization/FindReferencersArchive.cpp b/Engine/Source/Runtime/CoreUObject/Private/Serialization/FindReferencersArchive.cpp
--- a/Engine/Source/Runtime/CoreUObject/Private/Serialization/FindReferencersArchive.cpp
+++ b/Engine/Source/Runtime/CoreUObject/Private/Serialization/FindReferencersArchive.cpp
@@ -11,7 +11,7 @@
  * @param	PotentialReferencer		the object to serialize which may contain references to our target objects
  * @param	InTargetObjects			array of objects to search for references to
  */
-FFindReferencersArchive::FFindReferencersArchive( UObject* PotentialReferencer, TArray<UObject*> InTargetObjects )
+FFindReferencersArchive::FFindReferencersArchive( UObject* InPotentialReferencer, TArray<UObject*> InTargetObjects )
 {
 	// use the optimized RefLink to skip over properties which don't contain object references
 	ArIsObjectReferenceCollector = true;
@@ -23,14 +23,33 @@ FFindReferencersArchive::FFindReferencersArchive( UObject* PotentialReferencer,
 	for ( int32 ObjIndex = 0; ObjIndex < InTargetObjects.Num(); ObjIndex++ )
 	{
 		UObject* TargetObject = InTargetObjects[ObjIndex];
-		if ( TargetObject != NULL && TargetObject != PotentialReferencer )
+		if ( TargetObject != nullptr )
 		{
 			TargetObjects.Add(TargetObject, 0);
 		}
 	}
 
+	PotentialReferencer = InPotentialReferencer;
+
+	// now start the search
+	if ( PotentialReferencer )
+	{
+		PotentialReferencer->Serialize(*this);
+	}
+}
+
+void FFindReferencersArchive::ResetPotentialReferencer(UObject* InPotentialReferencer)
+{
+	PotentialReferencer = InPotentialReferencer;
+
+	// Reset all reference counts
+	for ( TMap<UObject*, int32>::TIterator It(TargetObjects); It; ++It )
+	{
+		It.Value() = 0;
+	}
+
 	// now start the search
-	PotentialReferencer->Serialize( *this );
+	PotentialReferencer->Serialize(*this);
 }
 
 /**
@@ -43,16 +62,16 @@ FFindReferencersArchive::FFindReferencersArchive( UObject* PotentialReferencer,
  * @return	the number of references to TargetObject which were encountered when PotentialReferencer
  *			was serialized.
  */
-int32 FFindReferencersArchive::GetReferenceCount( UObject* TargetObject, TArray<UProperty*>* out_ReferencingProperties/*=NULL*/ ) const
+int32 FFindReferencersArchive::GetReferenceCount( UObject* TargetObject, TArray<UProperty*>* out_ReferencingProperties/*=nullptr*/ ) const
 {
 	int32 Result = 0;
-	if ( TargetObject != NULL )
+	if ( TargetObject != nullptr && PotentialReferencer != TargetObject )
 	{
 		const int32* pCount = TargetObjects.Find(TargetObject);
-		if ( pCount != NULL && (*pCount) > 0 )
+		if ( pCount != nullptr && ( *pCount ) > 0 )
 		{
 			Result = *pCount;
-			if ( out_ReferencingProperties != NULL )
+			if ( out_ReferencingProperties != nullptr )
 			{
 				TArray<UProperty*> PropertiesReferencingObj;
 				ReferenceMap.MultiFind(TargetObject, PropertiesReferencingObj);
@@ -81,7 +100,7 @@ int32 FFindReferencersArchive::GetReferenceCounts( TMap<class UObject*, int32>&
 	out_ReferenceCounts.Empty();
 	for ( TMap<UObject*,int32>::TConstIterator It(TargetObjects); It; ++It )
 	{
-		if ( It.Value() > 0 )
+		if ( It.Value() > 0 && It.Key() != PotentialReferencer )
 		{
 			out_ReferenceCounts.Add(It.Key(), It.Value());
 		}
@@ -98,7 +117,7 @@ int32 FFindReferencersArchive::GetReferenceCounts( TMap<class UObject*, int32>&
  *
  * @return	the number of objects which were referenced by PotentialReferencer.
  */
-int32 FFindReferencersArchive::GetReferenceCounts( TMap<class UObject*, int32>& out_ReferenceCounts, TMultiMap<class UObject*,class UProperty*>& out_ReferencingProperties ) const
+int32 FFindReferencersArchive::GetReferenceCounts( TMap<class UObject*, int32>& out_ReferenceCounts, TMultiMap<class UObject*, class UProperty*>& out_ReferencingProperties ) const
 {
 	GetReferenceCounts(out_ReferenceCounts);
 	if ( out_ReferenceCounts.Num() > 0 )
@@ -126,13 +145,13 @@ int32 FFindReferencersArchive::GetReferenceCounts( TMap<class UObject*, int32>&
  */
 FArchive& FFindReferencersArchive::operator<<( UObject*& Obj )
 {
-	if ( Obj != NULL )
+	if ( Obj != nullptr && Obj != PotentialReferencer )
 	{
 		int32* pReferenceCount = TargetObjects.Find(Obj);
-		if ( pReferenceCount != NULL )
+		if ( pReferenceCount != nullptr )
 		{
 			// if this object was serialized via a UProperty, add it to the list
-			if ( GSerializedProperty != NULL )
+			if ( GSerializedProperty != nullptr )
 			{
 				ReferenceMap.AddUnique(Obj, GSerializedProperty);
 			}
@@ -143,4 +162,4 @@ FArchive& FFindReferencersArchive::operator<<( UObject*& Obj )
 	}
 
 	return *this;
-}
\ No newline at end of file
+}
diff --git a/Engine/Source/Runtime/CoreUObject/Public/Serialization/ArchiveUObject.h b/Engine/Source/Runtime/CoreUObject/Public/Serialization/ArchiveUObject.h
--- a/Engine/Source/Runtime/CoreUObject/Public/Serialization/ArchiveUObject.h
+++ b/Engine/Source/Runtime/CoreUObject/Public/Serialization/ArchiveUObject.h
@@ -36,7 +36,7 @@ public:
 		Max += InMax;
 	}
 	/**
-  	 * Returns the name of the Archive.  Useful for getting the name of the package a struct or object
+	 * Returns the name of the Archive.  Useful for getting the name of the package a struct or object
 	 * is in when a loading error occurs.
 	 *
 	 * This is overridden for the specific Archive Types
@@ -228,7 +228,7 @@ public:
 	virtual ~FReloadObjectArc();
 
 	/**
-  	 * Returns the name of the Archive.  Useful for getting the name of the package a struct or object
+	 * Returns the name of the Archive.  Useful for getting the name of the package a struct or object
 	 * is in when a loading error occurs.
 	 *
 	 * This is overridden for the specific Archive Types
@@ -328,7 +328,7 @@ public:
 	FArchiveReplaceArchetype();
 
 	/**
-  	 * Returns the name of the Archive.  Useful for getting the name of the package a struct or object
+	 * Returns the name of the Archive.  Useful for getting the name of the package a struct or object
 	 * is in when a loading error occurs.
 	 *
 	 * This is overridden for the specific Archive Types
@@ -384,7 +384,7 @@ public:
 	FArchiveShowReferences( FOutputDevice& inOutputAr, UObject* inOuter, UObject* inSource, TArray<UObject*>& InExclude );
 
 	/**
-  	 * Returns the name of the Archive.  Useful for getting the name of the package a struct or object
+	 * Returns the name of the Archive.  Useful for getting the name of the package a struct or object
 	 * is in when a loading error occurs.
 	 *
 	 * This is overridden for the specific Archive Types
@@ -441,19 +441,28 @@ public:
 	COREUOBJECT_API int32 GetReferenceCounts( TMap<class UObject*, int32>& out_ReferenceCounts, TMultiMap<class UObject*,class UProperty*>& out_ReferencingProperties ) const;
 
 	/**
-  	 * Returns the name of the Archive.  Useful for getting the name of the package a struct or object
+	 * Returns the name of the Archive.  Useful for getting the name of the package a struct or object
 	 * is in when a loading error occurs.
 	 *
 	 * This is overridden for the specific Archive Types
 	 **/
 	COREUOBJECT_API virtual FString GetArchiveName() const { return TEXT("FFindReferencersArchive"); }
 
+	/**
+	 * Resets the reference counts.  Keeps the same target objects but sets up everything to test a new potential referencer.
+	 * @param	PotentialReferencer		the object to serialize which may contain references to our target objects
+	 **/
+	COREUOBJECT_API void ResetPotentialReferencer(UObject* InPotentialReferencer);
+
 protected:
-	TMap<class UObject*,int32>	TargetObjects;
+	TMap<class UObject*, int32>	TargetObjects;
 
 	/** a mapping of target object => the properties in PotentialReferencer that hold the reference to target object */
 	TMultiMap<class UObject*,class UProperty*> ReferenceMap;
 
+	/** The potential referencer we ignore */
+	class UObject* PotentialReferencer;
+
 private:
 
 	/**
@@ -485,16 +494,19 @@ public:
 	TFindObjectReferencers( TArray< T* > TargetObjects, UPackage* PackageToCheck=NULL, bool bIgnoreTemplates = true )
 	: TMultiMap< T*, UObject* >()
 	{
+		FFindReferencersArchive FindReferencerAr(nullptr, ( TArray<UObject*>& )TargetObjects);
+
+		// Loop over every object to find any reference that may exist for the target objects
 		for (FObjectIterator It; It; ++It)
 		{
 			UObject* PotentialReferencer = *It;
-			if (!TargetObjects.Contains(dynamic_cast<T*>(PotentialReferencer))
+			if ( !TargetObjects.Contains(dynamic_cast<T*>( PotentialReferencer ))
 			&&	(PackageToCheck == NULL || PotentialReferencer->IsIn(PackageToCheck))
 			&&	(!bIgnoreTemplates || !PotentialReferencer->IsTemplate()) )
 			{
-				FFindReferencersArchive FindReferencerAr(PotentialReferencer, (TArray<UObject*>&)TargetObjects);
+				FindReferencerAr.ResetPotentialReferencer(PotentialReferencer);
 
-				TMap<UObject*,int32> ReferenceCounts;
+				TMap<UObject*, int32> ReferenceCounts;
 				if ( FindReferencerAr.GetReferenceCounts(ReferenceCounts) > 0 )
 				{
 					// here we don't really care about the number of references from PotentialReferencer to the target object...just that he's a referencer
@@ -548,7 +560,7 @@ public:
 	}
 
 	/**
-  	 * Returns the name of the Archive.  Useful for getting the name of the package a struct or object
+	 * Returns the name of the Archive.  Useful for getting the name of the package a struct or object
 	 * is in when a loading error occurs.
 	 *
 	 * This is overridden for the specific Archive Types
@@ -727,7 +739,7 @@ public:
 	static COREUOBJECT_API FString PrintRootPath( const TMap<UObject*,UProperty*>& Route, const UObject* TargetObject );
 
 	/**
-  	 * Returns the name of the Archive.  Useful for getting the name of the package a struct or object
+	 * Returns the name of the Archive.  Useful for getting the name of the package a struct or object
 	 * is in when a loading error occurs.
 	 *
 	 * This is overridden for the specific Archive Types
@@ -869,7 +881,7 @@ private:
 
 public:
 	/**
-  	 * Returns the name of the Archive.  Useful for getting the name of the package a struct or object
+	 * Returns the name of the Archive.  Useful for getting the name of the package a struct or object
 	 * is in when a loading error occurs.
 	 *
 	 * This is overridden for the specific Archive Types
@@ -954,7 +966,7 @@ private:
 
 public:
 	/**
-  	 * Returns the name of the Archive.  Useful for getting the name of the package a struct or object
+	 * Returns the name of the Archive.  Useful for getting the name of the package a struct or object
 	 * is in when a loading error occurs.
 	 *
 	 * This is overridden for the specific Archive Types
@@ -1157,7 +1169,7 @@ public:
 	}
 
 	/**
-  	 * Returns the name of this archive.
+	 * Returns the name of this archive.
 	 **/
 	virtual FString GetArchiveName() const { return TEXT("ReplaceObjectRef"); }
 
@@ -1328,7 +1340,7 @@ class COREUOBJECT_API FArchiveAsync : public FArchive
 {
 public:
 	/**
- 	 * Constructor, initializing member variables.
+	 * Constructor, initializing member variables.
 	 */
 	FArchiveAsync( const TCHAR* InFileName );
 
