Make Core module public headers individually compilable. Part 1 of n.

diff --git a/Engine/Source/Runtime/Core/Public/Containers/EnumAsByte.h b/Engine/Source/Runtime/Core/Public/Containers/EnumAsByte.h
--- a/Engine/Source/Runtime/Core/Public/Containers/EnumAsByte.h
+++ b/Engine/Source/Runtime/Core/Public/Containers/EnumAsByte.h
@@ -1,7 +1,7 @@
 // Copyright 1998-2014 Epic Games, Inc. All Rights Reserved.
 
 #pragma once
-
+#include "Templates/UnrealTypeTraits.h"
 
 /**
  * Template to store enums as BYTEs with type safety
diff --git a/Engine/Source/Runtime/Core/Public/Misc/AssertionMacros.h b/Engine/Source/Runtime/Core/Public/Misc/AssertionMacros.h
--- a/Engine/Source/Runtime/Core/Public/Misc/AssertionMacros.h
+++ b/Engine/Source/Runtime/Core/Public/Misc/AssertionMacros.h
@@ -5,7 +5,8 @@
 =============================================================================*/
 
 #pragma once
-
+#include "HAL/Platform.h"
+#include "Misc/Build.h"
 
 /*----------------------------------------------------------------------------
 	Check, verify, etc macros
diff --git a/Engine/Source/Runtime/Core/Public/Misc/Char.h b/Engine/Source/Runtime/Core/Public/Misc/Char.h
--- a/Engine/Source/Runtime/Core/Public/Misc/Char.h
+++ b/Engine/Source/Runtime/Core/Public/Misc/Char.h
@@ -1,6 +1,9 @@
 // Copyright 1998-2014 Epic Games, Inc. All Rights Reserved.
 
 #pragma once
+#include "HAL/Platform.h"
+#include <ctype.h>
+#include <wctype.h>
 
 /*-----------------------------------------------------------------------------
 	Character type functions.
diff --git a/Engine/Source/Runtime/Core/Public/Misc/Crc.h b/Engine/Source/Runtime/Core/Public/Misc/Crc.h
--- a/Engine/Source/Runtime/Core/Public/Misc/Crc.h
+++ b/Engine/Source/Runtime/Core/Public/Misc/Crc.h
@@ -1,9 +1,9 @@
 // Copyright 1998-2014 Epic Games, Inc. All Rights Reserved.
 
 #pragma once
-
+#include "HAL/Platform.h"
+#include "Misc/Char.h"
 #include "Templates/EnableIf.h"
-
 /** 
  * CRC hash generation for different types of input data
  **/
diff --git a/Engine/Source/Runtime/Core/Public/Misc/DateTime.h b/Engine/Source/Runtime/Core/Public/Misc/DateTime.h
--- a/Engine/Source/Runtime/Core/Public/Misc/DateTime.h
+++ b/Engine/Source/Runtime/Core/Public/Misc/DateTime.h
@@ -1,10 +1,12 @@
 // Copyright 1998-2014 Epic Games, Inc. All Rights Reserved.
 
 #pragma once
-
+#include "HAL/Platform.h"
+#include "Misc/Timespan.h"
 
 class FArchive;
 class FOutputDevice;
+class FString;
 class UObject;
 
 
diff --git a/Engine/Source/Runtime/Core/Public/Misc/Parse.h b/Engine/Source/Runtime/Core/Public/Misc/Parse.h
--- a/Engine/Source/Runtime/Core/Public/Misc/Parse.h
+++ b/Engine/Source/Runtime/Core/Public/Misc/Parse.h
@@ -1,6 +1,11 @@
 // Copyright 1998-2014 Epic Games, Inc. All Rights Reserved.
 
 #pragma once
+#include "HAL/Platform.h"
+
+class FName;
+class FString;
+class FText;
 
 /*-----------------------------------------------------------------------------
 	Parsing functions.
@@ -14,7 +19,7 @@ struct CORE_API FParse
 	 */
 	static bool Command( const TCHAR** Stream, const TCHAR* Match, bool bParseMightTriggerExecution = true );
 	/** Parses a name. */
-	static bool Value( const TCHAR* Stream, const TCHAR* Match, class FName& Name );
+	static bool Value( const TCHAR* Stream, const TCHAR* Match, FName& Name );
 	/** Parses a uint32. */
 	static bool Value( const TCHAR* Stream, const TCHAR* Match, uint32& Value );
 	/** Parses a globally unique identifier. */
diff --git a/Engine/Source/Runtime/Core/Public/Misc/Timespan.h b/Engine/Source/Runtime/Core/Public/Misc/Timespan.h
--- a/Engine/Source/Runtime/Core/Public/Misc/Timespan.h
+++ b/Engine/Source/Runtime/Core/Public/Misc/Timespan.h
@@ -1,12 +1,14 @@
 // Copyright 1998-2014 Epic Games, Inc. All Rights Reserved.
 
 #pragma once
-
+#include "HAL/Platform.h"
 
 class FArchive;
 class FOutputDevice;
+class FString;
 class UObject;
 
+inline uint32 GetTypeHash(const int64 A);
 
 namespace ETimespan
 {
diff --git a/Engine/Source/Runtime/Core/Public/ProfilingDebugging/UMemoryDefines.h b/Engine/Source/Runtime/Core/Public/ProfilingDebugging/UMemoryDefines.h
--- a/Engine/Source/Runtime/Core/Public/ProfilingDebugging/UMemoryDefines.h
+++ b/Engine/Source/Runtime/Core/Public/ProfilingDebugging/UMemoryDefines.h
@@ -1,6 +1,7 @@
 // Copyright 1998-2014 Epic Games, Inc. All Rights Reserved.
 
 #pragma once
+#include "HAL/Platform.h"
 
 /**
  * true we want to force an ansi allocator
diff --git a/Engine/Source/Runtime/Core/Public/Templates/UnrealTemplate.h b/Engine/Source/Runtime/Core/Public/Templates/UnrealTemplate.h
--- a/Engine/Source/Runtime/Core/Public/Templates/UnrealTemplate.h
+++ b/Engine/Source/Runtime/Core/Public/Templates/UnrealTemplate.h
@@ -5,8 +5,9 @@
 =============================================================================*/
 
 #pragma once
-
 #include "EnableIf.h"
+#include "HAL/Platform.h"
+#include "Templates/UnrealTypeTraits.h"
 
 /*-----------------------------------------------------------------------------
 	Standard templates.
diff --git a/Engine/Source/Runtime/CoreUObject/Public/UObject/UObjectAllocator.h b/Engine/Source/Runtime/CoreUObject/Public/UObject/UObjectAllocator.h
--- a/Engine/Source/Runtime/CoreUObject/Public/UObject/UObjectAllocator.h
+++ b/Engine/Source/Runtime/CoreUObject/Public/UObject/UObjectAllocator.h
@@ -6,7 +6,9 @@
 
 #ifndef __UNOBJALLOCATOR_H__
 #define __UNOBJALLOCATOR_H__
+#include "HAL/Platform.h"
 
+class UObjectBase;
 
 class COREUOBJECT_API FUObjectAllocator
 {
@@ -41,7 +43,7 @@ public:
 	 * @param Object object to test as a member of permanent object pool
 	 * @return true if object is part of permanent object pool, false otherwise
 	 */
-	FORCEINLINE bool ResidesInPermanentPool(const class UObjectBase *Object) const
+	FORCEINLINE bool ResidesInPermanentPool(const UObjectBase *Object) const
 	{
 		return ((const uint8*)Object >= PermanentObjectPool) && ((const uint8*)Object < PermanentObjectPoolTail);
 	}
@@ -61,7 +63,7 @@ public:
 	 *
 	 * @param Object object to free
 	 */
-	void FreeUObject(class UObjectBase *Object) const;
+	void FreeUObject(UObjectBase *Object) const;
 
 private:
 
