CIS Fix.

diff --git a/Engine/Source/Runtime/CoreUObject/Public/UObject/ObjectBase.h b/Engine/Source/Runtime/CoreUObject/Public/UObject/ObjectBase.h
--- a/Engine/Source/Runtime/CoreUObject/Public/UObject/ObjectBase.h
+++ b/Engine/Source/Runtime/CoreUObject/Public/UObject/ObjectBase.h
@@ -1178,10 +1178,10 @@ public: \
 	static_assert(false, "You have to define " #TClass "::" #TClass "() or " #TClass "::" #TClass "(const FPostConstructInitializeProperties&). This is required by UObject system to work correctly.");
 
 #define DEFINE_DEFAULT_CONSTRUCTOR_CALL(TClass) \
-	static void __DefaultConstructor(const class FPostConstructInitializeProperties& X) { const_cast<FPostConstructInitializeProperties&>(X).FinalizeSubobjectClassInitialization(); new((EInternal*)X.GetObject())TClass(); }
+	static void __DefaultConstructor(const class FPostConstructInitializeProperties& X) { const_cast<FPostConstructInitializeProperties&>(X).FinalizeSubobjectClassInitialization(); new((EInternal*)X.GetObj())TClass(); }
 
 #define DEFINE_DEFAULT_PCIP_CONSTRUCTOR_CALL(TClass) \
-	static void __DefaultConstructor(const class FPostConstructInitializeProperties& X) { new((EInternal*)X.GetObject())TClass(X); }
+	static void __DefaultConstructor(const class FPostConstructInitializeProperties& X) { new((EInternal*)X.GetObj())TClass(X); }
 
 #define DECLARE_CLASS_INTRINSIC(TClass,TSuperClass,TStaticFlags,TPackage) \
 	DECLARE_CLASS(TClass,TSuperClass,TStaticFlags|CLASS_Intrinsic,CASTCLASS_None,TPackage,NO_API ) \
diff --git a/Engine/Source/Runtime/CoreUObject/Public/UObject/UObjectGlobals.h b/Engine/Source/Runtime/CoreUObject/Public/UObject/UObjectGlobals.h
--- a/Engine/Source/Runtime/CoreUObject/Public/UObject/UObjectGlobals.h
+++ b/Engine/Source/Runtime/CoreUObject/Public/UObject/UObjectGlobals.h
@@ -677,7 +677,7 @@ public:
 	/**
 	* Return the object that is being constructed
 	**/
-	FORCEINLINE UObject* GetObject() const
+	FORCEINLINE UObject* GetObj() const
 	{
 		return Obj;
 	}
