Allow dynamically adding game name redirects at runtime

diff --git a/Engine/Source/Runtime/CoreUObject/Private/UObject/LinkerLoad.cpp b/Engine/Source/Runtime/CoreUObject/Private/UObject/LinkerLoad.cpp
--- a/Engine/Source/Runtime/CoreUObject/Private/UObject/LinkerLoad.cpp
+++ b/Engine/Source/Runtime/CoreUObject/Private/UObject/LinkerLoad.cpp
@@ -59,6 +59,11 @@ TMap<FName, FName> ULinkerLoad::StructNameRedirects;				// Old struct name to ne
 TMap<FString, FString> ULinkerLoad::PluginNameRedirects;			// Old plugin name to new plugin name mapping
 TMap<FName, ULinkerLoad::FSubobjectRedirect> ULinkerLoad::SubobjectNameRedirects;	
 
+void ULinkerLoad::AddGameNameRedirect(const FName OldName, const FName NewName)
+{
+	GameNameRedirects.Add(OldName, NewName);
+}
+
 /*----------------------------------------------------------------------------
 	Helpers
 ----------------------------------------------------------------------------*/
diff --git a/Engine/Source/Runtime/CoreUObject/Public/UObject/Linker.h b/Engine/Source/Runtime/CoreUObject/Public/UObject/Linker.h
--- a/Engine/Source/Runtime/CoreUObject/Public/UObject/Linker.h
+++ b/Engine/Source/Runtime/CoreUObject/Public/UObject/Linker.h
@@ -1260,6 +1260,8 @@ public:
 	static TMap<FName, FName> ObjectNameRedirectsObjectOnly;	
 	/** Old game name to new game name for ImportMap */
 	static TMap<FName, FName> GameNameRedirects;
+	/** Add a new redirect from old game name to new game name for ImportMap */
+	COREUOBJECT_API static void AddGameNameRedirect(const FName OldName, const FName NewName);
 	/** Old struct name to new struct name mapping */
 	static TMap<FName, FName> StructNameRedirects;
 	/** Old plugin name to new plugin name mapping */
