Rename FPostConstructInitializeProperties to something simpler

diff --git a/Engine/Source/Runtime/AIModule/Private/EnvironmentQuery/Tests/EnvQueryTest_Pathfinding.cpp b/Engine/Source/Runtime/AIModule/Private/EnvironmentQuery/Tests/EnvQueryTest_Pathfinding.cpp
--- a/Engine/Source/Runtime/AIModule/Private/EnvironmentQuery/Tests/EnvQueryTest_Pathfinding.cpp
+++ b/Engine/Source/Runtime/AIModule/Private/EnvironmentQuery/Tests/EnvQueryTest_Pathfinding.cpp
@@ -9,7 +9,7 @@
 
 #define LOCTEXT_NAMESPACE "EnvQueryGenerator"
 
-UEnvQueryTest_Pathfinding::UEnvQueryTest_Pathfinding(const class FPostConstructInitializeProperties& PCIP) : Super(PCIP)
+UEnvQueryTest_Pathfinding::UEnvQueryTest_Pathfinding(const FObjectInitializer& ObjectInitializer) : Super(ObjectInitializer)
 {
 	Context = UEnvQueryContext_Querier::StaticClass();
 	Cost = EEnvTestCost::High;
diff --git a/Engine/Source/Runtime/AIModule/Private/EnvironmentQuery/Tests/EnvQueryTest_Random.cpp b/Engine/Source/Runtime/AIModule/Private/EnvironmentQuery/Tests/EnvQueryTest_Random.cpp
--- a/Engine/Source/Runtime/AIModule/Private/EnvironmentQuery/Tests/EnvQueryTest_Random.cpp
+++ b/Engine/Source/Runtime/AIModule/Private/EnvironmentQuery/Tests/EnvQueryTest_Random.cpp
@@ -4,7 +4,7 @@
 #include "EnvironmentQuery/Tests/EnvQueryTest_Random.h"
 #include "EnvironmentQuery/Items/EnvQueryItemType.h"
 
-UEnvQueryTest_Random::UEnvQueryTest_Random(const class FPostConstructInitializeProperties& PCIP) : Super(PCIP)
+UEnvQueryTest_Random::UEnvQueryTest_Random(const FObjectInitializer& ObjectInitializer) : Super(ObjectInitializer)
 {
 	Cost = EEnvTestCost::Low;
 
diff --git a/Engine/Source/Runtime/AIModule/Private/EnvironmentQuery/Tests/EnvQueryTest_Trace.cpp b/Engine/Source/Runtime/AIModule/Private/EnvironmentQuery/Tests/EnvQueryTest_Trace.cpp
--- a/Engine/Source/Runtime/AIModule/Private/EnvironmentQuery/Tests/EnvQueryTest_Trace.cpp
+++ b/Engine/Source/Runtime/AIModule/Private/EnvironmentQuery/Tests/EnvQueryTest_Trace.cpp
@@ -7,7 +7,7 @@
 
 #define LOCTEXT_NAMESPACE "EnvQueryGenerator"
 
-UEnvQueryTest_Trace::UEnvQueryTest_Trace(const class FPostConstructInitializeProperties& PCIP) : Super(PCIP)
+UEnvQueryTest_Trace::UEnvQueryTest_Trace(const FObjectInitializer& ObjectInitializer) : Super(ObjectInitializer)
 {
 	Cost = EEnvTestCost::High;
 	ValidItemType = UEnvQueryItemType_VectorBase::StaticClass();
diff --git a/Engine/Source/Runtime/AIModule/Private/HotSpots/AIHotSpotManager.cpp b/Engine/Source/Runtime/AIModule/Private/HotSpots/AIHotSpotManager.cpp
--- a/Engine/Source/Runtime/AIModule/Private/HotSpots/AIHotSpotManager.cpp
+++ b/Engine/Source/Runtime/AIModule/Private/HotSpots/AIHotSpotManager.cpp
@@ -4,7 +4,7 @@
 #include "HotSpots/AIHotSpotManager.h"
 
 
-UAIHotSpotManager::UAIHotSpotManager(const FPostConstructInitializeProperties& PCIP)
-: Super(PCIP)
+UAIHotSpotManager::UAIHotSpotManager(const FObjectInitializer& ObjectInitializer)
+: Super(ObjectInitializer)
 {
 }
\ No newline at end of file
diff --git a/Engine/Source/Runtime/AIModule/Private/Navigation/CrowdFollowingComponent.cpp b/Engine/Source/Runtime/AIModule/Private/Navigation/CrowdFollowingComponent.cpp
--- a/Engine/Source/Runtime/AIModule/Private/Navigation/CrowdFollowingComponent.cpp
+++ b/Engine/Source/Runtime/AIModule/Private/Navigation/CrowdFollowingComponent.cpp
@@ -10,7 +10,7 @@
 
 DEFINE_LOG_CATEGORY(LogCrowdFollowing);
 
-UCrowdFollowingComponent::UCrowdFollowingComponent(const class FPostConstructInitializeProperties& PCIP) : Super(PCIP)
+UCrowdFollowingComponent::UCrowdFollowingComponent(const FObjectInitializer& ObjectInitializer) : Super(ObjectInitializer)
 {
 	bAffectFallingVelocity = false;
 	bRotateToVelocity = true;
diff --git a/Engine/Source/Runtime/AIModule/Private/Navigation/CrowdManager.cpp b/Engine/Source/Runtime/AIModule/Private/Navigation/CrowdManager.cpp
--- a/Engine/Source/Runtime/AIModule/Private/Navigation/CrowdManager.cpp
+++ b/Engine/Source/Runtime/AIModule/Private/Navigation/CrowdManager.cpp
@@ -98,7 +98,7 @@ void FCrowdAvoidanceSamplingPattern::AddSampleWithMirror(float AngleInDegrees, f
 	AddSample(-AngleInDegrees, NormalizedRadius);
 }
 
-UCrowdManager::UCrowdManager(const class FPostConstructInitializeProperties& PCIP) : Super(PCIP)
+UCrowdManager::UCrowdManager(const FObjectInitializer& ObjectInitializer) : Super(ObjectInitializer)
 {
 	MyNavData = NULL;
 #if WITH_RECAST
diff --git a/Engine/Source/Runtime/AIModule/Private/Navigation/NavigationComponent.cpp b/Engine/Source/Runtime/AIModule/Private/Navigation/NavigationComponent.cpp
--- a/Engine/Source/Runtime/AIModule/Private/Navigation/NavigationComponent.cpp
+++ b/Engine/Source/Runtime/AIModule/Private/Navigation/NavigationComponent.cpp
@@ -9,8 +9,8 @@
 // Life cycle                                                        
 //----------------------------------------------------------------------//
 
-UNavigationComponent::UNavigationComponent(const class FPostConstructInitializeProperties& PCIP)
-	: Super(PCIP)
+UNavigationComponent::UNavigationComponent(const FObjectInitializer& ObjectInitializer)
+	: Super(ObjectInitializer)
 	, bIsPathPartial(false)
 	, AsynPathQueryID(INVALID_NAVQUERYID)
 {
diff --git a/Engine/Source/Runtime/AIModule/Private/Navigation/NavigationInterfaces.cpp b/Engine/Source/Runtime/AIModule/Private/Navigation/NavigationInterfaces.cpp
--- a/Engine/Source/Runtime/AIModule/Private/Navigation/NavigationInterfaces.cpp
+++ b/Engine/Source/Runtime/AIModule/Private/Navigation/NavigationInterfaces.cpp
@@ -3,7 +3,7 @@
 #include "AIModulePrivate.h"
 #include "Navigation/CrowdAgentInterface.h"
 
-UCrowdAgentInterface::UCrowdAgentInterface(const class FPostConstructInitializeProperties& PCIP)
-	: Super(PCIP)
+UCrowdAgentInterface::UCrowdAgentInterface(const FObjectInitializer& ObjectInitializer)
+	: Super(ObjectInitializer)
 {
 }
diff --git a/Engine/Source/Runtime/AIModule/Private/Navigation/PathFollowingComponent.cpp b/Engine/Source/Runtime/AIModule/Private/Navigation/PathFollowingComponent.cpp
--- a/Engine/Source/Runtime/AIModule/Private/Navigation/PathFollowingComponent.cpp
+++ b/Engine/Source/Runtime/AIModule/Private/Navigation/PathFollowingComponent.cpp
@@ -25,7 +25,7 @@ uint32 UPathFollowingComponent::NextRequestId = 0;
 const float UPathFollowingComponent::DefaultAcceptanceRadius = -1.f;
 UPathFollowingComponent::FRequestCompletedSignature UPathFollowingComponent::UnboundRequestDelegate;
 
-UPathFollowingComponent::UPathFollowingComponent(const class FPostConstructInitializeProperties& PCIP) : Super(PCIP)
+UPathFollowingComponent::UPathFollowingComponent(const FObjectInitializer& ObjectInitializer) : Super(ObjectInitializer)
 {
 	PrimaryComponentTick.bCanEverTick = true;
 	bAutoActivate = true;
diff --git a/Engine/Source/Runtime/AIModule/Private/Perception/AIPerceptionComponent.cpp b/Engine/Source/Runtime/AIModule/Private/Perception/AIPerceptionComponent.cpp
--- a/Engine/Source/Runtime/AIModule/Private/Perception/AIPerceptionComponent.cpp
+++ b/Engine/Source/Runtime/AIModule/Private/Perception/AIPerceptionComponent.cpp
@@ -28,8 +28,8 @@ void FActorPerceptionInfo::Merge(const FActorPerceptionInfo& Other)
 //----------------------------------------------------------------------//
 const int32 UAIPerceptionComponent::InitialStimuliToProcessArraySize = 10;
 
-UAIPerceptionComponent::UAIPerceptionComponent(const class FPostConstructInitializeProperties& PCIP)
-	: Super(PCIP)
+UAIPerceptionComponent::UAIPerceptionComponent(const FObjectInitializer& ObjectInitializer)
+	: Super(ObjectInitializer)
 	, DominantSense(ECorePerceptionTypes::MAX)
 	, PerceptionListenerId(AIPerception::InvalidListenerId)
 {
diff --git a/Engine/Source/Runtime/AIModule/Private/Perception/AIPerceptionListenerInterface.cpp b/Engine/Source/Runtime/AIModule/Private/Perception/AIPerceptionListenerInterface.cpp
--- a/Engine/Source/Runtime/AIModule/Private/Perception/AIPerceptionListenerInterface.cpp
+++ b/Engine/Source/Runtime/AIModule/Private/Perception/AIPerceptionListenerInterface.cpp
@@ -3,7 +3,7 @@
 #include "AIModulePrivate.h"
 #include "Perception/AIPerceptionListenerInterface.h"
 
-UAIPerceptionListenerInterface::UAIPerceptionListenerInterface(const class FPostConstructInitializeProperties& PCIP)
-	: Super(PCIP)
+UAIPerceptionListenerInterface::UAIPerceptionListenerInterface(const FObjectInitializer& ObjectInitializer)
+	: Super(ObjectInitializer)
 {
 }
diff --git a/Engine/Source/Runtime/AIModule/Private/Perception/AIPerceptionSystem.cpp b/Engine/Source/Runtime/AIModule/Private/Perception/AIPerceptionSystem.cpp
--- a/Engine/Source/Runtime/AIModule/Private/Perception/AIPerceptionSystem.cpp
+++ b/Engine/Source/Runtime/AIModule/Private/Perception/AIPerceptionSystem.cpp
@@ -16,8 +16,8 @@ DEFINE_LOG_CATEGORY(LogAIPerception);
 //----------------------------------------------------------------------//
 // UAIPerceptionSystem
 //----------------------------------------------------------------------//
-UAIPerceptionSystem::UAIPerceptionSystem(const class FPostConstructInitializeProperties& PCIP)
-	: Super(PCIP)
+UAIPerceptionSystem::UAIPerceptionSystem(const FObjectInitializer& ObjectInitializer)
+	: Super(ObjectInitializer)
 	, PerceptionAgingRate(0.3f)
 	, NextFreeListenerId(AIPerception::InvalidListenerId + 1)
 	, CurrentTime(0.f)
diff --git a/Engine/Source/Runtime/AIModule/Private/Perception/AISense.cpp b/Engine/Source/Runtime/AIModule/Private/Perception/AISense.cpp
--- a/Engine/Source/Runtime/AIModule/Private/Perception/AISense.cpp
+++ b/Engine/Source/Runtime/AIModule/Private/Perception/AISense.cpp
@@ -5,8 +5,8 @@
 
 const float UAISense::SuspendNextUpdate = FLT_MAX;
 
-UAISense::UAISense(const class FPostConstructInitializeProperties& PCIP)
-	: Super(PCIP)
+UAISense::UAISense(const FObjectInitializer& ObjectInitializer)
+	: Super(ObjectInitializer)
 	, TimeUntilNextUpdate(SuspendNextUpdate)
 {
 }
diff --git a/Engine/Source/Runtime/AIModule/Private/Perception/AISense_Damage.cpp b/Engine/Source/Runtime/AIModule/Private/Perception/AISense_Damage.cpp
--- a/Engine/Source/Runtime/AIModule/Private/Perception/AISense_Damage.cpp
+++ b/Engine/Source/Runtime/AIModule/Private/Perception/AISense_Damage.cpp
@@ -3,8 +3,8 @@
 #include "AIModulePrivate.h"
 #include "Perception/AISense_Damage.h"
 
-UAISense_Damage::UAISense_Damage(const class FPostConstructInitializeProperties& PCIP) :
-	Super(PCIP)
+UAISense_Damage::UAISense_Damage(const FObjectInitializer& ObjectInitializer) :
+	Super(ObjectInitializer)
 {
 	
 }
diff --git a/Engine/Source/Runtime/AIModule/Private/Perception/AISense_Hearing.cpp b/Engine/Source/Runtime/AIModule/Private/Perception/AISense_Hearing.cpp
--- a/Engine/Source/Runtime/AIModule/Private/Perception/AISense_Hearing.cpp
+++ b/Engine/Source/Runtime/AIModule/Private/Perception/AISense_Hearing.cpp
@@ -15,8 +15,8 @@ FAINoiseEvent::FAINoiseEvent(AActor* InInstigator, const FVector& InNoiseLocatio
 //----------------------------------------------------------------------//
 // UAISense_Hearing
 //----------------------------------------------------------------------//
-UAISense_Hearing::UAISense_Hearing(const class FPostConstructInitializeProperties& PCIP) :
-	Super(PCIP)
+UAISense_Hearing::UAISense_Hearing(const FObjectInitializer& ObjectInitializer) :
+	Super(ObjectInitializer)
 {
 	
 }
diff --git a/Engine/Source/Runtime/AIModule/Private/Perception/AISense_Prediction.cpp b/Engine/Source/Runtime/AIModule/Private/Perception/AISense_Prediction.cpp
--- a/Engine/Source/Runtime/AIModule/Private/Perception/AISense_Prediction.cpp
+++ b/Engine/Source/Runtime/AIModule/Private/Perception/AISense_Prediction.cpp
@@ -3,8 +3,8 @@
 #include "AIModulePrivate.h"
 #include "Perception/AISense_Prediction.h"
 
-UAISense_Prediction::UAISense_Prediction(const class FPostConstructInitializeProperties& PCIP) :
-	Super(PCIP)
+UAISense_Prediction::UAISense_Prediction(const FObjectInitializer& ObjectInitializer) :
+	Super(ObjectInitializer)
 {
 }
 
diff --git a/Engine/Source/Runtime/AIModule/Private/Perception/AISense_Sight.cpp b/Engine/Source/Runtime/AIModule/Private/Perception/AISense_Sight.cpp
--- a/Engine/Source/Runtime/AIModule/Private/Perception/AISense_Sight.cpp
+++ b/Engine/Source/Runtime/AIModule/Private/Perception/AISense_Sight.cpp
@@ -43,8 +43,8 @@ FAISightTarget::FAISightTarget(AActor* InTarget, FGenericTeamId InTeamId)
 //----------------------------------------------------------------------//
 // UAISense_Sight
 //----------------------------------------------------------------------//
-UAISense_Sight::UAISense_Sight(const class FPostConstructInitializeProperties& PCIP)
-	: Super(PCIP)
+UAISense_Sight::UAISense_Sight(const FObjectInitializer& ObjectInitializer)
+	: Super(ObjectInitializer)
 	, MaxTracesPerTick(DefaultMaxTracesPerTick)
 	, HighImportanceQueryDistanceThreshold(300.f)
 	, MaxQueryImportance(60.f)
diff --git a/Engine/Source/Runtime/AIModule/Private/Perception/AISense_Team.cpp b/Engine/Source/Runtime/AIModule/Private/Perception/AISense_Team.cpp
--- a/Engine/Source/Runtime/AIModule/Private/Perception/AISense_Team.cpp
+++ b/Engine/Source/Runtime/AIModule/Private/Perception/AISense_Team.cpp
@@ -17,8 +17,8 @@ FAITeamStimulusEvent::FAITeamStimulusEvent(AActor* InBroadcaster, AActor* InEnem
 //----------------------------------------------------------------------//
 // 
 //----------------------------------------------------------------------//
-UAISense_Team::UAISense_Team(const class FPostConstructInitializeProperties& PCIP) :
-	Super(PCIP)
+UAISense_Team::UAISense_Team(const FObjectInitializer& ObjectInitializer) :
+	Super(ObjectInitializer)
 {
 	
 }
diff --git a/Engine/Source/Runtime/AIModule/Private/Perception/AISense_Touch.cpp b/Engine/Source/Runtime/AIModule/Private/Perception/AISense_Touch.cpp
--- a/Engine/Source/Runtime/AIModule/Private/Perception/AISense_Touch.cpp
+++ b/Engine/Source/Runtime/AIModule/Private/Perception/AISense_Touch.cpp
@@ -3,8 +3,8 @@
 #include "AIModulePrivate.h"
 #include "Perception/AISense_Touch.h"
 
-UAISense_Touch::UAISense_Touch(const class FPostConstructInitializeProperties& PCIP) :
-	Super(PCIP)
+UAISense_Touch::UAISense_Touch(const FObjectInitializer& ObjectInitializer) :
+	Super(ObjectInitializer)
 {
 	
 }
diff --git a/Engine/Source/Runtime/AIModule/Private/Perception/AISightTargetInterface.cpp b/Engine/Source/Runtime/AIModule/Private/Perception/AISightTargetInterface.cpp
--- a/Engine/Source/Runtime/AIModule/Private/Perception/AISightTargetInterface.cpp
+++ b/Engine/Source/Runtime/AIModule/Private/Perception/AISightTargetInterface.cpp
@@ -3,7 +3,7 @@
 #include "AIModulePrivate.h"
 #include "Perception/AISightTargetInterface.h"
 
-UAISightTargetInterface::UAISightTargetInterface(const class FPostConstructInitializeProperties& PCIP)
-	: Super(PCIP)
+UAISightTargetInterface::UAISightTargetInterface(const FObjectInitializer& ObjectInitializer)
+	: Super(ObjectInitializer)
 {
 }
diff --git a/Engine/Source/Runtime/AIModule/Private/Perception/PawnSensingComponent.cpp b/Engine/Source/Runtime/AIModule/Private/Perception/PawnSensingComponent.cpp
--- a/Engine/Source/Runtime/AIModule/Private/Perception/PawnSensingComponent.cpp
+++ b/Engine/Source/Runtime/AIModule/Private/Perception/PawnSensingComponent.cpp
@@ -6,8 +6,8 @@
 
 DECLARE_CYCLE_STAT(TEXT("Sensing"),STAT_AI_Sensing,STATGROUP_AI);
 
-UPawnSensingComponent::UPawnSensingComponent(const class FPostConstructInitializeProperties& PCIP)
-	: Super(PCIP)
+UPawnSensingComponent::UPawnSensingComponent(const FObjectInitializer& ObjectInitializer)
+	: Super(ObjectInitializer)
 {
 	LOSHearingThreshold = 2800.f;
 	HearingThreshold = 1400.0f;
diff --git a/Engine/Source/Runtime/Android/AndroidRuntimeSettings/Private/AndroidRuntimeSettings.cpp b/Engine/Source/Runtime/Android/AndroidRuntimeSettings/Private/AndroidRuntimeSettings.cpp
--- a/Engine/Source/Runtime/Android/AndroidRuntimeSettings/Private/AndroidRuntimeSettings.cpp
+++ b/Engine/Source/Runtime/Android/AndroidRuntimeSettings/Private/AndroidRuntimeSettings.cpp
@@ -4,8 +4,8 @@
 
 #include "AndroidRuntimeSettings.h"
 
-UAndroidRuntimeSettings::UAndroidRuntimeSettings(const class FPostConstructInitializeProperties& PCIP)
-	: Super(PCIP)
+UAndroidRuntimeSettings::UAndroidRuntimeSettings(const FObjectInitializer& ObjectInitializer)
+	: Super(ObjectInitializer)
 	, Orientation(EAndroidScreenOrientation::Landscape)
 	, bEnableGooglePlaySupport(false)
 {
diff --git a/Engine/Source/Runtime/CoreUObject/Classes/Object.h b/Engine/Source/Runtime/CoreUObject/Classes/Object.h
--- a/Engine/Source/Runtime/CoreUObject/Classes/Object.h
+++ b/Engine/Source/Runtime/CoreUObject/Classes/Object.h
@@ -843,7 +843,7 @@ public:
 	/**
 	 * Default UObject constructor.
 	 */
-	UObject(const class FPostConstructInitializeProperties& PCIP);
+	UObject(const FObjectInitializer& ObjectInitializer);
 
 	//=============================================================================
 	// K2 support functions.
diff --git a/Engine/Source/Runtime/CoreUObject/Private/Blueprint/BlueprintSupport.cpp b/Engine/Source/Runtime/CoreUObject/Private/Blueprint/BlueprintSupport.cpp
--- a/Engine/Source/Runtime/CoreUObject/Private/Blueprint/BlueprintSupport.cpp
+++ b/Engine/Source/Runtime/CoreUObject/Private/Blueprint/BlueprintSupport.cpp
@@ -192,7 +192,7 @@ void UObject::DestroyNonNativeProperties()
  * @param	Data				Default data
  * @return	Returns true if that property was a non-native one, otherwise false
  */
-bool FPostConstructInitializeProperties::InitNonNativeProperty(UProperty* Property, UObject* Data)
+bool FObjectInitializer::InitNonNativeProperty(UProperty* Property, UObject* Data)
 {
 	if (!Property->GetOwnerClass()->HasAnyClassFlags(CLASS_Native | CLASS_Intrinsic)) // if this property belongs to a native class, it was already initialized by the class constructor
 	{
diff --git a/Engine/Source/Runtime/CoreUObject/Private/Misc/TextBuffer.cpp b/Engine/Source/Runtime/CoreUObject/Private/Misc/TextBuffer.cpp
--- a/Engine/Source/Runtime/CoreUObject/Private/Misc/TextBuffer.cpp
+++ b/Engine/Source/Runtime/CoreUObject/Private/Misc/TextBuffer.cpp
@@ -13,8 +13,8 @@ IMPLEMENT_CORE_INTRINSIC_CLASS(UTextBuffer, UObject, { });
 /* UTextBuffer constructors
  *****************************************************************************/
 
-UTextBuffer::UTextBuffer (const class FPostConstructInitializeProperties& PCIP, const TCHAR* InText)
-	: UObject(PCIP)
+UTextBuffer::UTextBuffer (const FObjectInitializer& ObjectInitializer, const TCHAR* InText)
+	: UObject(ObjectInitializer)
 	, Text(InText)
 {}
 
diff --git a/Engine/Source/Runtime/CoreUObject/Private/UObject/Class.cpp b/Engine/Source/Runtime/CoreUObject/Private/UObject/Class.cpp
--- a/Engine/Source/Runtime/CoreUObject/Private/UObject/Class.cpp
+++ b/Engine/Source/Runtime/CoreUObject/Private/UObject/Class.cpp
@@ -408,8 +408,8 @@ UStruct::UStruct( EStaticConstructor, int32 InSize, EObjectFlags InFlags )
 {
 }
 
-UStruct::UStruct(const class FPostConstructInitializeProperties& PCIP, UStruct* InSuperStruct, SIZE_T ParamsSize, SIZE_T Alignment )
-:	UField			(PCIP)
+UStruct::UStruct(const FObjectInitializer& ObjectInitializer, UStruct* InSuperStruct, SIZE_T ParamsSize, SIZE_T Alignment )
+:	UField			(ObjectInitializer)
 ,   SuperStruct		( InSuperStruct )
 ,	Children		( NULL )
 ,	PropertiesSize	( ParamsSize ? ParamsSize : (InSuperStruct ? InSuperStruct->GetPropertiesSize() : 0) )
@@ -1812,8 +1812,8 @@ UScriptStruct::UScriptStruct( EStaticConstructor, int32 InSize, EObjectFlags InF
 {
 }
 
-UScriptStruct::UScriptStruct(const class FPostConstructInitializeProperties& PCIP, UScriptStruct* InSuperStruct, ICppStructOps* InCppStructOps, EStructFlags InStructFlags, SIZE_T ExplicitSize, SIZE_T ExplicitAlignment )
-	: UStruct(PCIP, InSuperStruct, InCppStructOps ? InCppStructOps->GetSize() : ExplicitSize, InCppStructOps ? InCppStructOps->GetAlignment() : ExplicitAlignment )
+UScriptStruct::UScriptStruct(const FObjectInitializer& ObjectInitializer, UScriptStruct* InSuperStruct, ICppStructOps* InCppStructOps, EStructFlags InStructFlags, SIZE_T ExplicitSize, SIZE_T ExplicitAlignment )
+	: UStruct(ObjectInitializer, InSuperStruct, InCppStructOps ? InCppStructOps->GetSize() : ExplicitSize, InCppStructOps ? InCppStructOps->GetAlignment() : ExplicitAlignment )
 	, StructFlags(EStructFlags(InStructFlags | (InCppStructOps ? STRUCT_Native : STRUCT_NoFlags)))
 #if HACK_HEADER_GENERATOR
 	, StructMacroDeclaredLineNumber(INDEX_NONE)
@@ -1825,8 +1825,8 @@ UScriptStruct::UScriptStruct(const class FPostConstructInitializeProperties& PCI
 	PrepareCppStructOps(); // propgate flags, etc
 }
 
-UScriptStruct::UScriptStruct(const class FPostConstructInitializeProperties& PCIP)
-	: UStruct(PCIP)
+UScriptStruct::UScriptStruct(const FObjectInitializer& ObjectInitializer)
+	: UStruct(ObjectInitializer)
 	, StructFlags(STRUCT_NoFlags)
 #if HACK_HEADER_GENERATOR
 	, StructMacroDeclaredLineNumber(INDEX_NONE)
@@ -2414,7 +2414,7 @@ class FRestoreClassInfo: public FRestoreForUObjectOverwrite
 	/** Saved ClassCastFlags **/
 	EClassCastFlags	CastFlags;
 	/** Saved ClassConstructor **/
-	void			(*Constructor)(const class FPostConstructInitializeProperties&);
+	void			(*Constructor)(const FObjectInitializer&);
 	/** Saved ClassConstructor **/
 	void			(*AddReferencedObjects)(UObject*, FReferenceCollector&);
 	/** Saved NativeFunctionLookupTable. */
@@ -2513,7 +2513,7 @@ UObject* UClass::CreateDefaultObject()
 				check(ClassDefaultObject);
 				// Blueprint CDOs have their properties always initialized.
 				const bool bShouldInitilizeProperties = !HasAnyClassFlags(CLASS_Native | CLASS_Intrinsic);
-				(*ClassConstructor)(FPostConstructInitializeProperties(ClassDefaultObject, ParentDefaultObject, false, bShouldInitilizeProperties));
+				(*ClassConstructor)(FObjectInitializer(ClassDefaultObject, ParentDefaultObject, false, bShouldInitilizeProperties));
 			}
 		}
 	}
@@ -3257,8 +3257,8 @@ bool UClass::HasProperty(UProperty* InProperty) const
 /**
  * Internal constructor.
  */
-UClass::UClass(const class FPostConstructInitializeProperties& PCIP)
-:	UStruct( PCIP )
+UClass::UClass(const FObjectInitializer& ObjectInitializer)
+:	UStruct( ObjectInitializer )
 ,	ClassFlags(0)
 ,	ClassCastFlags(0)
 ,	ClassUnique(0)
@@ -3273,8 +3273,8 @@ UClass::UClass(const class FPostConstructInitializeProperties& PCIP)
 /**
  * Create a new UClass given its superclass.
  */
-UClass::UClass(const class FPostConstructInitializeProperties& PCIP, UClass* InBaseClass )
-:	UStruct( PCIP, InBaseClass )
+UClass::UClass(const FObjectInitializer& ObjectInitializer, UClass* InBaseClass )
+:	UStruct( ObjectInitializer, InBaseClass )
 ,	ClassFlags(0)
 ,	ClassCastFlags(0)
 ,	ClassUnique(0)
@@ -3317,7 +3317,7 @@ UClass::UClass
 	EClassCastFlags	InClassCastFlags,
 	const TCHAR*    InConfigName,
 	EObjectFlags	InFlags,
-	void			(*InClassConstructor)(const class FPostConstructInitializeProperties&),
+	void			(*InClassConstructor)(const FObjectInitializer&),
 	void			(*InClassAddReferencedObjects)(UObject*, class FReferenceCollector&)
 )
 :	UStruct					( EC_StaticConstructor, InSize, InFlags )
@@ -3344,7 +3344,7 @@ bool UClass::HotReloadPrivateStaticClass(
 	uint32			InClassFlags,
 	EClassCastFlags	InClassCastFlags,
 	const TCHAR*    InConfigName,
-	void			(*InClassConstructor)(const class FPostConstructInitializeProperties&),
+	void			(*InClassConstructor)(const FObjectInitializer&),
 	void			(*InAddReferencedObjects)(UObject*, class FReferenceCollector&),
 	class UClass* TClass_Super_StaticClass,
 	class UClass* TClass_WithinClass_StaticClass
@@ -3363,7 +3363,7 @@ bool UClass::HotReloadPrivateStaticClass(
 	//@todo safe? ClassFlags = InClassFlags | CLASS_Native;
 	//@todo safe? ClassCastFlags = InClassCastFlags;
 	//@todo safe? ClassConfigName = InConfigName;
-	void(*OldClassConstructor)(const class FPostConstructInitializeProperties&) = ClassConstructor;
+	void(*OldClassConstructor)(const FObjectInitializer&) = ClassConstructor;
 	ClassConstructor = InClassConstructor;
 	ClassAddReferencedObjects = InAddReferencedObjects;
 	/* No recursive ::StaticClass calls allowed. Setup extras. */
@@ -3644,8 +3644,8 @@ IMPLEMENT_CORE_INTRINSIC_CLASS(UClass, UStruct,
 	UFunction.
 -----------------------------------------------------------------------------*/
 
-UFunction::UFunction(const class FPostConstructInitializeProperties& PCIP, UFunction* InSuperFunction, uint32 InFunctionFlags, uint16 InRepOffset, SIZE_T ParamsSize )
-: UStruct( PCIP, InSuperFunction, ParamsSize )
+UFunction::UFunction(const FObjectInitializer& ObjectInitializer, UFunction* InSuperFunction, uint32 InFunctionFlags, uint16 InRepOffset, SIZE_T ParamsSize )
+: UStruct( ObjectInitializer, InSuperFunction, ParamsSize )
 , FunctionFlags(InFunctionFlags)
 , RepOffset(InRepOffset)
 , RPCId(0)
diff --git a/Engine/Source/Runtime/CoreUObject/Private/UObject/Linker.cpp b/Engine/Source/Runtime/CoreUObject/Private/UObject/Linker.cpp
--- a/Engine/Source/Runtime/CoreUObject/Private/UObject/Linker.cpp
+++ b/Engine/Source/Runtime/CoreUObject/Private/UObject/Linker.cpp
@@ -109,8 +109,8 @@ FName FLinkerTables::GetExportClassName( int32 i )
 /*----------------------------------------------------------------------------
 	ULinker.
 ----------------------------------------------------------------------------*/
-ULinker::ULinker(const class FPostConstructInitializeProperties& PCIP, UPackage* InRoot, const TCHAR* InFilename )
-:	UObject(PCIP)
+ULinker::ULinker(const FObjectInitializer& ObjectInitializer, UPackage* InRoot, const TCHAR* InFilename )
+:	UObject(ObjectInitializer)
 ,	LinkerRoot( InRoot )
 ,	Summary()
 ,	Filename( InFilename )
diff --git a/Engine/Source/Runtime/CoreUObject/Private/UObject/LinkerLoad.cpp b/Engine/Source/Runtime/CoreUObject/Private/UObject/LinkerLoad.cpp
--- a/Engine/Source/Runtime/CoreUObject/Private/UObject/LinkerLoad.cpp
+++ b/Engine/Source/Runtime/CoreUObject/Private/UObject/LinkerLoad.cpp
@@ -464,7 +464,7 @@ ULinkerLoad* ULinkerLoad::CreateLinkerAsync( UPackage* Parent, const TCHAR* File
 		{
 			LoadFlags |= LOAD_SeekFree;
 		}
-		Linker = new ULinkerLoad( FPostConstructInitializeProperties(), Parent, Filename, LoadFlags );
+		Linker = new ULinkerLoad( FObjectInitializer(), Parent, Filename, LoadFlags );
 		// Add to the list of linkers that haven't been finalized yet
 		GObjPendingLoaders.Add(Parent, Linker);
 	}
@@ -593,8 +593,8 @@ ULinkerLoad::ELinkerStatus ULinkerLoad::Tick( float InTimeLimit, bool bInUseTime
  * @param	Filename	Name of file on disk to load
  * @param	LoadFlags	Load flags determining behavior
  */
-ULinkerLoad::ULinkerLoad( const class FPostConstructInitializeProperties& PCIP, UPackage* InParent, const TCHAR* InFilename, uint32 InLoadFlags )
-:	ULinker( PCIP, InParent, InFilename )
+ULinkerLoad::ULinkerLoad( const class FObjectInitializer& ObjectInitializer, UPackage* InParent, const TCHAR* InFilename, uint32 InLoadFlags )
+:	ULinker( ObjectInitializer, InParent, InFilename )
 ,	LoadFlags( InLoadFlags )
 ,	bHaveImportsBeenVerified( false )
 #if WITH_EDITOR
diff --git a/Engine/Source/Runtime/CoreUObject/Private/UObject/LinkerSave.cpp b/Engine/Source/Runtime/CoreUObject/Private/UObject/LinkerSave.cpp
--- a/Engine/Source/Runtime/CoreUObject/Private/UObject/LinkerSave.cpp
+++ b/Engine/Source/Runtime/CoreUObject/Private/UObject/LinkerSave.cpp
@@ -9,8 +9,8 @@
 /** A mapping of package name to generated script SHA keys */
 TMap<FString, TArray<uint8> > ULinkerSave::PackagesToScriptSHAMap;
 
-ULinkerSave::ULinkerSave( const class FPostConstructInitializeProperties& PCIP, UPackage* InParent, const TCHAR* InFilename, bool bForceByteSwapping, bool bInSaveUnversioned )
-:	ULinker(PCIP, InParent, InFilename )
+ULinkerSave::ULinkerSave( const FObjectInitializer& ObjectInitializer, UPackage* InParent, const TCHAR* InFilename, bool bForceByteSwapping, bool bInSaveUnversioned )
+:	ULinker(ObjectInitializer, InParent, InFilename )
 ,	Saver( NULL )
 {
 	check(!HasAnyFlags(RF_ClassDefaultObject));
@@ -45,8 +45,8 @@ ULinkerSave::ULinkerSave( const class FPostConstructInitializeProperties& PCIP,
 	}
 }
 
-ULinkerSave::ULinkerSave(const class FPostConstructInitializeProperties& PCIP, UPackage* InParent, bool bForceByteSwapping, bool bInSaveUnversioned )
-:	ULinker( PCIP, InParent,TEXT("$$Memory$$") )
+ULinkerSave::ULinkerSave(const FObjectInitializer& ObjectInitializer, UPackage* InParent, bool bForceByteSwapping, bool bInSaveUnversioned )
+:	ULinker( ObjectInitializer, InParent,TEXT("$$Memory$$") )
 ,	Saver( NULL )
 {
 	check(!HasAnyFlags(RF_ClassDefaultObject));
diff --git a/Engine/Source/Runtime/CoreUObject/Private/UObject/Obj.cpp b/Engine/Source/Runtime/CoreUObject/Private/UObject/Obj.cpp
--- a/Engine/Source/Runtime/CoreUObject/Private/UObject/Obj.cpp
+++ b/Engine/Source/Runtime/CoreUObject/Private/UObject/Obj.cpp
@@ -3447,7 +3447,7 @@ void StaticUObjectInit()
 	UObjectBaseInit();
 
 	// Allocate special packages.
-	GObjTransientPkg = new( NULL, TEXT("/Engine/Transient") )UPackage(FPostConstructInitializeProperties());
+	GObjTransientPkg = new( NULL, TEXT("/Engine/Transient") )UPackage(FObjectInitializer());
 	GObjTransientPkg->AddToRoot();
 
 	if( FParse::Param( FCommandLine::Get(), TEXT("VERIFYGC") ) )
diff --git a/Engine/Source/Runtime/CoreUObject/Private/UObject/Property.cpp b/Engine/Source/Runtime/CoreUObject/Private/UObject/Property.cpp
--- a/Engine/Source/Runtime/CoreUObject/Private/UObject/Property.cpp
+++ b/Engine/Source/Runtime/CoreUObject/Private/UObject/Property.cpp
@@ -341,13 +341,13 @@ const TCHAR* UPropertyHelpers::ReadToken( const TCHAR* Buffer, FString& String,
 //
 // Constructors.
 //
-UProperty::UProperty(const class FPostConstructInitializeProperties& PCIP)
-: UField(PCIP)	
+UProperty::UProperty(const FObjectInitializer& ObjectInitializer)
+: UField(ObjectInitializer)	
 , ArrayDim(1)
 {
 }
-UProperty::UProperty(const class FPostConstructInitializeProperties& PCIP, ECppProperty, int32 InOffset, uint64 InFlags )
-: UField(PCIP)	
+UProperty::UProperty(const FObjectInitializer& ObjectInitializer, ECppProperty, int32 InOffset, uint64 InFlags )
+: UField(ObjectInitializer)	
 , ArrayDim(1)
 , PropertyFlags(InFlags)
 , Offset_Internal(InOffset)
diff --git a/Engine/Source/Runtime/CoreUObject/Private/UObject/PropertyBool.cpp b/Engine/Source/Runtime/CoreUObject/Private/UObject/PropertyBool.cpp
--- a/Engine/Source/Runtime/CoreUObject/Private/UObject/PropertyBool.cpp
+++ b/Engine/Source/Runtime/CoreUObject/Private/UObject/PropertyBool.cpp
@@ -7,8 +7,8 @@
 	UBoolProperty.
 -----------------------------------------------------------------------------*/
 
-UBoolProperty::UBoolProperty( const class FPostConstructInitializeProperties& PCIP )
-: UProperty( PCIP )
+UBoolProperty::UBoolProperty( const FObjectInitializer& ObjectInitializer )
+: UProperty( ObjectInitializer )
 , FieldSize(0)
 , ByteOffset(0)
 , ByteMask(1)
@@ -17,8 +17,8 @@ UBoolProperty::UBoolProperty( const class FPostConstructInitializeProperties& PC
 	SetBoolSize( 1, false, 1 );
 }
 
-UBoolProperty::UBoolProperty( const class FPostConstructInitializeProperties& PCIP, ECppProperty, int32 InOffset, uint64 InFlags, uint32 InBitMask, uint32 InElementSize, bool bIsNativeBool )
-: UProperty( PCIP, EC_CppProperty, InOffset, InFlags )
+UBoolProperty::UBoolProperty( const FObjectInitializer& ObjectInitializer, ECppProperty, int32 InOffset, uint64 InFlags, uint32 InBitMask, uint32 InElementSize, bool bIsNativeBool )
+: UProperty( ObjectInitializer, EC_CppProperty, InOffset, InFlags )
 , FieldSize(0)
 , ByteOffset(0)
 , ByteMask(1)
diff --git a/Engine/Source/Runtime/CoreUObject/Private/UObject/PropertyStruct.cpp b/Engine/Source/Runtime/CoreUObject/Private/UObject/PropertyStruct.cpp
--- a/Engine/Source/Runtime/CoreUObject/Private/UObject/PropertyStruct.cpp
+++ b/Engine/Source/Runtime/CoreUObject/Private/UObject/PropertyStruct.cpp
@@ -8,8 +8,8 @@
 	UStructProperty.
 -----------------------------------------------------------------------------*/
 
-UStructProperty::UStructProperty( const class FPostConstructInitializeProperties& PCIP, ECppProperty, int32 InOffset, uint64 InFlags, UScriptStruct* InStruct )
-	:	UProperty( PCIP, EC_CppProperty, InOffset, InFlags )
+UStructProperty::UStructProperty( const FObjectInitializer& ObjectInitializer, ECppProperty, int32 InOffset, uint64 InFlags, UScriptStruct* InStruct )
+	:	UProperty( ObjectInitializer, EC_CppProperty, InOffset, InFlags )
 	,	Struct( InStruct )
 {
 	ElementSize = Struct->PropertiesSize;
diff --git a/Engine/Source/Runtime/CoreUObject/Private/UObject/SavePackage.cpp b/Engine/Source/Runtime/CoreUObject/Private/UObject/SavePackage.cpp
--- a/Engine/Source/Runtime/CoreUObject/Private/UObject/SavePackage.cpp
+++ b/Engine/Source/Runtime/CoreUObject/Private/UObject/SavePackage.cpp
@@ -2608,12 +2608,12 @@ bool UPackage::SavePackage( UPackage* InOuter, UObject* Base, EObjectFlags TopLe
 				if (bCompressFromMemory || bSaveAsync)
 				{
 					// Allocate the linker with a memory writer, forcing byte swapping if wanted.
-					Linker = new ULinkerSave( FPostConstructInitializeProperties(), InOuter, bForceByteSwapping, bSaveUnversioned );
+					Linker = new ULinkerSave( FObjectInitializer(), InOuter, bForceByteSwapping, bSaveUnversioned );
 				}
 				else
 				{
 					// Allocate the linker, forcing byte swapping if wanted.
-					Linker = new ULinkerSave( FPostConstructInitializeProperties(), InOuter, *TempFilename, bForceByteSwapping, bSaveUnversioned );
+					Linker = new ULinkerSave( FObjectInitializer(), InOuter, *TempFilename, bForceByteSwapping, bSaveUnversioned );
 				}
 
 				// Use the custom versions we had previously gleaned from the export tag pass
diff --git a/Engine/Source/Runtime/CoreUObject/Private/UObject/UObjectGlobals.cpp b/Engine/Source/Runtime/CoreUObject/Private/UObject/UObjectGlobals.cpp
--- a/Engine/Source/Runtime/CoreUObject/Private/UObject/UObjectGlobals.cpp
+++ b/Engine/Source/Runtime/CoreUObject/Private/UObject/UObjectGlobals.cpp
@@ -471,7 +471,7 @@ UPackage* CreatePackage( UObject* InOuter, const TCHAR* PackageName )
 			}
 			else
 			{
-				Result = new( InOuter, NewPackageName, RF_Public )UPackage(FPostConstructInitializeProperties());
+				Result = new( InOuter, NewPackageName, RF_Public )UPackage(FObjectInitializer());
 			}
 		}
 	}
@@ -1810,19 +1810,19 @@ void UObject::PostInitProperties()
 #endif
 }
 
-UObject::UObject(const class FPostConstructInitializeProperties& PCIP)
+UObject::UObject(const FObjectInitializer& ObjectInitializer)
 {
-	check(!PCIP.Obj || PCIP.Obj == this);
-	const_cast<FPostConstructInitializeProperties&>(PCIP).Obj = this;
-	const_cast<FPostConstructInitializeProperties&>(PCIP).FinalizeSubobjectClassInitialization();
+	check(!ObjectInitializer.Obj || ObjectInitializer.Obj == this);
+	const_cast<FObjectInitializer&>(ObjectInitializer).Obj = this;
+	const_cast<FObjectInitializer&>(ObjectInitializer).FinalizeSubobjectClassInitialization();
 }
 
 /* Global flag so that FObjectFinders know if they are called from inside the UObject constructors or not. */
 static int32 GIsInConstructor = 0;
-/* Object that is currently being constructed with PCIP */
+/* Object that is currently being constructed with ObjectInitializer */
 static UObject* GConstructedObject = NULL;
 
-FPostConstructInitializeProperties::FPostConstructInitializeProperties() :
+FObjectInitializer::FObjectInitializer() :
 	Obj(NULL),
 	ObjectArchetype(NULL),
 	bCopyTransientsFromClassDefaults(false),
@@ -1836,7 +1836,7 @@ FPostConstructInitializeProperties::FPostConstructInitializeProperties() :
 	GConstructedObject = Obj;
 }	
 
-FPostConstructInitializeProperties::FPostConstructInitializeProperties(UObject* InObj, UObject* InObjectArchetype, bool bInCopyTransientsFromClassDefaults, bool bInShouldIntializeProps, struct FObjectInstancingGraph* InInstanceGraph) :
+FObjectInitializer::FObjectInitializer(UObject* InObj, UObject* InObjectArchetype, bool bInCopyTransientsFromClassDefaults, bool bInShouldIntializeProps, struct FObjectInstancingGraph* InInstanceGraph) :
 	Obj(InObj),
 	ObjectArchetype(InObjectArchetype),
 	// if the SubobjectRoot NULL, then we want to copy the transients from the template, otherwise we are doing a duplicate and we want to copy the transients from the class defaults
@@ -1854,7 +1854,7 @@ FPostConstructInitializeProperties::FPostConstructInitializeProperties(UObject*
 /**
  * Destructor for internal class to finalize UObject creation (initialize properties) after the real C++ constructor is called.
  **/
-FPostConstructInitializeProperties::~FPostConstructInitializeProperties()
+FObjectInitializer::~FObjectInitializer()
 {
 	// Let the FObjectFinders know we left the constructor.
 	GIsInConstructor--;
@@ -1947,7 +1947,7 @@ FPostConstructInitializeProperties::~FPostConstructInitializeProperties()
 				UObject* PropertyValue = ObjProp->GetObjectPropertyValue(ObjProp->ContainerPtrToValuePtr<void>(Obj));
 				if (!FSubobjectPtr::IsInitialized(PropertyValue))
 				{
-					UE_LOG(LogUObjectGlobals, Fatal, TEXT("%s must be initialized in the constructor (at least to NULL) by calling PCIP.CreateDefaultSubobject"), *ObjProp->GetFullName() );
+					UE_LOG(LogUObjectGlobals, Fatal, TEXT("%s must be initialized in the constructor (at least to NULL) by calling ObjectInitializer.CreateDefaultSubobject"), *ObjProp->GetFullName() );
 				}
 				else if (PropertyValue && P->HasAnyPropertyFlags(CPF_Transient))
 				{
@@ -1971,12 +1971,12 @@ FPostConstructInitializeProperties::~FPostConstructInitializeProperties()
 	}
 }
 
-bool FPostConstructInitializeProperties::IsInstancingAllowed() const
+bool FObjectInitializer::IsInstancingAllowed() const
 {
 	return (InstanceGraph == NULL) || InstanceGraph->IsSubobjectInstancingEnabled();
 }
 
-bool FPostConstructInitializeProperties::InitSubobjectProperties(bool bAllowInstancing) const
+bool FObjectInitializer::InitSubobjectProperties(bool bAllowInstancing) const
 {
 	bool bNeedSubobjectInstancing = false;
 	// initialize any subobjects, now that the constructors have run
@@ -1994,7 +1994,7 @@ bool FPostConstructInitializeProperties::InitSubobjectProperties(bool bAllowInst
 	return bNeedSubobjectInstancing;
 }
 
-void FPostConstructInitializeProperties::InstanceSubobjects(UClass* Class, bool bNeedInstancing, bool bNeedSubobjectInstancing) const
+void FObjectInitializer::InstanceSubobjects(UClass* Class, bool bNeedInstancing, bool bNeedSubobjectInstancing) const
 {
 	FObjectInstancingGraph TempInstancingGraph;
 	FObjectInstancingGraph* UseInstancingGraph = InstanceGraph ? InstanceGraph : &TempInstancingGraph;
@@ -2024,7 +2024,7 @@ void FPostConstructInitializeProperties::InstanceSubobjects(UClass* Class, bool
 	}
 }
 
-UClass* FPostConstructInitializeProperties::GetClass() const
+UClass* FObjectInitializer::GetClass() const
 {
 	return Obj->GetClass();
 }
@@ -2042,14 +2042,14 @@ void FSubobjectPtr::Set(UObject* InObject)
 }
 
 // Binary initialize object properties to zero or defaults.
-void FPostConstructInitializeProperties::InitProperties(UObject* Obj, UClass* DefaultsClass, UObject* DefaultData, bool bCopyTransientsFromClassDefaults)
+void FObjectInitializer::InitProperties(UObject* Obj, UClass* DefaultsClass, UObject* DefaultData, bool bCopyTransientsFromClassDefaults)
 {
 	SCOPE_CYCLE_COUNTER(STAT_InitProperties);
 
 	check(DefaultsClass && Obj);
 
 	UClass* Class = Obj->GetClass();
-	// bool to indicate that we need to initialize any non-native properties (native ones were done when the native constructor was called by the code that created and passed in a FPostConstructInitializeProperties object)
+	// bool to indicate that we need to initialize any non-native properties (native ones were done when the native constructor was called by the code that created and passed in a FObjectInitializer object)
 	bool bNeedInitialize = !Class->HasAnyClassFlags(CLASS_Native | CLASS_Intrinsic);
 
 	if (Obj->HasAnyFlags(RF_NeedLoad))
@@ -2105,7 +2105,7 @@ void FPostConstructInitializeProperties::InitProperties(UObject* Obj, UClass* De
 	}
 }
 
-bool FPostConstructInitializeProperties::IslegalOverride(FName InComponentName, class UClass *DerivedComponentClass, class UClass *BaseComponentClass) const
+bool FObjectInitializer::IslegalOverride(FName InComponentName, class UClass *DerivedComponentClass, class UClass *BaseComponentClass) const
 {
 	if (DerivedComponentClass && BaseComponentClass && !DerivedComponentClass->IsChildOf(BaseComponentClass) )
 	{
@@ -2116,7 +2116,7 @@ bool FPostConstructInitializeProperties::IslegalOverride(FName InComponentName,
 	return true;
 }
 
-void FPostConstructInitializeProperties::AssertIfSubobjectSetupIsNotAllowed(const TCHAR* SubobjectName) const
+void FObjectInitializer::AssertIfSubobjectSetupIsNotAllowed(const TCHAR* SubobjectName) const
 {
 	UE_CLOG(!bSubobjectClassInitializationAllowed, LogUObjectGlobals, Fatal,
 		TEXT("%s.%s: Subobject class setup is only allowed in base class constructor call (in the initialization list)"), Obj ? *Obj->GetFullName() : TEXT("NULL"), SubobjectName);
@@ -2152,7 +2152,7 @@ UObject* StaticConstructObject
 	if (!bRecycledSubobject)
 	{		
 		FScopeCycleCounterUObject ConstructorScope(Result);
-		(*InClass->ClassConstructor)( FPostConstructInitializeProperties(Result, InTemplate, bCopyTransientsFromClassDefaults, true, InInstanceGraph) );
+		(*InClass->ClassConstructor)( FObjectInitializer(Result, InTemplate, bCopyTransientsFromClassDefaults, true, InInstanceGraph) );
 	}
 	
 	if( GIsEditor && GUndo && (InFlags & RF_Transactional) && !(InFlags & RF_NeedLoad) && !InClass->IsChildOf(UField::StaticClass()) )
@@ -2165,7 +2165,7 @@ UObject* StaticConstructObject
 	return Result;
 }
 
-void FPostConstructInitializeProperties::AssertIfInConstructor(UObject* Outer, const TCHAR* ErrorMessage)
+void FObjectInitializer::AssertIfInConstructor(UObject* Outer, const TCHAR* ErrorMessage)
 {
 	UE_CLOG(GIsInConstructor && Outer == GConstructedObject, LogUObjectGlobals, Fatal, TEXT("%s"), ErrorMessage);
 }
@@ -2528,7 +2528,7 @@ UScriptStruct* GetFallbackStruct()
 	return FallbackStruct;
 }
 
-UObject* FPostConstructInitializeProperties::CreateDefaultSubobject(UObject* Outer, FName SubobjectFName, UClass* ReturnType, UClass* ClassToCreateByDefault, bool bIsRequired, bool bAbstract, bool bIsTransient) const
+UObject* FObjectInitializer::CreateDefaultSubobject(UObject* Outer, FName SubobjectFName, UClass* ReturnType, UClass* ClassToCreateByDefault, bool bIsRequired, bool bAbstract, bool bIsTransient) const
 {
 	if (SubobjectFName == NAME_None)
 	{
diff --git a/Engine/Source/Runtime/CoreUObject/Public/Misc/TextBuffer.h b/Engine/Source/Runtime/CoreUObject/Public/Misc/TextBuffer.h
--- a/Engine/Source/Runtime/CoreUObject/Public/Misc/TextBuffer.h
+++ b/Engine/Source/Runtime/CoreUObject/Public/Misc/TextBuffer.h
@@ -25,10 +25,10 @@ public:
 	/**
 	 * Creates and initializes a new text buffer.
 	 *
-	 * @param PCIP - Initialization properties.
+	 * @param ObjectInitializer - Initialization properties.
 	 * @param InText - The initial text.
 	 */
-	COREUOBJECT_API UTextBuffer (const class FPostConstructInitializeProperties& PCIP, const TCHAR* InText);
+	COREUOBJECT_API UTextBuffer (const FObjectInitializer& ObjectInitializer, const TCHAR* InText);
 
 
 public:
diff --git a/Engine/Source/Runtime/CoreUObject/Public/UObject/Class.h b/Engine/Source/Runtime/CoreUObject/Public/UObject/Class.h
--- a/Engine/Source/Runtime/CoreUObject/Public/UObject/Class.h
+++ b/Engine/Source/Runtime/CoreUObject/Public/UObject/Class.h
@@ -256,7 +256,7 @@ public:
 public:
 	// Constructors.
 	UStruct( EStaticConstructor, int32 InSize, EObjectFlags InFlags );
-	explicit UStruct(const class FPostConstructInitializeProperties& PCIP, UStruct* InSuperStruct, SIZE_T ParamsSize = 0, SIZE_T Alignment = 0 );
+	explicit UStruct(const FObjectInitializer& ObjectInitializer, UStruct* InSuperStruct, SIZE_T ParamsSize = 0, SIZE_T Alignment = 0 );
 
 	// UObject interface.
 	virtual void Serialize(FArchive& Ar) override;
@@ -953,8 +953,8 @@ public:
 	DECLARE_CASTED_CLASS_INTRINSIC_NO_CTOR(UScriptStruct,UStruct,0,CoreUObject,CASTCLASS_UScriptStruct,COREUOBJECT_API)
 
 	COREUOBJECT_API UScriptStruct( EStaticConstructor, int32 InSize, EObjectFlags InFlags );
-	COREUOBJECT_API explicit UScriptStruct(const class FPostConstructInitializeProperties& PCIP, UScriptStruct* InSuperStruct, ICppStructOps* InCppStructOps = NULL, EStructFlags InStructFlags = STRUCT_NoFlags, SIZE_T ExplicitSize = 0, SIZE_T ExplicitAlignment = 0);
-	COREUOBJECT_API explicit UScriptStruct(const class FPostConstructInitializeProperties& PCIP);
+	COREUOBJECT_API explicit UScriptStruct(const FObjectInitializer& ObjectInitializer, UScriptStruct* InSuperStruct, ICppStructOps* InCppStructOps = NULL, EStructFlags InStructFlags = STRUCT_NoFlags, SIZE_T ExplicitSize = 0, SIZE_T ExplicitAlignment = 0);
+	COREUOBJECT_API explicit UScriptStruct(const FObjectInitializer& ObjectInitializer);
 
 public:
 	EStructFlags StructFlags;
@@ -1195,7 +1195,7 @@ public:
 	void Invoke(UObject* Obj, FFrame& Stack, RESULT_DECL);
 
 	// Constructors.
-	explicit UFunction(const class FPostConstructInitializeProperties& PCIP, UFunction* InSuperFunction, uint32 InFunctionFlags = 0, uint16 InRepOffset = 0, SIZE_T ParamsSize = 0 );
+	explicit UFunction(const FObjectInitializer& ObjectInitializer, UFunction* InSuperFunction, uint32 InFunctionFlags = 0, uint16 InRepOffset = 0, SIZE_T ParamsSize = 0 );
 
 	void InitializeDerivedMembers();
 
@@ -1716,7 +1716,7 @@ public:
 
 	friend class FRestoreClassInfo;
 
-	void(*ClassConstructor)(const class FPostConstructInitializeProperties&);
+	void(*ClassConstructor)(const FObjectInitializer&);
 	/** Pointer to a static AddReferencedObjects method. */
 	void(*ClassAddReferencedObjects)(UObject*, class FReferenceCollector&);
 
@@ -1817,10 +1817,10 @@ public:
 
 public:
 	// Constructors
-	UClass(const class FPostConstructInitializeProperties& PCIP);
-	explicit UClass(const class FPostConstructInitializeProperties& PCIP, UClass* InSuperClass);
+	UClass(const FObjectInitializer& ObjectInitializer);
+	explicit UClass(const FObjectInitializer& ObjectInitializer, UClass* InSuperClass);
 	UClass( EStaticConstructor, FName InName, uint32 InSize, uint32 InClassFlags, EClassCastFlags InClassCastFlags,
-		const TCHAR* InClassConfigName, EObjectFlags InFlags, void(*InClassConstructor)(const class FPostConstructInitializeProperties&),
+		const TCHAR* InClassConfigName, EObjectFlags InFlags, void(*InClassConstructor)(const FObjectInitializer&),
 		void(*InClassAddReferencedObjects)(UObject*, class FReferenceCollector&));
 
 #if WITH_HOT_RELOAD
@@ -1839,7 +1839,7 @@ public:
 		uint32			InClassFlags,
 		EClassCastFlags	InClassCastFlags,
 		const TCHAR*    InConfigName,
-		void			(*InClassConstructor)(const class FPostConstructInitializeProperties&),
+		void			(*InClassConstructor)(const FObjectInitializer&),
 		void			(*InAddReferencedObjects)(UObject*, class FReferenceCollector&),
 		class UClass* TClass_Super_StaticClass,
 		class UClass* TClass_WithinClass_StaticClass
@@ -2186,7 +2186,7 @@ protected:
  * Helper template to call the default constructor for a class
  */
 template<class T>
-void InternalConstructor( const class FPostConstructInitializeProperties& X )
+void InternalConstructor( const FObjectInitializer& X )
 { 
 	T::__DefaultConstructor(X);
 }
@@ -2226,7 +2226,7 @@ void GetPrivateStaticClassBody( const TCHAR* PackageName, const TCHAR* Name, UCl
 				TClass::StaticClassFlags,
 				TClass::StaticClassCastFlags(),
 				TClass::StaticConfigName(),
-				(void(*)(const class FPostConstructInitializeProperties&))InternalConstructor<TClass>,
+				(void(*)(const FObjectInitializer&))InternalConstructor<TClass>,
 				&TClass::AddReferencedObjects,
 				TClass::Super::StaticClass(),
 				TClass::WithinClass::StaticClass()
@@ -2254,7 +2254,7 @@ void GetPrivateStaticClassBody( const TCHAR* PackageName, const TCHAR* Name, UCl
 		TClass::StaticClassCastFlags(),
 		TClass::StaticConfigName(),
 		EObjectFlags(RF_Public | RF_Standalone | RF_Transient | RF_Native | RF_RootSet),
-		(void(*)(const class FPostConstructInitializeProperties&))InternalConstructor<TClass>,
+		(void(*)(const FObjectInitializer&))InternalConstructor<TClass>,
 		&TClass::AddReferencedObjects
 		);
 	check(ReturnClass);
diff --git a/Engine/Source/Runtime/CoreUObject/Public/UObject/GCObject.h b/Engine/Source/Runtime/CoreUObject/Public/UObject/GCObject.h
--- a/Engine/Source/Runtime/CoreUObject/Public/UObject/GCObject.h
+++ b/Engine/Source/Runtime/CoreUObject/Public/UObject/GCObject.h
@@ -80,7 +80,7 @@ public:
 	{
 		if (GGCObjectReferencer == NULL)
 		{
-			GGCObjectReferencer = new UGCObjectReferencer(FPostConstructInitializeProperties());
+			GGCObjectReferencer = new UGCObjectReferencer(FObjectInitializer());
 			GGCObjectReferencer->AddToRoot();
 		}
 	}
diff --git a/Engine/Source/Runtime/CoreUObject/Public/UObject/Linker.h b/Engine/Source/Runtime/CoreUObject/Public/UObject/Linker.h
--- a/Engine/Source/Runtime/CoreUObject/Public/UObject/Linker.h
+++ b/Engine/Source/Runtime/CoreUObject/Public/UObject/Linker.h
@@ -957,7 +957,7 @@ class ULinker : public UObject, public FLinkerTables
 	class FSHA1*			ScriptSHA;
 
 	/** Constructor. */
-	ULinker(const class FPostConstructInitializeProperties& PCIP, UPackage* InRoot, const TCHAR* InFilename );
+	ULinker(const FObjectInitializer& ObjectInitializer, UPackage* InRoot, const TCHAR* InFilename );
 	/**
 	 * I/O function
 	 * 
@@ -1744,7 +1744,7 @@ private:
 	 * @param	Filename	Name of file on disk to load
 	 * @param	LoadFlags	Load flags determining behavior
 	 */
-	ULinkerLoad( const class FPostConstructInitializeProperties& PCIP, UPackage* InParent, const TCHAR* InFilename, uint32 InLoadFlags );
+	ULinkerLoad( const FObjectInitializer& ObjectInitializer, UPackage* InParent, const TCHAR* InFilename, uint32 InLoadFlags );
 
 	/**
 	 * Returns whether the time limit allotted has been exceeded, if enabled.
@@ -1878,9 +1878,9 @@ class ULinkerSave : public ULinker, public FArchiveUObject
 	COREUOBJECT_API static TMap<FString, TArray<uint8> > PackagesToScriptSHAMap;
 
 	/** Constructor */
-	ULinkerSave( const class FPostConstructInitializeProperties& PCIP, UPackage* InParent, const TCHAR* InFilename, bool bForceByteSwapping, bool bInSaveUnversioned = false );
+	ULinkerSave( const FObjectInitializer& ObjectInitializer, UPackage* InParent, const TCHAR* InFilename, bool bForceByteSwapping, bool bInSaveUnversioned = false );
 	/** Constructor for memory writer */
-	ULinkerSave( const class FPostConstructInitializeProperties& PCIP, UPackage* InParent, bool bForceByteSwapping, bool bInSaveUnversioned = false );
+	ULinkerSave( const FObjectInitializer& ObjectInitializer, UPackage* InParent, bool bForceByteSwapping, bool bInSaveUnversioned = false );
 	void BeginDestroy();
 
 	/** Returns the appropriate name index for the source name, or 0 if not found in NameIndices */
diff --git a/Engine/Source/Runtime/CoreUObject/Public/UObject/ObjectBase.h b/Engine/Source/Runtime/CoreUObject/Public/UObject/ObjectBase.h
--- a/Engine/Source/Runtime/CoreUObject/Public/UObject/ObjectBase.h
+++ b/Engine/Source/Runtime/CoreUObject/Public/UObject/ObjectBase.h
@@ -1122,7 +1122,7 @@ namespace UM
 	};
 }
 
-#define RELAY_CONSTRUCTOR(TClass, TSuperClass) TClass(const class FPostConstructInitializeProperties& PCIP) : TSuperClass(PCIP) {}
+#define RELAY_CONSTRUCTOR(TClass, TSuperClass) TClass(const FObjectInitializer& ObjectInitializer) : TSuperClass(ObjectInitializer) {}
 
 #if !USE_COMPILED_IN_NATIVES
 	#define COMPILED_IN_FLAGS(TStaticFlags) (TStaticFlags& ~(CLASS_Intrinsic))
@@ -1175,13 +1175,13 @@ public: \
 	}
 
 #define DEFINE_FORBIDDEN_DEFAULT_CONSTRUCTOR_CALL(TClass) \
-	static_assert(false, "You have to define " #TClass "::" #TClass "() or " #TClass "::" #TClass "(const FPostConstructInitializeProperties&). This is required by UObject system to work correctly.");
+	static_assert(false, "You have to define " #TClass "::" #TClass "() or " #TClass "::" #TClass "(const FObjectInitializer&). This is required by UObject system to work correctly.");
 
 #define DEFINE_DEFAULT_CONSTRUCTOR_CALL(TClass) \
-	static void __DefaultConstructor(const class FPostConstructInitializeProperties& X) { const_cast<FPostConstructInitializeProperties&>(X).FinalizeSubobjectClassInitialization(); new((EInternal*)X.GetObj())TClass(); }
+	static void __DefaultConstructor(const class FObjectInitializer& X) { const_cast<FObjectInitializer&>(X).FinalizeSubobjectClassInitialization(); new((EInternal*)X.GetObj())TClass(); }
 
-#define DEFINE_DEFAULT_PCIP_CONSTRUCTOR_CALL(TClass) \
-	static void __DefaultConstructor(const class FPostConstructInitializeProperties& X) { new((EInternal*)X.GetObj())TClass(X); }
+#define DEFINE_DEFAULT_OBJECT_INITIALIZER_CONSTRUCTOR_CALL(TClass) \
+	static void __DefaultConstructor(const class FObjectInitializer& X) { new((EInternal*)X.GetObj())TClass(X); }
 
 #define DECLARE_CLASS_INTRINSIC(TClass,TSuperClass,TStaticFlags,TPackage) \
 	DECLARE_CLASS(TClass,TSuperClass,TStaticFlags|CLASS_Intrinsic,CASTCLASS_None,TPackage,NO_API ) \
@@ -1189,7 +1189,7 @@ public: \
 	enum {IsIntrinsic=1}; \
 	static void StaticRegisterNatives##TClass() {} \
 	DECLARE_SERIALIZER(TClass) \
-	DEFINE_DEFAULT_PCIP_CONSTRUCTOR_CALL(TClass)
+	DEFINE_DEFAULT_OBJECT_INITIALIZER_CONSTRUCTOR_CALL(TClass)
 
 #define DECLARE_CASTED_CLASS_INTRINSIC_WITH_API( TClass, TSuperClass, TStaticFlags, TPackage, TStaticCastFlags, TRequiredAPI ) \
 	DECLARE_CLASS(TClass,TSuperClass,TStaticFlags|CLASS_Intrinsic,TStaticCastFlags,TPackage,TRequiredAPI ) \
@@ -1197,7 +1197,7 @@ public: \
 	enum {IsIntrinsic=1}; \
 	static void StaticRegisterNatives##TClass() {} \
 	DECLARE_SERIALIZER(TClass) \
-	DEFINE_DEFAULT_PCIP_CONSTRUCTOR_CALL(TClass)
+	DEFINE_DEFAULT_OBJECT_INITIALIZER_CONSTRUCTOR_CALL(TClass)
 
 #define DECLARE_CASTED_CLASS_INTRINSIC( TClass, TSuperClass, TStaticFlags, TPackage, TStaticCastFlags ) \
 	DECLARE_CASTED_CLASS_INTRINSIC_WITH_API( TClass, TSuperClass, TStaticFlags, TPackage, TStaticCastFlags, NO_API) \
@@ -1207,7 +1207,7 @@ public: \
 	enum {IsIntrinsic=1}; \
 	static void StaticRegisterNatives##TClass() {} \
 	DECLARE_SERIALIZER(TClass) \
-	DEFINE_DEFAULT_PCIP_CONSTRUCTOR_CALL(TClass)
+	DEFINE_DEFAULT_OBJECT_INITIALIZER_CONSTRUCTOR_CALL(TClass)
 
 
 
diff --git a/Engine/Source/Runtime/CoreUObject/Public/UObject/UObject.h b/Engine/Source/Runtime/CoreUObject/Public/UObject/UObject.h
--- a/Engine/Source/Runtime/CoreUObject/Public/UObject/UObject.h
+++ b/Engine/Source/Runtime/CoreUObject/Public/UObject/UObject.h
@@ -37,14 +37,14 @@ class COREUOBJECT_API UObject : public UObjectBaseUtility
 {
 	// Declarations.
 	DECLARE_CLASS(UObject,UObject,CLASS_Abstract|CLASS_NoExport,CASTCLASS_None,CoreUObject,NO_API)
-	DEFINE_DEFAULT_PCIP_CONSTRUCTOR_CALL(UObject)
+	DEFINE_DEFAULT_OBJECT_INITIALIZER_CONSTRUCTOR_CALL(UObject)
 
 	typedef UObject WithinClass;
 	static const TCHAR* StaticConfigName() {return TEXT("Engine");}
 
 	// Constructors and destructors.
 	UObject() { };
-	UObject(const class FPostConstructInitializeProperties& PCIP);
+	UObject(const FObjectInitializer& ObjectInitializer);
 	UObject( EStaticConstructor, EObjectFlags InFlags );
 
 	static void StaticRegisterNativesUObject() 
diff --git a/Engine/Source/Runtime/CoreUObject/Public/UObject/UObjectGlobals.h b/Engine/Source/Runtime/CoreUObject/Public/UObject/UObjectGlobals.h
--- a/Engine/Source/Runtime/CoreUObject/Public/UObject/UObjectGlobals.h
+++ b/Engine/Source/Runtime/CoreUObject/Public/UObject/UObjectGlobals.h
@@ -537,18 +537,18 @@ template<> struct TIsWeakPointerType<FSubobjectPtr> { enum { Value = false }; };
 /**
  * TSubobjectPtrConstructor - internal private class to disallow the use of 
  * TSubobjectPtr's assignment operator as well as to disable easy access to
- * TSubobjectPtr::Get() when creating a subobject with PCIP.CreateDefaultSubobject().Get() <- No public 'Get'
+ * TSubobjectPtr::Get() when creating a subobject with ObjectInitializer.CreateDefaultSubobject().Get() <- No public 'Get'
  */
 template <class SubobjectType>
 class TSubobjectPtrConstructor
 {
-	/** Only FPostConstructInitializeProperties can construct this class. */
-	friend class FPostConstructInitializeProperties;
+	/** Only FObjectInitializer can construct this class. */
+	friend class FObjectInitializer;
 	/** TSubobjectPtr needs to have access to the subobject pointer. */
 	template <class AnySubobjectType> friend class TSubobjectPtr;
 	/** Subobject pointer. */
 	UObject* Object;	
-	/** Constructor used by PCIP. */
+	/** Constructor used by ObjectInitializer. */
 	TSubobjectPtrConstructor(SubobjectType* InObject)
 		: Object(InObject)
 	{}
@@ -567,14 +567,14 @@ class TSubobjectPtrConstructor
  *   UPROPERTY()
  *   TSubobjectPtr<UActorComponent> MyComponent;
  *
- * It can only be assigned to with PCIP.CreateDefaultSubobject (via TSubobjectPtrConstructor) in the owning object's constructor.
+ * It can only be assigned to with ObjectInitializer.CreateDefaultSubobject (via TSubobjectPtrConstructor) in the owning object's constructor.
  *
- *   MyComponent = PCIP.CreateDefaultSubobject<UPathFollowingComponent>(this, TEXT("PathFollowingComponent"));
+ *   MyComponent = ObjectInitializer.CreateDefaultSubobject<UPathFollowingComponent>(this, TEXT("PathFollowingComponent"));
  *
  * Initialized with InvalidPtrValue by default because it always needs to be initialized (either with NULL or a pointer to sub-object in the constructor
  * of the owner object).
  * Usually used with Actor components to specify that the actor-derived class is the owner of the component and prevent other derived classes from
- * overwriting it in any other way than through PCIP object.
+ * overwriting it in any other way than through ObjectInitializer object.
  *
  * Implements a structure dereference operator for convenience
  *
@@ -607,7 +607,7 @@ public:
 	{
 		static_assert((CanConvertPointerFromTo<DerivedSubobjectType, SubobjectType>::Result), "Subobject pointers must be compatible.");
 	}
-	/** Initialization constructor. Can only be used with PCIP.CreateDefaultSubobject(). */
+	/** Initialization constructor. Can only be used with ObjectInitializer.CreateDefaultSubobject(). */
 	template <class DerivedSubobjectType>
 	TSubobjectPtr(const TSubobjectPtrConstructor<DerivedSubobjectType>& Other)
 		: FSubobjectPtr(Other.Object)
@@ -619,7 +619,7 @@ public:
 	{
 		return (SubobjectType*)Object;
 	}
-	/** Assignment operator - can only be used with PCIP.CreateDefaultSubobject(). */
+	/** Assignment operator - can only be used with ObjectInitializer.CreateDefaultSubobject(). */
 	template <class DerivedSubobjectType>
 	FORCEINLINE TSubobjectPtr& operator=(const TSubobjectPtrConstructor<DerivedSubobjectType>& Other)
 	{
@@ -646,13 +646,13 @@ template<class T> struct TIsWeakPointerType< TSubobjectPtr<T> > { enum { Value =
 /**
  * Internal class to finalize UObject creation (initialize properties) after the real C++ constructor is called.
  **/
-class COREUOBJECT_API FPostConstructInitializeProperties
+class COREUOBJECT_API FObjectInitializer
 {
 public:
 	/**
 	 * Default Constructor, used when you are using the C++ "new" syntax. UObject::UObject will set the object pointer
 	 **/
-	FPostConstructInitializeProperties();
+	FObjectInitializer();
 
 	/**
 	 * Constructor
@@ -662,9 +662,9 @@ public:
 	 * @param	bInShouldIntializeProps false is a special case for changing base classes in UCCMake
 	 * @param	InInstanceGraph passed instance graph
 	 */
-	FPostConstructInitializeProperties(UObject* InObj, UObject* InObjectArchetype, bool bInCopyTransientsFromClassDefaults, bool bInShouldIntializeProps, struct FObjectInstancingGraph* InInstanceGraph=NULL);
+	FObjectInitializer(UObject* InObj, UObject* InObjectArchetype, bool bInCopyTransientsFromClassDefaults, bool bInShouldIntializeProps, struct FObjectInstancingGraph* InInstanceGraph = NULL);
 
-	~FPostConstructInitializeProperties();
+	~FObjectInitializer();
 
 	/** 
 	 * Return the archetype that this object will copy properties from later
@@ -787,7 +787,7 @@ public:
 	 * @param	SubobjectName	name of the new component or subobject
 	 */
 	template<class T>
-	FPostConstructInitializeProperties const& SetDefaultSubobjectClass(FName SubobjectName) const
+	FObjectInitializer const& SetDefaultSubobjectClass(FName SubobjectName) const
 	{
 		AssertIfSubobjectSetupIsNotAllowed(*SubobjectName.GetPlainNameString());
 		ComponentOverrides.Add(SubobjectName, T::StaticClass(), *this);
@@ -798,7 +798,7 @@ public:
 	 * @param	SubobjectName	name of the new component or subobject
 	 */
 	template<class T>
-	FORCEINLINE FPostConstructInitializeProperties const& SetDefaultSubobjectClass(TCHAR const*SubobjectName) const
+	FORCEINLINE FObjectInitializer const& SetDefaultSubobjectClass(TCHAR const*SubobjectName) const
 	{
 		AssertIfSubobjectSetupIsNotAllowed(SubobjectName);
 		ComponentOverrides.Add(SubobjectName, T::StaticClass(), *this);
@@ -809,7 +809,7 @@ public:
 	 * Indicates that a base class should not create a component
 	 * @param	SubobjectName	name of the new component or subobject to not create
 	 */
-	FPostConstructInitializeProperties const& DoNotCreateDefaultSubobject(FName SubobjectName) const
+	FObjectInitializer const& DoNotCreateDefaultSubobject(FName SubobjectName) const
 	{
 		AssertIfSubobjectSetupIsNotAllowed(*SubobjectName.GetPlainNameString());
 		ComponentOverrides.Add(SubobjectName, NULL, *this);
@@ -820,7 +820,7 @@ public:
 	 * Indicates that a base class should not create a component
 	 * @param	ComponentName	name of the new component or subobject to not create
 	 */
-	FORCEINLINE FPostConstructInitializeProperties const& DoNotCreateDefaultSubobject(TCHAR const*SubobjectName) const
+	FORCEINLINE FObjectInitializer const& DoNotCreateDefaultSubobject(TCHAR const*SubobjectName) const
 	{
 		AssertIfSubobjectSetupIsNotAllowed(SubobjectName);
 		ComponentOverrides.Add(SubobjectName, NULL, *this);
@@ -848,7 +848,7 @@ private:
 	friend class FScriptIntegrationObjectHelper;
 
 	template<class T>
-	friend void InternalConstructor( const class FPostConstructInitializeProperties& X );
+	friend void InternalConstructor(const class FObjectInitializer& X);
 
 	/**
 	 * Binary initialize object properties to zero or defaults.
@@ -863,7 +863,7 @@ private:
 	bool IsInstancingAllowed() const;
 
 	/**
-	 * Calls InitProperties for any default subobjects created through this PCIP.
+	 * Calls InitProperties for any default subobjects created through this ObjectInitializer.
 	 * @param bAllowInstancing	Indicates whether the object's components may be copied from their templates.
 	 * @return true if there are any subobjects which require instancing.
 	*/
@@ -891,7 +891,7 @@ private:
 	struct FOverrides
 	{
 		/**  Add an override, make sure it is legal **/
-		void Add(FName InComponentName, UClass *InComponentClass, FPostConstructInitializeProperties const& PCIP)
+		void Add(FName InComponentName, UClass *InComponentClass, FObjectInitializer const& ObjectInitializer)
 		{
 			int32 Index = Find(InComponentName);
 			if (Index == INDEX_NONE)
@@ -900,11 +900,11 @@ private:
 			}
 			else if (InComponentClass && Overrides[Index].ComponentClass)
 			{
-				PCIP.IslegalOverride(InComponentName, Overrides[Index].ComponentClass, InComponentClass); // if a base class is asking for an override, the existing override (which we are going to use) had better be derived
+				ObjectInitializer.IslegalOverride(InComponentName, Overrides[Index].ComponentClass, InComponentClass); // if a base class is asking for an override, the existing override (which we are going to use) had better be derived
 			}
 		}
 		/**  Retrieve an override, or TClassToConstructByDefault::StaticClass or NULL if this was removed by a derived class **/
-		UClass* Get(FName InComponentName, UClass* ReturnType, UClass* ClassToConstructByDefault, FPostConstructInitializeProperties const& PCIP)
+		UClass* Get(FName InComponentName, UClass* ReturnType, UClass* ClassToConstructByDefault, FObjectInitializer const& ObjectInitializer)
 		{
 			int32 Index = Find(InComponentName);
 			UClass *BaseComponentClass = ClassToConstructByDefault;
@@ -914,7 +914,7 @@ private:
 			}
 			else if (Overrides[Index].ComponentClass)
 			{
-				if (PCIP.IslegalOverride(InComponentName, Overrides[Index].ComponentClass, ReturnType)) // if THE base class is asking for a T, the existing override (which we are going to use) had better be derived
+				if (ObjectInitializer.IslegalOverride(InComponentName, Overrides[Index].ComponentClass, ReturnType)) // if THE base class is asking for a T, the existing override (which we are going to use) had better be derived
 				{
 					return Overrides[Index].ComponentClass; // the override is of an acceptable class, so use it
 				}
@@ -987,7 +987,7 @@ private:
 	bool bCopyTransientsFromClassDefaults;
 	/**  If true, intialize the properties **/
 	bool bShouldIntializePropsFromArchetype;
-	/**  Only true until PCIP has not reached the base UObject class */
+	/**  Only true until ObjectInitializer has not reached the base UObject class */
 	bool bSubobjectClassInitializationAllowed;
 	/**  Instance graph **/
 	struct FObjectInstancingGraph* InstanceGraph;
@@ -999,6 +999,8 @@ private:
 	UObject* LastConstructedObject;
 };
 
+typedef DEPRECATED(4.6, "FPostConstructInitializeProperties was renamed to FObjectInitializer.") FObjectInitializer FPostConstructInitializeProperties;
+
 /**
 * Helper class for script integrations to access some UObject innards. Needed for script-generated UObject classes
 */
@@ -1008,36 +1010,36 @@ public:
 	/**
 	* Binary initialize object properties to zero or defaults.
 	*
-	* @param	PCIP				FPostConstructInitializeProperties helper
+	* @param	ObjectInitializer	FObjectInitializer helper
 	* @param	Obj					object to initialize data for
 	* @param	DefaultsClass		the class to use for initializing the data
 	* @param	DefaultData			the buffer containing the source data for the initialization
 	*/
-	inline static void InitProperties(const FPostConstructInitializeProperties& PCIP, UObject* Obj, UClass* DefaultsClass, UObject* DefaultData)
+	inline static void InitProperties(const FObjectInitializer& ObjectInitializer, UObject* Obj, UClass* DefaultsClass, UObject* DefaultData)
 	{
-		FPostConstructInitializeProperties::InitProperties(Obj, DefaultsClass, DefaultData, PCIP.bCopyTransientsFromClassDefaults);
+		FObjectInitializer::InitProperties(Obj, DefaultsClass, DefaultData, ObjectInitializer.bCopyTransientsFromClassDefaults);
 	}
 
 	/**
-	* Calls InitProperties for any default subobjects created through this PCIP.
+	* Calls InitProperties for any default subobjects created through this ObjectInitializer.
 	* @param bAllowInstancing	Indicates whether the object's components may be copied from their templates.
 	* @return true if there are any subobjects which require instancing.
 	*/
-	inline static bool InitSubobjectProperties(const FPostConstructInitializeProperties& PCIP)
+	inline static bool InitSubobjectProperties(const FObjectInitializer& ObjectInitializer)
 	{
-		return PCIP.InitSubobjectProperties(PCIP.IsInstancingAllowed());
+		return ObjectInitializer.InitSubobjectProperties(ObjectInitializer.IsInstancingAllowed());
 	}
 
 	/**
 	* Create copies of the object's components from their templates.
-	* @param PCIP					FPostConstructInitializeProperties helper
+	* @param ObjectInitializer			FObjectInitializer helper
 	* @param Class						Class of the object we are initializing
 	* @param bNeedInstancing			Indicates whether the object's components need to be instanced
 	* @param bNeedSubobjectInstancing	Indicates whether subobjects of the object's components need to be instanced
 	*/
-	inline static void InstanceSubobjects(const FPostConstructInitializeProperties& PCIP, UClass* Class, bool bNeedInstancing, bool bNeedSubobjectInstancing)
+	inline static void InstanceSubobjects(const FObjectInitializer& ObjectInitializer, UClass* Class, bool bNeedInstancing, bool bNeedSubobjectInstancing)
 	{
-		PCIP.InstanceSubobjects(Class, bNeedInstancing, bNeedSubobjectInstancing);
+		ObjectInitializer.InstanceSubobjects(Class, bNeedInstancing, bNeedSubobjectInstancing);
 	}
 };
 
@@ -1073,7 +1075,7 @@ T* ConstructObject(UClass* Class, UObject* Outer = (UObject*)GetTransientPackage
 template< class T >
 T* NewObject(UObject* Outer=(UObject*)GetTransientPackage(),UClass* Class=T::StaticClass() )
 {
-	FPostConstructInitializeProperties::AssertIfInConstructor(Outer, TEXT("NewObject can't be used to create default subobjects (inside of UObject derived class constructor) as it produces inconsistent object names. Use PCIP.CreateDefaultSuobject<> instead."));
+	FObjectInitializer::AssertIfInConstructor(Outer, TEXT("NewObject can't be used to create default subobjects (inside of UObject derived class constructor) as it produces inconsistent object names. Use ObjectInitializer.CreateDefaultSuobject<> instead."));
 	return ConstructObject<T>(Class,Outer);
 }
 
diff --git a/Engine/Source/Runtime/CoreUObject/Public/UObject/UTextProperty.h b/Engine/Source/Runtime/CoreUObject/Public/UObject/UTextProperty.h
--- a/Engine/Source/Runtime/CoreUObject/Public/UObject/UTextProperty.h
+++ b/Engine/Source/Runtime/CoreUObject/Public/UObject/UTextProperty.h
@@ -15,8 +15,8 @@ public:
 	typedef UTextProperty_Super::TTypeFundamentals TTypeFundamentals;
 	typedef TTypeFundamentals::TCppType TCppType;
 
-	UTextProperty( const class FPostConstructInitializeProperties& PCIP, ECppProperty, int32 InOffset, uint64 InFlags )
-		:	TProperty( PCIP, EC_CppProperty, InOffset, InFlags)
+	UTextProperty( const FObjectInitializer& ObjectInitializer, ECppProperty, int32 InOffset, uint64 InFlags )
+		:	TProperty( ObjectInitializer, EC_CppProperty, InOffset, InFlags)
 	{
 	}
 
diff --git a/Engine/Source/Runtime/CoreUObject/Public/UObject/UnrealType.h b/Engine/Source/Runtime/CoreUObject/Public/UObject/UnrealType.h
--- a/Engine/Source/Runtime/CoreUObject/Public/UObject/UnrealType.h
+++ b/Engine/Source/Runtime/CoreUObject/Public/UObject/UnrealType.h
@@ -160,8 +160,8 @@ public:
 
 public:
 	// Constructors.
-	UProperty(const class FPostConstructInitializeProperties& PCIP);
-	UProperty( const class FPostConstructInitializeProperties& PCIP, ECppProperty, int32 InOffset, uint64 InFlags );
+	UProperty(const FObjectInitializer& ObjectInitializer);
+	UProperty( const FObjectInitializer& ObjectInitializer, ECppProperty, int32 InOffset, uint64 InFlags );
 
 	// UObject interface
 	virtual void Serialize( FArchive& Ar ) override;
@@ -938,13 +938,13 @@ public:
 	typedef TInPropertyBaseClass Super;
 	typedef TPropertyTypeFundamentals<InTCppType> TTypeFundamentals;
 
-	TProperty( const class FPostConstructInitializeProperties& PCIP )
-		:	Super( PCIP)
+	TProperty( const FObjectInitializer& ObjectInitializer )
+		:	Super( ObjectInitializer)
 	{
 		SetElementSize();
 	}
-	TProperty( const class FPostConstructInitializeProperties& PCIP, ECppProperty, int32 InOffset, uint64 InFlags )
-		:	Super( PCIP, EC_CppProperty, InOffset, InFlags | TTypeFundamentals::GetComputedFlagsPropertyFlags())
+	TProperty( const FObjectInitializer& ObjectInitializer, ECppProperty, int32 InOffset, uint64 InFlags )
+		:	Super( ObjectInitializer, EC_CppProperty, InOffset, InFlags | TTypeFundamentals::GetComputedFlagsPropertyFlags())
 	{
 		SetElementSize();
 	}
@@ -1038,12 +1038,12 @@ public:
 	typedef InTCppType TCppType;
 	typedef typename Super::TTypeFundamentals TTypeFundamentals;
 
-	TProperty_WithEqualityAndSerializer( const class FPostConstructInitializeProperties& PCIP )
-		:	Super( PCIP)
+	TProperty_WithEqualityAndSerializer( const FObjectInitializer& ObjectInitializer )
+		:	Super( ObjectInitializer)
 	{
 	}
-	TProperty_WithEqualityAndSerializer( const class FPostConstructInitializeProperties& PCIP, ECppProperty, int32 InOffset, uint64 InFlags )
-		:	Super( PCIP, EC_CppProperty, InOffset, InFlags)
+	TProperty_WithEqualityAndSerializer( const FObjectInitializer& ObjectInitializer, ECppProperty, int32 InOffset, uint64 InFlags )
+		:	Super( ObjectInitializer, EC_CppProperty, InOffset, InFlags)
 	{
 	}
 
@@ -1064,8 +1064,8 @@ class COREUOBJECT_API UNumericProperty : public UProperty
 {
 	DECLARE_CASTED_CLASS_INTRINSIC(UNumericProperty,UProperty,CLASS_Abstract,CoreUObject,CASTCLASS_UNumericProperty)
 
-	UNumericProperty( const class FPostConstructInitializeProperties& PCIP, ECppProperty, int32 InOffset, uint64 InFlags )
-		:	UProperty( PCIP, EC_CppProperty, InOffset, InFlags )
+	UNumericProperty( const FObjectInitializer& ObjectInitializer, ECppProperty, int32 InOffset, uint64 InFlags )
+		:	UProperty( ObjectInitializer, EC_CppProperty, InOffset, InFlags )
 	{}
 
 	// UProperty interface.
@@ -1196,12 +1196,12 @@ public:
 	typedef TProperty_WithEqualityAndSerializer<InTCppType, UNumericProperty> Super;
 	typedef InTCppType TCppType;
 	typedef typename Super::TTypeFundamentals TTypeFundamentals;
-	TProperty_Numeric( const class FPostConstructInitializeProperties& PCIP )
-		:	Super( PCIP)
+	TProperty_Numeric( const FObjectInitializer& ObjectInitializer )
+		:	Super( ObjectInitializer)
 	{
 	}
-	TProperty_Numeric( const class FPostConstructInitializeProperties& PCIP, ECppProperty, int32 InOffset, uint64 InFlags )
-		:	Super( PCIP, EC_CppProperty, InOffset, InFlags)
+	TProperty_Numeric( const FObjectInitializer& ObjectInitializer, ECppProperty, int32 InOffset, uint64 InFlags )
+		:	Super( ObjectInitializer, EC_CppProperty, InOffset, InFlags)
 	{
 	}
 
@@ -1271,8 +1271,8 @@ class COREUOBJECT_API UByteProperty : public TProperty_Numeric<uint8>
 	// Variables.
 	UEnum* Enum;
 
-	UByteProperty( const class FPostConstructInitializeProperties& PCIP, ECppProperty, int32 InOffset, uint64 InFlags, UEnum* InEnum=NULL )
-	:	TProperty_Numeric( PCIP, EC_CppProperty, InOffset, InFlags)
+	UByteProperty( const FObjectInitializer& ObjectInitializer, ECppProperty, int32 InOffset, uint64 InFlags, UEnum* InEnum=NULL )
+	:	TProperty_Numeric( ObjectInitializer, EC_CppProperty, InOffset, InFlags)
 	,	Enum( InEnum )
 	{
 	}
@@ -1313,8 +1313,8 @@ class COREUOBJECT_API UInt8Property : public TProperty_Numeric<int8>
 {
 	DECLARE_CASTED_CLASS_INTRINSIC(UInt8Property,TProperty_Numeric<int8>,0,CoreUObject,CASTCLASS_UInt8Property)
 
-	UInt8Property( const class FPostConstructInitializeProperties& PCIP, ECppProperty, int32 InOffset, uint64 InFlags )
-		:	TProperty_Numeric( PCIP, EC_CppProperty, InOffset, InFlags)
+	UInt8Property( const FObjectInitializer& ObjectInitializer, ECppProperty, int32 InOffset, uint64 InFlags )
+		:	TProperty_Numeric( ObjectInitializer, EC_CppProperty, InOffset, InFlags)
 	{
 	}
 };
@@ -1330,8 +1330,8 @@ class COREUOBJECT_API UInt16Property : public TProperty_Numeric<int16>
 {
 	DECLARE_CASTED_CLASS_INTRINSIC(UInt16Property,TProperty_Numeric<int16>,0,CoreUObject,CASTCLASS_UInt16Property)
 
-	UInt16Property( const class FPostConstructInitializeProperties& PCIP, ECppProperty, int32 InOffset, uint64 InFlags )
-	:	TProperty_Numeric( PCIP, EC_CppProperty, InOffset, InFlags)
+	UInt16Property( const FObjectInitializer& ObjectInitializer, ECppProperty, int32 InOffset, uint64 InFlags )
+	:	TProperty_Numeric( ObjectInitializer, EC_CppProperty, InOffset, InFlags)
 	{
 	}
 };
@@ -1348,8 +1348,8 @@ class COREUOBJECT_API UIntProperty : public TProperty_Numeric<int32>
 {
 	DECLARE_CASTED_CLASS_INTRINSIC(UIntProperty,TProperty_Numeric<int32>,0,CoreUObject,CASTCLASS_UIntProperty)
 
-	UIntProperty( const class FPostConstructInitializeProperties& PCIP, ECppProperty, int32 InOffset, uint64 InFlags )
-	:	TProperty_Numeric( PCIP, EC_CppProperty, InOffset, InFlags)
+	UIntProperty( const FObjectInitializer& ObjectInitializer, ECppProperty, int32 InOffset, uint64 InFlags )
+	:	TProperty_Numeric( ObjectInitializer, EC_CppProperty, InOffset, InFlags)
 	{
 	}
 };
@@ -1365,8 +1365,8 @@ class COREUOBJECT_API UInt64Property : public TProperty_Numeric<int64>
 {
 	DECLARE_CASTED_CLASS_INTRINSIC(UInt64Property,TProperty_Numeric<int64>,0,CoreUObject,CASTCLASS_UInt64Property)
 
-	UInt64Property( const class FPostConstructInitializeProperties& PCIP, ECppProperty, int32 InOffset, uint64 InFlags )
-	:	TProperty_Numeric( PCIP, EC_CppProperty, InOffset, InFlags )
+	UInt64Property( const FObjectInitializer& ObjectInitializer, ECppProperty, int32 InOffset, uint64 InFlags )
+	:	TProperty_Numeric( ObjectInitializer, EC_CppProperty, InOffset, InFlags )
 	{
 	}
 };
@@ -1382,8 +1382,8 @@ class COREUOBJECT_API UUInt16Property : public TProperty_Numeric<uint16>
 {
 	DECLARE_CASTED_CLASS_INTRINSIC(UUInt16Property,TProperty_Numeric<uint16>,0,CoreUObject,CASTCLASS_UUInt16Property)
 
-	UUInt16Property( const class FPostConstructInitializeProperties& PCIP, ECppProperty, int32 InOffset, uint64 InFlags )
-	:	TProperty_Numeric( PCIP, EC_CppProperty, InOffset, InFlags )
+	UUInt16Property( const FObjectInitializer& ObjectInitializer, ECppProperty, int32 InOffset, uint64 InFlags )
+	:	TProperty_Numeric( ObjectInitializer, EC_CppProperty, InOffset, InFlags )
 	{
 	}
 };
@@ -1399,8 +1399,8 @@ class COREUOBJECT_API UUInt32Property : public TProperty_Numeric<uint32>
 {
 	DECLARE_CASTED_CLASS_INTRINSIC(UUInt32Property,TProperty_Numeric<uint32>,0,CoreUObject,CASTCLASS_UUInt32Property)
 
-	UUInt32Property( const class FPostConstructInitializeProperties& PCIP, ECppProperty, int32 InOffset, uint64 InFlags )
-	:	TProperty_Numeric( PCIP, EC_CppProperty, InOffset, InFlags )
+	UUInt32Property( const FObjectInitializer& ObjectInitializer, ECppProperty, int32 InOffset, uint64 InFlags )
+	:	TProperty_Numeric( ObjectInitializer, EC_CppProperty, InOffset, InFlags )
 	{
 	}
 };
@@ -1416,8 +1416,8 @@ class COREUOBJECT_API UUInt64Property : public TProperty_Numeric<uint64>
 {
 	DECLARE_CASTED_CLASS_INTRINSIC(UUInt64Property,TProperty_Numeric<uint64>,0,CoreUObject,CASTCLASS_UUInt64Property)
 
-	UUInt64Property( const class FPostConstructInitializeProperties& PCIP, ECppProperty, int32 InOffset, uint64 InFlags )
-	:	TProperty_Numeric( PCIP, EC_CppProperty, InOffset, InFlags )
+	UUInt64Property( const FObjectInitializer& ObjectInitializer, ECppProperty, int32 InOffset, uint64 InFlags )
+	:	TProperty_Numeric( ObjectInitializer, EC_CppProperty, InOffset, InFlags )
 	{
 	}
 };
@@ -1433,8 +1433,8 @@ class COREUOBJECT_API UFloatProperty : public TProperty_Numeric<float>
 {
 	DECLARE_CASTED_CLASS_INTRINSIC(UFloatProperty,TProperty_Numeric<float>,0,CoreUObject,CASTCLASS_UFloatProperty)
 
-	UFloatProperty( const class FPostConstructInitializeProperties& PCIP, ECppProperty, int32 InOffset, uint64 InFlags )
-		:	TProperty_Numeric( PCIP, EC_CppProperty, InOffset, InFlags )
+	UFloatProperty( const FObjectInitializer& ObjectInitializer, ECppProperty, int32 InOffset, uint64 InFlags )
+		:	TProperty_Numeric( ObjectInitializer, EC_CppProperty, InOffset, InFlags )
 	{
 	}
 };
@@ -1450,8 +1450,8 @@ class COREUOBJECT_API UDoubleProperty : public TProperty_Numeric<double>
 {
 	DECLARE_CASTED_CLASS_INTRINSIC(UDoubleProperty,TProperty_Numeric<double>,0,CoreUObject,CASTCLASS_UDoubleProperty)
 
-	UDoubleProperty( const class FPostConstructInitializeProperties& PCIP, ECppProperty, int32 InOffset, uint64 InFlags )
-		:	TProperty_Numeric( PCIP, EC_CppProperty, InOffset, InFlags )
+	UDoubleProperty( const FObjectInitializer& ObjectInitializer, ECppProperty, int32 InOffset, uint64 InFlags )
+		:	TProperty_Numeric( ObjectInitializer, EC_CppProperty, InOffset, InFlags )
 	{
 	}
 };
@@ -1483,12 +1483,12 @@ private:
 
 public:
 
-	UBoolProperty( const class FPostConstructInitializeProperties& PCIP );
+	UBoolProperty( const FObjectInitializer& ObjectInitializer );
 	
 	/**
 	 * Constructor.
 	 *
-	 * @param PCIP Properties.
+	 * @param ObjectInitializer Properties.
 	 * @param ECppProperty Unused.
 	 * @param InOffset Offset of the property.
 	 * @param InCategory Category of the property.
@@ -1497,7 +1497,7 @@ public:
 	 * @param InElementSize Sizeof of the boolean type this property represents.
 	 * @param bIsNativeBool true if this property represents C++ bool type.
 	 */
-	UBoolProperty( const class FPostConstructInitializeProperties& PCIP, ECppProperty, int32 InOffset, uint64 InFlags, uint32 InBitMask, uint32 InElementSize, bool bIsNativeBool );
+	UBoolProperty( const FObjectInitializer& ObjectInitializer, ECppProperty, int32 InOffset, uint64 InFlags, uint32 InBitMask, uint32 InElementSize, bool bIsNativeBool );
 
 	// UObject interface.
 	virtual void Serialize( FArchive& Ar ) override;
@@ -1590,8 +1590,8 @@ class COREUOBJECT_API UObjectPropertyBase : public UProperty
 	// Variables.
 	class UClass* PropertyClass;
 
-	UObjectPropertyBase( const class FPostConstructInitializeProperties& PCIP, ECppProperty, int32 InOffset, uint64 InFlags, UClass* InClass=NULL )
-	:	UProperty( PCIP, EC_CppProperty, InOffset, InFlags )
+	UObjectPropertyBase( const FObjectInitializer& ObjectInitializer, ECppProperty, int32 InOffset, uint64 InFlags, UClass* InClass=NULL )
+	:	UProperty( ObjectInitializer, EC_CppProperty, InOffset, InFlags )
 	,	PropertyClass( InClass )
 	{}
 
@@ -1728,12 +1728,12 @@ public:
 	typedef InTCppType TCppType;
 	typedef typename Super::TTypeFundamentals TTypeFundamentals;
 
-	TUObjectPropertyBase( const class FPostConstructInitializeProperties& PCIP )
-		:	Super( PCIP)
+	TUObjectPropertyBase( const FObjectInitializer& ObjectInitializer )
+		:	Super( ObjectInitializer)
 	{
 	}
-	TUObjectPropertyBase( const class FPostConstructInitializeProperties& PCIP, ECppProperty, int32 InOffset, uint64 InFlags, UClass* InClass )
-		:	Super( PCIP, EC_CppProperty, InOffset, InFlags)
+	TUObjectPropertyBase( const FObjectInitializer& ObjectInitializer, ECppProperty, int32 InOffset, uint64 InFlags, UClass* InClass )
+		:	Super( ObjectInitializer, EC_CppProperty, InOffset, InFlags)
 	{
 		this->PropertyClass = InClass;
 	}
@@ -1758,8 +1758,8 @@ class COREUOBJECT_API UObjectProperty : public TUObjectPropertyBase<UObject*>
 {
 	DECLARE_CASTED_CLASS_INTRINSIC(UObjectProperty,TUObjectPropertyBase<UObject*>,0,CoreUObject,CASTCLASS_UObjectProperty)
 
-	UObjectProperty( const class FPostConstructInitializeProperties& PCIP, ECppProperty, int32 InOffset, uint64 InFlags, UClass* InClass )
-	:	TUObjectPropertyBase( PCIP, EC_CppProperty, InOffset, InFlags, InClass )
+	UObjectProperty( const FObjectInitializer& ObjectInitializer, ECppProperty, int32 InOffset, uint64 InFlags, UClass* InClass )
+	:	TUObjectPropertyBase( ObjectInitializer, EC_CppProperty, InOffset, InFlags, InClass )
 	{
 	}
 
@@ -1791,8 +1791,8 @@ class COREUOBJECT_API UWeakObjectProperty : public TUObjectPropertyBase<FWeakObj
 {
 	DECLARE_CASTED_CLASS_INTRINSIC(UWeakObjectProperty,TUObjectPropertyBase<FWeakObjectPtr>,0,CoreUObject,CASTCLASS_UWeakObjectProperty)
 
-	UWeakObjectProperty( const class FPostConstructInitializeProperties& PCIP, ECppProperty, int32 InOffset, uint64 InFlags, UClass* InClass )
-	:	TUObjectPropertyBase( PCIP, EC_CppProperty, InOffset, InFlags, InClass )
+	UWeakObjectProperty( const FObjectInitializer& ObjectInitializer, ECppProperty, int32 InOffset, uint64 InFlags, UClass* InClass )
+	:	TUObjectPropertyBase( ObjectInitializer, EC_CppProperty, InOffset, InFlags, InClass )
 	{
 	}
 	
@@ -1824,8 +1824,8 @@ class COREUOBJECT_API ULazyObjectProperty : public TUObjectPropertyBase<FLazyObj
 {
 	DECLARE_CASTED_CLASS_INTRINSIC(ULazyObjectProperty,TUObjectPropertyBase<FLazyObjectPtr>,0,CoreUObject,CASTCLASS_ULazyObjectProperty)
 
-	ULazyObjectProperty( const class FPostConstructInitializeProperties& PCIP, ECppProperty, int32 InOffset, uint64 InFlags, UClass* InClass )
-		:	TUObjectPropertyBase( PCIP, EC_CppProperty, InOffset, InFlags, InClass )
+	ULazyObjectProperty( const FObjectInitializer& ObjectInitializer, ECppProperty, int32 InOffset, uint64 InFlags, UClass* InClass )
+		:	TUObjectPropertyBase( ObjectInitializer, EC_CppProperty, InOffset, InFlags, InClass )
 	{
 	}
 
@@ -1863,8 +1863,8 @@ class COREUOBJECT_API UAssetObjectProperty : public TUObjectPropertyBase<FAssetP
 {
 	DECLARE_CASTED_CLASS_INTRINSIC(UAssetObjectProperty,TUObjectPropertyBase<FAssetPtr>,0,CoreUObject,CASTCLASS_UAssetObjectProperty)
 
-	UAssetObjectProperty( const class FPostConstructInitializeProperties& PCIP, ECppProperty, int32 InOffset, uint64 InFlags, UClass* InClass )
-		:	TUObjectPropertyBase( PCIP, EC_CppProperty, InOffset, InFlags, InClass )
+	UAssetObjectProperty( const FObjectInitializer& ObjectInitializer, ECppProperty, int32 InOffset, uint64 InFlags, UClass* InClass )
+		:	TUObjectPropertyBase( ObjectInitializer, EC_CppProperty, InOffset, InFlags, InClass )
 	{}
 
 	// UHT interface
@@ -1910,8 +1910,8 @@ class COREUOBJECT_API UClassProperty : public UObjectProperty
 	// Variables.
 	class UClass* MetaClass;
 public:
-	UClassProperty( const class FPostConstructInitializeProperties& PCIP, ECppProperty, int32 InOffset, uint64 InFlags, UClass* InMetaClass )
-	:	UObjectProperty( PCIP, EC_CppProperty, InOffset, InFlags, UClass::StaticClass() )
+	UClassProperty( const FObjectInitializer& ObjectInitializer, ECppProperty, int32 InOffset, uint64 InFlags, UClass* InMetaClass )
+	:	UObjectProperty( ObjectInitializer, EC_CppProperty, InOffset, InFlags, UClass::StaticClass() )
 	,	MetaClass( InMetaClass )
 	{
 	}
@@ -1949,8 +1949,8 @@ class COREUOBJECT_API UAssetClassProperty : public UAssetObjectProperty
 	// Variables.
 	class UClass* MetaClass;
 public:
-	UAssetClassProperty( const class FPostConstructInitializeProperties& PCIP, ECppProperty, int32 InOffset, uint64 InFlags, UClass* InMetaClass )
-		:	Super(PCIP, EC_CppProperty, InOffset, InFlags, UClass::StaticClass() )
+	UAssetClassProperty( const FObjectInitializer& ObjectInitializer, ECppProperty, int32 InOffset, uint64 InFlags, UClass* InMetaClass )
+		:	Super(ObjectInitializer, EC_CppProperty, InOffset, InFlags, UClass::StaticClass() )
 		,	MetaClass( InMetaClass )
 	{}
 
@@ -1991,8 +1991,8 @@ public:
 	typedef UInterfaceProperty_Super::TTypeFundamentals TTypeFundamentals;
 	typedef TTypeFundamentals::TCppType TCppType;
 
-	UInterfaceProperty( const class FPostConstructInitializeProperties& PCIP, ECppProperty, int32 InOffset, uint64 InFlags, UClass* InInterfaceClass )
-		:	UInterfaceProperty_Super( PCIP, EC_CppProperty, InOffset, (InFlags & ~CPF_InterfaceClearMask) )
+	UInterfaceProperty( const FObjectInitializer& ObjectInitializer, ECppProperty, int32 InOffset, uint64 InFlags, UClass* InInterfaceClass )
+		:	UInterfaceProperty_Super( ObjectInitializer, EC_CppProperty, InOffset, (InFlags & ~CPF_InterfaceClearMask) )
 		,	InterfaceClass( InInterfaceClass )
 	{
 	}
@@ -2037,8 +2037,8 @@ public:
 	typedef UNameProperty_Super::TTypeFundamentals TTypeFundamentals;
 	typedef TTypeFundamentals::TCppType TCppType;
 
-	UNameProperty( const class FPostConstructInitializeProperties& PCIP, ECppProperty, int32 InOffset, uint64 InFlags )
-	:	UNameProperty_Super( PCIP, EC_CppProperty, InOffset, InFlags )
+	UNameProperty( const FObjectInitializer& ObjectInitializer, ECppProperty, int32 InOffset, uint64 InFlags )
+	:	UNameProperty_Super( ObjectInitializer, EC_CppProperty, InOffset, InFlags )
 	{
 	}
 
@@ -2066,8 +2066,8 @@ public:
 	typedef UStrProperty_Super::TTypeFundamentals TTypeFundamentals;
 	typedef TTypeFundamentals::TCppType TCppType;
 
-	UStrProperty( const class FPostConstructInitializeProperties& PCIP, ECppProperty, int32 InOffset, uint64 InFlags )
-	:	UStrProperty_Super( PCIP, EC_CppProperty, InOffset, InFlags)
+	UStrProperty( const FObjectInitializer& ObjectInitializer, ECppProperty, int32 InOffset, uint64 InFlags )
+	:	UStrProperty_Super( ObjectInitializer, EC_CppProperty, InOffset, InFlags)
 	{
 	}
 
@@ -2099,8 +2099,8 @@ public:
 	typedef UArrayProperty_Super::TTypeFundamentals TTypeFundamentals;
 	typedef TTypeFundamentals::TCppType TCppType;
 
-	UArrayProperty( const class FPostConstructInitializeProperties& PCIP, ECppProperty, int32 InOffset, uint64 InFlags )
-	:	UArrayProperty_Super( PCIP, EC_CppProperty, InOffset, InFlags)
+	UArrayProperty( const FObjectInitializer& ObjectInitializer, ECppProperty, int32 InOffset, uint64 InFlags )
+	:	UArrayProperty_Super( ObjectInitializer, EC_CppProperty, InOffset, InFlags)
 	{
 	}
 
@@ -2471,7 +2471,7 @@ class COREUOBJECT_API UStructProperty : public UProperty
 	// Variables.
 	class UScriptStruct* Struct;
 public:
-	UStructProperty( const class FPostConstructInitializeProperties& PCIP, ECppProperty, int32 InOffset, uint64 InFlags, UScriptStruct* InStruct );
+	UStructProperty( const FObjectInitializer& ObjectInitializer, ECppProperty, int32 InOffset, uint64 InFlags, UScriptStruct* InStruct );
 
 	// UObject interface
 	virtual void Serialize( FArchive& Ar ) override;
@@ -2541,8 +2541,8 @@ public:
 	typedef UDelegateProperty_Super::TTypeFundamentals TTypeFundamentals;
 	typedef TTypeFundamentals::TCppType TCppType;
 
-	UDelegateProperty( const class FPostConstructInitializeProperties& PCIP, ECppProperty, int32 InOffset, uint64 InFlags, UFunction* InSignatureFunction = NULL )
-		: UDelegateProperty_Super( PCIP, EC_CppProperty, InOffset, InFlags)
+	UDelegateProperty( const FObjectInitializer& ObjectInitializer, ECppProperty, int32 InOffset, uint64 InFlags, UFunction* InSignatureFunction = NULL )
+		: UDelegateProperty_Super( ObjectInitializer, EC_CppProperty, InOffset, InFlags)
 		, SignatureFunction(InSignatureFunction)
 	{
 	}
@@ -2586,8 +2586,8 @@ public:
 	typedef UMulticastDelegateProperty_Super::TTypeFundamentals TTypeFundamentals;
 	typedef TTypeFundamentals::TCppType TCppType;
 
-	UMulticastDelegateProperty( const class FPostConstructInitializeProperties& PCIP, ECppProperty, int32 InOffset, uint64 InFlags, UFunction* InSignatureFunction = NULL )
-		: TProperty( PCIP, EC_CppProperty, InOffset, InFlags )
+	UMulticastDelegateProperty( const FObjectInitializer& ObjectInitializer, ECppProperty, int32 InOffset, uint64 InFlags, UFunction* InSignatureFunction = NULL )
+		: TProperty( ObjectInitializer, EC_CppProperty, InOffset, InFlags )
 		, SignatureFunction(InSignatureFunction)
 	{
 	}
@@ -3087,8 +3087,8 @@ inline bool UObject::IsBasedOnArchetype(  const UObject* const SomeObject ) cons
 	C++ property macros.
 -----------------------------------------------------------------------------*/
 
-#define CPP_PROPERTY(name)	FPostConstructInitializeProperties(), EC_CppProperty, STRUCT_OFFSET(ThisClass, name)
-#define CPP_PROPERTY_BASE(name, base)	FPostConstructInitializeProperties(), EC_CppProperty, STRUCT_OFFSET(base, name)
+#define CPP_PROPERTY(name)	FObjectInitializer(), EC_CppProperty, STRUCT_OFFSET(ThisClass, name)
+#define CPP_PROPERTY_BASE(name, base)	FObjectInitializer(), EC_CppProperty, STRUCT_OFFSET(base, name)
 
 /** 
 	The mac does not interpret a pointer to a bool* that is say 0x40 as true!, so we need to use uint8 for that. 
diff --git a/Engine/Source/Runtime/Engine/Classes/AI/Navigation/NavigationTypes.h b/Engine/Source/Runtime/Engine/Classes/AI/Navigation/NavigationTypes.h
--- a/Engine/Source/Runtime/Engine/Classes/AI/Navigation/NavigationTypes.h
+++ b/Engine/Source/Runtime/Engine/Classes/AI/Navigation/NavigationTypes.h
@@ -452,7 +452,7 @@ class UNavigationTypes : public UObject
 {
 	GENERATED_UCLASS_BODY()
 
-	UNavigationTypes(const class FPostConstructInitializeProperties& PCIP) : Super(PCIP) { }
+	UNavigationTypes(const FObjectInitializer& ObjectInitializer) : Super(ObjectInitializer) { }
 };
 
 //////////////////////////////////////////////////////////////////////////
diff --git a/Engine/Source/Runtime/Engine/Classes/Components/ActorComponent.h b/Engine/Source/Runtime/Engine/Classes/Components/ActorComponent.h
--- a/Engine/Source/Runtime/Engine/Classes/Components/ActorComponent.h
+++ b/Engine/Source/Runtime/Engine/Classes/Components/ActorComponent.h
@@ -17,7 +17,7 @@ public:
 	/**
 	 * Default UObject constructor.
 	 */
-	UActorComponent(const class FPostConstructInitializeProperties& PCIP);
+	UActorComponent(const FObjectInitializer& ObjectInitializer);
 
 	/** Main tick function for the Actor */
 	UPROPERTY()
diff --git a/Engine/Source/Runtime/Engine/Classes/Components/ModelComponent.h b/Engine/Source/Runtime/Engine/Classes/Components/ModelComponent.h
--- a/Engine/Source/Runtime/Engine/Classes/Components/ModelComponent.h
+++ b/Engine/Source/Runtime/Engine/Classes/Components/ModelComponent.h
@@ -40,7 +40,7 @@ public:
 	/**
 	 * Minimal initialization constructor.
 	 */
-	UModelComponent( const class FPostConstructInitializeProperties& PCIP,UModel* InModel, uint16 InComponentIndex, uint32 MaskedSurfaceFlags, const TArray<uint16>& InNodes );
+	UModelComponent( const FObjectInitializer& ObjectInitializer,UModel* InModel, uint16 InComponentIndex, uint32 MaskedSurfaceFlags, const TArray<uint16>& InNodes );
 #endif // WITH_EDITOR
 
 	/**
diff --git a/Engine/Source/Runtime/Engine/Classes/Components/PrimitiveComponent.h b/Engine/Source/Runtime/Engine/Classes/Components/PrimitiveComponent.h
--- a/Engine/Source/Runtime/Engine/Classes/Components/PrimitiveComponent.h
+++ b/Engine/Source/Runtime/Engine/Classes/Components/PrimitiveComponent.h
@@ -107,7 +107,7 @@ public:
 	/**
 	 * Default UObject constructor.
 	 */
-	UPrimitiveComponent(const class FPostConstructInitializeProperties& PCIP);
+	UPrimitiveComponent(const FObjectInitializer& ObjectInitializer);
 
 	// Rendering
 	
diff --git a/Engine/Source/Runtime/Engine/Classes/Components/SceneComponent.h b/Engine/Source/Runtime/Engine/Classes/Components/SceneComponent.h
--- a/Engine/Source/Runtime/Engine/Classes/Components/SceneComponent.h
+++ b/Engine/Source/Runtime/Engine/Classes/Components/SceneComponent.h
@@ -85,7 +85,7 @@ public:
 	/**
 	 * Default UObject constructor.
 	 */
-	USceneComponent(const class FPostConstructInitializeProperties& PCIP);
+	USceneComponent(const FObjectInitializer& ObjectInitializer);
 
 	/** Current transform of this component, relative to the world */
 	FTransform ComponentToWorld;
diff --git a/Engine/Source/Runtime/Engine/Classes/Distributions/DistributionFloat.h b/Engine/Source/Runtime/Engine/Classes/Distributions/DistributionFloat.h
--- a/Engine/Source/Runtime/Engine/Classes/Distributions/DistributionFloat.h
+++ b/Engine/Source/Runtime/Engine/Classes/Distributions/DistributionFloat.h
@@ -87,8 +87,8 @@ class UDistributionFloat : public UDistribution
 	virtual float GetFloatValue(float F = 0);
 
 
-	UDistributionFloat(const class FPostConstructInitializeProperties& PCIP)
-	:	Super(PCIP)
+	UDistributionFloat(const FObjectInitializer& ObjectInitializer)
+	:	Super(ObjectInitializer)
 	,   bCanBeBaked(true)
 	,   bIsDirty(true) // make sure the FRawDistribution is initialized
 	{
diff --git a/Engine/Source/Runtime/Engine/Classes/Distributions/DistributionVector.h b/Engine/Source/Runtime/Engine/Classes/Distributions/DistributionVector.h
--- a/Engine/Source/Runtime/Engine/Classes/Distributions/DistributionVector.h
+++ b/Engine/Source/Runtime/Engine/Classes/Distributions/DistributionVector.h
@@ -118,8 +118,8 @@ class ENGINE_API UDistributionVector : public UDistribution
 	virtual FVector GetVectorValue(float F = 0);
 
 
-	UDistributionVector(const class FPostConstructInitializeProperties& PCIP)
-	:	Super(PCIP)
+	UDistributionVector(const FObjectInitializer& ObjectInitializer)
+	:	Super(ObjectInitializer)
 	,   bCanBeBaked(true)
 	,   bIsDirty(true) // make sure the FRawDistribution is initialized
 	{
diff --git a/Engine/Source/Runtime/Engine/Classes/Engine/ActorChannel.h b/Engine/Source/Runtime/Engine/Classes/Engine/ActorChannel.h
--- a/Engine/Source/Runtime/Engine/Classes/Engine/ActorChannel.h
+++ b/Engine/Source/Runtime/Engine/Classes/Engine/ActorChannel.h
@@ -73,8 +73,8 @@ class ENGINE_API UActorChannel : public UChannel
 	/**
 	 * Default constructor
 	 */
-	UActorChannel(const class FPostConstructInitializeProperties& PCIP)
-		: UChannel(PCIP)
+	UActorChannel(const FObjectInitializer& ObjectInitializer)
+		: UChannel(ObjectInitializer)
 	{
 		ChannelClasses[CHTYPE_Actor] = GetClass();
 		ChType = CHTYPE_Actor;
diff --git a/Engine/Source/Runtime/Engine/Classes/Engine/ControlChannel.h b/Engine/Source/Runtime/Engine/Classes/Engine/ControlChannel.h
--- a/Engine/Source/Runtime/Engine/Classes/Engine/ControlChannel.h
+++ b/Engine/Source/Runtime/Engine/Classes/Engine/ControlChannel.h
@@ -47,8 +47,8 @@ class UControlChannel : public UChannel
 	/**
 	 * Default constructor
 	 */
-	UControlChannel(const class FPostConstructInitializeProperties& PCIP)
-		: UChannel(PCIP)
+	UControlChannel(const FObjectInitializer& ObjectInitializer)
+		: UChannel(ObjectInitializer)
 	{
 		ChannelClasses[CHTYPE_Control]      = GetClass();
 		ChType								= CHTYPE_Control;
diff --git a/Engine/Source/Runtime/Engine/Classes/Engine/IntSerialization.h b/Engine/Source/Runtime/Engine/Classes/Engine/IntSerialization.h
--- a/Engine/Source/Runtime/Engine/Classes/Engine/IntSerialization.h
+++ b/Engine/Source/Runtime/Engine/Classes/Engine/IntSerialization.h
@@ -10,7 +10,7 @@ class ENGINE_API UIntSerialization	: public UObject
 	GENERATED_UCLASS_BODY()
 
 public:
-//	ENGINE_API UIntSerialization(const class FPostConstructInitializeProperties& PCIP, uint8 InUnsignedInt8Variable, uint16 InUnsignedInt16Variable, uint32 InUnsignedInt32Variable, uint64 InUnsignedInt64Variable, int8 InSignedInt8Variable, int16 InSignedInt16Variable, int InSignedInt32Variable, int64 InSignedInt64Variable);
+//	ENGINE_API UIntSerialization(const FObjectInitializer& ObjectInitializer, uint8 InUnsignedInt8Variable, uint16 InUnsignedInt16Variable, uint32 InUnsignedInt32Variable, uint64 InUnsignedInt64Variable, int8 InSignedInt8Variable, int16 InSignedInt16Variable, int InSignedInt32Variable, int64 InSignedInt64Variable);
 
 //	virtual void Serialize (FArchive& Ar) override;
 	/**
diff --git a/Engine/Source/Runtime/Engine/Classes/Engine/Level.h b/Engine/Source/Runtime/Engine/Classes/Engine/Level.h
--- a/Engine/Source/Runtime/Engine/Classes/Engine/Level.h
+++ b/Engine/Source/Runtime/Engine/Classes/Engine/Level.h
@@ -455,8 +455,8 @@ public:
 	ENGINE_API static FSimpleMulticastDelegate LevelDirtiedEvent;
 
 	// Constructor.
-	ENGINE_API ULevel(const class FPostConstructInitializeProperties& PCIP, const FURL& InURL );
-	ULevel(const class FPostConstructInitializeProperties& PCIP );
+	ENGINE_API ULevel(const FObjectInitializer& ObjectInitializer, const FURL& InURL );
+	ULevel(const FObjectInitializer& ObjectInitializer );
 	~ULevel();
 
 	// Begin UObject interface.
diff --git a/Engine/Source/Runtime/Engine/Classes/Engine/LevelBase.h b/Engine/Source/Runtime/Engine/Classes/Engine/LevelBase.h
--- a/Engine/Source/Runtime/Engine/Classes/Engine/LevelBase.h
+++ b/Engine/Source/Runtime/Engine/Classes/Engine/LevelBase.h
@@ -19,15 +19,15 @@ class ULevelBase : public UObject
 	TTransArray<AActor*> Actors;
 
 	// Constructors.
-	ULevelBase( const class FPostConstructInitializeProperties& PCIP,const FURL& InURL );
+	ULevelBase( const FObjectInitializer& ObjectInitializer,const FURL& InURL );
 
 	// Begin UObject Interface
 	virtual void Serialize( FArchive& Ar ) override;
 	static void AddReferencedObjects(UObject* InThis, FReferenceCollector& Collector);
 	// End UObject Interface
 
-	ULevelBase(const class FPostConstructInitializeProperties& PCIP)
-		:	UObject(PCIP)
+	ULevelBase(const FObjectInitializer& ObjectInitializer)
+		:	UObject(ObjectInitializer)
 		,Actors( this )
 	{}
 };
diff --git a/Engine/Source/Runtime/Engine/Classes/Engine/NetConnection.h b/Engine/Source/Runtime/Engine/Classes/Engine/NetConnection.h
--- a/Engine/Source/Runtime/Engine/Classes/Engine/NetConnection.h
+++ b/Engine/Source/Runtime/Engine/Classes/Engine/NetConnection.h
@@ -311,7 +311,7 @@ public:
 	static class UNetConnection* GNetConnectionBeingCleanedUp;
 
 	// Constructors and destructors.
-	ENGINE_API UNetConnection(const class FPostConstructInitializeProperties& PCIP);
+	ENGINE_API UNetConnection(const FObjectInitializer& ObjectInitializer);
 
 	// Begin UObject interface.
 
