Async linker loads no longer create slow tasks

Also compiled out the slow tasks for non-editor builds, and fixed disabled slow tasks causing a UI update
This is a fix for TTP#349473 - CRASH: MAIN: EDITOR: Editor crashes when attempting to PIE in viewport.

diff --git a/Engine/Source/Runtime/Core/Public/Misc/FeedbackContext.h b/Engine/Source/Runtime/Core/Public/Misc/FeedbackContext.h
--- a/Engine/Source/Runtime/Core/Public/Misc/FeedbackContext.h
+++ b/Engine/Source/Runtime/Core/Public/Misc/FeedbackContext.h
@@ -242,7 +242,10 @@ public:
 		verifySlow(ExpectedWorkThisFrame <= WorkRemaining);
 		CurrentFrameScope = FMath::Min(WorkRemaining, ExpectedWorkThisFrame);
 
-		Context.RequestUpdateUI(bCreatedDialog || (bEnabled && Context.ScopeStack[0] == this));
+		if (bEnabled)
+		{
+			Context.RequestUpdateUI(bCreatedDialog || Context.ScopeStack[0] == this);
+		}
 	}
 
 	/**
diff --git a/Engine/Source/Runtime/CoreUObject/Private/UObject/LinkerLoad.cpp b/Engine/Source/Runtime/CoreUObject/Private/UObject/LinkerLoad.cpp
--- a/Engine/Source/Runtime/CoreUObject/Private/UObject/LinkerLoad.cpp
+++ b/Engine/Source/Runtime/CoreUObject/Private/UObject/LinkerLoad.cpp
@@ -652,8 +652,7 @@ ULinkerLoad::ELinkerStatus ULinkerLoad::CreateLoader()
 
 	if (!LoadProgressScope)
 	{
-		const bool bShowProgress = ( LoadFlags & ( LOAD_Quiet | LOAD_SeekFree ) ) == 0;
-		LoadProgressScope = new FScopedSlowTask(ULinkerDefs::TotalProgressSteps, NSLOCTEXT("Core", "GenericLoading", "Loading..."), bShowProgress);
+		LoadProgressScope = new FScopedSlowTask(ULinkerDefs::TotalProgressSteps, NSLOCTEXT("Core", "GenericLoading", "Loading..."), ShouldReportProgress());
 	}
 
 #endif
@@ -1695,14 +1694,17 @@ void ULinkerLoad::Verify()
 	{
 		if (!bHaveImportsBeenVerified)
 		{
-			FScopedSlowTask SlowTask(Summary.ImportCount, NSLOCTEXT("Core", "LinkerLoad_Imports", "Loading Imports"));
-
+#if WITH_EDITOR
+			FScopedSlowTask SlowTask(Summary.ImportCount, NSLOCTEXT("Core", "LinkerLoad_Imports", "Loading Imports"), ShouldReportProgress());
+#endif
 			// Validate all imports and map them to their remote linkers.
 			for (int32 ImportIndex = 0; ImportIndex < Summary.ImportCount; ImportIndex++)
 			{
 				FObjectImport& Import = ImportMap[ImportIndex];
-				
+
+#if WITH_EDITOR
 				SlowTask.EnterProgressFrame(1, FText::Format(NSLOCTEXT("Core", "LinkerLoad_LoadingImportName", "Loading Import '{0}'"), FText::FromString(Import.ObjectName.ToString())));
+#endif
 				VerifyImport( ImportIndex );
 			}
 		}
@@ -2074,7 +2076,9 @@ bool ULinkerLoad::VerifyImportInner(const int32 ImportIndex, FString& WarningSuf
 
 	FObjectImport& Import = ImportMap[ImportIndex];
 
-	FScopedSlowTask SlowTask(100, FText::Format(NSLOCTEXT("Core", "VerifyPackage_Scope", "Verifying '{0}'"), FText::FromString(Import.ObjectName.ToString())));
+#if WITH_EDITOR
+	FScopedSlowTask SlowTask(100, FText::Format(NSLOCTEXT("Core", "VerifyPackage_Scope", "Verifying '{0}'"), FText::FromString(Import.ObjectName.ToString())), ShouldReportProgress());
+#endif
 
 	if
 	(	(Import.SourceLinker && Import.SourceIndex != INDEX_NONE)
@@ -2112,7 +2116,9 @@ bool ULinkerLoad::VerifyImportInner(const int32 ImportIndex, FString& WarningSuf
 			bWasFullyLoaded = TmpPkg && TmpPkg->IsFullyLoaded();
 		}
 
+#if WITH_EDITOR
 		SlowTask.EnterProgressFrame(30);
+#endif
 
 		if (!bWasFullyLoaded)
 		{
@@ -2121,7 +2127,9 @@ bool ULinkerLoad::VerifyImportInner(const int32 ImportIndex, FString& WarningSuf
 			TmpPkg = LoadPackage(NULL, *Import.ObjectName.ToString(), InternalLoadFlags);
 		}
 
+#if WITH_EDITOR
 		SlowTask.EnterProgressFrame(30);
+#endif
 
 		// following is the original VerifyImport code
 		// @todo linkers: This could quite possibly be cleaned up
@@ -2143,7 +2151,9 @@ bool ULinkerLoad::VerifyImportInner(const int32 ImportIndex, FString& WarningSuf
 			InternalLoadFlags |= LOAD_NoVerify;
 		}
 
+#if WITH_EDITOR
 		SlowTask.EnterProgressFrame(40);
+#endif
 
 		// Get the linker if the package hasn't been fully loaded already.
 		if (!bWasFullyLoaded)
@@ -2156,7 +2166,9 @@ bool ULinkerLoad::VerifyImportInner(const int32 ImportIndex, FString& WarningSuf
 		// this resource's Outer is not a UPackage
 		checkf(Import.OuterIndex.IsImport(),TEXT("Outer for Import %s (%i) is not an import - OuterIndex:%i"), *GetImportFullName(ImportIndex), ImportIndex, Import.OuterIndex.ForDebugging());
 
+#if WITH_EDITOR
 		SlowTask.EnterProgressFrame(50);
+#endif
 
 		VerifyImport( Import.OuterIndex.ToImport() );
 
@@ -2186,7 +2198,9 @@ bool ULinkerLoad::VerifyImportInner(const int32 ImportIndex, FString& WarningSuf
 			Import.SourceLinker = OuterImport.SourceLinker;
 		}
 
+#if WITH_EDITOR
 		SlowTask.EnterProgressFrame(50);
+#endif
 
 		//check(Import.SourceLinker);
 		//@todo what does it mean if we don't have a SourceLinker here?
@@ -2468,9 +2482,11 @@ int32 ULinkerLoad::LoadMetaDataFromExportMap(bool bForcePreload/* = false */)
  */
 void ULinkerLoad::LoadAllObjects( bool bForcePreload )
 {
-	FScopedSlowTask SlowTask(ExportMap.Num(), NSLOCTEXT("Core", "LinkerLoad_LoadingObjects", "Loading Objects"), (LoadFlags & ( LOAD_Quiet | LOAD_SeekFree ) ) == 0);
+#if WITH_EDITOR
+	FScopedSlowTask SlowTask(ExportMap.Num(), NSLOCTEXT("Core", "LinkerLoad_LoadingObjects", "Loading Objects"), ShouldReportProgress());
 	SlowTask.bVisibleOnUI = false;
-	
+#endif
+
 	if ( (LoadFlags&LOAD_SeekFree) != 0 )
 	{
 		bForcePreload = true;
@@ -2488,7 +2504,9 @@ void ULinkerLoad::LoadAllObjects( bool bForcePreload )
 	
 	for(int32 ExportIndex = 0; ExportIndex < ExportMap.Num(); ++ExportIndex)
 	{
+#if WITH_EDITOR
 		SlowTask.EnterProgressFrame(1);
+#endif
 		if(ExportIndex == MetaDataIndex)
 		{
 			continue;
diff --git a/Engine/Source/Runtime/CoreUObject/Private/UObject/UObjectGlobals.cpp b/Engine/Source/Runtime/CoreUObject/Private/UObject/UObjectGlobals.cpp
--- a/Engine/Source/Runtime/CoreUObject/Private/UObject/UObjectGlobals.cpp
+++ b/Engine/Source/Runtime/CoreUObject/Private/UObject/UObjectGlobals.cpp
@@ -76,6 +76,13 @@ FCoreUObjectDelegates::FOnLoadObjectsOnTop FCoreUObjectDelegates::ShouldLoadOnTo
 FCoreUObjectDelegates::FStringAssetReferenceLoaded FCoreUObjectDelegates::StringAssetReferenceLoaded;
 FCoreUObjectDelegates::FStringAssetReferenceSaving FCoreUObjectDelegates::StringAssetReferenceSaving;
 FCoreUObjectDelegates::FPackageCreatedForLoad FCoreUObjectDelegates::PackageCreatedForLoad;
+
+/** Check wehether we should report progress or not */
+bool ShouldReportProgress()
+{
+	return GIsEditor && !IsRunningCommandlet() && !GIsAsyncLoading;
+}
+
 /**
  * Fast version of StaticFindObject that relies on the passed in FName being the object name
  * without any group/ package qualifiers.
@@ -798,7 +805,7 @@ UPackage* LoadPackage( UPackage* InOuter, const TCHAR* InLongPackageName, uint32
 	TGuardValue<bool> IsEditorLoadingPackage(GIsEditorLoadingPackage, GIsEditor || GIsEditorLoadingPackage);
 #endif
 
-	FScopedSlowTask SlowTask(100, FText::Format(NSLOCTEXT("Core", "LoadingPackage_Scope", "Loading Package '{0}'"), FText::FromString(FileToLoad)));
+	FScopedSlowTask SlowTask(100, FText::Format(NSLOCTEXT("Core", "LoadingPackage_Scope", "Loading Package '{0}'"), FText::FromString(FileToLoad)), ShouldReportProgress());
 	SlowTask.bVisibleOnUI = false;
 	
 	SlowTask.EnterProgressFrame(10);
@@ -1008,9 +1015,7 @@ void EndLoad()
 	}
 
 #if WITH_EDITOR
-	const bool bReportProgress = GIsEditor && !IsRunningCommandlet();
-	
-	FScopedSlowTask SlowTask(0, NSLOCTEXT("Core", "PerformingPostLoad", "Performing post-load..."), bReportProgress);
+	FScopedSlowTask SlowTask(0, NSLOCTEXT("Core", "PerformingPostLoad", "Performing post-load..."), ShouldReportProgress());
 
 	int32 NumObjectsLoaded = 0, NumObjectsFound = 0;
 #endif
diff --git a/Engine/Source/Runtime/CoreUObject/Public/UObject/Linker.h b/Engine/Source/Runtime/CoreUObject/Public/UObject/Linker.h
--- a/Engine/Source/Runtime/CoreUObject/Public/UObject/Linker.h
+++ b/Engine/Source/Runtime/CoreUObject/Public/UObject/Linker.h
@@ -1380,6 +1380,12 @@ private:
 #if WITH_EDITOR
 	/** Feedback scope that is created to house the slow task of an asynchronous linker load. Raw ptr so we don't pull in TUniquePtr for everything. */
 	FScopedSlowTask* LoadProgressScope;
+
+	/** Test whether we should report progress or not */
+	FORCEINLINE bool ShouldReportProgress() const
+	{
+		return !GIsAsyncLoading && ( LoadFlags & ( LOAD_Quiet | LOAD_SeekFree ) ) == 0;
+	}
 #endif 
 
 public:
