First phase of deprecating TSubobjectPtr<> replacing TSubobjectPtr properties in the Engine with normal pointers and adding subobject accessors for derived classes/client code.

diff --git a/Engine/Source/Runtime/AIModule/Classes/AIController.h b/Engine/Source/Runtime/AIModule/Classes/AIController.h
--- a/Engine/Source/Runtime/AIModule/Classes/AIController.h
+++ b/Engine/Source/Runtime/AIModule/Classes/AIController.h
@@ -90,13 +90,17 @@ public:
 	UPROPERTY()
 	uint32 bWantsPlayerState:1;
 
+private:
 	/** Component used for pathfinding and querying environment's navigation. */
 	UPROPERTY()
-	TSubobjectPtr<UNavigationComponent> NavComponent;
+	UNavigationComponent* NavComponent;
+public:
 
+private:
 	/** Component used for moving along a path. */
 	UPROPERTY()
-	TSubobjectPtr<UPathFollowingComponent> PathFollowingComponent;
+	UPathFollowingComponent* PathFollowingComponent;
+public:
 
 	/** Component responsible for behaviors. */
 	UPROPERTY()
@@ -107,8 +111,10 @@ public:
 	
 public:
 
-	UPROPERTY(BlueprintReadOnly, Category = AI)
-	TSubobjectPtr<UPawnActionsComponent> ActionsComp;
+private:
+	UPROPERTY(BlueprintReadOnly, Category = AI, meta = (AllowPrivateAccess = "true"))
+	UPawnActionsComponent* ActionsComp;
+public:
 
 	/** Event called when PossessedPawn is possesed by this controller. */
 	UFUNCTION(BlueprintImplementableEvent, Category = "AI")
@@ -162,7 +168,7 @@ public:
 	virtual void OnMoveCompleted(FAIRequestID RequestID, EPathFollowingResult::Type Result);
 
 	/** Returns the Move Request ID for the current move */
-	FORCEINLINE FAIRequestID GetCurrentMoveRequestID() const { return PathFollowingComponent.IsValid() ? PathFollowingComponent->GetCurrentRequestId() : FAIRequestID::InvalidRequest; }
+	FORCEINLINE FAIRequestID GetCurrentMoveRequestID() const { return PathFollowingComponent ? PathFollowingComponent->GetCurrentRequestId() : FAIRequestID::InvalidRequest; }
 
 	/** Blueprint notification that we've completed the current movement request */
 	UPROPERTY(BlueprintAssignable, meta=(DisplayName="MoveCompleted"))
@@ -306,6 +312,14 @@ public:
 
 	/** If true, AI controllers will ignore players. */
 	static bool bAIIgnorePlayers;
+
+public:
+	/** Returns NavComponent subobject **/
+	FORCEINLINE UNavigationComponent* GetNavComponent() const { return NavComponent; }
+	/** Returns PathFollowingComponent subobject **/
+	FORCEINLINE UPathFollowingComponent* GetPathFollowingComponent() const { return PathFollowingComponent; }
+	/** Returns ActionsComp subobject **/
+	FORCEINLINE UPawnActionsComponent* GetActionsComp() const { return ActionsComp; }
 };
 
 
diff --git a/Engine/Source/Runtime/AIModule/Classes/EnvironmentQuery/EQSTestingPawn.h b/Engine/Source/Runtime/AIModule/Classes/EnvironmentQuery/EQSTestingPawn.h
--- a/Engine/Source/Runtime/AIModule/Classes/EnvironmentQuery/EQSTestingPawn.h
+++ b/Engine/Source/Runtime/AIModule/Classes/EnvironmentQuery/EQSTestingPawn.h
@@ -45,9 +45,10 @@ class AIMODULE_API AEQSTestingPawn : public ACharacter, public IEQSQueryResultSo
 	
 
 #if WITH_EDITORONLY_DATA
+private:
 	/** Editor Preview */
 	UPROPERTY()
-	TSubobjectPtr<class UEQSRenderingComponent> EdRenderComp;
+	class UEQSRenderingComponent* EdRenderComp;
 #endif // WITH_EDITORONLY_DATA
 
 protected:
@@ -82,4 +83,10 @@ protected:
 	void UpdateDrawing();
 
 	static void OnEditorSelectionChanged(UObject* NewSelection);
+
+public:
+#if WITH_EDITORONLY_DATA
+	/** Returns EdRenderComp subobject **/
+	FORCEINLINE class UEQSRenderingComponent* GetEdRenderComp() { return EdRenderComp; }
+#endif
 };
diff --git a/Engine/Source/Runtime/AIModule/Private/AIController.cpp b/Engine/Source/Runtime/AIModule/Private/AIController.cpp
--- a/Engine/Source/Runtime/AIModule/Private/AIController.cpp
+++ b/Engine/Source/Runtime/AIModule/Private/AIController.cpp
@@ -127,12 +127,12 @@ void AAIController::GrabDebugSnapshot(FVisualLogEntry* Snapshot) const
 		Snapshot->Location = GetPawn()->GetActorLocation();
 	}
 
-	if (NavComponent.IsValid())
+	if (NavComponent)
 	{
 		NavComponent->DescribeSelfToVisLog(Snapshot);
 	}
 
-	if (PathFollowingComponent.IsValid())
+	if (PathFollowingComponent)
 	{
 		PathFollowingComponent->DescribeSelfToVisLog(Snapshot);
 	}
@@ -510,7 +510,7 @@ EPathFollowingRequestResult::Type AAIController::MoveToActor(class AActor* Goal,
 		}
 		else
 		{
-			FNavPathSharedPtr Path = FindPath(Goal, bUsePathfinding, UNavigationQueryFilter::GetQueryFilter(NavComponent.Get() ? NavComponent->GetNavData() : NULL, FilterClass));
+			FNavPathSharedPtr Path = FindPath(Goal, bUsePathfinding, UNavigationQueryFilter::GetQueryFilter(NavComponent ? NavComponent->GetNavData() : NULL, FilterClass));
 			if (Path.IsValid())
 			{
 				bAllowStrafe = bCanStrafe;
@@ -557,7 +557,7 @@ EPathFollowingRequestResult::Type AAIController::MoveToLocation(const FVector& D
 	FVector GoalLocation = Dest;
 
 	// fail if projection to navigation is required but it either failed or there's not navigation component
-	if (bCanRequestMove && bProjectDestinationToNavigation && (NavComponent.Get() == NULL || !NavComponent->ProjectPointToNavigation(Dest, GoalLocation)))
+	if (bCanRequestMove && bProjectDestinationToNavigation && (NavComponent == NULL || !NavComponent->ProjectPointToNavigation(Dest, GoalLocation)))
 	{
 		UE_VLOG_LOCATION(this, LogAINavigation, Error, Dest, 30.f, FLinearColor::Red, TEXT("AAIController::MoveToLocation failed to project destination location to navmesh"));
 		bCanRequestMove = false;
@@ -579,7 +579,7 @@ EPathFollowingRequestResult::Type AAIController::MoveToLocation(const FVector& D
 
 	if (bCanRequestMove)
 	{
-		FNavPathSharedPtr Path = FindPath(GoalLocation, bUsePathfinding, UNavigationQueryFilter::GetQueryFilter(NavComponent.Get() != NULL ? NavComponent->GetNavData() : NULL, FilterClass));
+		FNavPathSharedPtr Path = FindPath(GoalLocation, bUsePathfinding, UNavigationQueryFilter::GetQueryFilter(NavComponent != NULL ? NavComponent->GetNavData() : NULL, FilterClass));
 		if (Path.IsValid())
 		{
 			bAllowStrafe = bCanStrafe;
@@ -674,7 +674,7 @@ EPathFollowingStatus::Type AAIController::GetMoveStatus() const
 
 bool AAIController::HasPartialPath() const
 {
-	return (PathFollowingComponent.Get() != NULL) && (PathFollowingComponent->HasPartialPath());
+	return (PathFollowingComponent != NULL) && (PathFollowingComponent->HasPartialPath());
 }
 
 FVector AAIController::GetImmediateMoveDestination() const
diff --git a/Engine/Source/Runtime/AIModule/Private/Actions/PawnAction_Move.cpp b/Engine/Source/Runtime/AIModule/Private/Actions/PawnAction_Move.cpp
--- a/Engine/Source/Runtime/AIModule/Private/Actions/PawnAction_Move.cpp
+++ b/Engine/Source/Runtime/AIModule/Private/Actions/PawnAction_Move.cpp
@@ -187,9 +187,9 @@ EPawnActionAbortState::Type UPawnAction_Move::PerformAbort(EAIForceParam::Type S
 
 	AAIController* MyController = Cast<AAIController>(GetController());
 
-	if (MyController && MyController->PathFollowingComponent)
+	if (MyController && MyController->GetPathFollowingComponent())
 	{
-		MyController->PathFollowingComponent->AbortMove(TEXT("BehaviorTree abort"), RequestID);
+		MyController->GetPathFollowingComponent()->AbortMove(TEXT("BehaviorTree abort"), RequestID);
 	}
 
 	return Super::PerformAbort(ShouldForce);
@@ -291,11 +291,11 @@ void UPawnAction_Move::ClearPendingRepath()
 
 bool UPawnAction_Move::CheckAlreadyAtGoal(class AAIController* Controller, const FVector& TestLocation, float Radius)
 {
-	const bool bAlreadyAtGoal = Controller->PathFollowingComponent->HasReached(TestLocation, Radius);
+	const bool bAlreadyAtGoal = Controller->GetPathFollowingComponent()->HasReached(TestLocation, Radius);
 	if (bAlreadyAtGoal)
 	{
-		Controller->PathFollowingComponent->AbortMove(TEXT("Aborting move due to new move request finishing with AlreadyAtGoal"), FAIRequestID::AnyRequest, true, false, EPathFollowingMessage::OtherRequest);
-		Controller->PathFollowingComponent->SetLastMoveAtGoal(true);
+		Controller->GetPathFollowingComponent()->AbortMove(TEXT("Aborting move due to new move request finishing with AlreadyAtGoal"), FAIRequestID::AnyRequest, true, false, EPathFollowingMessage::OtherRequest);
+		Controller->GetPathFollowingComponent()->SetLastMoveAtGoal(true);
 	}
 
 	return bAlreadyAtGoal;
@@ -303,11 +303,11 @@ bool UPawnAction_Move::CheckAlreadyAtGoal(class AAIController* Controller, const
 
 bool UPawnAction_Move::CheckAlreadyAtGoal(class AAIController* Controller, const AActor* TestGoal, float Radius)
 {
-	const bool bAlreadyAtGoal = Controller->PathFollowingComponent->HasReached(TestGoal, Radius);
+	const bool bAlreadyAtGoal = Controller->GetPathFollowingComponent()->HasReached(TestGoal, Radius);
 	if (bAlreadyAtGoal)
 	{
-		Controller->PathFollowingComponent->AbortMove(TEXT("Aborting move due to new move request finishing with AlreadyAtGoal"), FAIRequestID::AnyRequest, true, false, EPathFollowingMessage::OtherRequest);
-		Controller->PathFollowingComponent->SetLastMoveAtGoal(true);
+		Controller->GetPathFollowingComponent()->AbortMove(TEXT("Aborting move due to new move request finishing with AlreadyAtGoal"), FAIRequestID::AnyRequest, true, false, EPathFollowingMessage::OtherRequest);
+		Controller->GetPathFollowingComponent()->SetLastMoveAtGoal(true);
 	}
 
 	return bAlreadyAtGoal;
diff --git a/Engine/Source/Runtime/AIModule/Private/BehaviorTree/Decorators/BTDecorator_DoesPathExist.cpp b/Engine/Source/Runtime/AIModule/Private/BehaviorTree/Decorators/BTDecorator_DoesPathExist.cpp
--- a/Engine/Source/Runtime/AIModule/Private/BehaviorTree/Decorators/BTDecorator_DoesPathExist.cpp
+++ b/Engine/Source/Runtime/AIModule/Private/BehaviorTree/Decorators/BTDecorator_DoesPathExist.cpp
@@ -61,7 +61,7 @@ bool UBTDecorator_DoesPathExist::CalculateRawConditionValue(UBehaviorTreeCompone
 	if (NavSys && bHasPointA && bHasPointB)
 	{
 		const AAIController* AIOwner = OwnerComp->GetAIOwner();
-		const ANavigationData* NavData = AIOwner && AIOwner->NavComponent ? AIOwner->NavComponent->GetNavData() : NULL;
+		const ANavigationData* NavData = AIOwner && AIOwner->GetNavComponent() ? AIOwner->GetNavComponent()->GetNavData() : NULL;
 		TSharedPtr<const FNavigationQueryFilter> QueryFilter = UNavigationQueryFilter::GetQueryFilter(NavData, FilterClass);
 
 		if (PathQueryType == EPathExistanceQueryType::NavmeshRaycast2D)
@@ -100,4 +100,4 @@ FName UBTDecorator_DoesPathExist::GetNodeIconName() const
 	return FName("BTEditor.Graph.BTNode.Decorator.DoesPathExist.Icon");
 }
 
-#endif	// WITH_EDITOR
\ No newline at end of file
+#endif	// WITH_EDITOR
diff --git a/Engine/Source/Runtime/AIModule/Private/BehaviorTree/Decorators/BTDecorator_ReachedMoveGoal.cpp b/Engine/Source/Runtime/AIModule/Private/BehaviorTree/Decorators/BTDecorator_ReachedMoveGoal.cpp
--- a/Engine/Source/Runtime/AIModule/Private/BehaviorTree/Decorators/BTDecorator_ReachedMoveGoal.cpp
+++ b/Engine/Source/Runtime/AIModule/Private/BehaviorTree/Decorators/BTDecorator_ReachedMoveGoal.cpp
@@ -17,7 +17,7 @@ UBTDecorator_ReachedMoveGoal::UBTDecorator_ReachedMoveGoal(const FObjectInitiali
 bool UBTDecorator_ReachedMoveGoal::CalculateRawConditionValue(UBehaviorTreeComponent* OwnerComp, uint8* NodeMemory) const 
 {
 	AAIController* AIOwner = OwnerComp->GetAIOwner();
-	const bool bReachedGoal = AIOwner && AIOwner->PathFollowingComponent && AIOwner->PathFollowingComponent->DidMoveReachGoal();
+	const bool bReachedGoal = AIOwner && AIOwner->GetPathFollowingComponent() && AIOwner->GetPathFollowingComponent()->DidMoveReachGoal();
 	return bReachedGoal;
 }
 
@@ -28,4 +28,4 @@ FName UBTDecorator_ReachedMoveGoal::GetNodeIconName() const
 	return FName("BTEditor.Graph.BTNode.Decorator.ReachedMoveGoal.Icon");
 }
 
-#endif	// WITH_EDITOR
\ No newline at end of file
+#endif	// WITH_EDITOR
diff --git a/Engine/Source/Runtime/AIModule/Private/BehaviorTree/Tasks/BTTask_MoveDirectlyToward.cpp b/Engine/Source/Runtime/AIModule/Private/BehaviorTree/Tasks/BTTask_MoveDirectlyToward.cpp
--- a/Engine/Source/Runtime/AIModule/Private/BehaviorTree/Tasks/BTTask_MoveDirectlyToward.cpp
+++ b/Engine/Source/Runtime/AIModule/Private/BehaviorTree/Tasks/BTTask_MoveDirectlyToward.cpp
@@ -70,9 +70,9 @@ EBTNodeResult::Type UBTTask_MoveDirectlyToward::AbortTask(UBehaviorTreeComponent
 	FBTMoveDirectlyTowardMemory* MyMemory = (FBTMoveDirectlyTowardMemory*)NodeMemory;
 	AAIController* MyController = OwnerComp ? Cast<AAIController>(OwnerComp->GetOwner()) : NULL;
 
-	if (MyController && MyController->PathFollowingComponent)
+	if (MyController && MyController->GetPathFollowingComponent())
 	{
-		MyController->PathFollowingComponent->AbortMove(TEXT("BehaviorTree abort"), MyMemory->MoveRequestID);
+		MyController->GetPathFollowingComponent()->AbortMove(TEXT("BehaviorTree abort"), MyMemory->MoveRequestID);
 	}
 
 	return Super::AbortTask(OwnerComp, NodeMemory);
@@ -116,4 +116,4 @@ FName UBTTask_MoveDirectlyToward::GetNodeIconName() const
 	return FName("BTEditor.Graph.BTNode.Task.MoveDirectlyToward.Icon");
 }
 
-#endif	// WITH_EDITOR
\ No newline at end of file
+#endif	// WITH_EDITOR
diff --git a/Engine/Source/Runtime/AIModule/Private/BehaviorTree/Tasks/BTTask_MoveTo.cpp b/Engine/Source/Runtime/AIModule/Private/BehaviorTree/Tasks/BTTask_MoveTo.cpp
--- a/Engine/Source/Runtime/AIModule/Private/BehaviorTree/Tasks/BTTask_MoveTo.cpp
+++ b/Engine/Source/Runtime/AIModule/Private/BehaviorTree/Tasks/BTTask_MoveTo.cpp
@@ -97,9 +97,9 @@ EBTNodeResult::Type UBTTask_MoveTo::AbortTask(class UBehaviorTreeComponent* Owne
 	{
 		AAIController* MyController = OwnerComp ? Cast<AAIController>(OwnerComp->GetOwner()) : NULL;
 
-		if (MyController && MyController->PathFollowingComponent)
+		if (MyController && MyController->GetPathFollowingComponent())
 		{
-			MyController->PathFollowingComponent->AbortMove(TEXT("BehaviorTree abort"), MyMemory->MoveRequestID);
+			MyController->GetPathFollowingComponent()->AbortMove(TEXT("BehaviorTree abort"), MyMemory->MoveRequestID);
 		}
 	}
 
@@ -171,4 +171,4 @@ FName UBTTask_MoveTo::GetNodeIconName() const
 	return FName("BTEditor.Graph.BTNode.Task.MoveTo.Icon");
 }
 
-#endif	// WITH_EDITOR
\ No newline at end of file
+#endif	// WITH_EDITOR
diff --git a/Engine/Source/Runtime/AIModule/Private/Blueprint/AIBlueprintHelperLibrary.cpp b/Engine/Source/Runtime/AIModule/Private/Blueprint/AIBlueprintHelperLibrary.cpp
--- a/Engine/Source/Runtime/AIModule/Private/Blueprint/AIBlueprintHelperLibrary.cpp
+++ b/Engine/Source/Runtime/AIModule/Private/Blueprint/AIBlueprintHelperLibrary.cpp
@@ -188,9 +188,9 @@ void UAIBlueprintHelperLibrary::LockAIResourcesWithAnimation(class UAnimInstance
 		AAIController* OwningAI = Cast<AAIController>(PawnOwner->Controller);
 		if (OwningAI)
 		{
-			if (bLockMovement && OwningAI->PathFollowingComponent)
+			if (bLockMovement && OwningAI->GetPathFollowingComponent())
 			{
-				OwningAI->PathFollowingComponent->LockResource(EAILockSource::Animation);
+				OwningAI->GetPathFollowingComponent()->LockResource(EAILockSource::Animation);
 			}
 			if (LockAILogic && OwningAI->BrainComponent)
 			{
@@ -213,9 +213,9 @@ void UAIBlueprintHelperLibrary::UnlockAIResourcesWithAnimation(class UAnimInstan
 		AAIController* OwningAI = Cast<AAIController>(PawnOwner->Controller);
 		if (OwningAI)
 		{
-			if (bUnlockMovement && OwningAI->PathFollowingComponent)
+			if (bUnlockMovement && OwningAI->GetPathFollowingComponent())
 			{
-				OwningAI->PathFollowingComponent->ClearResourceLock(EAILockSource::Animation);
+				OwningAI->GetPathFollowingComponent()->ClearResourceLock(EAILockSource::Animation);
 			}
 			if (UnlockAILogic && OwningAI->BrainComponent)
 			{
@@ -223,4 +223,4 @@ void UAIBlueprintHelperLibrary::UnlockAIResourcesWithAnimation(class UAnimInstan
 			}
 		}
 	}
-}
\ No newline at end of file
+}
diff --git a/Engine/Source/Runtime/AIModule/Private/EnvironmentQuery/EQSTestingPawn.cpp b/Engine/Source/Runtime/AIModule/Private/EnvironmentQuery/EQSTestingPawn.cpp
--- a/Engine/Source/Runtime/AIModule/Private/EnvironmentQuery/EQSTestingPawn.cpp
+++ b/Engine/Source/Runtime/AIModule/Private/EnvironmentQuery/EQSTestingPawn.cpp
@@ -34,7 +34,7 @@ AEQSTestingPawn::AEQSTestingPawn(const FObjectInitializer& ObjectInitializer)
 	, QueryingMode(EEnvQueryRunMode::AllMatching)
 {
 	static FName CollisionProfileName(TEXT("NoCollision"));
-	CapsuleComponent->SetCollisionProfileName(CollisionProfileName);
+	GetCapsuleComponent()->SetCollisionProfileName(CollisionProfileName);
 
 #if WITH_EDITORONLY_DATA
 	EdRenderComp = ObjectInitializer.CreateEditorOnlyDefaultSubobject<UEQSRenderingComponent>(this, TEXT("EQSRender"));
@@ -47,7 +47,7 @@ AEQSTestingPawn::AEQSTestingPawn(const FObjectInitializer& ObjectInitializer)
 			ArrowComponent->bIsScreenSizeScaled = true;
 		}
 
-		TSubobjectPtr<UBillboardComponent> SpriteComponent = ObjectInitializer.CreateEditorOnlyDefaultSubobject<UBillboardComponent>(this, TEXT("Sprite"));
+		UBillboardComponent* SpriteComponent = ObjectInitializer.CreateEditorOnlyDefaultSubobject<UBillboardComponent>(this, TEXT("Sprite"));
 		if (!IsRunningCommandlet() && (SpriteComponent != nullptr))
 		{
 			struct FConstructorStatics
diff --git a/Engine/Source/Runtime/AIModule/Private/Navigation/PathFollowingComponent.cpp b/Engine/Source/Runtime/AIModule/Private/Navigation/PathFollowingComponent.cpp
--- a/Engine/Source/Runtime/AIModule/Private/Navigation/PathFollowingComponent.cpp
+++ b/Engine/Source/Runtime/AIModule/Private/Navigation/PathFollowingComponent.cpp
@@ -623,7 +623,7 @@ int32 UPathFollowingComponent::OptimizeSegmentVisibility(int32 StartIndex)
 	//SCOPE_CYCLE_COUNTER(STAT_Navigation_PathVisibilityOptimisation);
 
 	const AAIController* MyAI = Cast<AAIController>(GetOwner());
-	const ANavigationData* NavData = MyAI && MyAI->NavComponent ? MyAI->NavComponent->GetNavData() : NULL;
+	const ANavigationData* NavData = MyAI && MyAI->GetNavComponent() ? MyAI->GetNavComponent()->GetNavData() : NULL;
 	if (Path.IsValid())
 	{
 		Path->ShortcutNodeRefs.Reset();
@@ -635,7 +635,7 @@ int32 UPathFollowingComponent::OptimizeSegmentVisibility(int32 StartIndex)
 	}
 
 	const APawn* MyPawn = MyAI->GetPawn();
-	const float PawnHalfHeight = (MyAI->GetCharacter() && MyAI->GetCharacter()->CapsuleComponent) ? MyAI->GetCharacter()->CapsuleComponent->GetScaledCapsuleHalfHeight() : 0.0f;
+	const float PawnHalfHeight = (MyAI->GetCharacter() && MyAI->GetCharacter()->GetCapsuleComponent()) ? MyAI->GetCharacter()->GetCapsuleComponent()->GetScaledCapsuleHalfHeight() : 0.0f;
 	const FVector PawnLocation = MyPawn->GetActorLocation();
 
 #if USE_PHYSIC_FOR_VISIBILITY_TESTS
diff --git a/Engine/Source/Runtime/CoreUObject/Public/UObject/UObjectGlobals.h b/Engine/Source/Runtime/CoreUObject/Public/UObject/UObjectGlobals.h
--- a/Engine/Source/Runtime/CoreUObject/Public/UObject/UObjectGlobals.h
+++ b/Engine/Source/Runtime/CoreUObject/Public/UObject/UObjectGlobals.h
@@ -534,38 +534,13 @@ template<> struct TIsPODType<FSubobjectPtr> { enum { Value = true }; };
 template<> struct TIsZeroConstructType<FSubobjectPtr> { enum { Value = true }; };
 template<> struct TIsWeakPointerType<FSubobjectPtr> { enum { Value = false }; };
 
-/**
- * TSubobjectPtrConstructor - internal private class to disallow the use of 
- * TSubobjectPtr's assignment operator as well as to disable easy access to
- * TSubobjectPtr::Get() when creating a subobject with ObjectInitializer.CreateDefaultSubobject().Get() <- No public 'Get'
- */
-template <class SubobjectType>
-class TSubobjectPtrConstructor
-{
-	/** Only FObjectInitializer can construct this class. */
-	friend class FObjectInitializer;
-	/** TSubobjectPtr needs to have access to the subobject pointer. */
-	template <class AnySubobjectType> friend class TSubobjectPtr;
-	/** Subobject pointer. */
-	UObject* Object;	
-	/** Constructor used by ObjectInitializer. */
-	TSubobjectPtrConstructor(SubobjectType* InObject)
-		: Object(InObject)
-	{}
-	/** Private constructors. */
-	TSubobjectPtrConstructor() {}
-	TSubobjectPtrConstructor(const TSubobjectPtrConstructor&) {}
-	/** Private assignment operator. */
-	TSubobjectPtrConstructor& operator=(const TSubobjectPtrConstructor&) { return *this; }
-};
-
 /**
  * TSubobjectPtr - Sub-object smart pointer, owns a reference to instanced objects (sub-objects / components).
  * Prevents anything in C++ from overwriting the sub-object pointer.
  * Can (and should) be declared as a UPROPERTY():
  *
  *   UPROPERTY()
- *   TSubobjectPtr<UActorComponent> MyComponent;
+ *   UActorComponent* MyComponent;
  *
  * It can only be assigned to with ObjectInitializer.CreateDefaultSubobject (via TSubobjectPtrConstructor) in the owning object's constructor.
  *
@@ -587,13 +562,19 @@ class TSubobjectPtrConstructor
 template <class SubobjectType>
 class TSubobjectPtr : public FSubobjectPtr
 {
+public:
+
 	/** Internal constructors. */
 	TSubobjectPtr(SubobjectType* InObject)
 		: FSubobjectPtr(InObject)
 	{}
-	TSubobjectPtr& operator=(const TSubobjectPtr& Other) { return *this; }
+	TSubobjectPtr& operator=(const TSubobjectPtr& Other) 
+	{ 
+		Set(Other.Object);
+		return *this; 
+	}
+
 
-public:
 	/** Default constructor. */
 	TSubobjectPtr()
 		: FSubobjectPtr((UObject*)FSubobjectPtr::InvalidPtrValue)
@@ -607,26 +588,11 @@ public:
 	{
 		static_assert((CanConvertPointerFromTo<DerivedSubobjectType, SubobjectType>::Result), "Subobject pointers must be compatible.");
 	}
-	/** Initialization constructor. Can only be used with ObjectInitializer.CreateDefaultSubobject(). */
-	template <class DerivedSubobjectType>
-	TSubobjectPtr(const TSubobjectPtrConstructor<DerivedSubobjectType>& Other)
-		: FSubobjectPtr(Other.Object)
-	{
-		static_assert((CanConvertPointerFromTo<DerivedSubobjectType, SubobjectType>::Result), "Subobject pointers must be compatible.");
-	}
 	/** Gets the sub-object pointer. */
 	FORCEINLINE SubobjectType* Get() const
 	{
 		return (SubobjectType*)Object;
 	}
-	/** Assignment operator - can only be used with ObjectInitializer.CreateDefaultSubobject(). */
-	template <class DerivedSubobjectType>
-	FORCEINLINE TSubobjectPtr& operator=(const TSubobjectPtrConstructor<DerivedSubobjectType>& Other)
-	{
-		static_assert((CanConvertPointerFromTo<DerivedSubobjectType, SubobjectType>::Result), "Subobject pointers must be compatible.");
-		Set(Other.Object);
-		return *this;
-	}
 	/** Gets the sub-object pointer. */
 	FORCEINLINE SubobjectType* operator->() const
 	{
@@ -695,10 +661,10 @@ public:
 	 * @param bTransient		true if the component is being assigned to a transient property
 	 */
 	template<class TReturnType>
-	TSubobjectPtrConstructor<TReturnType> CreateDefaultSubobject(UObject* Outer, FName SubobjectName, bool bTransient = false) const
+	TReturnType* CreateDefaultSubobject(UObject* Outer, FName SubobjectName, bool bTransient = false) const
 	{
 		UClass* ReturnType = TReturnType::StaticClass();
-		return TSubobjectPtrConstructor<TReturnType>(static_cast<TReturnType*>(CreateDefaultSubobject(Outer, SubobjectName, ReturnType, ReturnType, /*bIsRequired =*/ true, /*bIsAbstract =*/ false, bTransient)));
+		return static_cast<TReturnType*>(CreateDefaultSubobject(Outer, SubobjectName, ReturnType, ReturnType, /*bIsRequired =*/ true, /*bIsAbstract =*/ false, bTransient));
 	}
 
 	/**
@@ -710,10 +676,10 @@ public:
 	 * @param bTransient		true if the component is being assigned to a transient property
 	 */
 	template<class TReturnType>
-	TSubobjectPtrConstructor<TReturnType> CreateOptionalDefaultSubobject(UObject* Outer, FName SubobjectName, bool bTransient = false) const
+	TReturnType* CreateOptionalDefaultSubobject(UObject* Outer, FName SubobjectName, bool bTransient = false) const
 	{
 		UClass* ReturnType = TReturnType::StaticClass();
-		return TSubobjectPtrConstructor<TReturnType>(static_cast<TReturnType*>(CreateDefaultSubobject(Outer, SubobjectName, ReturnType, ReturnType, /*bIsRequired =*/ false, /*bIsAbstract =*/ false, bTransient)));
+		return static_cast<TReturnType*>(CreateDefaultSubobject(Outer, SubobjectName, ReturnType, ReturnType, /*bIsRequired =*/ false, /*bIsAbstract =*/ false, bTransient));
 	}
 
 	/**
@@ -725,10 +691,10 @@ public:
 	 * @param bTransient		true if the component is being assigned to a transient property
 	 */
 	template<class TReturnType>
-	TSubobjectPtrConstructor<TReturnType> CreateAbstractDefaultSubobject(UObject* Outer, FName SubobjectName, bool bTransient = false) const
+	TReturnType* CreateAbstractDefaultSubobject(UObject* Outer, FName SubobjectName, bool bTransient = false) const
 	{
 		UClass* ReturnType = TReturnType::StaticClass();
-		return TSubobjectPtrConstructor<TReturnType>(static_cast<TReturnType*>(CreateDefaultSubobject(Outer, SubobjectName, ReturnType, ReturnType, /*bIsRequired =*/ true, /*bIsAbstract =*/ true, bTransient)));
+		return static_cast<TReturnType*>(CreateDefaultSubobject(Outer, SubobjectName, ReturnType, ReturnType, /*bIsRequired =*/ true, /*bIsAbstract =*/ true, bTransient));
 	}
 
 	/** 
@@ -740,9 +706,9 @@ public:
 	* @param bTransient		true if the component is being assigned to a transient property
 	*/ 
 	template<class TReturnType, class TClassToConstructByDefault> 
-	TSubobjectPtrConstructor<TReturnType> CreateDefaultSubobject(UObject* Outer, FName SubobjectName, bool bTransient = false) const 
+	TReturnType* CreateDefaultSubobject(UObject* Outer, FName SubobjectName, bool bTransient = false) const 
 	{ 
-		return TSubobjectPtrConstructor<TReturnType>(static_cast<TReturnType*>(CreateDefaultSubobject(Outer, SubobjectName, TReturnType::StaticClass(), TClassToConstructByDefault::StaticClass(), /*bIsRequired =*/ true, /*bIsAbstract =*/ false, bTransient)));
+		return static_cast<TReturnType*>(CreateDefaultSubobject(Outer, SubobjectName, TReturnType::StaticClass(), TClassToConstructByDefault::StaticClass(), /*bIsRequired =*/ true, /*bIsAbstract =*/ false, bTransient));
 	}
 
 	/**
@@ -753,7 +719,7 @@ public:
 	 * @param bTransient		true if the component is being assigned to a transient property
 	 */
 	template<class TReturnType>
-	TSubobjectPtrConstructor<TReturnType> CreateEditorOnlyDefaultSubobject(UObject* Outer, FName SubobjectName, bool bTransient = false) const
+	TReturnType* CreateEditorOnlyDefaultSubobject(UObject* Outer, FName SubobjectName, bool bTransient = false) const
 	{
 #if WITH_EDITOR
 		if (GIsEditor)
@@ -765,7 +731,7 @@ public:
 				EditorSubobject->AlwaysLoadOnClient = false;
 				EditorSubobject->AlwaysLoadOnServer = false;
 			}
-			return TSubobjectPtrConstructor<TReturnType>(EditorSubobject);	
+			return EditorSubobject;	
 		}
 #endif
 		return NULL;
diff --git a/Engine/Source/Runtime/Engine/Classes/AI/Navigation/NavLinkProxy.h b/Engine/Source/Runtime/Engine/Classes/AI/Navigation/NavLinkProxy.h
--- a/Engine/Source/Runtime/Engine/Classes/AI/Navigation/NavLinkProxy.h
+++ b/Engine/Source/Runtime/Engine/Classes/AI/Navigation/NavLinkProxy.h
@@ -23,21 +23,25 @@ class ENGINE_API ANavLinkProxy : public AActor, public INavLinkHostInterface, pu
 	UPROPERTY()
 	TArray<FNavigationSegmentLink> SegmentLinks;
 
+private:
 	/** Smart link: can affect path following */
 	UPROPERTY(VisibleAnywhere, Category=SmartLink)
-	TSubobjectPtr<class UNavLinkCustomComponent> SmartLinkComp;
+	class UNavLinkCustomComponent* SmartLinkComp;
+public:
 
 	/** Smart link: toggle relevancy */
 	UPROPERTY(EditAnywhere, Category=SmartLink)
 	bool bSmartLinkIsRelevant;
 
 #if WITH_EDITORONLY_DATA
+private:
 	/** Editor Preview */
 	UPROPERTY()
-	TSubobjectPtr<class UNavLinkRenderingComponent> EdRenderComp;
+	class UNavLinkRenderingComponent* EdRenderComp;
 
 	UPROPERTY()
-	TSubobjectPtr<class UBillboardComponent> SpriteComponent;
+	class UBillboardComponent* SpriteComponent;
+public:
 #endif // WITH_EDITORONLY_DATA
 
 	// BEGIN INavRelevantInterface
@@ -86,4 +90,14 @@ protected:
 	FSmartLinkReachedSignature OnSmartLinkReached;
 
 	void NotifySmartLinkReached(UNavLinkCustomComponent* LinkComp, class UPathFollowingComponent* PathComp, const FVector& DestPoint);
+
+public:
+	/** Returns SmartLinkComp subobject **/
+	FORCEINLINE class UNavLinkCustomComponent* GetSmartLinkComp() const { return SmartLinkComp; }
+#if WITH_EDITORONLY_DATA
+	/** Returns EdRenderComp subobject **/
+	FORCEINLINE class UNavLinkRenderingComponent* GetEdRenderComp() const { return EdRenderComp; }
+	/** Returns SpriteComponent subobject **/
+	FORCEINLINE class UBillboardComponent* GetSpriteComponent() const { return SpriteComponent; }
+#endif
 };
diff --git a/Engine/Source/Runtime/Engine/Classes/AI/Navigation/NavigationTestingActor.h b/Engine/Source/Runtime/Engine/Classes/AI/Navigation/NavigationTestingActor.h
--- a/Engine/Source/Runtime/Engine/Classes/AI/Navigation/NavigationTestingActor.h
+++ b/Engine/Source/Runtime/Engine/Classes/AI/Navigation/NavigationTestingActor.h
@@ -31,14 +31,16 @@ class ENGINE_API ANavigationTestingActor : public AActor, public INavAgentInterf
 {
 	GENERATED_UCLASS_BODY()
 
+private:
 	UPROPERTY()
-	TSubobjectPtr<class UCapsuleComponent> CapsuleComponent;
+	class UCapsuleComponent* CapsuleComponent;
 
 #if WITH_EDITORONLY_DATA
 	/** Editor Preview */
 	UPROPERTY()
-	TSubobjectPtr<class UNavTestRenderingComponent> EdRenderComp;
+	class UNavTestRenderingComponent* EdRenderComp;
 #endif // WITH_EDITORONLY_DATA
+public:
 
 	/** @todo document */
 	UPROPERTY(EditAnywhere, Category=Agent)
@@ -163,4 +165,12 @@ class ENGINE_API ANavigationTestingActor : public AActor, public INavAgentInterf
 	// constructed for the path find (e.g. change the filter or add 
 	// constraints/goal evaluators).
 	virtual FPathFindingQuery BuildPathFindingQuery(const ANavigationTestingActor* Goal) const;
+
+public:
+	/** Returns CapsuleComponent subobject **/
+	FORCEINLINE class UCapsuleComponent* GetCapsuleComponent() const { return CapsuleComponent; }
+#if WITH_EDITORONLY_DATA
+	/** Returns EdRenderComp subobject **/
+	FORCEINLINE class UNavTestRenderingComponent* GetEdRenderComp() const { return EdRenderComp; }
+#endif
 };
diff --git a/Engine/Source/Runtime/Engine/Classes/Animation/SkeletalMeshActor.h b/Engine/Source/Runtime/Engine/Classes/Animation/SkeletalMeshActor.h
--- a/Engine/Source/Runtime/Engine/Classes/Animation/SkeletalMeshActor.h
+++ b/Engine/Source/Runtime/Engine/Classes/Animation/SkeletalMeshActor.h
@@ -24,8 +24,10 @@ class ENGINE_API ASkeletalMeshActor : public AActor, public IMatineeAnimInterfac
 	UPROPERTY()
 	uint32 bWakeOnLevelStart_DEPRECATED:1;
 
-	UPROPERTY(Category=SkeletalMeshActor, VisibleAnywhere, BlueprintReadOnly, meta=(ExposeFunctionCategories="Mesh,Components|SkeletalMesh,Animation,Physics"))
-	TSubobjectPtr<class USkeletalMeshComponent> SkeletalMeshComponent;
+private:
+	UPROPERTY(Category = SkeletalMeshActor, VisibleAnywhere, BlueprintReadOnly, meta = (ExposeFunctionCategories = "Mesh,Components|SkeletalMesh,Animation,Physics", AllowPrivateAccess = "true"))
+	class USkeletalMeshComponent* SkeletalMeshComponent;
+public:
 
 	/** Used to replicate mesh to clients */
 	UPROPERTY(replicatedUsing=OnRep_ReplicatedMesh, transient)
@@ -88,6 +90,10 @@ private:
 	EAnimationMode::Type	SavedAnimationMode;
 	// utility function to see if it can play animation or not
 	bool CanPlayAnimation();
+
+public:
+	/** Returns SkeletalMeshComponent subobject **/
+	FORCEINLINE class USkeletalMeshComponent* GetSkeletalMeshComponent() { return SkeletalMeshComponent; }
 };
 
 
diff --git a/Engine/Source/Runtime/Engine/Classes/Atmosphere/AtmosphericFog.h b/Engine/Source/Runtime/Engine/Classes/Atmosphere/AtmosphericFog.h
--- a/Engine/Source/Runtime/Engine/Classes/Atmosphere/AtmosphericFog.h
+++ b/Engine/Source/Runtime/Engine/Classes/Atmosphere/AtmosphericFog.h
@@ -13,19 +13,30 @@ class AAtmosphericFog : public AInfo
 {
 	GENERATED_UCLASS_BODY()
 
+private:
 	/** Main fog component */
-	UPROPERTY(BlueprintReadOnly, VisibleAnywhere, Category=Atmosphere)
-	TSubobjectPtr<class UAtmosphericFogComponent> AtmosphericFogComponent;
+	UPROPERTY(BlueprintReadOnly, VisibleAnywhere, Category = Atmosphere, meta = (AllowPrivateAccess = "true"))
+	class UAtmosphericFogComponent* AtmosphericFogComponent;
 
 #if WITH_EDITORONLY_DATA
 	/** Arrow component to indicate default sun rotation */
 	UPROPERTY()
-	TSubobjectPtr<class UArrowComponent> ArrowComponent;
+	class UArrowComponent* ArrowComponent;
 #endif
 
+public:
+
 #if WITH_EDITOR
 	virtual void PostActorCreated() override;
 #endif
+
+public:
+	/** Returns AtmosphericFogComponent subobject **/
+	FORCEINLINE class UAtmosphericFogComponent* GetAtmosphericFogComponent() { return AtmosphericFogComponent; }
+#if WITH_EDITORONLY_DATA
+	/** Returns ArrowComponent subobject **/
+	FORCEINLINE class UArrowComponent* GetArrowComponent() { return ArrowComponent; }
+#endif
 };
 
 
diff --git a/Engine/Source/Runtime/Engine/Classes/Camera/CameraActor.h b/Engine/Source/Runtime/Engine/Classes/Camera/CameraActor.h
--- a/Engine/Source/Runtime/Engine/Classes/Camera/CameraActor.h
+++ b/Engine/Source/Runtime/Engine/Classes/Camera/CameraActor.h
@@ -19,9 +19,11 @@ private:
 	TEnumAsByte<EAutoReceiveInput::Type> AutoActivateForPlayer;
 
 public:
+private:
 	/** The camera component for this camera */
-	UPROPERTY(Category=CameraActor, VisibleAnywhere, BlueprintReadOnly)
-	TSubobjectPtr<class UCameraComponent> CameraComponent;
+	UPROPERTY(Category = CameraActor, VisibleAnywhere, BlueprintReadOnly, meta = (AllowPrivateAccess = "true"))
+	class UCameraComponent* CameraComponent;
+public:
 
 	/** If this CameraActor is being used to preview a CameraAnim in the editor, this is the anim being previewed. */
 	TWeakObjectPtr<class UCameraAnim> PreviewedCameraAnim;
@@ -60,4 +62,8 @@ public:
 	// Begin AActor interface
 	ENGINE_API virtual void BeginPlay() override;
 	// End AActor interface
+
+public:
+	/** Returns CameraComponent subobject **/
+	FORCEINLINE class UCameraComponent* GetCameraComponent() const { return CameraComponent; }
 };
diff --git a/Engine/Source/Runtime/Engine/Classes/Camera/CameraAnimInst.h b/Engine/Source/Runtime/Engine/Classes/Camera/CameraAnimInst.h
--- a/Engine/Source/Runtime/Engine/Classes/Camera/CameraAnimInst.h
+++ b/Engine/Source/Runtime/Engine/Classes/Camera/CameraAnimInst.h
@@ -17,9 +17,11 @@ class ENGINE_API UCameraAnimInst : public UObject
 	class UCameraAnim* CamAnim;
 
 protected:
+private:
 	/** the UInterpGroupInst used to do the interpolation */
 	UPROPERTY(instanced)
-	TSubobjectPtr<class UInterpGroupInst> InterpGroupInst;
+	class UInterpGroupInst* InterpGroupInst;
+protected:
 
 public:
 	/** Current time for the animation */
@@ -133,6 +135,10 @@ public:
 	void SetDuration(float NewDuration);
 
 	
+
+protected:
+	/** Returns InterpGroupInst subobject **/
+	FORCEINLINE class UInterpGroupInst* GetInterpGroupInst() const { return InterpGroupInst; }
 };
 
 
diff --git a/Engine/Source/Runtime/Engine/Classes/Camera/PlayerCameraManager.h b/Engine/Source/Runtime/Engine/Classes/Camera/PlayerCameraManager.h
--- a/Engine/Source/Runtime/Engine/Classes/Camera/PlayerCameraManager.h
+++ b/Engine/Source/Runtime/Engine/Classes/Camera/PlayerCameraManager.h
@@ -144,9 +144,11 @@ class ENGINE_API APlayerCameraManager : public AActor
 	UPROPERTY()
 	class APlayerController* PCOwner;
 
+private:
 	/** Dummy component we can use to attach things to the camera. */
-	UPROPERTY(Category=PlayerCameraManager, VisibleAnywhere, BlueprintReadOnly)
-	TSubobjectPtr<class USceneComponent> TransformComponent;
+	UPROPERTY(Category = PlayerCameraManager, VisibleAnywhere, BlueprintReadOnly, meta = (AllowPrivateAccess = "true"))
+	class USceneComponent* TransformComponent;
+public:
 
 	/** Camera Mode */
 	UPROPERTY()
@@ -639,4 +641,8 @@ private:
 
 	// Buried to prevent use; use GetCameraRotation instead
 	FVector GetActorLocation() const { return Super::GetActorLocation(); }
+
+public:
+	/** Returns TransformComponent subobject **/
+	FORCEINLINE class USceneComponent* GetTransformComponent() const { return TransformComponent; }
 };
diff --git a/Engine/Source/Runtime/Engine/Classes/Engine/Brush.h b/Engine/Source/Runtime/Engine/Classes/Engine/Brush.h
--- a/Engine/Source/Runtime/Engine/Classes/Engine/Brush.h
+++ b/Engine/Source/Runtime/Engine/Classes/Engine/Brush.h
@@ -109,8 +109,10 @@ class ENGINE_API ABrush : public AActor
 	UPROPERTY(export)
 	class UModel* Brush;
 
-	UPROPERTY(Category=Collision, VisibleAnywhere, BlueprintReadOnly)
-	TSubobjectPtr<class UBrushComponent> BrushComponent;
+private:
+	UPROPERTY(Category = Collision, VisibleAnywhere, BlueprintReadOnly, meta = (AllowPrivateAccess = "true"))
+	class UBrushComponent* BrushComponent;
+public:
 
 #if WITH_EDITORONLY_DATA
 	UPROPERTY(VisibleAnywhere, Instanced, Category=BrushBuilder)
@@ -251,6 +253,10 @@ public:
 	/** @return the brush builder that created the current brush shape */
 	const UBrushBuilder* GetBrushBuilder() const { return BrushBuilder; }
 #endif
+
+public:
+	/** Returns BrushComponent subobject **/
+	FORCEINLINE class UBrushComponent* GetBrushComponent() const { return BrushComponent; }
 };
 
 
diff --git a/Engine/Source/Runtime/Engine/Classes/Engine/DecalActor.h b/Engine/Source/Runtime/Engine/Classes/Engine/DecalActor.h
--- a/Engine/Source/Runtime/Engine/Classes/Engine/DecalActor.h
+++ b/Engine/Source/Runtime/Engine/Classes/Engine/DecalActor.h
@@ -18,24 +18,27 @@ class ADecalActor : public AActor
 {
 	GENERATED_UCLASS_BODY()
 
+private:
 	/** The decal component for this decal actor */
-	UPROPERTY(Category=Decal, VisibleAnywhere, BlueprintReadOnly, meta=(ExposeFunctionCategories="Decal,Rendering|Components|Decal"))
-	TSubobjectPtr<class UDecalComponent> Decal;
+	UPROPERTY(Category = Decal, VisibleAnywhere, BlueprintReadOnly, meta = (ExposeFunctionCategories = "Decal,Rendering|Components|Decal", AllowPrivateAccess = "true"))
+	class UDecalComponent* Decal;
 
 #if WITH_EDITORONLY_DATA
 	/* Reference to the editor only arrow visualization component */
 	UPROPERTY()
-	TSubobjectPtr<class UArrowComponent> ArrowComponent;
+	class UArrowComponent* ArrowComponent;
 
 	/* Reference to the billboard component */
 	UPROPERTY()
-	TSubobjectPtr<UBillboardComponent> SpriteComponent;
+	UBillboardComponent* SpriteComponent;
 
 	/* Reference to the selected visualization box component */
 	UPROPERTY()
-	TSubobjectPtr<UBoxComponent> BoxComponent;
+	UBoxComponent* BoxComponent;
 #endif
 
+public:
+
 	// BEGIN DEPRECATED (use component functions now in level script)
 	UFUNCTION(BlueprintCallable, Category="Rendering|Components|Decal", meta=(DeprecatedFunction))
 	void SetDecalMaterial(class UMaterialInterface* NewDecalMaterial);
@@ -56,6 +59,18 @@ class ADecalActor : public AActor
 	virtual void EditorApplyScale(const FVector& DeltaScale, const FVector* PivotLocation, bool bAltDown, bool bShiftDown, bool bCtrlDown) override;
 	// End AActor interface.
 #endif // WITH_EDITOR
+
+public:
+	/** Returns Decal subobject **/
+	FORCEINLINE class UDecalComponent* GetDecal() const { return Decal; }
+#if WITH_EDITORONLY_DATA
+	/** Returns ArrowComponent subobject **/
+	FORCEINLINE class UArrowComponent* GetArrowComponent() const { return ArrowComponent; }
+	/** Returns SpriteComponent subobject **/
+	FORCEINLINE UBillboardComponent* GetSpriteComponent() const { return SpriteComponent; }
+	/** Returns BoxComponent subobject **/
+	FORCEINLINE UBoxComponent* GetBoxComponent() const { return BoxComponent; }
+#endif
 };
 
 
diff --git a/Engine/Source/Runtime/Engine/Classes/Engine/DirectionalLight.h b/Engine/Source/Runtime/Engine/Classes/Engine/DirectionalLight.h
--- a/Engine/Source/Runtime/Engine/Classes/Engine/DirectionalLight.h
+++ b/Engine/Source/Runtime/Engine/Classes/Engine/DirectionalLight.h
@@ -11,8 +11,9 @@ class ADirectionalLight : public ALight
 
 #if WITH_EDITORONLY_DATA
 	// Reference to editor visualization arrow
+private:
 	UPROPERTY()
-	TSubobjectPtr<UArrowComponent> ArrowComponent;
+	UArrowComponent* ArrowComponent;
 #endif
 
 public:
@@ -24,6 +25,12 @@ public:
 	virtual void PostEditChangeProperty(FPropertyChangedEvent& PropertyChangedEvent) override;
 #endif
 	// End UObject Interface
+
+public:
+#if WITH_EDITORONLY_DATA
+	/** Returns ArrowComponent subobject **/
+	FORCEINLINE UArrowComponent* GetArrowComponent() const { return ArrowComponent; }
+#endif
 };
 
 
diff --git a/Engine/Source/Runtime/Engine/Classes/Engine/DocumentationActor.h b/Engine/Source/Runtime/Engine/Classes/Engine/DocumentationActor.h
--- a/Engine/Source/Runtime/Engine/Classes/Engine/DocumentationActor.h
+++ b/Engine/Source/Runtime/Engine/Classes/Engine/DocumentationActor.h
@@ -45,9 +45,11 @@ class ENGINE_API ADocumentationActor : public AActor
 	UPROPERTY(Category = HelpDocumentation, EditAnywhere, BlueprintReadWrite, AdvancedDisplay)
 	FString DocumentLink; 
 	
- 	UPROPERTY(Category = Sprite, VisibleAnywhere, BlueprintReadOnly)
-	TSubobjectPtr<class UMaterialBillboardComponent> Billboard;
+private:
+	UPROPERTY(Category = Sprite, VisibleAnywhere, BlueprintReadOnly, meta = (AllowPrivateAccess = "true"))
+	class UMaterialBillboardComponent* Billboard;
 #endif
+public:
 
 	/** Determine the type of link contained in the DocumentLink member */
 	EDocumentationActorType::Type GetLinkType() const;
@@ -59,6 +61,12 @@ private:
 
 	EDocumentationActorType::Type	LinkType;
 
+
+public:
+#if WITH_EDITORONLY_DATA
+	/** Returns Billboard subobject **/
+	FORCEINLINE class UMaterialBillboardComponent* GetBillboard() const { return Billboard; }
+#endif
 };
 
 
diff --git a/Engine/Source/Runtime/Engine/Classes/Engine/ExponentialHeightFog.h b/Engine/Source/Runtime/Engine/Classes/Engine/ExponentialHeightFog.h
--- a/Engine/Source/Runtime/Engine/Classes/Engine/ExponentialHeightFog.h
+++ b/Engine/Source/Runtime/Engine/Classes/Engine/ExponentialHeightFog.h
@@ -9,9 +9,11 @@ class AExponentialHeightFog : public AInfo
 {
 	GENERATED_UCLASS_BODY()
 
+private:
 	/** @todo document */
-	UPROPERTY(Category=ExponentialHeightFog, VisibleAnywhere, BlueprintReadOnly)
-	TSubobjectPtr<class UExponentialHeightFogComponent> Component;
+	UPROPERTY(Category = ExponentialHeightFog, VisibleAnywhere, BlueprintReadOnly, meta = (AllowPrivateAccess = "true"))
+	class UExponentialHeightFogComponent* Component;
+public:
 
 	/** replicated copy of ExponentialHeightFogComponent's bEnabled property */
 	UPROPERTY(replicatedUsing=OnRep_bEnabled)
@@ -25,6 +27,10 @@ class AExponentialHeightFog : public AInfo
 	//Begin AActor Interface
 	virtual void PostInitializeComponents() override;
 	//End AActor Interface
+
+public:
+	/** Returns Component subobject **/
+	FORCEINLINE class UExponentialHeightFogComponent* GetComponent() const { return Component; }
 };
 
 
diff --git a/Engine/Source/Runtime/Engine/Classes/Engine/InteractiveFoliageActor.h b/Engine/Source/Runtime/Engine/Classes/Engine/InteractiveFoliageActor.h
--- a/Engine/Source/Runtime/Engine/Classes/Engine/InteractiveFoliageActor.h
+++ b/Engine/Source/Runtime/Engine/Classes/Engine/InteractiveFoliageActor.h
@@ -15,7 +15,7 @@ class AInteractiveFoliageActor : public AStaticMeshActor
 private:
 	/** Collision cylinder */
 	UPROPERTY()
-	TSubobjectPtr<class UCapsuleComponent> CapsuleComponent;
+	class UCapsuleComponent* CapsuleComponent;
 
 	/**
 	 * Position of the last actor to enter the collision cylinder.
diff --git a/Engine/Source/Runtime/Engine/Classes/Engine/Light.h b/Engine/Source/Runtime/Engine/Classes/Engine/Light.h
--- a/Engine/Source/Runtime/Engine/Classes/Engine/Light.h
+++ b/Engine/Source/Runtime/Engine/Classes/Engine/Light.h
@@ -9,9 +9,11 @@ class ENGINE_API ALight : public AActor
 {
 	GENERATED_UCLASS_BODY()
 
+private:
 	/** @todo document */
-	UPROPERTY(Category=Light, VisibleAnywhere, BlueprintReadOnly, meta=(ExposeFunctionCategories="Light,Rendering,Rendering|Components|Light"))
-	TSubobjectPtr<class ULightComponent> LightComponent;
+	UPROPERTY(Category = Light, VisibleAnywhere, BlueprintReadOnly, meta = (ExposeFunctionCategories = "Light,Rendering,Rendering|Components|Light", AllowPrivateAccess = "true"))
+	class ULightComponent* LightComponent;
+public:
 
 	/** replicated copy of LightComponent's bEnabled property */
 	UPROPERTY(replicatedUsing=OnRep_bEnabled)
@@ -70,6 +72,10 @@ public:
 	void Destroyed();
 	virtual bool IsLevelBoundsRelevant() const override { return false; }
 	// End AActor interface.
+
+public:
+	/** Returns LightComponent subobject **/
+	FORCEINLINE class ULightComponent* GetLightComponent() const { return LightComponent; }
 };
 
 
diff --git a/Engine/Source/Runtime/Engine/Classes/Engine/NavigationObjectBase.h b/Engine/Source/Runtime/Engine/Classes/Engine/NavigationObjectBase.h
--- a/Engine/Source/Runtime/Engine/Classes/Engine/NavigationObjectBase.h
+++ b/Engine/Source/Runtime/Engine/Classes/Engine/NavigationObjectBase.h
@@ -9,16 +9,22 @@ class ENGINE_API ANavigationObjectBase : public AActor, public INavAgentInterfac
 {
 	GENERATED_UCLASS_BODY()
 
+private:
 	UPROPERTY()
-	TSubobjectPtr<class UCapsuleComponent> CapsuleComponent;
+	class UCapsuleComponent* CapsuleComponent;
+public:
 
+private:
 	/** Normal editor sprite. */
 	UPROPERTY()
-	TSubobjectPtr<class UBillboardComponent> GoodSprite;
+	class UBillboardComponent* GoodSprite;
+public:
 
+private:
 	/** Used to draw bad collision intersection in editor. */
 	UPROPERTY()
-	TSubobjectPtr<class UBillboardComponent> BadSprite;
+	class UBillboardComponent* BadSprite;
+public:
 
 	/** True if this nav point was spawned to be a PIE player start. */
 	UPROPERTY()
@@ -53,6 +59,14 @@ class ENGINE_API ANavigationObjectBase : public AActor, public INavAgentInterfac
 	virtual FVector GetNavAgentLocation() const { return GetActorLocation(); }
 	virtual void GetMoveGoalReachTest(class AActor* MovingActor, const FVector& MoveOffset, FVector& GoalOffset, float& GoalRadius, float& GoalHalfHeight) const;
 	// INavAgentInterface end
+
+public:
+	/** Returns CapsuleComponent subobject **/
+	FORCEINLINE class UCapsuleComponent* GetCapsuleComponent() const { return CapsuleComponent; }
+	/** Returns GoodSprite subobject **/
+	FORCEINLINE class UBillboardComponent* GetGoodSprite() const { return GoodSprite; }
+	/** Returns BadSprite subobject **/
+	FORCEINLINE class UBillboardComponent* GetBadSprite() const { return BadSprite; }
 };
 
 
diff --git a/Engine/Source/Runtime/Engine/Classes/Engine/NiagaraActor.h b/Engine/Source/Runtime/Engine/Classes/Engine/NiagaraActor.h
--- a/Engine/Source/Runtime/Engine/Classes/Engine/NiagaraActor.h
+++ b/Engine/Source/Runtime/Engine/Classes/Engine/NiagaraActor.h
@@ -8,17 +8,29 @@ class ANiagaraActor : public AActor
 {
 	GENERATED_UCLASS_BODY()
 
+private:
 	/** Pointer to effect component */
 	UPROPERTY(VisibleAnywhere, Category=NiagaraActor)
-	TSubobjectPtr<class UNiagaraComponent> NiagaraComponent;
+	class UNiagaraComponent* NiagaraComponent;
 
 #if WITH_EDITORONLY_DATA
 	// Reference to sprite visualization component
 	UPROPERTY()
-	TSubobjectPtr<class UBillboardComponent> SpriteComponent;
+	class UBillboardComponent* SpriteComponent;
 
 	// Reference to arrow visualization component
 	UPROPERTY()
-	TSubobjectPtr<class UArrowComponent> ArrowComponent;
+	class UArrowComponent* ArrowComponent;
+
+#endif
+
+public:
+	/** Returns NiagaraComponent subobject **/
+	FORCEINLINE class UNiagaraComponent* GetNiagaraComponent() const { return NiagaraComponent; }
+#if WITH_EDITORONLY_DATA
+	/** Returns SpriteComponent subobject **/
+	FORCEINLINE class UBillboardComponent* GetSpriteComponent() const { return SpriteComponent; }
+	/** Returns ArrowComponent subobject **/
+	FORCEINLINE class UArrowComponent* GetArrowComponent() const { return ArrowComponent; }
 #endif
 };
diff --git a/Engine/Source/Runtime/Engine/Classes/Engine/Note.h b/Engine/Source/Runtime/Engine/Classes/Engine/Note.h
--- a/Engine/Source/Runtime/Engine/Classes/Engine/Note.h
+++ b/Engine/Source/Runtime/Engine/Classes/Engine/Note.h
@@ -18,12 +18,14 @@ class ANote : public AActor
 	FString Text;
 
 	// Reference to sprite visualization component
+private:
 	UPROPERTY()
-	TSubobjectPtr<class UBillboardComponent> SpriteComponent;
+	class UBillboardComponent* SpriteComponent;
 
 	// Reference to arrow visualization component
 	UPROPERTY()
-	TSubobjectPtr<class UArrowComponent> ArrowComponent;
+	class UArrowComponent* ArrowComponent;
+public:
 
 #endif // WITH_EDITORONLY_DATA
 
@@ -32,6 +34,14 @@ class ANote : public AActor
 	virtual void CheckForErrors();
 	// End AActor Interface
 #endif
+
+public:
+#if WITH_EDITORONLY_DATA
+	/** Returns SpriteComponent subobject **/
+	FORCEINLINE class UBillboardComponent* GetSpriteComponent() const { return SpriteComponent; }
+	/** Returns ArrowComponent subobject **/
+	FORCEINLINE class UArrowComponent* GetArrowComponent() const { return ArrowComponent; }
+#endif
 };
 
 
diff --git a/Engine/Source/Runtime/Engine/Classes/Engine/ReflectionCapture.h b/Engine/Source/Runtime/Engine/Classes/Engine/ReflectionCapture.h
--- a/Engine/Source/Runtime/Engine/Classes/Engine/ReflectionCapture.h
+++ b/Engine/Source/Runtime/Engine/Classes/Engine/ReflectionCapture.h
@@ -8,21 +8,32 @@ class AReflectionCapture : public AActor
 {
 	GENERATED_UCLASS_BODY()
 
+private:
 	/** Reflection capture component. */
-	UPROPERTY(Category=DecalActor, VisibleAnywhere, BlueprintReadOnly)
-	TSubobjectPtr<class UReflectionCaptureComponent> CaptureComponent;
+	UPROPERTY(Category = DecalActor, VisibleAnywhere, BlueprintReadOnly, meta = (AllowPrivateAccess = "true"))
+	class UReflectionCaptureComponent* CaptureComponent;
 
 #if WITH_EDITORONLY_DATA
 	UPROPERTY()
-	TSubobjectPtr<UBillboardComponent> SpriteComponent;
+	UBillboardComponent* SpriteComponent;
 #endif // WITH_EDITORONLY_DATA
-	
+
+public:	
+
 	virtual bool IsLevelBoundsRelevant() const override { return false; }
-public:
+
 #if WITH_EDITOR
 	virtual void PostEditMove(bool bFinished) override;
 #endif // WITH_EDITOR
 
+
+public:
+	/** Returns CaptureComponent subobject **/
+	FORCEINLINE class UReflectionCaptureComponent* GetCaptureComponent() const { return CaptureComponent; }
+#if WITH_EDITORONLY_DATA
+	/** Returns SpriteComponent subobject **/
+	FORCEINLINE UBillboardComponent* GetSpriteComponent() const { return SpriteComponent; }
+#endif
 };
 
 
diff --git a/Engine/Source/Runtime/Engine/Classes/Engine/SceneCapture.h b/Engine/Source/Runtime/Engine/Classes/Engine/SceneCapture.h
--- a/Engine/Source/Runtime/Engine/Classes/Engine/SceneCapture.h
+++ b/Engine/Source/Runtime/Engine/Classes/Engine/SceneCapture.h
@@ -13,10 +13,14 @@ class ASceneCapture : public AActor
 {
 	GENERATED_UCLASS_BODY()
 
+private:
 	/** To display the 3d camera in the editor. */
 	UPROPERTY()
-	TSubobjectPtr<class UStaticMeshComponent> MeshComp;
+	class UStaticMeshComponent* MeshComp;
+
 public:
+	/** Returns MeshComp subobject **/
+	FORCEINLINE class UStaticMeshComponent* GetMeshComp() const { return MeshComp; }
 };
 
 
diff --git a/Engine/Source/Runtime/Engine/Classes/Engine/SceneCapture2D.h b/Engine/Source/Runtime/Engine/Classes/Engine/SceneCapture2D.h
--- a/Engine/Source/Runtime/Engine/Classes/Engine/SceneCapture2D.h
+++ b/Engine/Source/Runtime/Engine/Classes/Engine/SceneCapture2D.h
@@ -13,13 +13,13 @@ class ASceneCapture2D : public ASceneCapture
 {
 	GENERATED_UCLASS_BODY()
 
+private:
 	/** Scene capture component. */
-	UPROPERTY(Category=DecalActor, VisibleAnywhere, BlueprintReadOnly)
-	TSubobjectPtr<class USceneCaptureComponent2D> CaptureComponent2D;
-
+	UPROPERTY(Category = DecalActor, VisibleAnywhere, BlueprintReadOnly, meta=(AllowPrivateAccess = "true"))
+	class USceneCaptureComponent2D* CaptureComponent2D;
 	/** To allow drawing the camera frustum in the editor. */
 	UPROPERTY()
-	TSubobjectPtr<class UDrawFrustumComponent> DrawFrustum;
+	class UDrawFrustumComponent* DrawFrustum;
 
 public:
 
@@ -32,6 +32,12 @@ public:
 
 	UFUNCTION(BlueprintCallable, Category="Rendering")
 	void OnInterpToggle(bool bEnable);
+
+public:
+	/** Returns CaptureComponent2D subobject **/
+	FORCEINLINE class USceneCaptureComponent2D* GetCaptureComponent2D() const { return CaptureComponent2D; }
+	/** Returns DrawFrustum subobject **/
+	FORCEINLINE class UDrawFrustumComponent* GetDrawFrustum() const { return DrawFrustum; }
 };
 
 
diff --git a/Engine/Source/Runtime/Engine/Classes/Engine/SceneCaptureCube.h b/Engine/Source/Runtime/Engine/Classes/Engine/SceneCaptureCube.h
--- a/Engine/Source/Runtime/Engine/Classes/Engine/SceneCaptureCube.h
+++ b/Engine/Source/Runtime/Engine/Classes/Engine/SceneCaptureCube.h
@@ -13,13 +13,14 @@ class ASceneCaptureCube : public ASceneCapture
 {
 	GENERATED_UCLASS_BODY()
 
+private:
 	/** Scene capture component. */
-	UPROPERTY(Category = DecalActor, VisibleAnywhere, BlueprintReadOnly)
-	TSubobjectPtr<class USceneCaptureComponentCube> CaptureComponentCube;
+	UPROPERTY(Category = DecalActor, VisibleAnywhere, BlueprintReadOnly, meta=(AllowPrivateAccess = "true"))
+	class USceneCaptureComponentCube* CaptureComponentCube;
 
 	/** To allow drawing the camera frustum in the editor. */
 	UPROPERTY()
-	TSubobjectPtr<class UDrawFrustumComponent> DrawFrustum;
+	class UDrawFrustumComponent* DrawFrustum;
 
 public:
 
@@ -35,6 +36,12 @@ public:
 
 	UFUNCTION(BlueprintCallable, Category = "Rendering")
 	void OnInterpToggle(bool bEnable);
+
+public:
+	/** Returns CaptureComponentCube subobject **/
+	FORCEINLINE class USceneCaptureComponentCube* GetCaptureComponentCube() const { return CaptureComponentCube; }
+	/** Returns DrawFrustum subobject **/
+	FORCEINLINE class UDrawFrustumComponent* GetDrawFrustum() const { return DrawFrustum; }
 };
 
 
diff --git a/Engine/Source/Runtime/Engine/Classes/Engine/SkyLight.h b/Engine/Source/Runtime/Engine/Classes/Engine/SkyLight.h
--- a/Engine/Source/Runtime/Engine/Classes/Engine/SkyLight.h
+++ b/Engine/Source/Runtime/Engine/Classes/Engine/SkyLight.h
@@ -9,9 +9,11 @@ class ENGINE_API ASkyLight : public AInfo
 {
 	GENERATED_UCLASS_BODY()
 
+private:
 	/** @todo document */
-	UPROPERTY(Category=Light, VisibleAnywhere, BlueprintReadOnly, meta=(ExposeFunctionCategories="Light,Rendering,Rendering|Components|SkyLight"))
-	TSubobjectPtr<class USkyLightComponent> LightComponent;
+	UPROPERTY(Category = Light, VisibleAnywhere, BlueprintReadOnly, meta = (ExposeFunctionCategories = "Light,Rendering,Rendering|Components|SkyLight", AllowPrivateAccess = "true"))
+	class USkyLightComponent* LightComponent;
+public:
 
 	/** replicated copy of LightComponent's bEnabled property */
 	UPROPERTY(replicatedUsing=OnRep_bEnabled)
@@ -20,6 +22,9 @@ class ENGINE_API ASkyLight : public AInfo
 	/** Replication Notification Callbacks */
 	UFUNCTION()
 	virtual void OnRep_bEnabled();
+
+	/** Returns LightComponent subobject **/
+	FORCEINLINE class USkyLightComponent* GetLightComponent() const { return LightComponent; }
 };
 
 
diff --git a/Engine/Source/Runtime/Engine/Classes/Engine/SphereReflectionCapture.h b/Engine/Source/Runtime/Engine/Classes/Engine/SphereReflectionCapture.h
--- a/Engine/Source/Runtime/Engine/Classes/Engine/SphereReflectionCapture.h
+++ b/Engine/Source/Runtime/Engine/Classes/Engine/SphereReflectionCapture.h
@@ -14,11 +14,10 @@ class ASphereReflectionCapture : public AReflectionCapture
 {
 	GENERATED_UCLASS_BODY()
 
+private:
 	/** Sphere component used to visualize the capture radius */
 	UPROPERTY()
-	TSubobjectPtr<UDrawSphereComponent> DrawCaptureRadius;
-
-public:
+	UDrawSphereComponent* DrawCaptureRadius;
 
 #if WITH_EDITOR
 	// Begin AActor interface.
@@ -26,6 +25,8 @@ public:
 	// End AActor interface.
 #endif
 
+	/** Returns DrawCaptureRadius subobject **/
+	FORCEINLINE UDrawSphereComponent* GetDrawCaptureRadius() const { return DrawCaptureRadius; }
 };
 
 
diff --git a/Engine/Source/Runtime/Engine/Classes/Engine/SpotLight.h b/Engine/Source/Runtime/Engine/Classes/Engine/SpotLight.h
--- a/Engine/Source/Runtime/Engine/Classes/Engine/SpotLight.h
+++ b/Engine/Source/Runtime/Engine/Classes/Engine/SpotLight.h
@@ -15,8 +15,10 @@ class ASpotLight : public ALight
 
 #if WITH_EDITORONLY_DATA
 	// Reference to editor arrow component visualization 
+private:
 	UPROPERTY()
-	TSubobjectPtr<class UArrowComponent> ArrowComponent;
+	class UArrowComponent* ArrowComponent;
+public:
 #endif
 
 	// BEGIN DEPRECATED (use component functions now in level script)
@@ -44,6 +46,12 @@ class ASpotLight : public ALight
 #endif
 	// End UObject Interface
 
+
+public:
+#if WITH_EDITORONLY_DATA
+	/** Returns ArrowComponent subobject **/
+	FORCEINLINE class UArrowComponent* GetArrowComponent() const { return ArrowComponent; }
+#endif
 };
 
 
diff --git a/Engine/Source/Runtime/Engine/Classes/Engine/StaticMeshActor.h b/Engine/Source/Runtime/Engine/Classes/Engine/StaticMeshActor.h
--- a/Engine/Source/Runtime/Engine/Classes/Engine/StaticMeshActor.h
+++ b/Engine/Source/Runtime/Engine/Classes/Engine/StaticMeshActor.h
@@ -16,8 +16,10 @@ class ENGINE_API AStaticMeshActor : public AActor
 {
 	GENERATED_UCLASS_BODY()
 
-	UPROPERTY(Category=StaticMeshActor, VisibleAnywhere, BlueprintReadOnly, meta=(ExposeFunctionCategories="Mesh,Rendering,Physics,Components|StaticMesh"))
-	TSubobjectPtr<class UStaticMeshComponent> StaticMeshComponent;
+private:
+	UPROPERTY(Category = StaticMeshActor, VisibleAnywhere, BlueprintReadOnly, meta = (ExposeFunctionCategories = "Mesh,Rendering,Physics,Components|StaticMesh", AllowPrivateAccess = "true"))
+	class UStaticMeshComponent* StaticMeshComponent;
+public:
 	
 	virtual void BeginPlay() override;
 
@@ -43,6 +45,10 @@ protected:
 	virtual void PostEditChangeChainProperty(FPropertyChangedChainEvent& PropertyChangedEvent) override;
 #endif // WITH_EDITOR	
 	// End UObject interface.
+
+public:
+	/** Returns StaticMeshComponent subobject **/
+	FORCEINLINE class UStaticMeshComponent* GetStaticMeshComponent() const { return StaticMeshComponent; }
 };
 
 
diff --git a/Engine/Source/Runtime/Engine/Classes/Engine/TargetPoint.h b/Engine/Source/Runtime/Engine/Classes/Engine/TargetPoint.h
--- a/Engine/Source/Runtime/Engine/Classes/Engine/TargetPoint.h
+++ b/Engine/Source/Runtime/Engine/Classes/Engine/TargetPoint.h
@@ -13,11 +13,18 @@ class ATargetPoint : public AActor
 	GENERATED_UCLASS_BODY()
 
 #if WITH_EDITORONLY_DATA
-	UPROPERTY(VisibleAnywhere, BlueprintReadOnly, Category=Display)
-	TSubobjectPtr<class UBillboardComponent> SpriteComponent;
+private:
+	UPROPERTY(VisibleAnywhere, BlueprintReadOnly, Category = Display, meta = (AllowPrivateAccess = "true"))
+	class UBillboardComponent* SpriteComponent;
 
 	UPROPERTY()
-	TSubobjectPtr<class UArrowComponent> ArrowComponent;
+	class UArrowComponent* ArrowComponent;
+
+public:
+	/** Returns SpriteComponent subobject **/
+	FORCEINLINE class UBillboardComponent* GetSpriteComponent() const { return SpriteComponent; }
+	/** Returns ArrowComponent subobject **/
+	FORCEINLINE class UArrowComponent* GetArrowComponent() const { return ArrowComponent; }
 #endif
 };
 
diff --git a/Engine/Source/Runtime/Engine/Classes/Engine/TextRenderActor.h b/Engine/Source/Runtime/Engine/Classes/Engine/TextRenderActor.h
--- a/Engine/Source/Runtime/Engine/Classes/Engine/TextRenderActor.h
+++ b/Engine/Source/Runtime/Engine/Classes/Engine/TextRenderActor.h
@@ -15,14 +15,23 @@ class ATextRenderActor : public AActor
 
 	friend class UActorFactoryTextRender;
 
+private:
 	/** Component to render a text in 3d with a font */
-	UPROPERTY(Category=TextRenderActor, VisibleAnywhere, BlueprintReadOnly, meta=(ExposeFunctionCategories="Rendering|Components|TextRender"))
-	TSubobjectPtr<class UTextRenderComponent> TextRender;
+	UPROPERTY(Category = TextRenderActor, VisibleAnywhere, BlueprintReadOnly, meta = (ExposeFunctionCategories = "Rendering|Components|TextRender", AllowPrivateAccess = "true"))
+	class UTextRenderComponent* TextRender;
 
 #if WITH_EDITORONLY_DATA
 	// Reference to the billboard component
 	UPROPERTY()
-	TSubobjectPtr<UBillboardComponent> SpriteComponent;
+	UBillboardComponent* SpriteComponent;
+#endif
+
+public:
+	/** Returns TextRender subobject **/
+	FORCEINLINE class UTextRenderComponent* GetTextRender() const { return TextRender; }
+#if WITH_EDITORONLY_DATA
+	/** Returns SpriteComponent subobject **/
+	FORCEINLINE UBillboardComponent* GetSpriteComponent() const { return SpriteComponent; }
 #endif
 };
 
diff --git a/Engine/Source/Runtime/Engine/Classes/Engine/TriggerBase.h b/Engine/Source/Runtime/Engine/Classes/Engine/TriggerBase.h
--- a/Engine/Source/Runtime/Engine/Classes/Engine/TriggerBase.h
+++ b/Engine/Source/Runtime/Engine/Classes/Engine/TriggerBase.h
@@ -13,13 +13,20 @@ class ATriggerBase : public AActor
 {
 	GENERATED_UCLASS_BODY()
 
+private:
 	/** Shape component used for collision */
-	UPROPERTY(Category=TriggerBase, VisibleAnywhere, BlueprintReadOnly)
-	TSubobjectPtr<UShapeComponent> CollisionComponent;
+	UPROPERTY(Category = TriggerBase, VisibleAnywhere, BlueprintReadOnly, meta = (AllowPrivateAccess = "true"))
+	UShapeComponent* CollisionComponent;
 
 	/** Billboard used to see the trigger in the editor */
 	UPROPERTY()
-	TSubobjectPtr<UBillboardComponent> SpriteComponent;
+	UBillboardComponent* SpriteComponent;
+
+public:
+	/** Returns CollisionComponent subobject **/
+	FORCEINLINE UShapeComponent* GetCollisionComponent() const { return CollisionComponent; }
+	/** Returns SpriteComponent subobject **/
+	FORCEINLINE UBillboardComponent* GetSpriteComponent() const { return SpriteComponent; }
 };
 
 
diff --git a/Engine/Source/Runtime/Engine/Classes/Engine/WindDirectionalSource.h b/Engine/Source/Runtime/Engine/Classes/Engine/WindDirectionalSource.h
--- a/Engine/Source/Runtime/Engine/Classes/Engine/WindDirectionalSource.h
+++ b/Engine/Source/Runtime/Engine/Classes/Engine/WindDirectionalSource.h
@@ -8,12 +8,22 @@ class ENGINE_API AWindDirectionalSource : public AInfo
 {
 	GENERATED_UCLASS_BODY()
 
-	UPROPERTY(Category=WindDirectionalSource, VisibleAnywhere, BlueprintReadOnly)
-	TSubobjectPtr<class UWindDirectionalSourceComponent> Component;
+private:
+	UPROPERTY(Category = WindDirectionalSource, VisibleAnywhere, BlueprintReadOnly, meta = (AllowPrivateAccess = "true"))
+	class UWindDirectionalSourceComponent* Component;
 
 #if WITH_EDITORONLY_DATA
 	UPROPERTY()
-	TSubobjectPtr<class UArrowComponent> ArrowComponent;
+	class UArrowComponent* ArrowComponent;
+#endif
+
+public:
+	/** Returns Component subobject **/
+	FORCEINLINE class UWindDirectionalSourceComponent* GetComponent() const { return Component; }
+
+#if WITH_EDITORONLY_DATA
+	/** Returns ArrowComponent subobject **/
+	FORCEINLINE class UArrowComponent* GetArrowComponent() const { return ArrowComponent; }
 #endif
 };
 
diff --git a/Engine/Source/Runtime/Engine/Classes/GameFramework/Character.h b/Engine/Source/Runtime/Engine/Classes/GameFramework/Character.h
--- a/Engine/Source/Runtime/Engine/Classes/GameFramework/Character.h
+++ b/Engine/Source/Runtime/Engine/Classes/GameFramework/Character.h
@@ -207,22 +207,24 @@ public:
 	 */
 	ACharacter(const FObjectInitializer& ObjectInitializer);
 
+private:
 	/** The main skeletal mesh associated with this Character (optional sub-object). */
-	UPROPERTY(Category=Character, VisibleAnywhere, BlueprintReadOnly)
-	TSubobjectPtr<class USkeletalMeshComponent> Mesh;
+	UPROPERTY(Category = Character, VisibleAnywhere, BlueprintReadOnly, meta = (AllowPrivateAccess = "true"))
+	class USkeletalMeshComponent* Mesh;
 
 #if WITH_EDITORONLY_DATA
 	UPROPERTY()
-	TSubobjectPtr<class UArrowComponent> ArrowComponent;
+	class UArrowComponent* ArrowComponent;
 #endif
 
 	/** Movement component used for movement logic in various movement modes (walking, falling, etc), containing relevant settings and functions to control movement. */
-	UPROPERTY(Category=Character, VisibleAnywhere, BlueprintReadOnly)
-	TSubobjectPtr<class UCharacterMovementComponent> CharacterMovement;
+	UPROPERTY(Category = Character, VisibleAnywhere, BlueprintReadOnly, meta = (AllowPrivateAccess = "true"))
+	class UCharacterMovementComponent* CharacterMovement;
 
 	/** The CapsuleComponent being used for movement collision (by CharacterMovement). Always treated as being vertically aligned in simple collision check functions. */
-	UPROPERTY(Category=Character, VisibleAnywhere, BlueprintReadOnly)
-	TSubobjectPtr<class UCapsuleComponent> CapsuleComponent;
+	UPROPERTY(Category = Character, VisibleAnywhere, BlueprintReadOnly, meta = (AllowPrivateAccess = "true"))
+	class UCapsuleComponent* CapsuleComponent;
+public:
 
 	/** Name of the MeshComponent. Use this name if you want to prevent creation of the component (with ObjectInitializer.DoNotCreateDefaultSubobject). */
 	static FName MeshComponentName;
@@ -679,4 +681,16 @@ public:
 
 	/** Called on the actor right before replication occurs */
 	virtual void PreReplication( IRepChangedPropertyTracker & ChangedPropertyTracker ) override;
+
+public:
+	/** Returns Mesh subobject **/
+	FORCEINLINE class USkeletalMeshComponent* GetMesh() const { return Mesh; }
+#if WITH_EDITORONLY_DATA
+	/** Returns ArrowComponent subobject **/
+	FORCEINLINE class UArrowComponent* GetArrowComponent() const { return ArrowComponent; }
+#endif
+	/** Returns CharacterMovement subobject **/
+	FORCEINLINE class UCharacterMovementComponent* GetCharacterMovement() const { return CharacterMovement; }
+	/** Returns CapsuleComponent subobject **/
+	FORCEINLINE class UCapsuleComponent* GetCapsuleComponent() const { return CapsuleComponent; }
 };
diff --git a/Engine/Source/Runtime/Engine/Classes/GameFramework/Controller.h b/Engine/Source/Runtime/Engine/Classes/GameFramework/Controller.h
--- a/Engine/Source/Runtime/Engine/Classes/GameFramework/Controller.h
+++ b/Engine/Source/Runtime/Engine/Classes/GameFramework/Controller.h
@@ -54,10 +54,10 @@ public:
 	UPROPERTY(replicatedUsing=OnRep_PlayerState, BlueprintReadOnly, Category="Controller")
 	class APlayerState* PlayerState;
 
-protected:
+private:
 	/** Component to give controllers a transform and enable attachment if desired. */
 	UPROPERTY()
-	TSubobjectPtr<class USceneComponent> TransformComponent;
+	class USceneComponent* TransformComponent;
 
 public:
 
@@ -304,4 +304,8 @@ private:
 	/** Called when the controller has instigated damage in any way */
 	UPROPERTY(BlueprintAssignable)
 	FInstigatedAnyDamageSignature OnInstigatedAnyDamage;
+
+protected:
+	/** Returns TransformComponent subobject **/
+	FORCEINLINE class USceneComponent* GetTransformComponent() const { return TransformComponent; }
 };
diff --git a/Engine/Source/Runtime/Engine/Classes/GameFramework/DefaultPawn.h b/Engine/Source/Runtime/Engine/Classes/GameFramework/DefaultPawn.h
--- a/Engine/Source/Runtime/Engine/Classes/GameFramework/DefaultPawn.h
+++ b/Engine/Source/Runtime/Engine/Classes/GameFramework/DefaultPawn.h
@@ -87,26 +87,37 @@ public:
 	/** Name of the MovementComponent.  Use this name if you want to use a different class (with ObjectInitializer.SetDefaultSubobjectClass). */
 	static FName MovementComponentName;
 
+private:
 	/** DefaultPawn movement component */
-	UPROPERTY(Category=Pawn, VisibleAnywhere, BlueprintReadOnly)
-	TSubobjectPtr<UFloatingPawnMovement> MovementComponent;
+	UPROPERTY(Category = Pawn, VisibleAnywhere, BlueprintReadOnly, meta = (AllowPrivateAccess = "true"))
+	UFloatingPawnMovement* MovementComponent;
 
+public:
 	/** Name of the CollisionComponent. */
 	static FName CollisionComponentName;
 
+private:
 	/** DefaultPawn collision component */
-	UPROPERTY(Category=Pawn, VisibleAnywhere, BlueprintReadOnly)
-	TSubobjectPtr<USphereComponent> CollisionComponent;
+	UPROPERTY(Category = Pawn, VisibleAnywhere, BlueprintReadOnly, meta = (AllowPrivateAccess = "true"))
+	USphereComponent* CollisionComponent;
+public:
 
 	/** Name of the MeshComponent. Use this name if you want to prevent creation of the component (with ObjectInitializer.DoNotCreateDefaultSubobject). */
 	static FName MeshComponentName;
 
+private:
 	/** The mesh associated with this Pawn. */
-	UPROPERTY(Category=Pawn, VisibleAnywhere, BlueprintReadOnly)
-	TSubobjectPtr<UStaticMeshComponent> MeshComponent;
+	UPROPERTY(Category = Pawn, VisibleAnywhere, BlueprintReadOnly, meta = (AllowPrivateAccess = "true"))
+	UStaticMeshComponent* MeshComponent;
+public:
 
 	/** If true, adds default input bindings for movement and camera look. */
 	UPROPERTY(Category=Pawn, EditAnywhere, BlueprintReadOnly)
 	uint32 bAddDefaultMovementBindings:1;
+
+	/** Returns CollisionComponent subobject **/
+	FORCEINLINE USphereComponent* GetCollisionComponent() const { return CollisionComponent; }
+	/** Returns MeshComponent subobject **/
+	FORCEINLINE UStaticMeshComponent* GetMeshComponent() const { return MeshComponent; }
 };
 
diff --git a/Engine/Source/Runtime/Engine/Classes/GameFramework/Info.h b/Engine/Source/Runtime/Engine/Classes/GameFramework/Info.h
--- a/Engine/Source/Runtime/Engine/Classes/GameFramework/Info.h
+++ b/Engine/Source/Runtime/Engine/Classes/GameFramework/Info.h
@@ -15,13 +15,21 @@ class AInfo : public AActor
 	GENERATED_UCLASS_BODY()
 
 #if WITH_EDITORONLY_DATA
+private:
 	/** Billboard Component displayed in editor */
 	UPROPERTY()
-	TSubobjectPtr<class UBillboardComponent> SpriteComponent;
+	class UBillboardComponent* SpriteComponent;
+public:
 #endif
 
 	/** Indicates whether this actor should participate in level bounds calculations. */
 	virtual bool IsLevelBoundsRelevant() const override { return false; }
+
+public:
+#if WITH_EDITORONLY_DATA
+	/** Returns SpriteComponent subobject **/
+	FORCEINLINE class UBillboardComponent* GetSpriteComponent() const { return SpriteComponent; }
+#endif
 };
 
 
diff --git a/Engine/Source/Runtime/Engine/Classes/GameFramework/PlayerStart.h b/Engine/Source/Runtime/Engine/Classes/GameFramework/PlayerStart.h
--- a/Engine/Source/Runtime/Engine/Classes/GameFramework/PlayerStart.h
+++ b/Engine/Source/Runtime/Engine/Classes/GameFramework/PlayerStart.h
@@ -17,13 +17,20 @@ class ENGINE_API APlayerStart : public ANavigationObjectBase
 	FName PlayerStartTag;
 
 #if WITH_EDITORONLY_DATA
+private:
 	UPROPERTY()
-	TSubobjectPtr<class UArrowComponent> ArrowComponent;
+	class UArrowComponent* ArrowComponent;
+public:
 #endif
 
 	virtual void PostInitializeComponents() override;
 	
 	virtual void PostUnregisterAllComponents() override;
+
+#if WITH_EDITORONLY_DATA
+	/** Returns ArrowComponent subobject **/
+	FORCEINLINE class UArrowComponent* GetArrowComponent() const { return ArrowComponent; }
+#endif
 };
 
 
diff --git a/Engine/Source/Runtime/Engine/Classes/GameFramework/WheeledVehicle.h b/Engine/Source/Runtime/Engine/Classes/GameFramework/WheeledVehicle.h
--- a/Engine/Source/Runtime/Engine/Classes/GameFramework/WheeledVehicle.h
+++ b/Engine/Source/Runtime/Engine/Classes/GameFramework/WheeledVehicle.h
@@ -11,13 +11,15 @@ class ENGINE_API AWheeledVehicle : public APawn
 {
 	GENERATED_UCLASS_BODY()
 
+private:
 	/**  The main skeletal mesh associated with this Vehicle */
-	UPROPERTY(Category=Vehicle, VisibleDefaultsOnly, BlueprintReadOnly)
-	TSubobjectPtr<class USkeletalMeshComponent> Mesh;
+	UPROPERTY(Category = Vehicle, VisibleDefaultsOnly, BlueprintReadOnly, meta = (AllowPrivateAccess = "true"))
+	class USkeletalMeshComponent* Mesh;
 
 	/** vehicle simulation component */
-	UPROPERTY(Category=Vehicle, VisibleDefaultsOnly, BlueprintReadOnly)
-	TSubobjectPtr<class UWheeledVehicleMovementComponent> VehicleMovement;
+	UPROPERTY(Category = Vehicle, VisibleDefaultsOnly, BlueprintReadOnly, meta = (AllowPrivateAccess = "true"))
+	class UWheeledVehicleMovementComponent* VehicleMovement;
+public:
 
 	/** Name of the MeshComponent. Use this name if you want to prevent creation of the component (with ObjectInitializer.DoNotCreateDefaultSubobject). */
 	static FName VehicleMeshComponentName;
@@ -34,4 +36,10 @@ class ENGINE_API AWheeledVehicle : public APawn
 	// Begin AActor interface
 	virtual void DisplayDebug(class UCanvas* Canvas, const FDebugDisplayInfo& DebugDisplay, float& YL, float& YPos) override;
 	// End Actor interface
+
+public:
+	/** Returns Mesh subobject **/
+	FORCEINLINE class USkeletalMeshComponent* GetMesh() const { return Mesh; }
+	/** Returns VehicleMovement subobject **/
+	FORCEINLINE class UWheeledVehicleMovementComponent* GetVehicleMovement() const { return VehicleMovement; }
 };
diff --git a/Engine/Source/Runtime/Engine/Classes/Materials/MaterialInstanceActor.h b/Engine/Source/Runtime/Engine/Classes/Materials/MaterialInstanceActor.h
--- a/Engine/Source/Runtime/Engine/Classes/Materials/MaterialInstanceActor.h
+++ b/Engine/Source/Runtime/Engine/Classes/Materials/MaterialInstanceActor.h
@@ -18,16 +18,25 @@ class AMaterialInstanceActor : public AActor
 	TArray<class AActor*> TargetActors;
 
 #if WITH_EDITORONLY_DATA
+private:
 	// Reference to actor sprite
-	TSubobjectPtr<UBillboardComponent> SpriteComponent;
+	UBillboardComponent* SpriteComponent;
 #endif
 
+public:
+
 	// Begin UObject Interface
 	virtual void PostLoad() override;
 #if WITH_EDITOR
 	virtual void PostEditChangeProperty(FPropertyChangedEvent& PropertyChangedEvent) override;
 #endif // WITH_EDITOR
 	// End UObject Interface
+
+public:
+#if WITH_EDITORONLY_DATA
+	/** Returns SpriteComponent subobject **/
+	FORCEINLINE UBillboardComponent* GetSpriteComponent() const { return SpriteComponent; }
+#endif
 };
 
 
diff --git a/Engine/Source/Runtime/Engine/Classes/Matinee/MatineeActor.h b/Engine/Source/Runtime/Engine/Classes/Matinee/MatineeActor.h
--- a/Engine/Source/Runtime/Engine/Classes/Matinee/MatineeActor.h
+++ b/Engine/Source/Runtime/Engine/Classes/Matinee/MatineeActor.h
@@ -162,8 +162,10 @@ class AMatineeActor : public AActor
 #if WITH_EDITORONLY_DATA
 
 	// Reference to the actor sprite
+private:
 	UPROPERTY()
-	TSubobjectPtr<class UBillboardComponent> SpriteComponent;
+	class UBillboardComponent* SpriteComponent;
+public:
 
 	UPROPERTY(transient)
 	uint32 bIsBeingEdited:1;
@@ -457,6 +459,12 @@ public:
 
 	/** Util to get the name of the function to find for the given event name */
 	ENGINE_API FName GetFunctionNameForEvent(FName EventName);
+
+public:
+#if WITH_EDITORONLY_DATA
+	/** Returns SpriteComponent subobject **/
+	FORCEINLINE class UBillboardComponent* GetSpriteComponent() const { return SpriteComponent; }
+#endif
 };
 
 
diff --git a/Engine/Source/Runtime/Engine/Classes/Particles/Emitter.h b/Engine/Source/Runtime/Engine/Classes/Particles/Emitter.h
--- a/Engine/Source/Runtime/Engine/Classes/Particles/Emitter.h
+++ b/Engine/Source/Runtime/Engine/Classes/Particles/Emitter.h
@@ -49,8 +49,10 @@ class ENGINE_API AEmitter : public AActor
 {
 	GENERATED_UCLASS_BODY()
 
-	UPROPERTY(Category=Emitter, VisibleAnywhere, BlueprintReadOnly, meta=(ExposeFunctionCategories="Particles|Beam,Particles|Parameters,Particles,Effects|Components|ParticleSystem,Rendering,Activation,Components|Activation"))
-	TSubobjectPtr<class UParticleSystemComponent> ParticleSystemComponent;
+private:
+	UPROPERTY(Category = Emitter, VisibleAnywhere, BlueprintReadOnly, meta = (ExposeFunctionCategories = "Particles|Beam,Particles|Parameters,Particles,Effects|Components|ParticleSystem,Rendering,Activation,Components|Activation", AllowPrivateAccess = "true"))
+	class UParticleSystemComponent* ParticleSystemComponent;
+public:
 
 	UPROPERTY()
 	uint32 bDestroyOnSystemFinish:1;
@@ -76,11 +78,13 @@ class ENGINE_API AEmitter : public AActor
 
 #if WITH_EDITORONLY_DATA
 
+private:
 	UPROPERTY()
-	TSubobjectPtr<class UBillboardComponent> SpriteComponent;
+	class UBillboardComponent* SpriteComponent;
 
 	UPROPERTY()
-	TSubobjectPtr<class UArrowComponent> ArrowComponent;
+	class UArrowComponent* ArrowComponent;
+public:
 
 #endif
 
@@ -136,6 +140,16 @@ class ENGINE_API AEmitter : public AActor
 	 */
 	void ResetInLevel();
 #endif
+
+public:
+	/** Returns ParticleSystemComponent subobject **/
+	FORCEINLINE class UParticleSystemComponent* GetParticleSystemComponent() { return ParticleSystemComponent; }
+#if WITH_EDITORONLY_DATA
+	/** Returns SpriteComponent subobject **/
+	FORCEINLINE class UBillboardComponent* GetSpriteComponent() const { return SpriteComponent; }
+	/** Returns ArrowComponent subobject **/
+	FORCEINLINE class UArrowComponent* GetArrowComponent() const { return ArrowComponent; }
+#endif
 };
 
 
diff --git a/Engine/Source/Runtime/Engine/Classes/PhysicsEngine/DestructibleActor.h b/Engine/Source/Runtime/Engine/Classes/PhysicsEngine/DestructibleActor.h
--- a/Engine/Source/Runtime/Engine/Classes/PhysicsEngine/DestructibleActor.h
+++ b/Engine/Source/Runtime/Engine/Classes/PhysicsEngine/DestructibleActor.h
@@ -19,8 +19,10 @@ class ADestructibleActor : public AActor
 	/**
 	 * The component which holds the skinned mesh and physics data for this actor.
 	 */
-	UPROPERTY(VisibleAnywhere, BlueprintReadOnly, Category=Destruction, meta=(ExposeFunctionCategories="Destruction,Components|Destructible"))
-	TSubobjectPtr<UDestructibleComponent> DestructibleComponent;
+private:
+	UPROPERTY(VisibleAnywhere, BlueprintReadOnly, Category = Destruction, meta = (ExposeFunctionCategories = "Destruction,Components|Destructible", AllowPrivateAccess = "true"))
+	UDestructibleComponent* DestructibleComponent;
+public:
 
 	UPROPERTY(config, EditAnywhere, BlueprintReadWrite, Category=Navigation)
 	uint32 bAffectNavigation : 1;
@@ -35,6 +37,10 @@ class ADestructibleActor : public AActor
 	// End AActor interface.
 
 
+
+public:
+	/** Returns DestructibleComponent subobject **/
+	FORCEINLINE UDestructibleComponent* GetDestructibleComponent() const { return DestructibleComponent; }
 };
 
 
diff --git a/Engine/Source/Runtime/Engine/Classes/PhysicsEngine/PhysicsConstraintActor.h b/Engine/Source/Runtime/Engine/Classes/PhysicsEngine/PhysicsConstraintActor.h
--- a/Engine/Source/Runtime/Engine/Classes/PhysicsEngine/PhysicsConstraintActor.h
+++ b/Engine/Source/Runtime/Engine/Classes/PhysicsEngine/PhysicsConstraintActor.h
@@ -13,8 +13,10 @@ class APhysicsConstraintActor : public ARigidBodyBase
 	GENERATED_UCLASS_BODY()
 
 	// Cached reference to constraint component
-	UPROPERTY(Category=ConstraintActor, VisibleAnywhere, BlueprintReadOnly, meta=(ExposeFunctionCategories="JointDrive,Physics|Components|PhysicsConstraint"))
-	TSubobjectPtr<class UPhysicsConstraintComponent> ConstraintComp;
+private:
+	UPROPERTY(Category = ConstraintActor, VisibleAnywhere, BlueprintReadOnly, meta = (ExposeFunctionCategories = "JointDrive,Physics|Components|PhysicsConstraint", AllowPrivateAccess = "true"))
+	class UPhysicsConstraintComponent* ConstraintComp;
+public:
 	
 	UPROPERTY()
 	class AActor* ConstraintActor1_DEPRECATED;
@@ -29,6 +31,10 @@ class APhysicsConstraintActor : public ARigidBodyBase
 	virtual void LoadedFromAnotherClass(const FName& OldClassName) override;
 #endif // WITH_EDITOR	
 	// End UObject Interface
+
+public:
+	/** Returns ConstraintComp subobject **/
+	FORCEINLINE class UPhysicsConstraintComponent* GetConstraintComp() const { return ConstraintComp; }
 };
 
 
diff --git a/Engine/Source/Runtime/Engine/Classes/PhysicsEngine/PhysicsThruster.h b/Engine/Source/Runtime/Engine/Classes/PhysicsEngine/PhysicsThruster.h
--- a/Engine/Source/Runtime/Engine/Classes/PhysicsEngine/PhysicsThruster.h
+++ b/Engine/Source/Runtime/Engine/Classes/PhysicsEngine/PhysicsThruster.h
@@ -14,19 +14,31 @@ class APhysicsThruster : public ARigidBodyBase
 {
 	GENERATED_UCLASS_BODY()
 
+private:
 	/** Thruster component */
-	UPROPERTY(Category=Physics, VisibleAnywhere, BlueprintReadOnly, meta=(ExposeFunctionCategories="Activation,Components|Activation"))
-	TSubobjectPtr<class UPhysicsThrusterComponent> ThrusterComponent;
-
+	UPROPERTY(Category = Physics, VisibleAnywhere, BlueprintReadOnly, meta = (ExposeFunctionCategories = "Activation,Components|Activation", AllowPrivateAccess = "true"))
+	class UPhysicsThrusterComponent* ThrusterComponent;
 
 #if WITH_EDITORONLY_DATA
+
 	UPROPERTY()
-	TSubobjectPtr<class UArrowComponent> ArrowComponent;
+	class UArrowComponent* ArrowComponent;
 
 	UPROPERTY()
-	TSubobjectPtr<class UBillboardComponent> SpriteComponent;
+	class UBillboardComponent* SpriteComponent;
+public:
 #endif
 
+
+public:
+	/** Returns ThrusterComponent subobject **/
+	FORCEINLINE class UPhysicsThrusterComponent* GetThrusterComponent() const { return ThrusterComponent; }
+#if WITH_EDITORONLY_DATA
+	/** Returns ArrowComponent subobject **/
+	FORCEINLINE class UArrowComponent* GetArrowComponent() const { return ArrowComponent; }
+	/** Returns SpriteComponent subobject **/
+	FORCEINLINE class UBillboardComponent* GetSpriteComponent() const { return SpriteComponent; }
+#endif
 };
 
 
diff --git a/Engine/Source/Runtime/Engine/Classes/PhysicsEngine/RadialForceActor.h b/Engine/Source/Runtime/Engine/Classes/PhysicsEngine/RadialForceActor.h
--- a/Engine/Source/Runtime/Engine/Classes/PhysicsEngine/RadialForceActor.h
+++ b/Engine/Source/Runtime/Engine/Classes/PhysicsEngine/RadialForceActor.h
@@ -9,13 +9,15 @@ class ARadialForceActor : public ARigidBodyBase
 {
 	GENERATED_UCLASS_BODY()
 
+private:
 	/** Force component */
-	UPROPERTY(Category=RadialForceActor, VisibleAnywhere, BlueprintReadOnly, meta=(ExposeFunctionCategories="Activation,Components|Activation,Physics,Physics|Components|RadialForce"))
-	TSubobjectPtr<class URadialForceComponent> ForceComponent;
+	UPROPERTY(Category = RadialForceActor, VisibleAnywhere, BlueprintReadOnly, meta = (ExposeFunctionCategories = "Activation,Components|Activation,Physics,Physics|Components|RadialForce", AllowPrivateAccess = "true"))
+	class URadialForceComponent* ForceComponent;
 
 #if WITH_EDITORONLY_DATA
 	UPROPERTY()
-	TSubobjectPtr<UBillboardComponent> SpriteComponent;
+	UBillboardComponent* SpriteComponent;
+public:
 #endif
 
 	// BEGIN DEPRECATED (use component functions now in level script)
@@ -35,6 +37,14 @@ class ARadialForceActor : public ARigidBodyBase
 	virtual void EditorApplyScale(const FVector& DeltaScale, const FVector* PivotLocation, bool bAltDown, bool bShiftDown, bool bCtrlDown) override;
 	// End AActor interface.
 #endif
+
+public:
+	/** Returns ForceComponent subobject **/
+	FORCEINLINE class URadialForceComponent* GetForceComponent() const { return ForceComponent; }
+#if WITH_EDITORONLY_DATA
+	/** Returns SpriteComponent subobject **/
+	FORCEINLINE UBillboardComponent* GetSpriteComponent() const { return SpriteComponent; }
+#endif
 };
 
 
diff --git a/Engine/Source/Runtime/Engine/Classes/Sound/AmbientSound.h b/Engine/Source/Runtime/Engine/Classes/Sound/AmbientSound.h
--- a/Engine/Source/Runtime/Engine/Classes/Sound/AmbientSound.h
+++ b/Engine/Source/Runtime/Engine/Classes/Sound/AmbientSound.h
@@ -12,9 +12,11 @@ class AAmbientSound : public AActor
 	UPROPERTY()
 	uint32 bAutoPlay_DEPRECATED:1;
 
+private:
 	/** Audio component that handles sound playing */
-	UPROPERTY(Category=Sound, VisibleAnywhere, BlueprintReadOnly,meta=(ExposeFunctionCategories="Sound,Audio,Audio|Components|Audio"))
-	TSubobjectPtr<class UAudioComponent> AudioComponent;
+	UPROPERTY(Category = Sound, VisibleAnywhere, BlueprintReadOnly, meta = (ExposeFunctionCategories = "Sound,Audio,Audio|Components|Audio", AllowPrivateAccess = "true"))
+	class UAudioComponent* AudioComponent;
+public:
 	
 	UPROPERTY(instanced)
 	class UDEPRECATED_SoundNodeAmbient* SoundNodeInstance_DEPRECATED;
@@ -44,6 +46,10 @@ class AAmbientSound : public AActor
 	UFUNCTION(BlueprintCallable, Category="Audio", meta=(DeprecatedFunction))
 	void Stop();
 	// END DEPRECATED
+
+public:
+	/** Returns AudioComponent subobject **/
+	FORCEINLINE class UAudioComponent* GetAudioComponent() const { return AudioComponent; }
 };
 
 
diff --git a/Engine/Source/Runtime/Engine/Classes/VectorField/VectorFieldVolume.h b/Engine/Source/Runtime/Engine/Classes/VectorField/VectorFieldVolume.h
--- a/Engine/Source/Runtime/Engine/Classes/VectorField/VectorFieldVolume.h
+++ b/Engine/Source/Runtime/Engine/Classes/VectorField/VectorFieldVolume.h
@@ -13,12 +13,21 @@ class AVectorFieldVolume : public AActor
 {
 	GENERATED_UCLASS_BODY()
 
-	UPROPERTY(VisibleAnywhere, BlueprintReadOnly, Category=VectorFieldVolume)
-	TSubobjectPtr<class UVectorFieldComponent> VectorFieldComponent;
+private:
+	UPROPERTY(VisibleAnywhere, BlueprintReadOnly, Category = VectorFieldVolume, meta = (AllowPrivateAccess = "true"))
+	class UVectorFieldComponent* VectorFieldComponent;
 
 #if WITH_EDITORONLY_DATA
 	UPROPERTY()
-	TSubobjectPtr<UBillboardComponent> SpriteComponent;
+	UBillboardComponent* SpriteComponent;
+#endif
+
+public:
+	/** Returns VectorFieldComponent subobject **/
+	FORCEINLINE class UVectorFieldComponent* GetVectorFieldComponent() const { return VectorFieldComponent; }
+#if WITH_EDITORONLY_DATA
+	/** Returns SpriteComponent subobject **/
+	FORCEINLINE UBillboardComponent* GetSpriteComponent() const { return SpriteComponent; }
 #endif
 };
 
diff --git a/Engine/Source/Runtime/Engine/Private/AI/Navigation/AvoidanceManager.cpp b/Engine/Source/Runtime/Engine/Private/AI/Navigation/AvoidanceManager.cpp
--- a/Engine/Source/Runtime/Engine/Private/AI/Navigation/AvoidanceManager.cpp
+++ b/Engine/Source/Runtime/Engine/Private/AI/Navigation/AvoidanceManager.cpp
@@ -115,7 +115,7 @@ bool UAvoidanceManager::RegisterMovementComponent(UCharacterMovementComponent* M
 {
 	if (MovementComp && MovementComp->GetCharacterOwner())
 	{
-		if (UCapsuleComponent *OurCapsule = MovementComp->GetCharacterOwner()->CapsuleComponent.Get())
+		if (UCapsuleComponent *OurCapsule = MovementComp->GetCharacterOwner()->GetCapsuleComponent())
 		{
 			const int32 NewAvoidanceUID = GetNewAvoidanceUID();
 			MovementComp->AvoidanceUID = NewAvoidanceUID;
diff --git a/Engine/Source/Runtime/Engine/Private/AI/Navigation/NavLinkProxy.cpp b/Engine/Source/Runtime/Engine/Private/AI/Navigation/NavLinkProxy.cpp
--- a/Engine/Source/Runtime/Engine/Private/AI/Navigation/NavLinkProxy.cpp
+++ b/Engine/Source/Runtime/Engine/Private/AI/Navigation/NavLinkProxy.cpp
@@ -11,7 +11,7 @@
 
 ANavLinkProxy::ANavLinkProxy(const FObjectInitializer& ObjectInitializer) : Super(ObjectInitializer)
 {
-	TSubobjectPtr<USceneComponent> SceneComponent = ObjectInitializer.CreateDefaultSubobject<USceneComponent>(this, TEXT("PositionComponent"));
+	USceneComponent* SceneComponent = ObjectInitializer.CreateDefaultSubobject<USceneComponent>(this, TEXT("PositionComponent"));
 	RootComponent = SceneComponent;
 
 	bHidden = true;
diff --git a/Engine/Source/Runtime/Engine/Private/AI/Navigation/NavModifierVolume.cpp b/Engine/Source/Runtime/Engine/Private/AI/Navigation/NavModifierVolume.cpp
--- a/Engine/Source/Runtime/Engine/Private/AI/Navigation/NavModifierVolume.cpp
+++ b/Engine/Source/Runtime/Engine/Private/AI/Navigation/NavModifierVolume.cpp
@@ -16,7 +16,7 @@ void ANavModifierVolume::GetNavigationData(struct FNavigationRelevantData& Data)
 {
 	if (Brush && AreaClass && AreaClass != UNavigationSystem::GetDefaultWalkableArea())
 	{
-		FAreaNavModifier AreaMod(BrushComponent, AreaClass);
+		FAreaNavModifier AreaMod(GetBrushComponent(), AreaClass);
 		Data.Modifiers.Add(AreaMod);
 	}
 }
@@ -38,4 +38,4 @@ void ANavModifierVolume::PostEditChangeProperty(FPropertyChangedEvent& PropertyC
 	}
 }
 
-#endif
\ No newline at end of file
+#endif
diff --git a/Engine/Source/Runtime/Engine/Private/ActorEditorUtils.cpp b/Engine/Source/Runtime/Engine/Private/ActorEditorUtils.cpp
--- a/Engine/Source/Runtime/Engine/Private/ActorEditorUtils.cpp
+++ b/Engine/Source/Runtime/Engine/Private/ActorEditorUtils.cpp
@@ -17,7 +17,7 @@ namespace FActorEditorUtils
 				// If the builder brush exists then it will be the 2nd actor in the actors array.
 				ABrush* BuilderBrush = Cast<ABrush>(ActorLevel->Actors[1]);
 				// If the second actor is not a brush then it certainly cannot be the builder brush.
-				if ((BuilderBrush != nullptr) && (BuilderBrush->BrushComponent != nullptr) && (BuilderBrush->Brush != nullptr))
+				if ((BuilderBrush != nullptr) && (BuilderBrush->GetBrushComponent() != nullptr) && (BuilderBrush->Brush != nullptr))
 				{
 					bIsBuilder = (BuilderBrush == InActor);
 				}
diff --git a/Engine/Source/Runtime/Engine/Private/AmbientSound.cpp b/Engine/Source/Runtime/Engine/Private/AmbientSound.cpp
--- a/Engine/Source/Runtime/Engine/Private/AmbientSound.cpp
+++ b/Engine/Source/Runtime/Engine/Private/AmbientSound.cpp
@@ -50,7 +50,7 @@ void AAmbientSound::CheckForErrors( void )
 {
 	Super::CheckForErrors();
 
-	if( !AudioComponent.IsValid() )
+	if( !AudioComponent )
 	{
 		FMessageLog("MapCheck").Warning()
 			->AddToken(FUObjectToken::Create(this))
diff --git a/Engine/Source/Runtime/Engine/Private/Atmosphere/Atmosphere.cpp b/Engine/Source/Runtime/Engine/Private/Atmosphere/Atmosphere.cpp
--- a/Engine/Source/Runtime/Engine/Private/Atmosphere/Atmosphere.cpp
+++ b/Engine/Source/Runtime/Engine/Private/Atmosphere/Atmosphere.cpp
@@ -37,13 +37,13 @@ AAtmosphericFog::AAtmosphericFog(const FObjectInitializer& ObjectInitializer)
 		};
 		static FConstructorStatics ConstructorStatics;
 
-		if (SpriteComponent)
+		if (GetSpriteComponent())
 		{
-			SpriteComponent->Sprite = ConstructorStatics.FogTextureObject.Get();
-			SpriteComponent->RelativeScale3D = FVector(0.5f, 0.5f, 0.5f);
-			SpriteComponent->SpriteInfo.Category = ConstructorStatics.ID_Fog;
-			SpriteComponent->SpriteInfo.DisplayName = ConstructorStatics.NAME_Fog;
-			SpriteComponent->AttachParent = AtmosphericFogComponent;
+			GetSpriteComponent()->Sprite = ConstructorStatics.FogTextureObject.Get();
+			GetSpriteComponent()->RelativeScale3D = FVector(0.5f, 0.5f, 0.5f);
+			GetSpriteComponent()->SpriteInfo.Category = ConstructorStatics.ID_Fog;
+			GetSpriteComponent()->SpriteInfo.DisplayName = ConstructorStatics.NAME_Fog;
+			GetSpriteComponent()->AttachParent = AtmosphericFogComponent;
 		}
 
 		if (ArrowComponent)
diff --git a/Engine/Source/Runtime/Engine/Private/AudioVolume.cpp b/Engine/Source/Runtime/Engine/Private/AudioVolume.cpp
--- a/Engine/Source/Runtime/Engine/Private/AudioVolume.cpp
+++ b/Engine/Source/Runtime/Engine/Private/AudioVolume.cpp
@@ -122,10 +122,10 @@ void FReverbSettings::PostSerialize(const FArchive& Ar)
 AAudioVolume::AAudioVolume(const FObjectInitializer& ObjectInitializer)
 	: Super(ObjectInitializer)
 {
-	BrushComponent->BodyInstance.bEnableCollision_DEPRECATED = false;
+	GetBrushComponent()->BodyInstance.bEnableCollision_DEPRECATED = false;
 
-	BrushComponent->SetCollisionProfileName(UCollisionProfile::NoCollision_ProfileName);
-	BrushComponent->bAlwaysCreatePhysicsState = true;
+	GetBrushComponent()->SetCollisionProfileName(UCollisionProfile::NoCollision_ProfileName);
+	GetBrushComponent()->bAlwaysCreatePhysicsState = true;
 
 	bColored = true;
 	BrushColor = FColor(255, 255, 0, 255);
diff --git a/Engine/Source/Runtime/Engine/Private/BillboardComponent.cpp b/Engine/Source/Runtime/Engine/Private/BillboardComponent.cpp
--- a/Engine/Source/Runtime/Engine/Private/BillboardComponent.cpp
+++ b/Engine/Source/Runtime/Engine/Private/BillboardComponent.cpp
@@ -63,9 +63,9 @@ public:
 			// to match the lights color.  
 			if (ALight* Light = Cast<ALight>(Owner))
 			{
-				if (Light->LightComponent)
+				if (Light->GetLightComponent())
 				{
-					Color = Light->LightComponent->LightColor.ReinterpretAsLinear();
+					Color = Light->GetLightComponent()->LightColor.ReinterpretAsLinear();
 					Color.A = 255;
 				}
 			}
diff --git a/Engine/Source/Runtime/Engine/Private/BlockingVolume.cpp b/Engine/Source/Runtime/Engine/Private/BlockingVolume.cpp
--- a/Engine/Source/Runtime/Engine/Private/BlockingVolume.cpp
+++ b/Engine/Source/Runtime/Engine/Private/BlockingVolume.cpp
@@ -8,9 +8,9 @@ static FName InvisibleWall_NAME(TEXT("InvisibleWall"));
 ABlockingVolume::ABlockingVolume(const FObjectInitializer& ObjectInitializer)
 	: Super(ObjectInitializer)
 {
-	BrushComponent->bCanEverAffectNavigation = true;
-	BrushComponent->BodyInstance.bEnableCollision_DEPRECATED = true;
-	BrushComponent->SetCollisionProfileName(InvisibleWall_NAME);
+	GetBrushComponent()->bCanEverAffectNavigation = true;
+	GetBrushComponent()->BodyInstance.bEnableCollision_DEPRECATED = true;
+	GetBrushComponent()->SetCollisionProfileName(InvisibleWall_NAME);
 }
 
 #if WITH_EDITOR
@@ -27,11 +27,11 @@ void ABlockingVolume::LoadedFromAnotherClass(const FName& OldClassName)
 
 		if(OldClassName == DynamicBlockingVolume_NAME)
 		{
-			BrushComponent->Mobility = EComponentMobility::Movable;
+			GetBrushComponent()->Mobility = EComponentMobility::Movable;
 
-			if(BrushComponent->GetCollisionProfileName() == InvisibleWall_NAME)
+			if(GetBrushComponent()->GetCollisionProfileName() == InvisibleWall_NAME)
 			{
-				BrushComponent->SetCollisionProfileName(InvisibleWallDynamic_NAME);
+				GetBrushComponent()->SetCollisionProfileName(InvisibleWallDynamic_NAME);
 			}
 		}
 	}
@@ -48,20 +48,20 @@ void ABlockingVolume::PostEditChangeChainProperty(FPropertyChangedChainEvent& Pr
 	{
 		// If the collision profile is one of the 'default' ones for a BlockingVolume, make sure it is the correct one
 		// If user has changed it to something else, don't touch it
-		FName CurrentProfileName = BrushComponent->GetCollisionProfileName();
+		FName CurrentProfileName = GetBrushComponent()->GetCollisionProfileName();
 		if(	CurrentProfileName == InvisibleWall_NAME ||
 			CurrentProfileName == InvisibleWallDynamic_NAME )
 		{
-			if(BrushComponent->Mobility == EComponentMobility::Movable)
+			if(GetBrushComponent()->Mobility == EComponentMobility::Movable)
 			{
-				BrushComponent->SetCollisionProfileName(InvisibleWallDynamic_NAME);
+				GetBrushComponent()->SetCollisionProfileName(InvisibleWallDynamic_NAME);
 			}
 			else
 			{
-				BrushComponent->SetCollisionProfileName(InvisibleWall_NAME);
+				GetBrushComponent()->SetCollisionProfileName(InvisibleWall_NAME);
 			}
 		}
 	}
 }
 
-#endif
\ No newline at end of file
+#endif
diff --git a/Engine/Source/Runtime/Engine/Private/Brush.cpp b/Engine/Source/Runtime/Engine/Private/Brush.cpp
--- a/Engine/Source/Runtime/Engine/Private/Brush.cpp
+++ b/Engine/Source/Runtime/Engine/Private/Brush.cpp
@@ -38,7 +38,7 @@ void ABrush::PostEditMove(bool bFinished)
 {
 	bInManipulation = !bFinished;
 
-	if( BrushComponent.IsValid() )
+	if( BrushComponent )
 	{
 		BrushComponent->ReregisterComponent();
 	}
@@ -143,7 +143,7 @@ void ABrush::SetIsTemporarilyHiddenInEditor( bool bIsHidden )
 FVector ABrush::GetPrePivot() const
 {
 	FVector Result(0.f);
-	if( BrushComponent.IsValid() )
+	if( BrushComponent )
 	{
 		Result = BrushComponent->PrePivot;
 	}
@@ -152,7 +152,7 @@ FVector ABrush::GetPrePivot() const
 
 void ABrush::SetPrePivot( const FVector& InPrePivot )
 {
-	if( BrushComponent.IsValid() )
+	if( BrushComponent )
 	{
 		BrushComponent->PrePivot = InPrePivot;
 		BrushComponent->UpdateComponentToWorld();
@@ -271,7 +271,7 @@ FColor ABrush::GetWireColor() const
 
 bool ABrush::IsStaticBrush() const
 {
-	return BrushComponent.IsValid() && (BrushComponent->Mobility == EComponentMobility::Static);
+	return BrushComponent && (BrushComponent->Mobility == EComponentMobility::Static);
 }
 
 bool ABrush::Modify(bool bAlwaysMarkDirty)
diff --git a/Engine/Source/Runtime/Engine/Private/BrushShape.cpp b/Engine/Source/Runtime/Engine/Private/BrushShape.cpp
--- a/Engine/Source/Runtime/Engine/Private/BrushShape.cpp
+++ b/Engine/Source/Runtime/Engine/Private/BrushShape.cpp
@@ -7,10 +7,10 @@ ABrushShape::ABrushShape(const FObjectInitializer& ObjectInitializer)
 	: Super(ObjectInitializer)
 {
 
-	BrushComponent->BodyInstance.bEnableCollision_DEPRECATED = false;
-	BrushComponent->SetCollisionProfileName(UCollisionProfile::NoCollision_ProfileName);
-	BrushComponent->AlwaysLoadOnClient = true;
-	BrushComponent->AlwaysLoadOnServer = false;
+	GetBrushComponent()->BodyInstance.bEnableCollision_DEPRECATED = false;
+	GetBrushComponent()->SetCollisionProfileName(UCollisionProfile::NoCollision_ProfileName);
+	GetBrushComponent()->AlwaysLoadOnClient = true;
+	GetBrushComponent()->AlwaysLoadOnServer = false;
 
 }
 
diff --git a/Engine/Source/Runtime/Engine/Private/CameraBlockingVolume.cpp b/Engine/Source/Runtime/Engine/Private/CameraBlockingVolume.cpp
--- a/Engine/Source/Runtime/Engine/Private/CameraBlockingVolume.cpp
+++ b/Engine/Source/Runtime/Engine/Private/CameraBlockingVolume.cpp
@@ -6,5 +6,5 @@
 ACameraBlockingVolume::ACameraBlockingVolume(const FObjectInitializer& ObjectInitializer)
 	: Super(ObjectInitializer)
 {
-	BrushComponent->SetCollisionResponseToChannel(ECC_Camera, ECR_Block);
-}
\ No newline at end of file
+	GetBrushComponent()->SetCollisionResponseToChannel(ECC_Camera, ECR_Block);
+}
diff --git a/Engine/Source/Runtime/Engine/Private/Character.cpp b/Engine/Source/Runtime/Engine/Private/Character.cpp
--- a/Engine/Source/Runtime/Engine/Private/Character.cpp
+++ b/Engine/Source/Runtime/Engine/Private/Character.cpp
@@ -111,13 +111,13 @@ void ACharacter::PostInitializeComponents()
 			BaseTranslationOffset = Mesh->RelativeLocation;
 
 			// force animation tick after movement component updates
-			if (Mesh->PrimaryComponentTick.bCanEverTick && CharacterMovement.IsValid())
+			if (Mesh->PrimaryComponentTick.bCanEverTick && CharacterMovement)
 			{
 				Mesh->PrimaryComponentTick.AddPrerequisite(CharacterMovement, CharacterMovement->PrimaryComponentTick);
 			}
 		}
 
-		if (CharacterMovement.IsValid() && CapsuleComponent.IsValid())
+		if (CharacterMovement && CapsuleComponent)
 		{
 			CharacterMovement->UpdateNavAgent(CapsuleComponent);
 		}
@@ -167,7 +167,7 @@ void ACharacter::ApplyWorldOffset(const FVector& InOffset, bool bWorldShift)
 	Super::ApplyWorldOffset(InOffset, bWorldShift);
 
 	// Update cached location values in movement component
-	if (CharacterMovement.IsValid())
+	if (CharacterMovement)
 	{
 		CharacterMovement->OldBaseLocation+= InOffset;
 	}
@@ -334,7 +334,7 @@ void ACharacter::ApplyDamageMomentum(float DamageTaken, FDamageEvent const& Dama
 	UDamageType const* const DmgTypeCDO = DamageEvent.DamageTypeClass->GetDefaultObject<UDamageType>();
 	float const ImpulseScale = DmgTypeCDO->DamageImpulse;
 
-	if ( (ImpulseScale > 3.f) && (CharacterMovement.Get() != NULL) )
+	if ( (ImpulseScale > 3.f) && (CharacterMovement != NULL) )
 	{
 		FHitResult HitInfo;
 		FVector ImpulseDir;
@@ -641,13 +641,13 @@ void ACharacter::SaveRelativeBasedMovement(const FVector& NewRelativeLocation, c
 
 void ACharacter::TurnOff()
 {
-	if (CharacterMovement.Get() != NULL)
+	if (CharacterMovement != NULL)
 	{
 		CharacterMovement->StopMovementImmediately();
 		CharacterMovement->DisableMovement();
 	}
 
-	if (GetNetMode() != NM_DedicatedServer && Mesh.Get() != NULL)
+	if (GetNetMode() != NM_DedicatedServer && Mesh != NULL)
 	{
 		Mesh->bPauseAnims = true;
 		if (Mesh->IsSimulatingPhysics())
@@ -677,7 +677,7 @@ void ACharacter::Restart()
 
 void ACharacter::PawnClientRestart()
 {
-	if (CharacterMovement.Get() != NULL)
+	if (CharacterMovement != NULL)
 	{
 		CharacterMovement->StopMovementImmediately();
 		CharacterMovement->ResetPredictionData_Client();
@@ -759,7 +759,7 @@ void ACharacter::DisplayDebug(UCanvas* Canvas, const FDebugDisplayInfo& DebugDis
 		FIndenter PhysicsIndent(Indent);
 
 		FString BaseString;
-		if ( CharacterMovement.Get() == NULL || BasedMovement.MovementBase == NULL )
+		if ( CharacterMovement == NULL || BasedMovement.MovementBase == NULL )
 		{
 			BaseString = "Not Based";
 		}
@@ -772,7 +772,7 @@ void ACharacter::DisplayDebug(UCanvas* Canvas, const FDebugDisplayInfo& DebugDis
 		Canvas->DrawText(RenderFont, FString::Printf(TEXT("RelativeLoc: %s Rot: %s %s"), *BasedMovement.Location.ToString(), *BasedMovement.Rotation.ToString(), *BaseString), Indent, YPos);
 		YPos += YL;
 
-		if ( CharacterMovement.Get() != NULL )
+		if ( CharacterMovement != NULL )
 		{
 			CharacterMovement->DisplayDebug(Canvas, DebugDisplay, YL, YPos);
 		}
