I18N/L10N: Removed the CoreDelegates::OnCultureChanged and replaced it with event FInternationalization::Get().OnCultureChanged() and added FTextLocalizationManager::Get().OnTranslationsChanged().

diff --git a/Engine/Source/Developer/SettingsEditor/Private/Widgets/SSettingsEditor.cpp b/Engine/Source/Developer/SettingsEditor/Private/Widgets/SSettingsEditor.cpp
--- a/Engine/Source/Developer/SettingsEditor/Private/Widgets/SSettingsEditor.cpp
+++ b/Engine/Source/Developer/SettingsEditor/Private/Widgets/SSettingsEditor.cpp
@@ -20,7 +20,7 @@ SSettingsEditor::~SSettingsEditor( )
 {
 	Model->OnSelectionChanged().RemoveAll(this);
 	SettingsContainer->OnCategoryModified().RemoveAll(this);
-	FCoreDelegates::OnCultureChanged.RemoveAll(this);
+	FInternationalization::Get().OnCultureChanged().RemoveAll(this);
 }
 
 
@@ -224,7 +224,7 @@ void SSettingsEditor::Construct( const FArguments& InArgs, const ISettingsEditor
 			]
 	];
 
-	FCoreDelegates::OnCultureChanged.AddSP(this, &SSettingsEditor::HandleCultureChanged);
+	FInternationalization::Get().OnCultureChanged().AddSP(this, &SSettingsEditor::HandleCultureChanged);
 	Model->OnSelectionChanged().AddSP(this, &SSettingsEditor::HandleModelSelectionChanged);
 	SettingsContainer->OnCategoryModified().AddSP(this, &SSettingsEditor::HandleSettingsContainerCategoryModified);
 
diff --git a/Engine/Source/Runtime/Core/Private/Internationalization/ICUInternationalization.cpp b/Engine/Source/Runtime/Core/Private/Internationalization/ICUInternationalization.cpp
--- a/Engine/Source/Runtime/Core/Private/Internationalization/ICUInternationalization.cpp
+++ b/Engine/Source/Runtime/Core/Private/Internationalization/ICUInternationalization.cpp
@@ -236,7 +236,7 @@ void FICUInternationalization::SetCurrentCulture(const FString& Name)
 			UErrorCode ICUStatus = U_ZERO_ERROR;
 			uloc_setDefault(StringCast<char>(*Name).Get(), &ICUStatus);
 
-			FCoreDelegates::OnCultureChanged.Broadcast();
+			FInternationalization::Get().BroadcastCultureChanged();
 		}
 	}
 }
diff --git a/Engine/Source/Runtime/Core/Private/Internationalization/TextLocalizationManager.cpp b/Engine/Source/Runtime/Core/Private/Internationalization/TextLocalizationManager.cpp
--- a/Engine/Source/Runtime/Core/Private/Internationalization/TextLocalizationManager.cpp
+++ b/Engine/Source/Runtime/Core/Private/Internationalization/TextLocalizationManager.cpp
@@ -15,7 +15,7 @@ void BeginInitTextLocalization()
 	// twice since FInternationalization::Initialize sets the culture
 	FInternationalization::Get();
 
-	FCoreDelegates::OnCultureChanged.AddRaw( &(FTextLocalizationManager::Get()), &FTextLocalizationManager::OnCultureChanged );
+	FInternationalization::Get().OnCultureChanged().AddRaw( &(FTextLocalizationManager::Get()), &FTextLocalizationManager::OnCultureChanged );
 }
 
 void EndInitTextLocalization()
@@ -471,6 +471,8 @@ void FTextLocalizationManager::UpdateLiveTable(const TArray<FLocalizationEntryTr
 			}
 		}
 	}
+
+	OnTranslationsChanged().Broadcast();
 }
 
 void FTextLocalizationManager::OnCultureChanged()
diff --git a/Engine/Source/Runtime/Core/Private/Misc/CallbackDevice.cpp b/Engine/Source/Runtime/Core/Private/Misc/CallbackDevice.cpp
--- a/Engine/Source/Runtime/Core/Private/Misc/CallbackDevice.cpp
+++ b/Engine/Source/Runtime/Core/Private/Misc/CallbackDevice.cpp
@@ -32,7 +32,6 @@ FCoreDelegates::FOnActorLabelChanged FCoreDelegates::OnActorLabelChanged;
 	FSimpleMulticastDelegate FCoreDelegates::PreModal;
 	FSimpleMulticastDelegate FCoreDelegates::PostModal;
 #endif	//WITH_EDITOR
-FSimpleMulticastDelegate FCoreDelegates::OnCultureChanged;
 FSimpleMulticastDelegate FCoreDelegates::OnShutdownAfterError;
 FSimpleMulticastDelegate FCoreDelegates::OnInit;
 FSimpleMulticastDelegate FCoreDelegates::OnExit;
diff --git a/Engine/Source/Runtime/Core/Public/Internationalization/Internationalization.h b/Engine/Source/Runtime/Core/Public/Internationalization/Internationalization.h
--- a/Engine/Source/Runtime/Core/Public/Internationalization/Internationalization.h
+++ b/Engine/Source/Runtime/Core/Public/Internationalization/Internationalization.h
@@ -61,9 +61,15 @@ public:
 	// Given some paths to look at, populate a list of cultures that we have available localization information for. If bIncludeDerivedCultures, include cultures that are derived from those we have localization data for.
 	CORE_API void GetCulturesWithAvailableLocalization(const TArray<FString>& InLocalizationPaths, TArray< FCultureRef >& OutAvailableCultures, const bool bIncludeDerivedCultures);
 
+	/** Broadcasts whenever the current culture changes */
+	DECLARE_EVENT(FInternationalization, FCultureChangedEvent)
+	CORE_API FCultureChangedEvent& OnCultureChanged() { return CultureChangedEvent; }
+
 private:
 	FInternationalization();
 
+	void BroadcastCultureChanged() { CultureChangedEvent.Broadcast(); }
+
 	void Initialize();
 	void Terminate();
 
@@ -71,6 +77,8 @@ private:
 	static FInternationalization* Instance;
 	bool bIsInitialized;
 
+	FCultureChangedEvent CultureChangedEvent;
+
 #if UE_ENABLE_ICU
 	friend class FICUInternationalization;
 	typedef FICUInternationalization FImplementation;
diff --git a/Engine/Source/Runtime/Core/Public/Internationalization/TextLocalizationManager.h b/Engine/Source/Runtime/Core/Public/Internationalization/TextLocalizationManager.h
--- a/Engine/Source/Runtime/Core/Public/Internationalization/TextLocalizationManager.h
+++ b/Engine/Source/Runtime/Core/Public/Internationalization/TextLocalizationManager.h
@@ -96,6 +96,10 @@ public:
 	/** Returns the current culture revision index */
 	int CORE_API GetHeadCultureRevision() const { return HeadCultureRevisionIndex; }
 
+	/** Broadcasts whenever the new translations are available */
+	DECLARE_EVENT(FInternationalization, FTranslationsChangedEvent)
+	CORE_API FTranslationsChangedEvent&  OnTranslationsChanged() { return TranslationsChangedEvent; }
+
 private:
 	FTextLocalizationManager() 
 		: bIsInitialized(false)
@@ -118,4 +122,6 @@ private:
 
 	/** The current culture revision, increments every time the culture is changed and allows for FTexts to know when to rebuild under a new culture */
 	int32 HeadCultureRevisionIndex;
+
+	FTranslationsChangedEvent TranslationsChangedEvent;
 };
\ No newline at end of file
diff --git a/Engine/Source/Runtime/Core/Public/Misc/CallbackDevice.h b/Engine/Source/Runtime/Core/Public/Misc/CallbackDevice.h
--- a/Engine/Source/Runtime/Core/Public/Misc/CallbackDevice.h
+++ b/Engine/Source/Runtime/Core/Public/Misc/CallbackDevice.h
@@ -113,10 +113,7 @@ public:
 	// Called after the editor dismisses a modal window, allowing other windows the opportunity to disable themselves to avoid reentrant calls
 	static FSimpleMulticastDelegate PostModal;
 #endif	//WITH_EDITOR
-
-	// Called when SetCurrentCulture is called.
-	static FSimpleMulticastDelegate OnCultureChanged;
-
+	
 	// Called when an error occurred.
 	static FSimpleMulticastDelegate OnShutdownAfterError;
 
diff --git a/Engine/Source/Runtime/CoreUObject/Private/UObject/Obj.cpp b/Engine/Source/Runtime/CoreUObject/Private/UObject/Obj.cpp
--- a/Engine/Source/Runtime/CoreUObject/Private/UObject/Obj.cpp
+++ b/Engine/Source/Runtime/CoreUObject/Private/UObject/Obj.cpp
@@ -1803,10 +1803,6 @@ void UObject::ReinitializeProperties( UObject* SourceObject/*=NULL*/, FObjectIns
 	StaticConstructObject( GetClass(), GetOuter(), GetFName(), GetFlags(), SourceObject, !HasAnyFlags(RF_ClassDefaultObject), InstanceGraph );
 }
 
-void UObject::CultureChange()
-{
-}
-
 
 /*-----------------------------------------------------------------------------
    Shutdown.
@@ -3391,7 +3387,6 @@ bool StaticExec( UWorld* InWorld, const TCHAR* Cmd, FOutputDevice& Ar )
 	StaticInit & StaticExit.
 -----------------------------------------------------------------------------*/
 
-void OnCultureChangedUObjectCallback();
 void StaticUObjectInit();
 void InitUObject();
 void StaticExit();
@@ -3403,7 +3398,6 @@ void PreInitUObject()
 
 void InitUObject()
 {
-	FCoreDelegates::OnCultureChanged.AddStatic(OnCultureChangedUObjectCallback);
 	FCoreDelegates::OnShutdownAfterError.AddStatic(StaticShutdownAfterError);
 	FCoreDelegates::OnExit.AddStatic(StaticExit);
 	FModuleManager::Get().OnProcessLoadedObjectsCallback().AddStatic(ProcessNewlyLoadedUObjects);
@@ -3462,14 +3456,6 @@ void StaticUObjectInit()
 	UE_LOG(LogInit, Log, TEXT("Object subsystem initialized") );
 }
 
-void OnCultureChangedUObjectCallback()
-{
-	for( FObjectIterator It; It; ++It )
-	{
-		It->CultureChange();
-	}
-}
-
 //
 // Shut down the object manager.
 //
diff --git a/Engine/Source/Runtime/CoreUObject/Public/UObject/UObject.h b/Engine/Source/Runtime/CoreUObject/Public/UObject/UObject.h
--- a/Engine/Source/Runtime/CoreUObject/Public/UObject/UObject.h
+++ b/Engine/Source/Runtime/CoreUObject/Public/UObject/UObject.h
@@ -531,9 +531,6 @@ public:
 	 */
 	void ReinitializeProperties( UObject* SourceObject=NULL, struct FObjectInstancingGraph* InstanceGraph=NULL );
 
-	/** Handle culture change. */
-	virtual void CultureChange();
-
 	/**
 	 * This will return detail info about this specific object. (e.g. AudioComponent will return the name of the cue,
 	 * ParticleSystemComponent will return the name of the ParticleSystem)  The idea here is that in many places
