FPostConstructInitializeProperties deprecation typedef was causing errors when used in combination with forward class declaration in parameter list and apparently this is used extensively. Changed it to deprecation macro which appears to work.

diff --git a/Engine/Source/Runtime/CoreUObject/Public/UObject/UObjectGlobals.h b/Engine/Source/Runtime/CoreUObject/Public/UObject/UObjectGlobals.h
--- a/Engine/Source/Runtime/CoreUObject/Public/UObject/UObjectGlobals.h
+++ b/Engine/Source/Runtime/CoreUObject/Public/UObject/UObjectGlobals.h
@@ -965,7 +965,9 @@ private:
 	UObject* LastConstructedObject;
 };
 
-typedef DEPRECATED(4.6, "FPostConstructInitializeProperties was renamed to FObjectInitializer.") FObjectInitializer FPostConstructInitializeProperties;
+#define FPostConstructInitializeProperties \
+	FObjectInitializer \
+	EMIT_DEPRECATED_WARNING_MESSAGE("FPostConstructInitializeProperties is deprecated and was renamed to FObjectInitializer. Please use that type instead.")
 
 /**
 * Helper class for script integrations to access some UObject innards. Needed for script-generated UObject classes
