class ArchiveObjectCrc32 - Calculates a checksum on an object's serialized data stream.
Written originally by [REDACTED_PERSON].

diff --git a/Engine/Source/Runtime/CoreUObject/Private/Serialization/ArchiveObjectCrc32.cpp b/Engine/Source/Runtime/CoreUObject/Private/Serialization/ArchiveObjectCrc32.cpp
new file mode 100644
--- /dev/null
+++ b/Engine/Source/Runtime/CoreUObject/Private/Serialization/ArchiveObjectCrc32.cpp
@@ -0,0 +1,112 @@
+// Copyright 1998-2014 Epic Games, Inc. All Rights Reserved.
+
+#include "CoreUObjectPrivate.h"
+
+DEFINE_LOG_CATEGORY_STATIC(LogArchiveObjectCrc32, Log, All);
+
+//#define DEBUG_ARCHIVE_OBJECT_CRC32
+
+/*----------------------------------------------------------------------------
+	FArchiveObjectCrc32
+----------------------------------------------------------------------------*/
+
+FArchiveObjectCrc32::FArchiveObjectCrc32()
+	: MemoryWriter(SerializedObjectData)
+	, ObjectBeingSerialized(NULL)
+	, RootObject(NULL)
+{
+	ArIgnoreOuterRef = true;
+}
+
+void FArchiveObjectCrc32::Serialize(void* Data, int64 Length)
+{
+	MemoryWriter.Serialize(Data, Length);
+}
+
+FArchive& FArchiveObjectCrc32::operator<<(class FName& Name)
+{
+	checkSlow(ObjectBeingSerialized);
+
+	// Don't include the name of the object being serialized, since that isn't technically part of the object's state
+	if(Name != ObjectBeingSerialized->GetFName())
+	{
+		MemoryWriter << Name;
+	}
+
+	return *this;
+}
+
+FArchive& FArchiveObjectCrc32::operator<<(class UObject*& Object)
+{
+	FArchive& Ar = *this;
+
+	// No linker should be a part of the object's state
+	if (Object && Object->IsA<ULinker>())
+	{
+		Object = NULL;
+	}
+
+	if (!Object || !Object->IsIn(RootObject))
+	{
+		auto UniqueName = GetPathNameSafe(Object);
+		Ar << UniqueName;
+	}
+	else
+	{
+		ObjectsToSerialize.Enqueue(Object);
+	}
+
+	return Ar;
+}
+
+uint32 FArchiveObjectCrc32::Crc32(UObject* Object, uint32 CRC)
+{
+#ifdef DEBUG_ARCHIVE_OBJECT_CRC32
+	const double StartTime = FPlatformTime::Seconds();
+	UE_LOG(LogArchiveObjectCrc32, Log, TEXT("### Calculating CRC for object: %s with outer: %s"), *Object->GetName(), Object->GetOuter() ? *Object->GetOuter()->GetName() : TEXT("NULL"));
+#endif
+	RootObject = Object;
+	if (Object)
+	{
+		TSet<UObject*> SerializedObjects;
+
+		// Start with the given object
+		ObjectsToSerialize.Enqueue(Object);
+
+		// Continue until we no longer have any objects to serialized
+		while (ObjectsToSerialize.Dequeue(Object))
+		{
+			bool bAlreadyProcessed = false;
+			SerializedObjects.Add(Object, &bAlreadyProcessed);
+			// If we haven't already serialized this object
+			if (!bAlreadyProcessed)
+			{
+#ifdef DEBUG_ARCHIVE_OBJECT_CRC32
+				UE_LOG(LogArchiveObjectCrc32, Log, TEXT("- Serializing object: %s with outer: %s"), *Object->GetName(), Object->GetOuter() ? *Object->GetOuter()->GetName() : TEXT("NULL"));
+#endif
+				// Serialize it
+				ObjectBeingSerialized = Object;
+				Object->Serialize(*this);
+				ObjectBeingSerialized = NULL;
+
+				// Calculate the CRC, compounding it with the checksum calculated from the previous object
+				CRC = FCrc::MemCrc32(SerializedObjectData.GetData(), SerializedObjectData.Num(), CRC);
+#ifdef DEBUG_ARCHIVE_OBJECT_CRC32
+				UE_LOG(LogArchiveObjectCrc32, Log, TEXT("=> total size: %d bytes, checksum: 0x%08x"), SerializedObjectData.Num(), CRC);
+#endif
+				// Cleanup
+				MemoryWriter.Seek(0L);
+				SerializedObjectData.Empty();
+			}
+		}
+
+		// Cleanup
+		SerializedObjects.Empty();
+		RootObject = NULL;
+	}
+
+#ifdef DEBUG_ARCHIVE_OBJECT_CRC32
+	UE_LOG(LogArchiveObjectCrc32, Log, TEXT("### Finished (%.02f ms), final checksum: 0x%08x"), (FPlatformTime::Seconds() - StartTime) * 1000.0f, CRC);
+#endif
+	return CRC;
+}
\ No newline at end of file
diff --git a/Engine/Source/Runtime/CoreUObject/Public/Serialization/ArchiveUObject.h b/Engine/Source/Runtime/CoreUObject/Public/Serialization/ArchiveUObject.h
--- a/Engine/Source/Runtime/CoreUObject/Public/Serialization/ArchiveUObject.h
+++ b/Engine/Source/Runtime/CoreUObject/Public/Serialization/ArchiveUObject.h
@@ -1508,4 +1508,44 @@ private:
 	int64							CurrentChunkIndex;
 	/** Compression flags determining compression of CompressedChunks.				*/
 	ECompressionFlags				CompressionFlags;
+};
+
+/*----------------------------------------------------------------------------
+FArchiveObjectCrc32
+----------------------------------------------------------------------------*/
+
+/**
+* Calculates a checksum on an object's serialized data stream.
+*/
+class COREUOBJECT_API FArchiveObjectCrc32 : public FArchiveUObject
+{
+public:
+	/**
+	* Default constructor.
+	*/
+	FArchiveObjectCrc32();
+
+	// Begin FArchive Interface
+	virtual void Serialize(void* Data, int64 Length);
+	virtual FArchive& operator<<(class FName& Name);
+	virtual FArchive& operator<<(class UObject*& Object);
+	virtual FString GetArchiveName() const { return TEXT("FArchiveObjectCrc32"); }
+	// End FArchive Interface
+
+	/**
+	* Serialize the given object, calculate and return its checksum.
+	*/
+	uint32 Crc32(UObject* Object, uint32 CRC = 0);
+
+private:
+	/** Internal archive used for serialization */
+	FMemoryWriter MemoryWriter;
+	/** Internal byte array used for serialization */
+	TArray<uint8> SerializedObjectData;
+	/** Internal queue of object references awaiting serialization */
+	TQueue<UObject*> ObjectsToSerialize;
+	/** Internal currently serialized object */
+	const UObject* ObjectBeingSerialized;
+	/** Internal root object */
+	const UObject* RootObject;
 };
\ No newline at end of file
