UberGraphPersistentFrame code can be disabled by setting USE_UBER_GRAPH_PERSISTENT_FRAME define false.

diff --git a/Engine/Source/Runtime/Core/Public/Misc/Build.h b/Engine/Source/Runtime/Core/Public/Misc/Build.h
--- a/Engine/Source/Runtime/Core/Public/Misc/Build.h
+++ b/Engine/Source/Runtime/Core/Public/Misc/Build.h
@@ -236,4 +236,7 @@
 #endif
 
 /** Enable the use of the network profiler as long as we are a build that includes stats */
-#define USE_NETWORK_PROFILER         STATS
\ No newline at end of file
+#define USE_NETWORK_PROFILER         STATS
+
+/** Enable UberGraphPersistentFrame feature. It can speed up BP compilation (re-instancing) in editor, but introduce an unnecessary overhead in runtime. */
+#define USE_UBER_GRAPH_PERSISTENT_FRAME 1
diff --git a/Engine/Source/Runtime/CoreUObject/Private/Blueprint/BlueprintSupport.cpp b/Engine/Source/Runtime/CoreUObject/Private/Blueprint/BlueprintSupport.cpp
--- a/Engine/Source/Runtime/CoreUObject/Private/Blueprint/BlueprintSupport.cpp
+++ b/Engine/Source/Runtime/CoreUObject/Private/Blueprint/BlueprintSupport.cpp
@@ -177,7 +177,9 @@ bool UObject::IsInBlueprint() const
 void UObject::DestroyNonNativeProperties()
 {
 	// Destroy properties that won't be destroyed by the native destructor
+#if USE_UBER_GRAPH_PERSISTENT_FRAME
 	GetClass()->DestroyPersistentUberGraphFrame(this);
+#endif
 	{
 		for (UProperty* P = GetClass()->DestructorLink; P; P = P->DestructorLinkNext)
 		{
diff --git a/Engine/Source/Runtime/CoreUObject/Private/UObject/ScriptCore.cpp b/Engine/Source/Runtime/CoreUObject/Private/UObject/ScriptCore.cpp
--- a/Engine/Source/Runtime/CoreUObject/Private/UObject/ScriptCore.cpp
+++ b/Engine/Source/Runtime/CoreUObject/Private/UObject/ScriptCore.cpp
@@ -402,7 +402,10 @@ void UObject::CallFunction( FFrame& Stack, RESULT_DECL, UFunction* Function )
 	}
 	else
 	{
-		uint8* Frame = GetClass()->GetPersistentUberGraphFrame(this, Function);
+		uint8* Frame = NULL;
+#if USE_UBER_GRAPH_PERSISTENT_FRAME
+		Frame = GetClass()->GetPersistentUberGraphFrame(this, Function);
+#endif
 		const bool bUsePersistentFrame = (NULL != Frame);
 		if (!bUsePersistentFrame)
 		{
@@ -839,7 +842,10 @@ void UObject::ProcessEvent( UFunction* Function, void* Parms )
 
 	// Scope required for scoped script stats.
 	{
-		uint8* Frame = GetClass()->GetPersistentUberGraphFrame(this, Function);
+		uint8* Frame = NULL;
+#if USE_UBER_GRAPH_PERSISTENT_FRAME
+		Frame = GetClass()->GetPersistentUberGraphFrame(this, Function);
+#endif
 		const bool bUsePersistentFrame = (NULL != Frame);
 		if (!bUsePersistentFrame)
 		{
@@ -1214,6 +1220,7 @@ IMPLEMENT_VM_FUNCTION( EX_PopExecutionFlowIfNot, execPopExecutionFlowIfNot );
 
 void UObject::execLetValueOnPersistentFrame(FFrame& Stack, RESULT_DECL)
 {
+#if USE_UBER_GRAPH_PERSISTENT_FRAME
 	Stack.MostRecentProperty = NULL;
 	Stack.MostRecentPropertyAddress = NULL;
 
@@ -1225,6 +1232,9 @@ void UObject::execLetValueOnPersistentFrame(FFrame& Stack, RESULT_DECL)
 	auto DestAddress = DestProperty->ContainerPtrToValuePtr<uint8>(FrameBase);
 
 	Stack.Step(Stack.Object, DestAddress);
+#else
+	checkf(false, TEXT("execLetValueOnPersistentFrame: UberGraphPersistentFrame is not supported by current build!")
+#endif
 }
 IMPLEMENT_VM_FUNCTION(Ex_LetValueOnPersistentFrame, execLetValueOnPersistentFrame);
 
diff --git a/Engine/Source/Runtime/CoreUObject/Private/UObject/UObjectBase.cpp b/Engine/Source/Runtime/CoreUObject/Private/UObject/UObjectBase.cpp
--- a/Engine/Source/Runtime/CoreUObject/Private/UObject/UObjectBase.cpp
+++ b/Engine/Source/Runtime/CoreUObject/Private/UObject/UObjectBase.cpp
@@ -257,9 +257,13 @@ void UObjectBase::SetClass(UClass* NewClass)
 	STAT(StatID = TStatId();) // reset the stat id since this thing now has a different name
 
 	UnhashObject(this);
+#if USE_UBER_GRAPH_PERSISTENT_FRAME
 	Class->DestroyPersistentUberGraphFrame((UObject*)this);
+#endif
 	Class = NewClass;
+#if USE_UBER_GRAPH_PERSISTENT_FRAME
 	Class->CreatePersistentUberGraphFrame((UObject*)this);
+#endif
 	HashObject(this);
 }
 
diff --git a/Engine/Source/Runtime/CoreUObject/Private/UObject/UObjectGlobals.cpp b/Engine/Source/Runtime/CoreUObject/Private/UObject/UObjectGlobals.cpp
--- a/Engine/Source/Runtime/CoreUObject/Private/UObject/UObjectGlobals.cpp
+++ b/Engine/Source/Runtime/CoreUObject/Private/UObject/UObjectGlobals.cpp
@@ -1813,7 +1813,9 @@ void UObject::PostInitProperties()
 #if !(UE_BUILD_SHIPPING || UE_BUILD_TEST)
 	PostInitPropertiesCheck.Push(this);
 #endif
+#if USE_UBER_GRAPH_PERSISTENT_FRAME
 	GetClass()->CreatePersistentUberGraphFrame(this);
+#endif
 }
 
 UObject::UObject(const FObjectInitializer& ObjectInitializer)
diff --git a/Engine/Source/Runtime/Engine/Private/BlueprintGeneratedClass.cpp b/Engine/Source/Runtime/Engine/Private/BlueprintGeneratedClass.cpp
--- a/Engine/Source/Runtime/Engine/Private/BlueprintGeneratedClass.cpp
+++ b/Engine/Source/Runtime/Engine/Private/BlueprintGeneratedClass.cpp
@@ -536,6 +536,10 @@ FName UBlueprintGeneratedClass::GetUberGraphFrameName()
 
 bool UBlueprintGeneratedClass::UsePersistentUberGraphFrame()
 {
+#if USE_UBER_GRAPH_PERSISTENT_FRAME
 	static const FBoolConfigValueHelper PersistentUberGraphFrame(TEXT("Kismet"), TEXT("bPersistentUberGraphFrame"), GEngineIni);
 	return PersistentUberGraphFrame;
+#else
+	return false;
+#endif
 }
\ No newline at end of file
