CoreUObject - Made code thread-safe when TRACK_BULKDATA_USE is set to 1, should no longer crash

diff --git a/Engine/Source/Runtime/CoreUObject/Private/Serialization/BulkData.cpp b/Engine/Source/Runtime/CoreUObject/Private/Serialization/BulkData.cpp
--- a/Engine/Source/Runtime/CoreUObject/Private/Serialization/BulkData.cpp
+++ b/Engine/Source/Runtime/CoreUObject/Private/Serialization/BulkData.cpp
@@ -14,21 +14,57 @@
 
 #if TRACK_BULKDATA_USE
 
-/** Map from bulk data pointer to object it is contained by */
-TMap<FUntypedBulkData*,UObject*> BulkDataToObjectMap;
+/** Simple wrapper for tracking the bulk data usage in the thread-safe way. */
+struct FThreadSafeBulkDataToObjectMap
+{
+	static FThreadSafeBulkDataToObjectMap& Get()
+	{
+		static FThreadSafeBulkDataToObjectMap Instance;
+		return Instance;
+	}
+
+	void Add( FUntypedBulkData* Key, UObject* Value )
+	{
+		FScopeLock ScopeLock(&CriticalSection);
+		BulkDataToObjectMap.Add(Key,Value);
+	}
+
+	void Remove( FUntypedBulkData* Key )
+	{
+		FScopeLock ScopeLock(&CriticalSection);
+		BulkDataToObjectMap.Remove(Key);
+	}
+
+	FCriticalSection& GetLock() 
+	{
+		return CriticalSection;
+	}
+
+	const TMap<FUntypedBulkData*,UObject*>::TConstIterator GetIterator() const
+	{
+		return BulkDataToObjectMap.CreateConstIterator();
+	}
+
+protected:
+	/** Map from bulk data pointer to object it is contained by */
+	TMap<FUntypedBulkData*,UObject*> BulkDataToObjectMap;
+
+	/** CriticalSection. */
+	FCriticalSection CriticalSection;
+};
 
 /**
  * Helper structure associating an object and a size for sorting purposes.
  */
 struct FObjectAndSize
 {
-	FObjectAndSize( UObject* InObject, int32 InSize )
+	FObjectAndSize( const UObject* InObject, int32 InSize )
 	:	Object( InObject )
 	,	Size( InSize )
 	{}
 
 	/** Object associated with size. */
-	UObject*	Object;
+	const UObject*	Object;
 	/** Size associated with object. */
 	int32			Size;
 };
@@ -69,7 +105,7 @@ FUntypedBulkData::FUntypedBulkData( const FUntypedBulkData& Other )
 	Copy( Other );
 
 #if TRACK_BULKDATA_USE
-	BulkDataToObjectMap.Add( this, NULL );
+	FThreadSafeBulkDataToObjectMap::Get().Add( this, NULL );
 #endif
 }
 
@@ -96,7 +132,7 @@ FUntypedBulkData::~FUntypedBulkData()
 #endif // WITH_EDITOR
 
 #if TRACK_BULKDATA_USE
-	BulkDataToObjectMap.Remove( this );
+	FThreadSafeBulkDataToObjectMap::Get().Remove( this );
 #endif
 }
 
@@ -140,35 +176,39 @@ void FUntypedBulkData::DumpBulkDataUsage( FOutputDevice& Log )
 	TArray<FObjectAndSize> PerObjectSizeArray;
 	TArray<FObjectAndSize> PerClassSizeArray;
 
-	// Iterate over all "live" bulk data and add size to arrays if it is loaded.
-	for( TMap<FUntypedBulkData*,UObject*>::TIterator It(BulkDataToObjectMap); It; ++It )
 	{
-		FUntypedBulkData*	BulkData	= It.Key();
-		UObject*			Owner		= It.Value();
-		// Only add bulk data that is consuming memory to array.
-		if( BulkData->IsBulkDataLoaded() && BulkData->GetBulkDataSize() > 0 )
+		FScopeLock Lock(&FThreadSafeBulkDataToObjectMap::Get().GetLock());
+
+		// Iterate over all "live" bulk data and add size to arrays if it is loaded.
+		for( auto It(FThreadSafeBulkDataToObjectMap::Get().GetIterator()); It; ++It )
 		{
-			// Per object stats.
-			PerObjectSizeArray.Add( FObjectAndSize( Owner, BulkData->GetBulkDataSize() ) );
-			
-			// Per class stats.
-			bool bFoundExistingPerClassSize = false;
-			// Iterate over array, trying to find existing entry.
-			for( int32 PerClassIndex=0; PerClassIndex<PerClassSizeArray.Num(); PerClassIndex++ )
+			const FUntypedBulkData*	BulkData	= It.Key();
+			const UObject*			Owner		= It.Value();
+			// Only add bulk data that is consuming memory to array.
+			if( Owner && BulkData->IsBulkDataLoaded() && BulkData->GetBulkDataSize() > 0 )
 			{
-				FObjectAndSize& PerClassSize = PerClassSizeArray( PerClassIndex );
-				// Add to existing entry if found.
-				if( PerClassSize.Object == Owner->GetClass() )
+				// Per object stats.
+				PerObjectSizeArray.Add( FObjectAndSize( Owner, BulkData->GetBulkDataSize() ) );
+
+				// Per class stats.
+				bool bFoundExistingPerClassSize = false;
+				// Iterate over array, trying to find existing entry.
+				for( int32 PerClassIndex=0; PerClassIndex<PerClassSizeArray.Num(); PerClassIndex++ )
 				{
-					PerClassSize.Size += BulkData->GetBulkDataSize();
-					bFoundExistingPerClassSize = true;
-					break;
+					FObjectAndSize& PerClassSize = PerClassSizeArray[ PerClassIndex ];
+					// Add to existing entry if found.
+					if( PerClassSize.Object == Owner->GetClass() )
+					{
+						PerClassSize.Size += BulkData->GetBulkDataSize();
+						bFoundExistingPerClassSize = true;
+						break;
+					}
+				}
+				// Add new entry if we didn't find an existing one.
+				if( !bFoundExistingPerClassSize )
+				{
+					PerClassSizeArray.Add( FObjectAndSize( Owner->GetClass(), BulkData->GetBulkDataSize() ) );
 				}
-			}
-			// Add new entry if we didn't find an existing one.
-			if( !bFoundExistingPerClassSize )
-			{
-				PerClassSizeArray.Add( FObjectAndSize( Owner->GetClass(), BulkData->GetBulkDataSize() ) );
 			}
 		}
 	}
@@ -191,14 +231,14 @@ void FUntypedBulkData::DumpBulkDataUsage( FOutputDevice& Log )
 	UE_LOG(LogSerialization, Log, TEXT("Per class summary of bulk data use:"));
 	for( int32 PerClassIndex=0; PerClassIndex<PerClassSizeArray.Num(); PerClassIndex++ )
 	{
-		const FObjectAndSize& PerClassSize = PerClassSizeArray( PerClassIndex );
+		const FObjectAndSize& PerClassSize = PerClassSizeArray[ PerClassIndex ];
 		Log.Logf( TEXT("  %5d KByte of bulk data for Class %s"), PerClassSize.Size / 1024, *PerClassSize.Object->GetPathName() );
 	}
 	UE_LOG(LogSerialization, Log, TEXT(""));
 	UE_LOG(LogSerialization, Log, TEXT("Detailed per object stats of bulk data use:"));
 	for( int32 PerObjectIndex=0; PerObjectIndex<PerObjectSizeArray.Num(); PerObjectIndex++ )
 	{
-		const FObjectAndSize& PerObjectSize = PerObjectSizeArray( PerObjectIndex );
+		const FObjectAndSize& PerObjectSize = PerObjectSizeArray[ PerObjectIndex ];
 		Log.Logf( TEXT("  %5d KByte of bulk data for %s"), PerObjectSize.Size / 1024, *PerObjectSize.Object->GetFullName() );
 	}
 	UE_LOG(LogSerialization, Log, TEXT(""));
@@ -609,7 +649,7 @@ void FUntypedBulkData::Serialize( FArchive& Ar, UObject* Owner, int32 Idx )
 	else if( Ar.IsPersistent() && !Ar.IsObjectReferenceCollector() && !Ar.ShouldSkipBulkData() )
 	{
 #if TRACK_BULKDATA_USE
-		BulkDataToObjectMap.Add( this, Owner );
+		FThreadSafeBulkDataToObjectMap::Get().Add( this, Owner );
 #endif
 		// Offset where the bulkdata flags are stored
 		int64 SavedBulkDataFlagsPos	= Ar.Tell();
