#ue4 - Added enum for metadata flag "AdvancedDisplay" on UFunctions, which is used to make pins advanced (initially hidden).

diff --git a/Engine/Source/Runtime/CoreUObject/Public/UObject/ObjectBase.h b/Engine/Source/Runtime/CoreUObject/Public/UObject/ObjectBase.h
--- a/Engine/Source/Runtime/CoreUObject/Public/UObject/ObjectBase.h
+++ b/Engine/Source/Runtime/CoreUObject/Public/UObject/ObjectBase.h
@@ -1040,6 +1040,10 @@ namespace UM
 	// Metadata usable in UFUNCTION
 	enum
 	{
+		/// [FunctionMetadata] Used with a comma-separated list of parameter names that should show up as advanced pins (requiring UI expansion).
+		/// Alternatively you can set a number, which is the number of paramaters from the start that should *not* be marked as advanced (eg 'AdvancedDisplay="2"' will mark all but the first two advanced).
+		AdvancedDisplay,
+
 		/// [FunctionMetadata] Indicates that a BlueprintCallable function should use a Call Array Function node and that the parameters specified in the comma delimited list should be treated as wild card array properties.
 		ArrayParm,
 
