Removing TSubobjectPtr<> from sample games

diff --git a/Engine/Source/Runtime/CoreUObject/Public/UObject/UObjectGlobals.h b/Engine/Source/Runtime/CoreUObject/Public/UObject/UObjectGlobals.h
--- a/Engine/Source/Runtime/CoreUObject/Public/UObject/UObjectGlobals.h
+++ b/Engine/Source/Runtime/CoreUObject/Public/UObject/UObjectGlobals.h
@@ -535,35 +535,12 @@ template<> struct TIsZeroConstructType<FSubobjectPtr> { enum { Value = true }; }
 template<> struct TIsWeakPointerType<FSubobjectPtr> { enum { Value = false }; };
 
 /**
- * TSubobjectPtr - Sub-object smart pointer, owns a reference to instanced objects (sub-objects / components).
- * Prevents anything in C++ from overwriting the sub-object pointer.
- * Can (and should) be declared as a UPROPERTY():
- *
- *   UPROPERTY()
- *   UActorComponent* MyComponent;
- *
- * It can only be assigned to with ObjectInitializer.CreateDefaultSubobject (via TSubobjectPtrConstructor) in the owning object's constructor.
- *
- *   MyComponent = ObjectInitializer.CreateDefaultSubobject<UPathFollowingComponent>(this, TEXT("PathFollowingComponent"));
- *
- * Initialized with InvalidPtrValue by default because it always needs to be initialized (either with NULL or a pointer to sub-object in the constructor
- * of the owner object).
- * Usually used with Actor components to specify that the actor-derived class is the owner of the component and prevent other derived classes from
- * overwriting it in any other way than through ObjectInitializer object.
- *
- * Implements a structure dereference operator for convenience
- *
- *   MyComponent->AttachTo(Owner);
- *
- * Can be reset to NULL, although this functionality is mostly for internal use:
- *
- *   MyComponent.Reset();
+ * TSubobjectPtr - Sub-object smart pointer, soon to be deprecated and should no longer be used.
  */
 template <class SubobjectType>
 class TSubobjectPtr : public FSubobjectPtr
 {
 public:
-
 	/** Internal constructors. */
 	TSubobjectPtr(SubobjectType* InObject)
 		: FSubobjectPtr(InObject)
@@ -574,7 +551,6 @@ public:
 		return *this; 
 	}
 
-
 	/** Default constructor. */
 	TSubobjectPtr()
 		: FSubobjectPtr((UObject*)FSubobjectPtr::InvalidPtrValue)
