Problem:
There is a crash when duplicating BP based on object, with single member variable of struct type. (or UE-4494)

Fix:
Don't call Link() on struct in Serialize() while duplication. Link() will be called later anyway (when all subobjects are serialized).

diff --git a/Engine/Source/Runtime/CoreUObject/Private/UObject/Class.cpp b/Engine/Source/Runtime/CoreUObject/Private/UObject/Class.cpp
--- a/Engine/Source/Runtime/CoreUObject/Private/UObject/Class.cpp
+++ b/Engine/Source/Runtime/CoreUObject/Private/UObject/Class.cpp
@@ -1411,7 +1411,7 @@ void UStruct::Serialize( FArchive& Ar )
 		}	
 	}
 
-	if (Ar.IsLoading() && !dynamic_cast<UClass*>(this)) // classes are linked in the UClass serializer, which just called me
+	if (Ar.IsLoading() && !dynamic_cast<UClass*>(this) && !(Ar.GetPortFlags() & PPF_Duplicate)) // classes are linked in the UClass serializer, which just called me
 	{
 		// Link the properties.
 		Link( Ar, true );
@@ -2941,7 +2941,10 @@ void UClass::Serialize( FArchive& Ar )
 		checkf(!HasAnyClassFlags(CLASS_Native), TEXT("Class %s loaded with CLASS_Native....we should not be loading any native classes."), *GetFullName());
 		checkf(!HasAnyClassFlags(CLASS_Intrinsic), TEXT("Class %s loaded with CLASS_Intrinsic....we should not be loading any intrinsic classes."), *GetFullName());
 		ClassFlags &= ~ CLASS_ShouldNeverBeLoaded;
-		Link(Ar, true);
+		if (!(Ar.GetPortFlags() & PPF_Duplicate))
+		{
+			Link(Ar, true);
+		}
 	}
 
 	if(Ar.IsLoading())
@@ -3026,7 +3029,7 @@ void UClass::Serialize( FArchive& Ar )
 
 	if( Ar.IsLoading() )
 	{
-		check(GetStructureSize() >= sizeof(UObject));
+		check((Ar.GetPortFlags() & PPF_Duplicate) || (GetStructureSize() >= sizeof(UObject)));
 		check(!GetSuperClass() || !GetSuperClass()->HasAnyFlags(RF_NeedLoad));
 		UObject* const OldCDO = ClassDefaultObject;
 		UObject* TempClassDefaultObject = NULL;
