Stats - Moved  STAT_FindObject(Fast) to STATGROUP_ObjectVerbose

diff --git a/Engine/Source/Runtime/CoreUObject/Public/UObject/UObjectGlobals.h b/Engine/Source/Runtime/CoreUObject/Public/UObject/UObjectGlobals.h
--- a/Engine/Source/Runtime/CoreUObject/Public/UObject/UObjectGlobals.h
+++ b/Engine/Source/Runtime/CoreUObject/Public/UObject/UObjectGlobals.h
@@ -16,8 +16,6 @@ DECLARE_CYCLE_STAT_EXTERN(TEXT("AllocateObject"),STAT_AllocateObject,STATGROUP_O
 DECLARE_CYCLE_STAT_EXTERN(TEXT("PostConstructInitializeProperties"),STAT_PostConstructInitializeProperties,STATGROUP_Object, );
 DECLARE_CYCLE_STAT_EXTERN(TEXT("LoadConfig"),STAT_LoadConfig,STATGROUP_Object, );
 DECLARE_CYCLE_STAT_EXTERN(TEXT("LoadObject"),STAT_LoadObject,STATGROUP_Object, );
-DECLARE_DWORD_COUNTER_STAT_EXTERN(TEXT("FindObject"),STAT_FindObject,STATGROUP_Object, );
-DECLARE_DWORD_COUNTER_STAT_EXTERN(TEXT("FindObjectFast"),STAT_FindObjectFast,STATGROUP_Object, );
 DECLARE_CYCLE_STAT_EXTERN(TEXT("InitProperties"),STAT_InitProperties,STATGROUP_Object, );
 DECLARE_DWORD_ACCUMULATOR_STAT_EXTERN(TEXT("NameTable Entries"),STAT_NameTableEntries,STATGROUP_Object, );
 DECLARE_DWORD_ACCUMULATOR_STAT_EXTERN(TEXT("NameTable ANSI Entries"),STAT_NameTableAnsiEntries,STATGROUP_Object, );
@@ -25,6 +23,9 @@ DECLARE_DWORD_ACCUMULATOR_STAT_EXTERN(TEXT("NameTable Wide Entries"),STAT_NameTa
 DECLARE_MEMORY_STAT_EXTERN(TEXT("NameTable Memory Size"),STAT_NameTableMemorySize,STATGROUP_Object, );
 DECLARE_CYCLE_STAT_EXTERN(TEXT("~UObject"),STAT_DestroyObject,STATGROUP_Object, );
 
+DECLARE_DWORD_COUNTER_STAT_EXTERN(TEXT("FindObject"),STAT_FindObject,STATGROUP_ObjectVerbose, );
+DECLARE_DWORD_COUNTER_STAT_EXTERN(TEXT("FindObjectFast"),STAT_FindObjectFast,STATGROUP_ObjectVerbose, );
+
 /**
  * Network stats counters
  */
