Core: Made the mock FGuid properties match the actual types

diff --git a/Engine/Source/Runtime/CoreUObject/Classes/Object.h b/Engine/Source/Runtime/CoreUObject/Classes/Object.h
--- a/Engine/Source/Runtime/CoreUObject/Classes/Object.h
+++ b/Engine/Source/Runtime/CoreUObject/Classes/Object.h
@@ -190,16 +190,16 @@ USTRUCT(immutable, noexport, BlueprintType)
 struct FGuid
 {
 	UPROPERTY(EditAnywhere, SaveGame, Category=Guid)
-	int32 A;
+	uint32 A;
 
 	UPROPERTY(EditAnywhere, SaveGame, Category=Guid)
-	int32 B;
+	uint32 B;
 
 	UPROPERTY(EditAnywhere, SaveGame, Category=Guid)
-	int32 C;
+	uint32 C;
 
 	UPROPERTY(EditAnywhere, SaveGame, Category=Guid)
-	int32 D;
+	uint32 D;
 };
 
 // A unique identifier for networking objects
