#ue4 Gating a warning about missing Super / Parent for functions when they're being regenerated.  Prevents false warnings from being emitted while in blueprint regeneration

diff --git a/Engine/Source/Runtime/CoreUObject/Private/UObject/LinkerLoad.cpp b/Engine/Source/Runtime/CoreUObject/Private/UObject/LinkerLoad.cpp
--- a/Engine/Source/Runtime/CoreUObject/Private/UObject/LinkerLoad.cpp
+++ b/Engine/Source/Runtime/CoreUObject/Private/UObject/LinkerLoad.cpp
@@ -2954,8 +2954,17 @@ UObject* ULinkerLoad::CreateExport( int32 Index )
 			{
 				if( LoadClass->IsChildOf(UFunction::StaticClass()) )
 				{
-					// If this is a function whose parent has been removed, give it a NULL parent, as we would have in the script compiler
-					UE_LOG(LogLinker, Warning, TEXT("CreateExport: Failed to load Parent for %s; removing parent information, but keeping function"), *GetExportFullName(Index));
+					UObject* ObjOuter = IndexToObject(Export.OuterIndex);
+					if (ObjOuter)
+					{
+						UClass* FuncClass = Cast<UClass>(ObjOuter);
+						if (FuncClass && FuncClass->ClassGeneratedBy && !FuncClass->ClassGeneratedBy->HasAnyFlags(RF_BeingRegenerated))
+						{
+							// If this is a function (NOT being regenerated) whose parent has been removed, give it a NULL parent, as we would have in the script compiler
+							UE_LOG(LogLinker, Warning, TEXT("CreateExport: Failed to load Parent for %s; removing parent information, but keeping function"), *GetExportFullName(Index));
+						}
+					}
+
 					Export.SuperIndex = FPackageIndex();
 				}
 				else
