GC will no longer flush streaming and can be run while async loading. Removed GSerializedProperty.

diff --git a/Engine/Source/Developer/HotReload/Private/HotReloadClassReinstancer.cpp b/Engine/Source/Developer/HotReload/Private/HotReloadClassReinstancer.cpp
--- a/Engine/Source/Developer/HotReload/Private/HotReloadClassReinstancer.cpp
+++ b/Engine/Source/Developer/HotReload/Private/HotReloadClassReinstancer.cpp
@@ -77,7 +77,7 @@ void FHotReloadClassReinstancer::SerializeCDOProperties(UObject* InObject, TArra
 				if (!VisitedObjects.Contains(InObj))
 				{
 					VisitedObjects.Add(InObj);
-					if (GSerializedProperty && GSerializedProperty->ContainsInstancedObjectProperty())
+					if (Ar.GetSerializedProperty() && Ar.GetSerializedProperty()->ContainsInstancedObjectProperty())
 					{
 						// Serialize all DSO properties too					
 						FCDOWriter DefaultSubobjectWriter(Bytes, InObj, VisitedObjects);
diff --git a/Engine/Source/Runtime/Core/Private/Serialization/Archive.cpp b/Engine/Source/Runtime/Core/Private/Serialization/Archive.cpp
--- a/Engine/Source/Runtime/Core/Private/Serialization/Archive.cpp
+++ b/Engine/Source/Runtime/Core/Private/Serialization/Archive.cpp
@@ -89,7 +89,8 @@ void FArchive::Reset()
 	ArMaxSerializeSize					= 0;
 	ArIsFilterEditorOnly				= false;
 	ArIsSaveGame						= false;
-	CookingTargetPlatform				= NULL;
+	CookingTargetPlatform = nullptr;
+	SerializedProperty = nullptr;
 
 	// Reset all custom versions to the current registered versions.
 	ResetCustomVersions();
diff --git a/Engine/Source/Runtime/Core/Public/Serialization/ArchiveBase.h b/Engine/Source/Runtime/Core/Public/Serialization/ArchiveBase.h
--- a/Engine/Source/Runtime/Core/Public/Serialization/ArchiveBase.h
+++ b/Engine/Source/Runtime/Core/Public/Serialization/ArchiveBase.h
@@ -886,6 +886,26 @@ public:
 		return false;
 	}
 
+	/**
+	 * Sets the property that is currently being serialized
+	 * 
+	 * @param InProperty Pointer to the property that is currently being serialized
+	 */
+	FORCEINLINE void SetSerializedProperty(class UProperty* InProperty)
+	{
+		SerializedProperty = InProperty;
+	}
+
+	/**
+	 * Gets the property that is currently being serialized
+	 *
+	 * @return Pointer to the property that is currently being serialized
+	 */
+	FORCEINLINE class UProperty* GetSerializedProperty() const
+	{
+		return SerializedProperty;
+	}
+
 protected:
 
 	/** Resets all of the base archive members. */
@@ -1003,6 +1023,9 @@ private:
 	/** Holds the cooking target platform. */
 	const ITargetPlatform* CookingTargetPlatform;
 
+	/** Holds the pointer to the property that is currently being serialized */
+	class UProperty* SerializedProperty;
+
 	/**
 	 * Indicates if the custom versions container is in a 'reset' state.  This will be used to defer the choice about how to
 	 * populate the container until it is needed, where the read/write state will be known.
diff --git a/Engine/Source/Runtime/CoreUObject/Private/Serialization/ArchiveDescribeReference.cpp b/Engine/Source/Runtime/CoreUObject/Private/Serialization/ArchiveDescribeReference.cpp
--- a/Engine/Source/Runtime/CoreUObject/Private/Serialization/ArchiveDescribeReference.cpp
+++ b/Engine/Source/Runtime/CoreUObject/Private/Serialization/ArchiveDescribeReference.cpp
@@ -15,7 +15,6 @@ FArchiveDescribeReference::FArchiveDescribeReference( UObject* Src, UObject* InT
 	ArIgnoreOuterRef					= true;
 	ArIgnoreClassRef					= false;
 
-	GSerializedProperty = NULL;
 	Source->Serialize( *this );
 }
 
@@ -23,9 +22,9 @@ FArchive& FArchiveDescribeReference::operator<<( class UObject*& Obj )
 {
 	if (Obj == Target)
 	{
-		if (GSerializedProperty)
+		if (GetSerializedProperty())
 		{
-			Output.Logf(TEXT("        [%s]"), *GSerializedProperty->GetFullName());
+			Output.Logf(TEXT("        [%s]"), *GetSerializedProperty()->GetFullName());
 		}
 		else
 		{
diff --git a/Engine/Source/Runtime/CoreUObject/Private/Serialization/ArchiveFindAllRefs.cpp b/Engine/Source/Runtime/CoreUObject/Private/Serialization/ArchiveFindAllRefs.cpp
--- a/Engine/Source/Runtime/CoreUObject/Private/Serialization/ArchiveFindAllRefs.cpp
+++ b/Engine/Source/Runtime/CoreUObject/Private/Serialization/ArchiveFindAllRefs.cpp
@@ -13,7 +13,6 @@ FArchiveFindAllRefs::FArchiveFindAllRefs( UObject* Src )
 	ArIgnoreOuterRef					= true;
 	ArIgnoreClassRef					= false;
 
-	GSerializedProperty = NULL;
 	Src->Serialize( *this );
 }
 
diff --git a/Engine/Source/Runtime/CoreUObject/Private/Serialization/ArchiveFindCulprit.cpp b/Engine/Source/Runtime/CoreUObject/Private/Serialization/ArchiveFindCulprit.cpp
--- a/Engine/Source/Runtime/CoreUObject/Private/Serialization/ArchiveFindCulprit.cpp
+++ b/Engine/Source/Runtime/CoreUObject/Private/Serialization/ArchiveFindCulprit.cpp
@@ -27,7 +27,6 @@ FArchiveFindCulprit::FArchiveFindCulprit( UObject* InFind, UObject* Src, bool In
 		ArIsPersistent	= true;
 	}
 
-	GSerializedProperty = NULL;
 	Src->Serialize( *this );
 }
 
@@ -35,9 +34,9 @@ FArchive& FArchiveFindCulprit::operator<<( UObject*& Obj )
 {
 	if( Obj==Find )
 	{
-		if ( GSerializedProperty != NULL )
+		if (GetSerializedProperty() != nullptr)
 		{
-			Referencers.AddUnique(GSerializedProperty);
+			Referencers.AddUnique(GetSerializedProperty());
 		}
 		Count++;
 	}
diff --git a/Engine/Source/Runtime/CoreUObject/Private/Serialization/ArchiveObjectGraph.cpp b/Engine/Source/Runtime/CoreUObject/Private/Serialization/ArchiveObjectGraph.cpp
--- a/Engine/Source/Runtime/CoreUObject/Private/Serialization/ArchiveObjectGraph.cpp
+++ b/Engine/Source/Runtime/CoreUObject/Private/Serialization/ArchiveObjectGraph.cpp
@@ -87,21 +87,21 @@ FArchive& FArchiveObjectGraph::operator<<( class UObject*& Obj )
 			// now record the references between this object and the one referencing it
 			if ( !Record )
 			{
-				ReferencerNode->ReferencedObjects.Add(Obj, FTraceRouteRecord(CurrentObjectNode));
+				ReferencerNode->ReferencedObjects.Add(Obj, FTraceRouteRecord(CurrentObjectNode, GetSerializedProperty()));
 			}
 			else
 			{
-				Record->Add();
+				Record->Add(GetSerializedProperty());
 			}
 
 			Record = CurrentObjectNode->ReferencerRecords.Find(CurrentReferencer);
 			if ( !Record )
 			{
-				CurrentObjectNode->ReferencerRecords.Add(CurrentReferencer, FTraceRouteRecord(ReferencerNode));
+				CurrentObjectNode->ReferencerRecords.Add(CurrentReferencer, FTraceRouteRecord(ReferencerNode, GetSerializedProperty()));
 			}
 			else
 			{
-				Record->Add();
+				Record->Add(GetSerializedProperty());
 			}
 		}
 
diff --git a/Engine/Source/Runtime/CoreUObject/Private/Serialization/ArchiveTraceRoute.cpp b/Engine/Source/Runtime/CoreUObject/Private/Serialization/ArchiveTraceRoute.cpp
--- a/Engine/Source/Runtime/CoreUObject/Private/Serialization/ArchiveTraceRoute.cpp
+++ b/Engine/Source/Runtime/CoreUObject/Private/Serialization/ArchiveTraceRoute.cpp
@@ -311,21 +311,21 @@ FArchive& FArchiveTraceRoute::operator<<( class UObject*& Obj )
 			// now record the references between this object and the one referencing it
 			if ( !Record )
 			{
-				ReferencerNode->ReferencedObjects.Add(Obj, FTraceRouteRecord(CurrentObjectNode));
+				ReferencerNode->ReferencedObjects.Add(Obj, FTraceRouteRecord(CurrentObjectNode, GetSerializedProperty()));
 			}
 			else
 			{
-				Record->Add();
+				Record->Add(GetSerializedProperty());
 			}
 
 			Record = CurrentObjectNode->ReferencerRecords.Find(CurrentReferencer);
 			if ( !Record )
 			{
-				CurrentObjectNode->ReferencerRecords.Add(CurrentReferencer, FTraceRouteRecord(ReferencerNode));
+				CurrentObjectNode->ReferencerRecords.Add(CurrentReferencer, FTraceRouteRecord(ReferencerNode, GetSerializedProperty()));
 			}
 			else
 			{
-				Record->Add();
+				Record->Add(GetSerializedProperty());
 			}
 		}
 
diff --git a/Engine/Source/Runtime/CoreUObject/Private/Serialization/AsyncLoading.cpp b/Engine/Source/Runtime/CoreUObject/Private/Serialization/AsyncLoading.cpp
--- a/Engine/Source/Runtime/CoreUObject/Private/Serialization/AsyncLoading.cpp
+++ b/Engine/Source/Runtime/CoreUObject/Private/Serialization/AsyncLoading.cpp
@@ -33,17 +33,57 @@ DECLARE_CYCLE_STAT(TEXT("Async Loading Time"),STAT_AsyncLoadingTime,STATGROUP_As
 /** Objects that have been constructed during async loading phase.						*/
 static TArray<UObject*>			GObjConstructedDuringAsyncLoading;
 
+/** Keeps a reference to all objects created during async load until streaming has finished */
+class FAsyncObjectsReferencer : FGCObject
+{
+	/** Private constructor */
+	FAsyncObjectsReferencer() {}
+
+	/** List of referenced objects */
+	TArray<UObject*> ReferencedObjects;
+public:
+	/** Returns the one and only instance of this object */
+	static FAsyncObjectsReferencer& Get();
+	/** FGCObject interface */
+	virtual void AddReferencedObjects(FReferenceCollector& Collector) override
+	{
+		Collector.AllowEliminatingReferences(false);
+		Collector.AddReferencedObjects(ReferencedObjects);
+		Collector.AllowEliminatingReferences(true);
+	}
+	/** Adds an object to be referenced */
+	FORCEINLINE void AddObject(UObject* InObject)
+	{
+		ReferencedObjects.Add(InObject);
+	}
+	/** Removes all objects from the list */
+	FORCEINLINE void EmptyReferencedObjects()
+	{
+		ReferencedObjects.Empty(ReferencedObjects.Num());
+	}
+};
+FAsyncObjectsReferencer& FAsyncObjectsReferencer::Get()
+{
+	static FAsyncObjectsReferencer Singleton;
+	return Singleton;
+}
+
 /**
  * Call back into the async loading code to inform of the creation of a new object
  * @param Object		Object created
+ * @param bSubObject	Object created as a sub-object of a loaded object
  */
-void NotifyConstructedDuringAsyncLoading(UObject* Object)
+void NotifyConstructedDuringAsyncLoading(UObject* Object, bool bSubObject)
 {
 	// Mark objects created during async loading process (e.g. from within PostLoad or CreateExport) as async loaded so they 
 	// cannot be found. This requires also keeping track of them so we can remove the async loading flag later one when we 
 	// finished routing PostLoad to all objects.
-	Object->SetFlags( RF_AsyncLoading );
-	GObjConstructedDuringAsyncLoading.Add(Object);
+	if (!bSubObject)
+	{
+		Object->SetFlags(RF_AsyncLoading);
+		GObjConstructedDuringAsyncLoading.Add(Object);
+	}
+	FAsyncObjectsReferencer::Get().AddObject(Object);
 }
 
 
@@ -247,7 +287,7 @@ EAsyncPackageState::Type FAsyncPackage::Tick( bool InbUseTimeLimit, bool InbUseF
 
 		// Call PostLoad on objects, this could cause new objects to be loaded that require
 		// another iteration of the PreLoad loop.
-		if( LoadingState == EAsyncPackageState::Complete )
+		if (LoadingState == EAsyncPackageState::Complete)
 		{
 			LoadingState = PostLoadObjects();
 		}
@@ -731,7 +771,14 @@ EAsyncPackageState::Type FAsyncPackage::PostLoadObjects()
 	}
 
 	// New objects might have been loaded during PostLoad.
-	return (PreLoadIndex == GObjLoaded.Num()) && (PostLoadIndex == GObjLoaded.Num()) ? EAsyncPackageState::Complete : EAsyncPackageState::TimeOut;
+	EAsyncPackageState::Type Result = (PreLoadIndex == GObjLoaded.Num()) && (PostLoadIndex == GObjLoaded.Num()) ? EAsyncPackageState::Complete : EAsyncPackageState::TimeOut;
+	if (PreLoadIndex < GObjLoaded.Num())
+	{
+		static volatile int32 xx = 0;
+		xx++;
+	}
+
+	return Result;
 }
 
 /**
@@ -965,6 +1012,7 @@ EAsyncPackageState::Type ProcessAsyncLoading( bool bUseTimeLimit, bool bUseFullT
 	EAsyncPackageState::Type LoadingState = EAsyncPackageState::Complete;
 	EAsyncPackageState::Type CompletionState = EAsyncPackageState::Complete;
 
+	bool bWasLoading = GObjAsyncPackages.Num() > 0;
 	// We need to loop as the function has to handle finish loading everything given no time limit
 	// like e.g. when called from FlushAsyncLoading.
 	for (int32 i = 0; LoadingState != EAsyncPackageState::TimeOut && i < GObjAsyncPackages.Num(); i++)
@@ -1023,6 +1071,13 @@ EAsyncPackageState::Type ProcessAsyncLoading( bool bUseTimeLimit, bool bUseFullT
 		CompletionState = FMath::Min(CompletionState, LoadingState);
 		// We cannot access Package anymore!
 	}
+
+	// Not that streaming has finished, we can stop force-referencing all objects that were created during async load.
+	if (CompletionState == EAsyncPackageState::Complete && bWasLoading)
+	{
+		FAsyncObjectsReferencer::Get().EmptyReferencedObjects();
+	}
+
 	return CompletionState;
 }
 
diff --git a/Engine/Source/Runtime/CoreUObject/Private/Serialization/DuplicateDataReader.cpp b/Engine/Source/Runtime/CoreUObject/Private/Serialization/DuplicateDataReader.cpp
--- a/Engine/Source/Runtime/CoreUObject/Private/Serialization/DuplicateDataReader.cpp
+++ b/Engine/Source/Runtime/CoreUObject/Private/Serialization/DuplicateDataReader.cpp
@@ -24,7 +24,7 @@ FDuplicateDataReader::FDuplicateDataReader( class FUObjectAnnotationSparse<FDupl
 void FDuplicateDataReader::SerializeFail()
 {
 	extern UObject* GSerializedObject;
-	UE_LOG(LogObj, Fatal, TEXT("FDuplicateDataReader Overread. GSerializedObject = %s GSerializedProperty = %s"), *GetFullNameSafe(GSerializedObject), *GetFullNameSafe(GSerializedProperty));
+	UE_LOG(LogObj, Fatal, TEXT("FDuplicateDataReader Overread. SerializedObject = %s SerializedProperty = %s"), *GetFullNameSafe(GSerializedObject), *GetFullNameSafe(GetSerializedProperty()));
 }
 
 FArchive& FDuplicateDataReader::operator<<(FName& N)
diff --git a/Engine/Source/Runtime/CoreUObject/Private/Serialization/FindReferencersArchive.cpp b/Engine/Source/Runtime/CoreUObject/Private/Serialization/FindReferencersArchive.cpp
--- a/Engine/Source/Runtime/CoreUObject/Private/Serialization/FindReferencersArchive.cpp
+++ b/Engine/Source/Runtime/CoreUObject/Private/Serialization/FindReferencersArchive.cpp
@@ -151,9 +151,9 @@ FArchive& FFindReferencersArchive::operator<<( UObject*& Obj )
 		if ( pReferenceCount != nullptr )
 		{
 			// if this object was serialized via a UProperty, add it to the list
-			if ( GSerializedProperty != nullptr )
+			if (GetSerializedProperty() != nullptr)
 			{
-				ReferenceMap.AddUnique(Obj, GSerializedProperty);
+				ReferenceMap.AddUnique(Obj, GetSerializedProperty());
 			}
 
 			// now increment the reference count for this target object
diff --git a/Engine/Source/Runtime/CoreUObject/Private/UObject/GarbageCollection.cpp b/Engine/Source/Runtime/CoreUObject/Private/UObject/GarbageCollection.cpp
--- a/Engine/Source/Runtime/CoreUObject/Private/UObject/GarbageCollection.cpp
+++ b/Engine/Source/Runtime/CoreUObject/Private/UObject/GarbageCollection.cpp
@@ -163,7 +163,7 @@ void SerializeRootSet( FArchive& Ar, EObjectFlags KeepFlags )
  * @param ReferencingObject UObject which owns the reference (can be NULL)
  * @param bAllowReferenceElimination	Whether to allow NULL'ing the reference if RF_PendingKill is set
  */
-static FORCEINLINE void HandleObjectReference( TArray<UObject*>& ObjectsToSerialize, UObject* ReferencingObject, UObject*& Object, bool bAllowReferenceElimination )
+static FORCEINLINE void HandleObjectReference(TArray<UObject*>& ObjectsToSerialize, UObject* ReferencingObject, UObject*& Object, bool bAllowReferenceElimination)
 {
 	// Disregard NULL objects and perform very fast check to see whether object is part of permanent
 	// object pool and should therefore be disregarded. The check doesn't touch the object and is
@@ -197,7 +197,7 @@ static FORCEINLINE void HandleObjectReference( TArray<UObject*>& ObjectsToSerial
 					// this message is to help track down culprits behind "Object in PIE world still referenced" errors
 					if ( GIsEditor && !GIsPlayInEditorWorld && ReferencingObject != NULL && !ReferencingObject->RootPackageHasAnyFlags(PKG_PlayInEditor) && Object->RootPackageHasAnyFlags(PKG_PlayInEditor) )
 					{
-						UE_LOG(LogGarbage, Warning, TEXT("GC detected illegal reference to PIE object from content [possibly via %s]:"), *GSerializedProperty->GetFullName());
+						UE_LOG(LogGarbage, Warning, TEXT("GC detected illegal reference to PIE object from content [possibly via [todo]]:"));
 						UE_LOG(LogGarbage, Warning, TEXT("      PIE object: %s"), *Object->GetFullName());
 						UE_LOG(LogGarbage, Warning, TEXT("  NON-PIE object: %s"), *ReferencingObject->GetFullName());
 					}
@@ -222,7 +222,7 @@ static FORCEINLINE void HandleObjectReference( TArray<UObject*>& ObjectsToSerial
 	}
 }
 
-static FORCEINLINE void HandleTokenStreamObjectReference( TArray<UObject*>& ObjectsToSerialize, UObject* ReferencingObject, UObject*& Object, const int32 TokenIndex, bool bAllowReferenceElimination )
+static FORCEINLINE void HandleTokenStreamObjectReference(TArray<UObject*>& ObjectsToSerialize, UObject* ReferencingObject, UObject*& Object, const int32 TokenIndex, bool bAllowReferenceElimination)
 {
 #if !(UE_BUILD_TEST || UE_BUILD_SHIPPING)
 	if (Object && !Object->IsValidLowLevelFast())
@@ -256,13 +256,13 @@ class FGCCollector : public FReferenceCollector
 
 public:
 
-	FGCCollector( TArray<UObject*>& InObjectArray )
-		: ObjectArray( InObjectArray )
-		, bAllowEliminatingReferences( true )
+	FGCCollector(TArray<UObject*>& InObjectArray)
+		: ObjectArray(InObjectArray)
+		, bAllowEliminatingReferences(true)
 	{
 	}
 
-	virtual void HandleObjectReference( UObject*& Object, const UObject* ReferencingObject, const UObject* ReferencingProperty ) override
+	virtual void HandleObjectReference(UObject*& Object, const UObject* ReferencingObject, const UObject* ReferencingProperty) override
 	{
 #if !(UE_BUILD_TEST || UE_BUILD_SHIPPING)
 		if (Object && !Object->IsValidLowLevelFast())
@@ -273,7 +273,7 @@ public:
 				ReferencingProperty ? *ReferencingProperty->GetFullName() : TEXT("NULL"));
 		}
 #endif
-		::HandleObjectReference( ObjectArray, const_cast<UObject*>(ReferencingObject), Object, bAllowEliminatingReferences );
+		::HandleObjectReference(ObjectArray, const_cast<UObject*>(ReferencingObject), Object, bAllowEliminatingReferences);
 	}
 	virtual bool IsIgnoringArchetypeRef() const override
 	{
@@ -296,27 +296,38 @@ public:
 FReferenceFinder::FReferenceFinder( TArray<UObject*>& InObjectArray, UObject* InOuter, bool bInRequireDirectOuter, bool bInShouldIgnoreArchetype, bool bInSerializeRecursively, bool bInShouldIgnoreTransient )
 	:	ObjectArray( InObjectArray )
 	,	LimitOuter( InOuter )
+	, SerializedProperty(nullptr)
 	,	bRequireDirectOuter( bInRequireDirectOuter )
 	, bShouldIgnoreArchetype( bInShouldIgnoreArchetype )
 	, bSerializeRecursively( false )
 	, bShouldIgnoreTransient( bInShouldIgnoreTransient )
 {
 	bSerializeRecursively = bInSerializeRecursively && LimitOuter != NULL;
+	if (InOuter)
+	{
+		// If the outer is specified, try to set the SerializedProperty based on its linker.
+		auto OuterLinker = InOuter->GetLinker();
+		if (OuterLinker)
+		{
+			SerializedProperty = OuterLinker->GetSerializedProperty();
+		}
+	}
 }
 
-void FReferenceFinder::FindReferences( UObject* Object )
+void FReferenceFinder::FindReferences(UObject* Object, UObject* InReferencingObject, UObject* InReferencingProperty)
 {
 	check(Object != NULL);
 
-	if( !Object->GetClass()->IsChildOf(UClass::StaticClass()) )
+	if (!Object->GetClass()->IsChildOf(UClass::StaticClass()))
 	{
-		FSimpleObjectReferenceCollectorArchive CollectorArchive( Object, *this );
-		Object->SerializeScriptProperties( CollectorArchive );
+		FSimpleObjectReferenceCollectorArchive CollectorArchive(Object, *this);
+		CollectorArchive.SetSerializedProperty(SerializedProperty);
+		Object->SerializeScriptProperties(CollectorArchive);
 	}
 	Object->CallAddReferencedObjects(*this);
 }
 
-void FReferenceFinder::HandleObjectReference( UObject*& InObject, const UObject* ReferencingObject /*= NULL*/, const UObject* ReferencingProperty /*= NULL*/ )
+void FReferenceFinder::HandleObjectReference( UObject*& InObject, const UObject* InReferencingObject /*= NULL*/, const UObject* InReferencingProperty /*= NULL*/ )
 {
 	// Avoid duplicate entries.
 	if ( InObject != NULL )
@@ -336,7 +347,7 @@ void FReferenceFinder::HandleObjectReference( UObject*& InObject, const UObject*
 			if ( bSerializeRecursively == true && !SerializedObjects.Find(Object) )
 			{
 				SerializedObjects.Add(Object);
-				FindReferences( Object );
+				FindReferences(Object, const_cast<UObject*>(InReferencingObject), const_cast<UObject*>(InReferencingProperty));
 			}
 		}
 	}
@@ -433,7 +444,7 @@ public:
 			//@todo UE4 - A prefetch was removed here. Re-add it. It wasn't right anyway, since it was ten items ahead and the consoles on have 8 prefetch slots
 			
 			// We can't collect garbage during an async load operation and by now all unreachable objects should've been purged.
-			checkf( !Object->HasAnyFlags(RF_AsyncLoading|RF_Unreachable), TEXT("%s"), *Object->GetFullName() );
+			checkf( !Object->HasAnyFlags(RF_Unreachable), TEXT("%s"), *Object->GetFullName() );
 	
 			// Keep track of how many objects are around.
 			GObjectCountDuringLastMarkPhase++;
@@ -1038,22 +1049,6 @@ static const auto CVarAllowParallelGC =
 
 void CollectGarbage( EObjectFlags KeepFlags, bool bPerformFullPurge )
 {
-	// Helper class to register FlushAsyncLoadingCallback on first GC run.
-	struct FAddFlushAsyncLoadingCallback
-	{
-		FAddFlushAsyncLoadingCallback()
-		{
-			FCoreUObjectDelegates::PreGarbageCollect.AddStatic(FlushAsyncLoadingCallback);
-		}
-		/** Wrapper function to handle default parameter when used as function pointer */
-		static void FlushAsyncLoadingCallback()
-		{
-			FlushAsyncLoading();
-		}
-	};
-	// Add FlushAsyncLoadingCallback the first time CollectGarbage is called.
-	static FAddFlushAsyncLoadingCallback AddFlushAsyncLoadingCallback;
-
 	// We can't collect garbage while there's a load in progress. E.g. one potential issue is Import.XObject
 	check( !IsLoading() );
 
@@ -1223,6 +1218,7 @@ void UObject::AddReferencedObjects(UObject* This, FReferenceCollector& Collector
 		{
 			// Script properties
 			FSimpleObjectReferenceCollectorArchive ObjectReferenceCollector( This, Collector );
+			ObjectReferenceCollector.SetSerializedProperty(Collector.GetSerializedProperty());
 			This->SerializeScriptProperties( ObjectReferenceCollector );
 		}
 
diff --git a/Engine/Source/Runtime/CoreUObject/Private/UObject/Linker.cpp b/Engine/Source/Runtime/CoreUObject/Private/UObject/Linker.cpp
--- a/Engine/Source/Runtime/CoreUObject/Private/UObject/Linker.cpp
+++ b/Engine/Source/Runtime/CoreUObject/Private/UObject/Linker.cpp
@@ -482,13 +482,13 @@ void DissociateImportsAndForcedExports()
 	GForcedExportCount = 0;
 }
 
-static void LogGetPackageLinkerError(const TCHAR* InFilename, const FText& InFullErrorMessage, const FText& InSummaryErrorMessage, UObject* InOuter, uint32 LoadFlags)
+static void LogGetPackageLinkerError(FArchiveUObject* LinkerArchive, const TCHAR* InFilename, const FText& InFullErrorMessage, const FText& InSummaryErrorMessage, UObject* InOuter, uint32 LoadFlags)
 {
 	static FName NAME_LoadErrors("LoadErrors");
 	struct Local
 	{
 		/** Helper function to output more detailed error info if available */
-		static void OutputErrorDetail(const FName& LogName)
+		static void OutputErrorDetail(FArchiveUObject* LinkerArchive, const FName& LogName)
 		{
 			if ( GSerializedObject && GSerializedImportLinker )
 			{
@@ -499,9 +499,10 @@ static void LogGetPackageLinkerError(const TCHAR* InFilename, const FText& InFul
 				Message->AddToken(FAssetNameToken::Create(GSerializedImportLinker->GetImportPathName(GSerializedImportIndex)));
 				Message->AddToken(FTextToken::Create(LOCTEXT("FailedLoad_Referenced", "Referenced by")));
 				Message->AddToken(FUObjectToken::Create(GSerializedObject));
-				if(GSerializedProperty != NULL)
+				auto SerializedProperty = LinkerArchive ? LinkerArchive->GetSerializedProperty() : nullptr;
+				if (SerializedProperty != nullptr)
 				{
-					FString PropertyPathName = GSerializedProperty->GetPathName();
+					FString PropertyPathName = SerializedProperty->GetPathName();
 					Message->AddToken(FTextToken::Create(LOCTEXT("FailedLoad_Property", "Property")));
 					Message->AddToken(FAssetNameToken::Create(PropertyPathName, FText::FromString( PropertyPathName) ) );
 				}
@@ -541,14 +542,14 @@ static void LogGetPackageLinkerError(const TCHAR* InFilename, const FText& InFul
 				Message->AddToken(FAssetNameToken::Create(FPackageName::FilenameToLongPackageName(InOuter->GetPathName())));
 			}
 
-			Local::OutputErrorDetail(NAME_LoadErrors);
+			Local::OutputErrorDetail(LinkerArchive, NAME_LoadErrors);
 		}
 	}
 	else
 	{
 		if (!(LoadFlags & LOAD_NoWarn))
 		{
-			Local::OutputErrorDetail(NAME_LoadErrors);
+			Local::OutputErrorDetail(LinkerArchive, NAME_LoadErrors);
 		}
 
 		FFormatNamedArguments Arguments;
@@ -599,7 +600,7 @@ ULinkerLoad* GetPackageLinker
 		{
 			// try to recover from this instead of throwing, it seems recoverable just by doing this
 			FText ErrorText(LOCTEXT("PackageResolveFailed", "Can't resolve asset name"));
-			LogGetPackageLinkerError(InLongPackageName, ErrorText, ErrorText, InOuter, LoadFlags);
+			LogGetPackageLinkerError(Result, InLongPackageName, ErrorText, ErrorText, InOuter, LoadFlags);
 			return nullptr;
 		}
 
@@ -611,7 +612,7 @@ ULinkerLoad* GetPackageLinker
 				FFormatNamedArguments Arguments;
 				Arguments.Add(TEXT("AssetName"), FText::FromString(InOuter->GetName()));
 				Arguments.Add(TEXT("PackageName"), FText::FromString(GSerializedPackageLinker ? *(GSerializedPackageLinker->Filename) : TEXT("NULL")));
-				LogGetPackageLinkerError(GSerializedPackageLinker ? *GSerializedPackageLinker->Filename : nullptr,
+				LogGetPackageLinkerError(Result, GSerializedPackageLinker ? *GSerializedPackageLinker->Filename : nullptr,
 											FText::Format(LOCTEXT("PackageNotFound", "Can't find file for asset '{AssetName}' while loading {PackageName}."), Arguments),
 											LOCTEXT("PackageNotFoundShort", "Can't find file for asset."),
 											InOuter,
@@ -628,7 +629,7 @@ ULinkerLoad* GetPackageLinker
 		{
 			// try to recover from this instead of throwing, it seems recoverable just by doing this
 			FText ErrorText(LOCTEXT("PackageResolveFailed", "Can't resolve asset name"));
-			LogGetPackageLinkerError(InLongPackageName, ErrorText, ErrorText, InOuter, LoadFlags);
+			LogGetPackageLinkerError(Result, InLongPackageName, ErrorText, ErrorText, InOuter, LoadFlags);
 			return nullptr;
 		}
 
@@ -648,7 +649,7 @@ ULinkerLoad* GetPackageLinker
 			Arguments.Add(TEXT("Filename"), FText::FromString(InLongPackageName));
 
 			// try to recover from this instead of throwing, it seems recoverable just by doing this
-			LogGetPackageLinkerError(InLongPackageName, FText::Format(LOCTEXT("FileNotFound", "Can't find file '{Filename}'"), Arguments), LOCTEXT("FileNotFoundShort", "Can't find file"), InOuter, LoadFlags);
+			LogGetPackageLinkerError(Result, InLongPackageName, FText::Format(LOCTEXT("FileNotFound", "Can't find file '{Filename}'"), Arguments), LOCTEXT("FileNotFoundShort", "Can't find file"), InOuter, LoadFlags);
 			return nullptr;
 		}
 
@@ -666,7 +667,7 @@ ULinkerLoad* GetPackageLinker
 			{
 				FFormatNamedArguments Arguments;
 				Arguments.Add(TEXT("Filename"), FText::FromString(InLongPackageName));
-				LogGetPackageLinkerError(InLongPackageName, FText::Format(LOCTEXT("FilenameToPackage", "Can't convert filename '{Filename}' to asset name"), Arguments), LOCTEXT("FilenameToPackageShort", "Can't convert filename to asset name"), InOuter, LoadFlags);
+				LogGetPackageLinkerError(Result, InLongPackageName, FText::Format(LOCTEXT("FilenameToPackage", "Can't convert filename '{Filename}' to asset name"), Arguments), LOCTEXT("FilenameToPackageShort", "Can't convert filename to asset name"), InOuter, LoadFlags);
 				return nullptr;
 			}
 			InOuter = FilenamePkg;
@@ -687,7 +688,7 @@ ULinkerLoad* GetPackageLinker
 		FFormatNamedArguments Arguments;
 		Arguments.Add(TEXT("AssetName"), FText::FromString(InOuter->GetName()));
 
-		LogGetPackageLinkerError(InLongPackageName, FText::Format(LOCTEXT("Sandbox", "Asset '{AssetName}' is not accessible in this sandbox"), Arguments), LOCTEXT("SandboxShort", "Asset is not accessible in this sandbox"), InOuter, LoadFlags);
+		LogGetPackageLinkerError(Result, InLongPackageName, FText::Format(LOCTEXT("Sandbox", "Asset '{AssetName}' is not accessible in this sandbox"), Arguments), LOCTEXT("SandboxShort", "Asset is not accessible in this sandbox"), InOuter, LoadFlags);
 		return nullptr;
 	}
 #endif
@@ -710,7 +711,7 @@ ULinkerLoad* GetPackageLinker
 		Arguments.Add(TEXT("AssetName"), FText::FromString(InOuter->GetName()));
 
 		// This should never fire, because FindPackageFile should never return an incompatible file
-		LogGetPackageLinkerError(InLongPackageName, FText::Format(LOCTEXT("PackageVersion", "Asset '{AssetName}' version mismatch"), Arguments), LOCTEXT("PackageVersionShort", "Asset version mismatch"), InOuter, LoadFlags);
+		LogGetPackageLinkerError(Result, InLongPackageName, FText::Format(LOCTEXT("PackageVersion", "Asset '{AssetName}' version mismatch"), Arguments), LOCTEXT("PackageVersionShort", "Asset version mismatch"), InOuter, LoadFlags);
 		return nullptr;
 	}
 
diff --git a/Engine/Source/Runtime/CoreUObject/Private/UObject/LinkerLoad.cpp b/Engine/Source/Runtime/CoreUObject/Private/UObject/LinkerLoad.cpp
--- a/Engine/Source/Runtime/CoreUObject/Private/UObject/LinkerLoad.cpp
+++ b/Engine/Source/Runtime/CoreUObject/Private/UObject/LinkerLoad.cpp
@@ -14,7 +14,8 @@ DECLARE_STATS_GROUP_VERBOSE(TEXT("Linker Load"), STATGROUP_LinkerLoad, STATCAT_A
 DECLARE_CYCLE_STAT(TEXT("Linker Preload"),STAT_LinkerPreload,STATGROUP_LinkerLoad);
 DECLARE_CYCLE_STAT(TEXT("Linker Precache"),STAT_LinkerPrecache,STATGROUP_LinkerLoad);
 DECLARE_CYCLE_STAT(TEXT("Linker Serialize"),STAT_LinkerSerialize,STATGROUP_LinkerLoad);
-
+DECLARE_DWORD_ACCUMULATOR_STAT(TEXT("Linker Count"), STAT_LinkerCount, STATGROUP_LinkerLoad);
+DECLARE_DWORD_ACCUMULATOR_STAT(TEXT("Live Linker Count"), STAT_LiveLinkerCount, STATGROUP_LinkerLoad);
 
 /** Points to the main PackageLinker currently being serialized (Defined in Linker.cpp) */
 extern ULinkerLoad* GSerializedPackageLinker;
@@ -602,6 +603,9 @@ ULinkerLoad::ULinkerLoad( const class FObjectInitializer& ObjectInitializer, UPa
 #endif
 {
 	check(!HasAnyFlags(RF_ClassDefaultObject));
+
+	INC_DWORD_STAT(STAT_LinkerCount);
+	INC_DWORD_STAT(STAT_LiveLinkerCount);
 }
 
 ULinkerLoad::~ULinkerLoad()
@@ -2062,6 +2066,9 @@ ULinkerLoad::EVerifyResult ULinkerLoad::VerifyImport(int32 ImportIndex)
 	return Result;
 }
 
+// Internal Load package call so that we can pass the linker that requested this package as an import dependency
+UPackage* LoadPackageInternal(UPackage* InOuter, const TCHAR* InLongPackageName, uint32 LoadFlags, ULinkerLoad* ImportLinker);
+
 /**
  * Safely verify that an import in the ImportMap points to a good object. This decides whether or not
  * a failure to load the object redirector in the wrapper is a fatal error or not (return value)
@@ -2124,7 +2131,7 @@ bool ULinkerLoad::VerifyImportInner(const int32 ImportIndex, FString& WarningSuf
 		{
 			// we now fully load the package that we need a single export from - however, we still use CreatePackage below as it handles all cases when the package
 			// didn't exist (native only), etc		
-			TmpPkg = LoadPackage(NULL, *Import.ObjectName.ToString(), InternalLoadFlags);
+			TmpPkg = LoadPackageInternal(NULL, *Import.ObjectName.ToString(), InternalLoadFlags, this);
 		}
 
 #if WITH_EDITOR
@@ -3499,6 +3506,7 @@ void ULinkerLoad::Detach( bool bEnsureAllBulkDataIsLoaded )
 
 void ULinkerLoad::BeginDestroy()
 {
+	DEC_DWORD_STAT(STAT_LiveLinkerCount);
 	// Detaches linker.
 	Detach( false );
 	Super::BeginDestroy();
diff --git a/Engine/Source/Runtime/CoreUObject/Private/UObject/Property.cpp b/Engine/Source/Runtime/CoreUObject/Private/UObject/Property.cpp
--- a/Engine/Source/Runtime/CoreUObject/Private/UObject/Property.cpp
+++ b/Engine/Source/Runtime/CoreUObject/Private/UObject/Property.cpp
@@ -11,9 +11,6 @@
 
 DEFINE_LOG_CATEGORY(LogProperty);
 
-UProperty* GSerializedProperty = NULL;
-
-
 // List the core ones here as they have already been included (and can be used without CoreUObject!)
 template<>
 struct TStructOpsTypeTraits<FVector> : public TStructOpsTypeTraitsBase
diff --git a/Engine/Source/Runtime/CoreUObject/Private/UObject/SavePackage.cpp b/Engine/Source/Runtime/CoreUObject/Private/UObject/SavePackage.cpp
--- a/Engine/Source/Runtime/CoreUObject/Private/UObject/SavePackage.cpp
+++ b/Engine/Source/Runtime/CoreUObject/Private/UObject/SavePackage.cpp
@@ -40,16 +40,16 @@ static bool HasDeprecatedOrPendingKillOuter(UObject* InObj, UPackage* InSavingPa
 	return false;
 }
 
-static void CheckObjectPriorToSave(UObject* InObj, UPackage* InSavingPackage)
+static void CheckObjectPriorToSave(FArchiveUObject& Ar, UObject* InObj, UPackage* InSavingPackage)
 {
 	if (!InObj)
 	{
 		return;
 	}
+	extern UObject* GSerializedObject;
 	if (!InObj->IsValidLowLevelFast() || !InObj->IsValidLowLevel())
 	{
-		extern UObject* GSerializedObject;
-		UE_LOG(LogLinker, Fatal, TEXT("Attempt to save bogus object %p GSerializedObject=%s  GSerializedProperty=%s"), (void*)InObj, *GetFullNameSafe(GSerializedObject), *GetFullNameSafe(GSerializedProperty));
+		UE_LOG(LogLinker, Fatal, TEXT("Attempt to save bogus object %p GSerializedObject=%s  SerializedProperty=%s"), (void*)InObj, *GetFullNameSafe(GSerializedObject), *GetFullNameSafe(Ar.GetSerializedProperty()));
 		return;
 	}
 	// if the object class is abstract or has been marked as deprecated, mark this
@@ -64,7 +64,7 @@ static void CheckObjectPriorToSave(UObject* InObj, UPackage* InSavingPackage)
 		{
 			TArray<UObject*> ComponentReferences;
 			FReferenceFinder ComponentCollector(ComponentReferences, InObj, false, true, true);
-			ComponentCollector.FindReferences(InObj);
+			ComponentCollector.FindReferences(InObj, GSerializedObject, Ar.GetSerializedProperty());
 
 			for ( int32 Index = 0; Index < ComponentReferences.Num(); Index++ )
 			{
@@ -323,7 +323,7 @@ FString FArchiveSaveTagExports::GetArchiveName() const
 FArchive& FArchiveSaveTagExports::operator<<( UObject*& Obj )
 {
 	check(Outer);
-	CheckObjectPriorToSave(Obj, Outer);
+	CheckObjectPriorToSave(*this, Obj, Outer);
 	if( Obj && Obj->IsIn(Outer) && !Obj->HasAnyFlags(RF_Transient) && !Obj->HasAnyMarks(OBJECTMARK_TagExp) )
 	{
 #if 0
@@ -583,7 +583,7 @@ FString FArchiveSaveTagImports::GetArchiveName() const
 
 FArchive& FArchiveSaveTagImports::operator<<( UObject*& Obj )
 {
-	CheckObjectPriorToSave(Obj, NULL);
+	CheckObjectPriorToSave(*this, Obj, NULL);
 
 	const EObjectMark ObjectMarks = UPackage::GetObjectMarksForTargetPlatform( CookingTarget(), IsCooking() );
 	
diff --git a/Engine/Source/Runtime/CoreUObject/Private/UObject/UObjectGlobals.cpp b/Engine/Source/Runtime/CoreUObject/Private/UObject/UObjectGlobals.cpp
--- a/Engine/Source/Runtime/CoreUObject/Private/UObject/UObjectGlobals.cpp
+++ b/Engine/Source/Runtime/CoreUObject/Private/UObject/UObjectGlobals.cpp
@@ -16,7 +16,7 @@ DEFINE_LOG_CATEGORY(LogUObjectGlobals);
 #include "MallocProfiler.h"
 #endif
 
-bool						GIsSavingPackage						= false;
+bool						GIsSavingPackage = false;
 int32							GImportCount							= 0;
 /** Forced exports for EndLoad optimization.											*/
 int32							GForcedExportCount						= 0;
@@ -763,14 +763,15 @@ UClass* StaticLoadClass( UClass* BaseClass, UObject* InOuter, const TCHAR* InNam
 }
 
 /**
- * Loads a package and all contained objects that match context flags.
- *
- * @param	InOuter		Package to load new package into (usually NULL or ULevel->GetOuter())
- * @param	Filename	Long package name to load.
- * @param	LoadFlags	Flags controlling loading behavior
- * @return	Loaded package if successful, NULL otherwise
- */
-UPackage* LoadPackage( UPackage* InOuter, const TCHAR* InLongPackageName, uint32 LoadFlags )
+* Loads a package and all contained objects that match context flags.
+*
+* @param	InOuter		Package to load new package into (usually NULL or ULevel->GetOuter())
+* @param	Filename	Long package name to load.
+* @param	LoadFlags	Flags controlling loading behavior
+* @param	ImportLinker	Linker that requests this package through one of its imports
+* @return	Loaded package if successful, NULL otherwise
+*/
+UPackage* LoadPackageInternal(UPackage* InOuter, const TCHAR* InLongPackageName, uint32 LoadFlags, ULinkerLoad* ImportLinker)
 {
 	UPackage* Result = NULL;
 
@@ -861,7 +862,17 @@ UPackage* LoadPackage( UPackage* InOuter, const TCHAR* InLongPackageName, uint32
 
 		if( !(LoadFlags & LOAD_Verify) )
 		{
+			// Make sure we pass the property that's currently being serialized by the linker that owns the import 
+			// that triggered this LoadPackage call
+			auto OldSerialziedProperty = Linker->GetSerializedProperty();
+			if (ImportLinker)
+			{
+				Linker->SetSerializedProperty(ImportLinker->GetSerializedProperty());
+			}
+
 			Linker->LoadAllObjects();
+
+			Linker->SetSerializedProperty(OldSerialziedProperty);
 		}
 
 		SlowTask.EnterProgressFrame(30);
@@ -939,6 +950,11 @@ UPackage* LoadPackage( UPackage* InOuter, const TCHAR* InLongPackageName, uint32
 	return Result;
 }
 
+UPackage* LoadPackage(UPackage* InOuter, const TCHAR* InLongPackageName, uint32 LoadFlags)
+{
+	return LoadPackageInternal(InOuter, InLongPackageName, LoadFlags, /*ImportLinker =*/ nullptr);
+}
+
 /**
  * Returns whether we are currently loading a package (sync or async)
  *
@@ -1612,6 +1628,10 @@ bool StaticAllocateObjectErrorTests( UClass* InClass, UObject* InOuter, FName In
 	return false;
 }
 
+/**
+* Call back into the async loading code to inform of the creation of a new object
+*/
+void NotifyConstructedDuringAsyncLoading(UObject* Object, bool bSubObject);
 
 /**
 * For object overwrites, the class may want to persist some info over the re-intialize
@@ -1788,9 +1808,9 @@ UObject* StaticAllocateObject
 		}
 	}
 
-	if( !bSubObject && GIsAsyncLoading )
+	if (GIsAsyncLoading)
 	{
-		NotifyConstructedDuringAsyncLoading(Obj);
+		NotifyConstructedDuringAsyncLoading(Obj, bSubObject);
 	}
 
 	// Let the caller know if a subobject has just been recycled.
@@ -2296,7 +2316,6 @@ public:
 		GObjectCountDuringLastMarkPhase = 0;
 		ReferenceSearchFlags = SearchFlags;
 		FoundReferencesList = FoundReferences;
-		GSerializedProperty = NULL;
 
 		checkSlow(!DebugSerialize.Num()); // should only be filled int he scope of this function
 		// Iterate over all objects.
diff --git a/Engine/Source/Runtime/CoreUObject/Public/Serialization/ArchiveUObject.h b/Engine/Source/Runtime/CoreUObject/Public/Serialization/ArchiveUObject.h
--- a/Engine/Source/Runtime/CoreUObject/Public/Serialization/ArchiveUObject.h
+++ b/Engine/Source/Runtime/CoreUObject/Public/Serialization/ArchiveUObject.h
@@ -584,7 +584,7 @@ struct FTraceRouteRecord
 	struct FObjectGraphNode*	GraphNode;
 	TArray<UProperty*>			ReferencerProperties;
 
-	FTraceRouteRecord( struct FObjectGraphNode* InGraphNode, UProperty* InReferencerProperty=GSerializedProperty )
+	FTraceRouteRecord( struct FObjectGraphNode* InGraphNode, UProperty* InReferencerProperty)
 	: GraphNode(InGraphNode)
 	{
 		ReferencerProperties.Add(InReferencerProperty);
@@ -596,7 +596,7 @@ struct FTraceRouteRecord
 		ReferencerProperties = InReferencerProperties;
 	}
 
-	void Add(UProperty* InReferencerProperty=GSerializedProperty)
+	void Add(UProperty* InReferencerProperty)
 	{
 		ReferencerProperties.Add(InReferencerProperty);
 	}
@@ -1260,7 +1260,7 @@ private:
 			if ((LimitClass == NULL || Object->IsA(LimitClass)) &&
 				(LimitOuter == NULL || (Object->GetOuter() == LimitOuter || (!bRequireDirectOuter && Object->IsIn(LimitOuter)))) )
 			{
-				ObjectGraph->Add(GSerializedProperty, Object);
+				ObjectGraph->Add(GetSerializedProperty(), Object);
 				if ( bSerializeRecursively && !ObjectArray.Contains(Object) )
 				{
 					ObjectArray.Add( Object );
diff --git a/Engine/Source/Runtime/CoreUObject/Public/UObject/GarbageCollection.h b/Engine/Source/Runtime/CoreUObject/Public/UObject/GarbageCollection.h
--- a/Engine/Source/Runtime/CoreUObject/Public/UObject/GarbageCollection.h
+++ b/Engine/Source/Runtime/CoreUObject/Public/UObject/GarbageCollection.h
@@ -6,11 +6,7 @@
 =============================================================================*/
 
 /** Context sensitive keep flags for garbage collection */
-#define GARBAGE_COLLECTION_KEEPFLAGS	(GIsEditor ? RF_Native | RF_Standalone : RF_Native)
-
-/** Points to the UProperty currently being serialized */
-extern COREUOBJECT_API UProperty* GSerializedProperty;
-
+#define GARBAGE_COLLECTION_KEEPFLAGS	(GIsEditor ? RF_Native|RF_AsyncLoading|RF_Standalone : RF_Native|RF_AsyncLoading)
 
 /*-----------------------------------------------------------------------------
 	Realtime garbage collection helper classes.
@@ -458,7 +454,10 @@ protected:
 	{
 		if( Object )
 		{
-			Collector.AddReferencedObject( Object, SerializingObject, (UObject*)GSerializedProperty );
+			auto OldCollectorSerializedProperty = Collector.GetSerializedProperty();
+			Collector.SetSerializedProperty(GetSerializedProperty());
+			Collector.AddReferencedObject(Object, SerializingObject, (UObject*)GetSerializedProperty());
+			Collector.SetSerializedProperty(OldCollectorSerializedProperty);
 		}
 		return *this;
 	}
diff --git a/Engine/Source/Runtime/CoreUObject/Public/UObject/Linker.h b/Engine/Source/Runtime/CoreUObject/Public/UObject/Linker.h
--- a/Engine/Source/Runtime/CoreUObject/Public/UObject/Linker.h
+++ b/Engine/Source/Runtime/CoreUObject/Public/UObject/Linker.h
@@ -1137,7 +1137,7 @@ class ULinker : public UObject, public FLinkerTables
 	void GetScriptSHAKey(uint8* OutKey);
 
 	// UObject interface.
-	void BeginDestroy();
+	virtual void BeginDestroy() override;
 
 	/**
 	 * Test and object against the load flag filters
@@ -1641,7 +1641,7 @@ private:
 	void DetachExport( int32 i );
 
 	// UObject interface.
-	void BeginDestroy();
+	virtual void BeginDestroy() override;
 
 	// FArchive interface.
 	/**
diff --git a/Engine/Source/Runtime/CoreUObject/Public/UObject/PropertyTag.h b/Engine/Source/Runtime/CoreUObject/Public/UObject/PropertyTag.h
--- a/Engine/Source/Runtime/CoreUObject/Public/UObject/PropertyTag.h
+++ b/Engine/Source/Runtime/CoreUObject/Public/UObject/PropertyTag.h
@@ -121,12 +121,12 @@ struct FPropertyTag
 		}
 		else
 		{
-			UProperty* OldSerializedProperty = GSerializedProperty;
-			GSerializedProperty = Property;
+			UProperty* OldSerializedProperty = Ar.GetSerializedProperty();
+			Ar.SetSerializedProperty(Property);
 
 			Property->SerializeItem( Ar, Value, MaxReadBytes, Defaults );
 
-			GSerializedProperty = OldSerializedProperty;
+			Ar.SetSerializedProperty(OldSerializedProperty);
 		}
 	}
 };
diff --git a/Engine/Source/Runtime/CoreUObject/Public/UObject/UObjectGlobals.h b/Engine/Source/Runtime/CoreUObject/Public/UObject/UObjectGlobals.h
--- a/Engine/Source/Runtime/CoreUObject/Public/UObject/UObjectGlobals.h
+++ b/Engine/Source/Runtime/CoreUObject/Public/UObject/UObjectGlobals.h
@@ -422,13 +422,6 @@ COREUOBJECT_API UPackage* CreatePackage( UObject* InOuter, const TCHAR* PackageN
 
 void GlobalSetProperty( const TCHAR* Value, UClass* Class, UProperty* Property, bool bNotifyObjectOfChange );
 
-/**
- * Call back into the async loading code to inform of the creation of a new object
- * @param Object		Object created
- */
-void NotifyConstructedDuringAsyncLoading(UObject* Object);
-
-
 /**
  * Save a copy of this object into the transaction buffer if we are currently recording into
  * one (undo/redo). If bMarkDirty is true, will also mark the package as needing to be saved.
@@ -1454,10 +1447,17 @@ public:
 	 */
 	virtual bool IsIgnoringTransient() const = 0;
 	/**
-	 * Allows reference limination by this collector.
+	 * Allows reference elimination by this collector.
 	 */
 	virtual void AllowEliminatingReferences(bool bAllow) {}
-
+	/**
+	* Sets the property that is currently being serialized
+	*/
+	virtual void SetSerializedProperty(class UProperty* Inproperty) {}
+	/**
+	* Gets the property that is currently being serialized
+	*/
+	virtual class UProperty* GetSerializedProperty() const { return nullptr; }
 protected:
 	/**
 	 * Handle object reference. Called by AddReferencedObject.
@@ -1495,14 +1495,23 @@ public:
 	 * Finds all objects referenced by Object.
 	 *
 	 * @param Object Object which references are to be found.
+	 * @param ReferencingObject object that's referencing the current object.
+	 * @param ReferencingProperty property the current object is being referenced through.
 	 */
-	virtual void FindReferences(UObject* Object);	
+	virtual void FindReferences(UObject* Object, UObject* ReferencingObject = nullptr, UObject* ReferencingProperty = nullptr);	
 
 	// FReferenceCollector interface.
-	virtual void HandleObjectReference(UObject*& Object, const UObject* ReferencingObject, const UObject* ReferencingProperty) override;
+	virtual void HandleObjectReference(UObject*& Object, const UObject* ReferencingObject, const UObject* InReferencingProperty) override;
 	virtual bool IsIgnoringArchetypeRef() const override { return bShouldIgnoreArchetype; }
 	virtual bool IsIgnoringTransient() const override { return bShouldIgnoreTransient; }
-
+	virtual void SetSerializedProperty(class UProperty* Inproperty) override
+	{
+		SerializedProperty = Inproperty;
+	}
+	virtual class UProperty* GetSerializedProperty() const override
+	{
+		return SerializedProperty;
+	}
 protected:
 
 	/** Stored reference to array of objects we add object references to. */
@@ -1511,6 +1520,8 @@ protected:
 	TSet<const UObject*>	SerializedObjects;
 	/** Only objects within this outer will be considered, NULL value indicates that outers are disregarded. */
 	UObject*		LimitOuter;
+	/** Property that is referencing the current object */
+	class UProperty* SerializedProperty;
 	/** Determines whether nested objects contained within LimitOuter are considered. */
 	bool			bRequireDirectOuter;
 	/** Determines whether archetype references are considered. */
diff --git a/Engine/Source/Runtime/CoreUObject/Public/UObject/UnrealType.h b/Engine/Source/Runtime/CoreUObject/Public/UObject/UnrealType.h
--- a/Engine/Source/Runtime/CoreUObject/Public/UObject/UnrealType.h
+++ b/Engine/Source/Runtime/CoreUObject/Public/UObject/UnrealType.h
@@ -279,13 +279,13 @@ public:
 	{
 		if( ShouldSerializeValue(Ar) )
 		{
-			UProperty* OldSerializedProperty = GSerializedProperty;
+			UProperty* OldSerializedProperty = Ar.GetSerializedProperty();
 			for (int32 Idx = 0; Idx < ArrayDim; Idx++)
 			{
-				GSerializedProperty = this;
+				Ar.SetSerializedProperty(this);
 				SerializeItem( Ar, ContainerPtrToValuePtr<void>(Data, Idx), 0 );
 			}
-			GSerializedProperty = OldSerializedProperty;
+			Ar.SetSerializedProperty(OldSerializedProperty);
 		}
 	}
 	/**
@@ -306,10 +306,10 @@ public:
 				void const* Default = ContainerPtrToValuePtrForDefaults<void>(DefaultStruct, DefaultData, Idx);
 				if ( !Identical(Target, Default, Ar.GetPortFlags()) )
 				{
-					UProperty* OldSerializedProperty = GSerializedProperty;
-					GSerializedProperty = this;
+					UProperty* OldSerializedProperty = Ar.GetSerializedProperty();
+					Ar.SetSerializedProperty(this);
 					SerializeItem( Ar, Target, 0, Default );
-					GSerializedProperty = OldSerializedProperty;
+					Ar.SetSerializedProperty(OldSerializedProperty);
 				}
 			}
 		}
diff --git a/Engine/Source/Runtime/Engine/Private/Particles/ParticleModules.cpp b/Engine/Source/Runtime/Engine/Private/Particles/ParticleModules.cpp
--- a/Engine/Source/Runtime/Engine/Private/Particles/ParticleModules.cpp
+++ b/Engine/Source/Runtime/Engine/Private/Particles/ParticleModules.cpp
@@ -885,7 +885,6 @@ public:
 		ArIgnoreClassRef					= true;
 		ArIsModifyingWeakAndStrongReferences = true;
 
-		GSerializedProperty = NULL;
 		Src->Serialize( *this );
 	}
 	virtual FString GetArchiveName() const { return TEXT("FArchiveFindDistributionRefs"); }
diff --git a/Engine/Source/Runtime/GameplayAbilities/Private/Abilities/GameplayAbilityWorldReticle_ActorVisualization.cpp b/Engine/Source/Runtime/GameplayAbilities/Private/Abilities/GameplayAbilityWorldReticle_ActorVisualization.cpp
--- a/Engine/Source/Runtime/GameplayAbilities/Private/Abilities/GameplayAbilityWorldReticle_ActorVisualization.cpp
+++ b/Engine/Source/Runtime/GameplayAbilities/Private/Abilities/GameplayAbilityWorldReticle_ActorVisualization.cpp
@@ -25,7 +25,7 @@ AGameplayAbilityWorldReticle_ActorVisualization::AGameplayAbilityWorldReticle_Ac
 	CollisionComponent->BodyInstance.bEnableCollision_DEPRECATED = false;
 	CollisionComponent->SetCollisionEnabled(ECollisionEnabled::NoCollision);
 
-	//TSubobjectPtr<USceneComponent> SceneComponent = ObjectInitializer.CreateDefaultSubobject<USceneComponent>(this, TEXT("RootComponent0"));
+	//USceneComponent* SceneComponent = ObjectInitializer.CreateDefaultSubobject<USceneComponent>(this, TEXT("RootComponent0"));
 	RootComponent = CollisionComponent;
 }
 
