Marking TSubobjectPtr<> as deprecated

diff --git a/Engine/Source/Runtime/CoreUObject/Public/UObject/UObjectGlobals.h b/Engine/Source/Runtime/CoreUObject/Public/UObject/UObjectGlobals.h
--- a/Engine/Source/Runtime/CoreUObject/Public/UObject/UObjectGlobals.h
+++ b/Engine/Source/Runtime/CoreUObject/Public/UObject/UObjectGlobals.h
@@ -532,28 +532,28 @@ template<> struct TIsWeakPointerType<FSubobjectPtr> { enum { Value = false }; };
  * TSubobjectPtr - Sub-object smart pointer, soon to be deprecated and should no longer be used.
  */
 template <class SubobjectType>
-class TSubobjectPtr : public FSubobjectPtr
+class TSubobjectPtrDeprecated : public FSubobjectPtr
 {
 public:
 	/** Internal constructors. */
-	TSubobjectPtr(SubobjectType* InObject)
+	TSubobjectPtrDeprecated(SubobjectType* InObject)
 		: FSubobjectPtr(InObject)
 	{}
-	TSubobjectPtr& operator=(const TSubobjectPtr& Other) 
+	TSubobjectPtrDeprecated& operator=(const TSubobjectPtrDeprecated& Other)
 	{ 
 		Set(Other.Object);
 		return *this; 
 	}
 
 	/** Default constructor. */
-	TSubobjectPtr()
+	TSubobjectPtrDeprecated()
 		: FSubobjectPtr((UObject*)FSubobjectPtr::InvalidPtrValue)
 	{
-		static_assert(sizeof(TSubobjectPtr) == sizeof(UObject*), "TSuobjectPtr should equal pointer size.");
+		static_assert(sizeof(TSubobjectPtrDeprecated) == sizeof(UObject*), "TSuobjectPtr should equal pointer size.");
 	}
 	/** Copy constructor */
 	template <class DerivedSubobjectType>
-	TSubobjectPtr(TSubobjectPtr<DerivedSubobjectType>& Other)
+	TSubobjectPtrDeprecated(TSubobjectPtrDeprecated<DerivedSubobjectType>& Other)
 		: FSubobjectPtr(Other.Object)
 	{
 		static_assert((CanConvertPointerFromTo<DerivedSubobjectType, SubobjectType>::Result), "Subobject pointers must be compatible.");
@@ -575,9 +575,14 @@ public:
 	}
 };
 
-template<class T> struct TIsPODType< TSubobjectPtr<T> > { enum { Value = true }; };
-template<class T> struct TIsZeroConstructType< TSubobjectPtr<T> > { enum { Value = true }; };
-template<class T> struct TIsWeakPointerType< TSubobjectPtr<T> > { enum { Value = false }; };
+template<class T> struct TIsPODType< TSubobjectPtrDeprecated<T> > { enum { Value = true }; };
+template<class T> struct TIsZeroConstructType< TSubobjectPtrDeprecated<T> > { enum { Value = true }; };
+template<class T> struct TIsWeakPointerType< TSubobjectPtrDeprecated<T> > { enum { Value = false }; };
+
+#define TSubobjectPtr \
+	EMIT_DEPRECATED_WARNING_MESSAGE("TSubobjectPtr is deprecated and and should no longer be used. Please use pointers instead.") \
+	TSubobjectPtrDeprecated
+	
 
 /**
  * Internal class to finalize UObject creation (initialize properties) after the real C++ constructor is called.
