Fixed a typo in deprecation warning message. Added a comment to private_subobject macro.

diff --git a/Engine/Source/Runtime/Core/Public/HAL/Platform.h b/Engine/Source/Runtime/Core/Public/HAL/Platform.h
--- a/Engine/Source/Runtime/Core/Public/HAL/Platform.h
+++ b/Engine/Source/Runtime/Core/Public/HAL/Platform.h
@@ -432,6 +432,7 @@
 	#define DEPRECATED_FORGAME(...)
 #endif
 
+// This is a temporary macro, will be removed when TSubobjectPtr can be safely removed
 #ifndef private_subobject
 #define private_subobject public
 #endif
diff --git a/Engine/Source/Runtime/CoreUObject/Public/UObject/UObjectGlobals.h b/Engine/Source/Runtime/CoreUObject/Public/UObject/UObjectGlobals.h
--- a/Engine/Source/Runtime/CoreUObject/Public/UObject/UObjectGlobals.h
+++ b/Engine/Source/Runtime/CoreUObject/Public/UObject/UObjectGlobals.h
@@ -580,7 +580,7 @@ template<class T> struct TIsZeroConstructType< TSubobjectPtrDeprecated<T> > { en
 template<class T> struct TIsWeakPointerType< TSubobjectPtrDeprecated<T> > { enum { Value = false }; };
 
 #define TSubobjectPtr \
-	EMIT_DEPRECATED_WARNING_MESSAGE("TSubobjectPtr is deprecated and and should no longer be used. Please use pointers instead.") \
+	EMIT_DEPRECATED_WARNING_MESSAGE("TSubobjectPtr is deprecated and should no longer be used. Please use pointers instead.") \
 	TSubobjectPtrDeprecated
 	
 
