Fix for a crash when performing hot-reload in a project that has deprecated classes.

UE-4701

diff --git a/Engine/Source/Runtime/CoreUObject/Private/UObject/UObjectBase.cpp b/Engine/Source/Runtime/CoreUObject/Private/UObject/UObjectBase.cpp
--- a/Engine/Source/Runtime/CoreUObject/Private/UObject/UObjectBase.cpp
+++ b/Engine/Source/Runtime/CoreUObject/Private/UObject/UObjectBase.cpp
@@ -606,6 +606,19 @@ static TArray<FFieldCompiledInInfo*>& GetHotReloadClasses()
 }
 #endif
 
+/** Removes prefix from the native class name */
+FString RemoveClassPrefix(const TCHAR* ClassName)
+{
+	const TCHAR* DeprecatedPrefix = TEXT("DEPRECATED_");
+	FString NameWithoutPrefix(ClassName);
+	NameWithoutPrefix = NameWithoutPrefix.Mid(1);
+	if (NameWithoutPrefix.StartsWith(DeprecatedPrefix))
+	{
+		NameWithoutPrefix = NameWithoutPrefix.Mid(FCString::Strlen(DeprecatedPrefix));
+	}
+	return NameWithoutPrefix;
+}
+
 void UClassCompiledInDefer(FFieldCompiledInInfo* ClassInfo, const TCHAR* Name, SIZE_T ClassSize, uint32 Crc)
 {
 	const FName CPPClassName(Name);
@@ -619,8 +632,7 @@ void UClassCompiledInDefer(FFieldCompiledInInfo* ClassInfo, const TCHAR* Name, S
 		check(GIsHotReload);
 
 		// Get the native name
-		FString NameWithoutPrefix(Name);
-		NameWithoutPrefix = NameWithoutPrefix.Mid(1);
+		FString NameWithoutPrefix(RemoveClassPrefix(Name));
 		auto ExistingClass = FindObjectChecked<UClass>(ANY_PACKAGE, *NameWithoutPrefix);
 
 		if (ClassInfo->bHasChanged)
