Changed FCoreUObjectDelegates::OnObjectModified to be called only once per frame for each unique object that has Modify() called on it

diff --git a/Engine/Source/Editor/UnrealEd/Private/Editor.cpp b/Engine/Source/Editor/UnrealEd/Private/Editor.cpp
--- a/Engine/Source/Editor/UnrealEd/Private/Editor.cpp
+++ b/Engine/Source/Editor/UnrealEd/Private/Editor.cpp
@@ -810,6 +810,9 @@ void UEditorEngine::Tick( float DeltaSeconds, bool bIdleMode )
 	check( CurrentGWorld );
 	check( CurrentGWorld != PlayWorld || bIsSimulatingInEditor );
 
+	// Clear out the list of objects modified this frame, used for OnObjectModified notification.
+	FCoreUObjectDelegates::ObjectsModifiedThisFrame.Empty();
+
 	// Always ensure we've got adequate slack for any worlds that are going to get created in this frame so that
 	// our EditorContext reference doesn't get invalidated
 	WorldList.Reserve(WorldList.Num() + 10);
diff --git a/Engine/Source/Runtime/CoreUObject/Private/UObject/Obj.cpp b/Engine/Source/Runtime/CoreUObject/Private/UObject/Obj.cpp
--- a/Engine/Source/Runtime/CoreUObject/Private/UObject/Obj.cpp
+++ b/Engine/Source/Runtime/CoreUObject/Private/UObject/Obj.cpp
@@ -719,7 +719,7 @@ bool UObject::Modify( bool bAlwaysMarkDirty/*=true*/ )
 		}
 	}
 #if WITH_EDITOR
-	FCoreUObjectDelegates::OnObjectModified.Broadcast(this);
+	FCoreUObjectDelegates::BroadcastOnObjectModified(this);
 #endif
 	return bSavedToTransactionBuffer;
 }
diff --git a/Engine/Source/Runtime/CoreUObject/Private/UObject/UObjectGlobals.cpp b/Engine/Source/Runtime/CoreUObject/Private/UObject/UObjectGlobals.cpp
--- a/Engine/Source/Runtime/CoreUObject/Private/UObject/UObjectGlobals.cpp
+++ b/Engine/Source/Runtime/CoreUObject/Private/UObject/UObjectGlobals.cpp
@@ -60,6 +60,8 @@ FCoreUObjectDelegates::FOnPreObjectPropertyChanged FCoreUObjectDelegates::OnPreO
 FCoreUObjectDelegates::FOnObjectPropertyChanged FCoreUObjectDelegates::OnObjectPropertyChanged;
 
 #if WITH_EDITOR
+// Set of objects modified this frame
+TSet<UObject*> FCoreUObjectDelegates::ObjectsModifiedThisFrame;
 FCoreUObjectDelegates::FOnObjectModified FCoreUObjectDelegates::OnObjectModified;
 FCoreUObjectDelegates::FOnAssetLoaded FCoreUObjectDelegates::OnAssetLoaded;
 FCoreUObjectDelegates::FOnObjectSaved FCoreUObjectDelegates::OnObjectSaved;
diff --git a/Engine/Source/Runtime/CoreUObject/Public/UObject/UObjectGlobals.h b/Engine/Source/Runtime/CoreUObject/Public/UObject/UObjectGlobals.h
--- a/Engine/Source/Runtime/CoreUObject/Public/UObject/UObjectGlobals.h
+++ b/Engine/Source/Runtime/CoreUObject/Public/UObject/UObjectGlobals.h
@@ -1569,6 +1569,19 @@ struct COREUOBJECT_API FCoreUObjectDelegates
 	// Called when any object is modified at all 
 	static FOnObjectModified OnObjectModified;
 
+	// Set of objects modified this frame, to prevent multiple triggerings of the OnObjectModified delegate.
+	static TSet<UObject*> ObjectsModifiedThisFrame;
+
+	// Broadcast OnObjectModified if the broadcast hasn't ocurred for this object in this frame
+	static void BroadcastOnObjectModified(UObject* Object)
+	{
+		if (OnObjectModified.IsBound() && !ObjectsModifiedThisFrame.Contains(Object))
+		{
+			ObjectsModifiedThisFrame.Add(Object);
+			OnObjectModified.Broadcast(Object);
+		}
+	}
+
 	// Callback for when an asset is loaded (Editor)
 	DECLARE_MULTICAST_DELEGATE_OneParam(FOnAssetLoaded, UObject*);
 
