Prevent UObject::Modify() from doing anything during GC.  Fixes error message when loading a map which breaks automation testing.

diff --git a/Engine/Source/Runtime/CoreUObject/Private/UObject/Obj.cpp b/Engine/Source/Runtime/CoreUObject/Private/UObject/Obj.cpp
--- a/Engine/Source/Runtime/CoreUObject/Private/UObject/Obj.cpp
+++ b/Engine/Source/Runtime/CoreUObject/Private/UObject/Obj.cpp
@@ -702,25 +702,30 @@ void UObject::PreSave()
 bool UObject::Modify( bool bAlwaysMarkDirty/*=true*/ )
 {
 	bool bSavedToTransactionBuffer = false;
-	// Do not consider PIE world objects or script packages, as they should never end up in the
-	// transaction buffer and we don't want to mark them dirty here either.
-	if ( (GetOutermost()->PackageFlags & (PKG_PlayInEditor|PKG_ContainsScript|PKG_CompiledIn)) == 0 )
-	{
-		// Attempt to mark the package dirty and save a copy of the object to the transaction
-		// buffer. The save will fail if there isn't a valid transactor, the object isn't
-		// transactional, etc.
-		bSavedToTransactionBuffer = SaveToTransactionBuffer(this, true );
-		
-		// If we failed to save to the transaction buffer, but the user requested the package
-		// marked dirty anyway, do so
-		if ( !bSavedToTransactionBuffer && bAlwaysMarkDirty )
+
+	if (!GIsGarbageCollecting)
+	{
+		// Do not consider PIE world objects or script packages, as they should never end up in the
+		// transaction buffer and we don't want to mark them dirty here either.
+		if ((GetOutermost()->PackageFlags & (PKG_PlayInEditor | PKG_ContainsScript | PKG_CompiledIn)) == 0)
 		{
-			MarkPackageDirty();
+			// Attempt to mark the package dirty and save a copy of the object to the transaction
+			// buffer. The save will fail if there isn't a valid transactor, the object isn't
+			// transactional, etc.
+			bSavedToTransactionBuffer = SaveToTransactionBuffer(this, true);
+
+			// If we failed to save to the transaction buffer, but the user requested the package
+			// marked dirty anyway, do so
+			if (!bSavedToTransactionBuffer && bAlwaysMarkDirty)
+			{
+				MarkPackageDirty();
+			}
 		}
-	}
 #if WITH_EDITOR
-	FCoreUObjectDelegates::BroadcastOnObjectModified(this);
+		FCoreUObjectDelegates::BroadcastOnObjectModified(this);
 #endif
+	}
+
 	return bSavedToTransactionBuffer;
 }
 
