Core: Reverted CL# 2342799 for mock GUID properties, because it broke blueprints (does not and will not support uint32)

diff --git a/Engine/Source/Runtime/CoreUObject/Classes/Object.h b/Engine/Source/Runtime/CoreUObject/Classes/Object.h
--- a/Engine/Source/Runtime/CoreUObject/Classes/Object.h
+++ b/Engine/Source/Runtime/CoreUObject/Classes/Object.h
@@ -190,16 +190,16 @@ USTRUCT(immutable, noexport, BlueprintType)
 struct FGuid
 {
 	UPROPERTY(EditAnywhere, SaveGame, Category=Guid)
-	uint32 A;
+	int32 A;
 
 	UPROPERTY(EditAnywhere, SaveGame, Category=Guid)
-	uint32 B;
+	int32 B;
 
 	UPROPERTY(EditAnywhere, SaveGame, Category=Guid)
-	uint32 C;
+	int32 C;
 
 	UPROPERTY(EditAnywhere, SaveGame, Category=Guid)
-	uint32 D;
+	int32 D;
 };
 
 // A unique identifier for networking objects
