Changing package names to FNames instead of FStrings to minimze the impact of string comparisons in async loading. Replaced NULLs with nullptrs in AsyncLoading code.

diff --git a/Engine/Source/Editor/UnrealEd/Private/Tests/EditorAutomationTests.cpp b/Engine/Source/Editor/UnrealEd/Private/Tests/EditorAutomationTests.cpp
--- a/Engine/Source/Editor/UnrealEd/Private/Tests/EditorAutomationTests.cpp
+++ b/Engine/Source/Editor/UnrealEd/Private/Tests/EditorAutomationTests.cpp
@@ -1961,7 +1961,7 @@ namespace StaticMeshUVTest
 		/** 
 		 * Callback when a package has finished loading
 		 */
-		void PackageLoadCallback(const FString& PackageName, UPackage* LoadedPackage)
+		void PackageLoadCallback(const FName& PackageName, UPackage* LoadedPackage)
 		{
 			CurrentPackage = LoadedPackage;
 			CurrentState = EStaticMeshUVTestState::PackageLoaded;
diff --git a/Engine/Source/Runtime/CoreUObject/Private/Serialization/AsyncLoading.cpp b/Engine/Source/Runtime/CoreUObject/Private/Serialization/AsyncLoading.cpp
--- a/Engine/Source/Runtime/CoreUObject/Private/Serialization/AsyncLoading.cpp
+++ b/Engine/Source/Runtime/CoreUObject/Private/Serialization/AsyncLoading.cpp
@@ -96,6 +96,56 @@ static TIndirectArray<struct FAsyncPackage>	GObjAsyncPackages;
 int32 FAsyncPackage::PreLoadIndex = 0;
 int32 FAsyncPackage::PostLoadIndex = 0;
 
+/**
+* Constructor
+*/
+FAsyncPackage::FAsyncPackage(const FName& InPackageName, const FGuid* InPackageGuid, FName InPackageType, const FName& InPackageNameToLoad)
+: PackageName(InPackageName)
+, PackageNameToLoad(InPackageNameToLoad)
+, PackageGuid(InPackageGuid != NULL ? *InPackageGuid : FGuid(0, 0, 0, 0))
+, PackageType(InPackageType)
+, Linker(NULL)
+, DependencyRefCount(0)
+, LoadImportIndex(0)
+, ImportIndex(0)
+, ExportIndex(0)
+, TimeLimit(FLT_MAX)
+, bUseTimeLimit(false)
+, bUseFullTimeLimit(false)
+, bTimeLimitExceeded(false)
+, bLoadHasFailed(false)
+, bLoadHasFinished(false)
+, TickStartTime(0)
+, LastObjectWorkWasPerformedOn(NULL)
+, LastTypeOfWorkPerformed(NULL)
+, LoadStartTime(0.0)
+, LoadPercentage(0)
+, PackageFlags(0)
+#if WITH_EDITOR
+, PIEInstanceID(INDEX_NONE)
+#endif	
+#if PERF_TRACK_DETAILED_ASYNC_STATS
+, TickCount(0)
+, TickLoopCount(0)
+, CreateLinkerCount(0)
+, FinishLinkerCount(0)
+, CreateImportsCount(0)
+, CreateExportsCount(0)
+, PreLoadObjectsCount(0)
+, PostLoadObjectsCount(0)
+, FinishObjectsCount(0)
+, TickTime(0.0)
+, CreateLinkerTime(0.0)
+, FinishLinkerTime(0.0)
+, CreateImportsTime(0.0)
+, CreateExportsTime(0.0)
+, PreLoadObjectsTime(0.0)
+, PostLoadObjectsTime(0.0)
+, FinishObjectsTime(0.0)
+#endif // PERF_TRACK_DETAILED_ASYNC_STATS
+{
+}
+
 /**
  * Returns an estimated load completion percentage.
  */
@@ -121,7 +171,7 @@ void FAsyncPackage::ResetLoader()
 	if( Linker )
 	{
 		Linker->Detach(false);
-		Linker = NULL;
+		Linker = nullptr;
 	}
 }
 
@@ -214,8 +264,8 @@ EAsyncPackageState::Type FAsyncPackage::Tick( bool InbUseTimeLimit, bool InbUseF
 {
 	SCOPE_CYCLE_COUNTER(STAT_FAsyncPackage_Tick);
 
- 	check(LastObjectWorkWasPerformedOn==NULL);
-	check(LastTypeOfWorkPerformed==NULL);
+ 	check(LastObjectWorkWasPerformedOn==nullptr);
+	check(LastTypeOfWorkPerformed==nullptr);
 
 	// Set up tick relevant variables.
 	bUseTimeLimit			= InbUseTimeLimit;
@@ -308,9 +358,9 @@ EAsyncPackageState::Type FAsyncPackage::Tick( bool InbUseTimeLimit, bool InbUseF
 	check( bUseTimeLimit || LoadingState != EAsyncPackageState::TimeOut );
 
 	// We can't have a reference to a UObject.
-	LastObjectWorkWasPerformedOn	= NULL;
+	LastObjectWorkWasPerformedOn	= nullptr;
 	// Reset type of work performed.
-	LastTypeOfWorkPerformed			= NULL;
+	LastTypeOfWorkPerformed			= nullptr;
 	// Mark this package as loaded if everything completed.
 	bLoadHasFinished = (LoadingState == EAsyncPackageState::Complete);
 	// Subtract the time it took to load this package from the global limit.
@@ -327,15 +377,15 @@ EAsyncPackageState::Type FAsyncPackage::Tick( bool InbUseTimeLimit, bool InbUseF
  */
 EAsyncPackageState::Type FAsyncPackage::CreateLinker()
 {
-	if( Linker == NULL )
+	if (Linker == nullptr)
 	{
 		SCOPE_CYCLE_COUNTER(STAT_FAsyncPackage_CreateLinker);
 
-		LastObjectWorkWasPerformedOn	= NULL;
+		LastObjectWorkWasPerformedOn	= nullptr;
 		LastTypeOfWorkPerformed			= TEXT("creating Linker");
 
 		// Try to find existing package or create it if not already present.
-		UPackage* Package = CreatePackage( NULL, *PackageName );
+		UPackage* Package = CreatePackage(nullptr, *PackageName.ToString());
 		FScopeCycleCounterUObject ConstructorScope(Package, GET_STATID(STAT_FAsyncPackage_CreateLinker));
 
 		// Set package specific data 
@@ -345,7 +395,7 @@ EAsyncPackageState::Type FAsyncPackage::CreateLinker()
 #endif
 
 		// Always store package filename we loading from
-		Package->FileName = FName(*PackageNameToLoad);
+		Package->FileName = PackageNameToLoad;
 
 		// if the linker already exists, we don't need to lookup the file (it may have been pre-created with
 		// a different filename)
@@ -354,10 +404,9 @@ EAsyncPackageState::Type FAsyncPackage::CreateLinker()
 		if (!Linker)
 		{
 			FString PackageFileName;
-
-			if (!FPackageName::DoesPackageExist(PackageNameToLoad, PackageGuid.IsValid() ? &PackageGuid : NULL, &PackageFileName))
+			if (!FPackageName::DoesPackageExist(PackageNameToLoad.ToString(), PackageGuid.IsValid() ? &PackageGuid : nullptr, &PackageFileName))
 			{
-				UE_LOG(LogStreaming, Error, TEXT("Couldn't find file for package %s requested by async loading code."),*PackageName);
+				UE_LOG(LogStreaming, Error, TEXT("Couldn't find file for package %s requested by async loading code."), *PackageName.ToString());
 				bLoadHasFailed = true;
 				return EAsyncPackageState::TimeOut;
 			}
@@ -366,7 +415,7 @@ EAsyncPackageState::Type FAsyncPackage::CreateLinker()
 			Linker = ULinkerLoad::CreateLinkerAsync( Package, *PackageFileName, (FApp::IsGame() && !GIsEditor) ? (LOAD_SeekFree | LOAD_NoVerify) : LOAD_None  );
 		}
 
-		UE_LOG(LogStreaming, Verbose, TEXT("FAsyncPackage::CreateLinker for %s finished."), *PackageNameToLoad);
+		UE_LOG(LogStreaming, Verbose, TEXT("FAsyncPackage::CreateLinker for %s finished."), *PackageNameToLoad.ToString());
 	}
 	return EAsyncPackageState::Complete;
 }
@@ -403,7 +452,7 @@ EAsyncPackageState::Type FAsyncPackage::FinishLinker()
  * @param PackageName long package name.
  * @return Index into the array if the package was found, otherwise INDEX_NONE
  */
-FORCEINLINE int32 ContainsDependencyPackage(TArray<FAsyncPackage*>& Dependencies, const FString& PackageName)
+FORCEINLINE int32 ContainsDependencyPackage(TArray<FAsyncPackage*>& Dependencies, const FName& PackageName)
 {
 	for (int32 Index = 0; Index < Dependencies.Num(); ++Index)
 	{
@@ -421,7 +470,7 @@ FORCEINLINE int32 ContainsDependencyPackage(TArray<FAsyncPackage*>& Dependencies
  * @param PackageName async package name.
  * @return Index of the async package in GObjAsyncPackages array or INDEX_NONE if not found.
  */
-FORCEINLINE int32 FindAsyncPackage(const FString& PackageName)
+FORCEINLINE int32 FindAsyncPackage(const FName& PackageName)
 {
 	for (int32 PackageIndex = 0; PackageIndex < GObjAsyncPackages.Num(); ++PackageIndex)
 	{
@@ -438,13 +487,13 @@ FORCEINLINE int32 FindAsyncPackage(const FString& PackageName)
 	*
 	* @param PendingImport Name of the package imported either directly or by one of the imported packages
 	*/
-void FAsyncPackage::AddImportDependency(int32 CurrentPackageIndex, const FString& PendingImport)
+void FAsyncPackage::AddImportDependency(int32 CurrentPackageIndex, const FName& PendingImport)
 {
-	FAsyncPackage* PackageToStream = NULL;
+	FAsyncPackage* PackageToStream = nullptr;
 	int32 ExistingAsyncPackageIndex = FindAsyncPackage(PendingImport);
 	if (ExistingAsyncPackageIndex == INDEX_NONE)
 	{
-		PackageToStream = new FAsyncPackage(PendingImport, NULL, NAME_None, PendingImport);
+		PackageToStream = new FAsyncPackage(PendingImport, nullptr, NAME_None, PendingImport);
 		GObjAsyncPackages.Insert(PackageToStream, CurrentPackageIndex);
 	}
 	else
@@ -475,10 +524,10 @@ bool FAsyncPackage::AddUniqueLinkerDependencyPackage(int32 CurrentPackageIndex,
 {
 	if (ContainsDependencyPackage(PendingImportedPackages, PendingImport.GetPackageName()) == INDEX_NONE)
 	{
-		if (PendingImport.Linker == NULL || !PendingImport.Linker->HasFinishedInitialization())
+		if (PendingImport.Linker == nullptr || !PendingImport.Linker->HasFinishedInitialization())
 		{
 			AddImportDependency(CurrentPackageIndex, PendingImport.GetPackageName());
-			UE_LOG(LogStreaming, Verbose, TEXT("  Adding linker dependency %s"), *PendingImport.GetPackageName());
+			UE_LOG(LogStreaming, Verbose, TEXT("  Adding linker dependency %s"), *PendingImport.GetPackageName().ToString());
 		}
 		else if (this != &PendingImport)
 		{
@@ -536,27 +585,29 @@ EAsyncPackageState::Type FAsyncPackage::LoadImports()
 			continue;
 		}
 
+		// Our import package name is the import name
+		const FName ImportPackageFName(Import->ObjectName);
+
 		// Handle circular dependencies - try to find existing packages.
-		FString ImportPackageName(Import->ObjectName.ToString());
-		UPackage* ExistingPackage = dynamic_cast<UPackage*>(StaticFindObjectFast(UPackage::StaticClass(), NULL, Import->ObjectName, true));
+		UPackage* ExistingPackage = dynamic_cast<UPackage*>(StaticFindObjectFast(UPackage::StaticClass(), nullptr, ImportPackageFName, true));
 		if (ExistingPackage && !(ExistingPackage->PackageFlags & PKG_CompiledIn) && !ExistingPackage->bHasBeenFullyLoaded)//!ExistingPackage->HasAnyFlags(RF_WasLoaded))
 		{
 			// The import package already exists. Check if it's currently being streamed as well. If so, make sure
 			// we add all dependencies that don't yet have linkers created otherwise we risk that if the current package
 			// doesn't depend on any other packages that have not yet started streaming, creating imports is going
 			// to load packages blocking the main thread.
-			int32 PendingAsyncPackageIndex = FindAsyncPackage(ImportPackageName);
+			int32 PendingAsyncPackageIndex = FindAsyncPackage(ImportPackageFName);
 			if (PendingAsyncPackageIndex != INDEX_NONE)
 			{
 				FAsyncPackage& PendingPackage = GObjAsyncPackages[PendingAsyncPackageIndex];
-				if (PendingPackage.Linker == NULL || !PendingPackage.Linker->HasFinishedInitialization())
+				if (PendingPackage.Linker == nullptr || !PendingPackage.Linker->HasFinishedInitialization())
 				{
 					// Add this import to the dependency list.
 					AddUniqueLinkerDependencyPackage(AsyncQueueIndex, PendingPackage);
 				}
 				else
 				{
-					UE_LOG(LogStreaming, Verbose, TEXT("FAsyncPackage::LoadImports for %s: Linker exists for %s"), *PackageNameToLoad, *ImportPackageName);
+					UE_LOG(LogStreaming, Verbose, TEXT("FAsyncPackage::LoadImports for %s: Linker exists for %s"), *PackageNameToLoad.ToString(), *ImportPackageFName.ToString());
 					// Only keep a reference to this package so that its linker doesn't go away too soon
 					PendingPackage.DependencyRefCount++;
 					ReferencedImports.Add(&PendingPackage);
@@ -566,18 +617,20 @@ EAsyncPackageState::Type FAsyncPackage::LoadImports()
 			}
 		}
 
-		if (!ExistingPackage && ContainsDependencyPackage(PendingImportedPackages, ImportPackageName) == INDEX_NONE)
+		if (!ExistingPackage && ContainsDependencyPackage(PendingImportedPackages, ImportPackageFName) == INDEX_NONE)
 		{			
+			const FString ImportPackageName(Import->ObjectName.ToString());
 			// The package doesn't exist and this import is not in the dependency list so add it now.
 			if (!FPackageName::IsShortPackageName(ImportPackageName))
 			{
-				UE_LOG(LogStreaming, Verbose, TEXT("FAsyncPackage::LoadImports for %s: Loading %s"), *PackageNameToLoad, *ImportPackageName);
-				AddImportDependency(AsyncQueueIndex, ImportPackageName);
+				UE_LOG(LogStreaming, Verbose, TEXT("FAsyncPackage::LoadImports for %s: Loading %s"), *PackageNameToLoad.ToString(), *ImportPackageName);
+
+				AddImportDependency(AsyncQueueIndex, ImportPackageFName);
 			}
 			else
 			{
 				// This usually means there's a reference to a script package from another project
-				UE_LOG(LogStreaming, Warning, TEXT("FAsyncPackage::LoadImports for %s: Short package name in imports list: %s"), *PackageNameToLoad, *ImportPackageName);
+				UE_LOG(LogStreaming, Warning, TEXT("FAsyncPackage::LoadImports for %s: Short package name in imports list: %s"), *PackageNameToLoad.ToString(), *ImportPackageName);
 			}
 		}
 	}
@@ -593,9 +646,9 @@ EAsyncPackageState::Type FAsyncPackage::LoadImports()
 /**
  * Function called when pending import package has been fully loaded.
  */
-void FAsyncPackage::ImportFullyLoadedCallback(const FString& InPackageName, UPackage* LoadedPackage)
+void FAsyncPackage::ImportFullyLoadedCallback(const FName& InPackageName, UPackage* LoadedPackage)
 {
-	UE_LOG(LogStreaming, Verbose, TEXT("FAsyncPackage::LoadImports for %s: Loaded %s"), *PackageNameToLoad, *InPackageName);
+	UE_LOG(LogStreaming, Verbose, TEXT("FAsyncPackage::LoadImports for %s: Loaded %s"), *PackageNameToLoad.ToString(), *InPackageName.ToString());
 	int32 Index = ContainsDependencyPackage(PendingImportedPackages, InPackageName);
 	check(Index != INDEX_NONE);
 	// Keep a reference to this package so that its linker doesn't go away too soon
@@ -715,7 +768,7 @@ void FAsyncPackage::FreeReferencedImports()
 	{
 		FAsyncPackage& Ref = *ReferencedImports[ReferenceIndex];
 		Ref.DependencyRefCount--;
-		UE_LOG(LogStreaming, Verbose, TEXT("FAsyncPackage::FreeReferencedImports for %s: Releasing %s (%d)"), *PackageNameToLoad, *Ref.GetPackageName(), Ref.DependencyRefCount);
+		UE_LOG(LogStreaming, Verbose, TEXT("FAsyncPackage::FreeReferencedImports for %s: Releasing %s (%d)"), *PackageNameToLoad.ToString(), *Ref.GetPackageName().ToString(), Ref.DependencyRefCount);
 		check(Ref.DependencyRefCount >= 0);
 	}
 	ReferencedImports.Empty();
@@ -789,7 +842,7 @@ EAsyncPackageState::Type FAsyncPackage::PostLoadObjects()
 EAsyncPackageState::Type FAsyncPackage::FinishObjects()
 {
 	SCOPE_CYCLE_COUNTER(STAT_FAsyncPackage_FinishObjects);
-	LastObjectWorkWasPerformedOn	= NULL;
+	LastObjectWorkWasPerformedOn	= nullptr;
 	LastTypeOfWorkPerformed			= TEXT("finishing all objects");
 
 	// All loaded objects are now finished loading so we can clear the RF_AsyncLoading flag so
@@ -833,7 +886,7 @@ EAsyncPackageState::Type FAsyncPackage::FinishObjects()
 	#endif		// WITH_ENGINE
 
 		// Flush cache to free memory
-		if (Linker != NULL)
+		if (Linker != nullptr)
 		{
 			Linker->FlushCache();
 		}
@@ -845,7 +898,7 @@ EAsyncPackageState::Type FAsyncPackage::FinishObjects()
 		// Call any completion callbacks specified.
 		for (int32 i = 0; i < CompletionCallbacks.Num(); i++)
 		{
-			CompletionCallbacks[i].ExecuteIfBound(PackageName, NULL);
+			CompletionCallbacks[i].ExecuteIfBound(PackageName, nullptr);
 		}
 	}
 
@@ -883,22 +936,22 @@ FAsyncPackage& LoadPackageAsync( const FString& InPackageName, const FGuid* Pack
 	}
 	else
 	{
+		// PackageName will get populated by the conditional function
 		FString ClassName;
-		if (FPackageName::ParseExportTextPath(InPackageName, &ClassName, &PackageName))
-		{
-			// PackageName got populated by the conditional function
-		}
-		else
+		if (!FPackageName::ParseExportTextPath(InPackageName, &ClassName, &PackageName))
 		{
 			UE_LOG(LogStreaming, Fatal, TEXT("LoadPackageAsync failed to begin to load a level because the supplied package name was neither a valid long package name nor a filename of a map within a content folder: '%s'"), *InPackageName);
 		}
 	}
 
-	// Check whether the file is already in the queue.
+	// When constructing the FName from filename, we need to make sure that if the filename ends with _number, we're not splitting the FName
+	FName PackageFName(/*ENAME_LinkerConstructor,*/ *PackageName);
+
+	// Check whether the file is already in the queue.	
 	for( int32 PackageIndex=0; PackageIndex<GObjAsyncPackages.Num(); PackageIndex++ )
 	{
 		FAsyncPackage& PendingPackage = GObjAsyncPackages[PackageIndex];
-		if( PendingPackage.GetPackageName() == PackageName )
+		if (PendingPackage.GetPackageName() == PackageFName)
 		{
 			// Early out as package is already being preloaded.
 			return PendingPackage;
@@ -916,7 +969,7 @@ FAsyncPackage& LoadPackageAsync( const FString& InPackageName, const FGuid* Pack
 		UE_LOG(LogStreaming, Fatal, TEXT("Async loading code requires long package names (%s)."), *InPackageName);
 	}
 	// Add to (FIFO) queue.
-	FAsyncPackage *Package = new(GObjAsyncPackages) FAsyncPackage(PackageName, PackageGuid, PackageType, PackageToLoadFrom);
+	FAsyncPackage *Package = new(GObjAsyncPackages)FAsyncPackage(PackageFName, PackageGuid, PackageType, FName(/*ENAME_LinkerConstructor,*/ *PackageToLoadFrom));
 	return *Package;
 }
 
@@ -933,7 +986,7 @@ FAsyncPackage& LoadPackageAsync( const FString& InPackageName, FLoadPackageAsync
  * @param	PackageName			Name of package to query load percentage for
  * @return	Async load percentage if package is currently being loaded, -1 otherwise
  */
-float GetAsyncLoadPercentage( const FString& PackageName )
+float GetAsyncLoadPercentage( const FName& PackageName )
 {
 	float LoadPercentage = -1.f;
 	for( int32 PackageIndex=0; PackageIndex<GObjAsyncPackages.Num(); PackageIndex++ )
@@ -1095,7 +1148,7 @@ FArchiveAsync::FArchiveAsync( const TCHAR* InFileName )
 ,	UncompressedFileSize		( INDEX_NONE	)
 ,	BulkDataAreaSize			( 0	)
 ,	CurrentPos					( 0				)
-,	CompressedChunks			( NULL			)
+,	CompressedChunks			( nullptr			)
 ,	CurrentChunkIndex			( 0				)
 ,	CompressionFlags			( COMPRESS_None	)
 {
@@ -1104,11 +1157,11 @@ FArchiveAsync::FArchiveAsync( const TCHAR* InFileName )
 
 	PrecacheStartPos[CURRENT]	= 0;
 	PrecacheEndPos[CURRENT]		= 0;
-	PrecacheBuffer[CURRENT]		= NULL;
+	PrecacheBuffer[CURRENT]		= nullptr;
 
 	PrecacheStartPos[NEXT]		= 0;
 	PrecacheEndPos[NEXT]		= 0;
-	PrecacheBuffer[NEXT]		= NULL;
+	PrecacheBuffer[NEXT]		= nullptr;
 
 	// Relies on default constructor initializing to 0.
 	check( PrecacheReadStatus[CURRENT].GetValue() == 0 );
@@ -1154,13 +1207,13 @@ void FArchiveAsync::FlushCache()
 	// Invalidate any precached data and free memory for current buffer.
 	Delta += PrecacheEndPos[CURRENT] - PrecacheStartPos[CURRENT];
 	FMemory::Free( PrecacheBuffer[CURRENT] );
-	PrecacheBuffer[CURRENT]		= NULL;
+	PrecacheBuffer[CURRENT]		= nullptr;
 	PrecacheStartPos[CURRENT]	= 0;
 	PrecacheEndPos[CURRENT]		= 0;
 	
 	// Invalidate any precached data and free memory for next buffer.
 	FMemory::Free( PrecacheBuffer[NEXT] );
-	PrecacheBuffer[NEXT]		= NULL;
+	PrecacheBuffer[NEXT]		= nullptr;
 	PrecacheStartPos[NEXT]		= 0;
 	PrecacheEndPos[NEXT]		= 0;
 
@@ -1202,7 +1255,7 @@ bool FArchiveAsync::Close()
  */
 bool FArchiveAsync::SetCompressionMap( TArray<FCompressedChunk>* InCompressedChunks, ECompressionFlags InCompressionFlags )
 {
-	// Set chunks. A value of NULL means to use direct reads again.
+	// Set chunks. A value of nullptr means to use direct reads again.
 	CompressedChunks	= InCompressedChunks;
 	CompressionFlags	= InCompressionFlags;
 	CurrentChunkIndex	= 0;
@@ -1240,7 +1293,7 @@ void FArchiveAsync::BufferSwitcheroo()
 	PrecacheEndPos[CURRENT]		= PrecacheEndPos[NEXT];
 
 	// Next buffer is unused/ free.
-	PrecacheBuffer[NEXT]		= NULL;
+	PrecacheBuffer[NEXT]		= nullptr;
 	PrecacheStartPos[NEXT]		= 0;
 	PrecacheEndPos[NEXT]		= 0;
 }
diff --git a/Engine/Source/Runtime/CoreUObject/Public/Serialization/AsyncLoading.h b/Engine/Source/Runtime/CoreUObject/Public/Serialization/AsyncLoading.h
--- a/Engine/Source/Runtime/CoreUObject/Public/Serialization/AsyncLoading.h
+++ b/Engine/Source/Runtime/CoreUObject/Public/Serialization/AsyncLoading.h
@@ -15,53 +15,7 @@ struct FAsyncPackage : public FGCObject
 	/**
 	 * Constructor
 	 */
-	FAsyncPackage(const FString& InPackageName, const FGuid* InPackageGuid, FName InPackageType, const FString& InPackageNameToLoad)
-	:	PackageName					( InPackageName			)
-	,	PackageNameToLoad			( InPackageNameToLoad	)
-	,	PackageGuid					( InPackageGuid != NULL ? *InPackageGuid : FGuid(0,0,0,0) )
-	,	PackageType					( InPackageType			)
-	,	Linker						( NULL					)
-	, DependencyRefCount	( 0						)
-	, LoadImportIndex			( 0						)
-	,	ImportIndex					( 0						)
-	,	ExportIndex					( 0						)
-	,	TimeLimit					( FLT_MAX				)
-	,	bUseTimeLimit				( false					)
-	,	bUseFullTimeLimit			( false					)
-	,	bTimeLimitExceeded			( false					)
-	,	bLoadHasFailed				( false					)
-	, bLoadHasFinished			( false					)
-	,	TickStartTime				( 0						)
-	,	LastObjectWorkWasPerformedOn( NULL					)
-	,	LastTypeOfWorkPerformed		( NULL					)
-	,	LoadStartTime				( 0.0					)
-	,	LoadPercentage				( 0						)
-	,	PackageFlags				( 0						)
-#if WITH_EDITOR
-	,	PIEInstanceID				( INDEX_NONE			)
-#endif	
-	
-#if PERF_TRACK_DETAILED_ASYNC_STATS
-	,	TickCount					( 0						)
-	,	TickLoopCount				( 0						)
-	,	CreateLinkerCount			( 0						)
-	,	FinishLinkerCount			( 0						)
-	,	CreateImportsCount			( 0						)
-	,	CreateExportsCount			( 0						)
-	,	PreLoadObjectsCount			( 0						)
-	,	PostLoadObjectsCount		( 0						)
-	,	FinishObjectsCount			( 0						)
-	,	TickTime					( 0.0					)
-	,	CreateLinkerTime			( 0.0					)
-	,	FinishLinkerTime			( 0.0					)
-	,	CreateImportsTime			( 0.0					)
-	,	CreateExportsTime			( 0.0					)
-	,	PreLoadObjectsTime			( 0.0					)
-	,	PostLoadObjectsTime			( 0.0					)
-	,	FinishObjectsTime			( 0.0					)
-#endif // PERF_TRACK_DETAILED_ASYNC_STATS
-	{
-	}
+	FAsyncPackage(const FName& InPackageName, const FGuid* InPackageGuid, FName InPackageType, const FName& InPackageNameToLoad);
 
 	/**
 	 * Ticks the async loading code.
@@ -92,7 +46,7 @@ struct FAsyncPackage : public FGCObject
 	/**
 	 * Returns the name of the package to load.
 	 */
-	const FString& GetPackageName() const
+	FORCEINLINE const FName& GetPackageName() const
 	{
 		return PackageName;
 	}
@@ -148,9 +102,9 @@ struct FAsyncPackage : public FGCObject
 
 private:
 	/** Name of the UPackage to create.																	*/
-	FString						PackageName;
+	FName						PackageName;
 	/** Name of the package to load.																	*/
-	FString						PackageNameToLoad;
+	FName						PackageNameToLoad;
 	/** GUID of the package to load, or the zeroed invalid GUID for "don't care" */
 	FGuid						PackageGuid;
 	/** An abstract type name associated with this package, for tagging use								*/
@@ -332,7 +286,7 @@ private:
 	/**
 	 * Function called when pending import package has been loaded.
 	 */
-	void ImportFullyLoadedCallback(const FString& PackageName, UPackage* LoadedPackage);
+	void ImportFullyLoadedCallback(const FName& PackageName, UPackage* LoadedPackage);
 	/**
 	 * Adds dependency tree to the list if packages to wait for until their linkers have been created.
 	 *
@@ -350,7 +304,7 @@ private:
 	 *
 	 * @param PendingImport Name of the package imported either directly or by one of the imported packages
 	 */
-	void AddImportDependency(int32 CurrentPackageIndex, const FString& PendingImport);
+	void AddImportDependency(int32 CurrentPackageIndex, const FName& PendingImport);
 	/**
 	 * Removes references to any imported packages.
 	 */
diff --git a/Engine/Source/Runtime/CoreUObject/Public/UObject/UObjectGlobals.h b/Engine/Source/Runtime/CoreUObject/Public/UObject/UObjectGlobals.h
--- a/Engine/Source/Runtime/CoreUObject/Public/UObject/UObjectGlobals.h
+++ b/Engine/Source/Runtime/CoreUObject/Public/UObject/UObjectGlobals.h
@@ -253,7 +253,7 @@ COREUOBJECT_API UPackage* LoadPackage( UPackage* InOuter, const TCHAR* InLongPac
  * @param	PackageName			Package name we were trying to load
  * @param	LoadedPackage		Loaded package if successful, NULL otherwise	
  */
-DECLARE_DELEGATE_TwoParams(FLoadPackageAsyncDelegate, const FString& /*PackageName*/, UPackage* /*LoadedPackage*/)
+DECLARE_DELEGATE_TwoParams(FLoadPackageAsyncDelegate, const FName& /*PackageName*/, UPackage* /*LoadedPackage*/)
 struct FAsyncPackage;
 /**
  * Asynchronously load a package and all contained objects that match context flags. Non- blocking.
@@ -284,7 +284,7 @@ COREUOBJECT_API FAsyncPackage& LoadPackageAsync( const FString& PackageName, con
  * @param	PackageName			Name of package to query load percentage for
  * @return	Async load percentage if package is currently being loaded, -1 otherwise
  */
-COREUOBJECT_API float GetAsyncLoadPercentage( const FString& PackageName );
+COREUOBJECT_API float GetAsyncLoadPercentage( const FName& PackageName );
 
 /** 
  * Deletes all unreferenced objects, keeping objects that have any of the passed in KeepFlags set
diff --git a/Engine/Source/Runtime/Engine/Classes/Engine/LevelStreaming.h b/Engine/Source/Runtime/Engine/Classes/Engine/LevelStreaming.h
--- a/Engine/Source/Runtime/Engine/Classes/Engine/LevelStreaming.h
+++ b/Engine/Source/Runtime/Engine/Classes/Engine/LevelStreaming.h
@@ -314,7 +314,7 @@ private:
 	 *
 	 * @param LevelPackage	Loaded level package
 	 */
-	void AsyncLevelLoadComplete( const FString& PackageName, UPackage* LevelPackage );
+	void AsyncLevelLoadComplete( const FName& PackageName, UPackage* LevelPackage );
 
 	/** Pointer to Level object if currently loaded/ streamed in.																*/
 	UPROPERTY(transient)
diff --git a/Engine/Source/Runtime/Engine/Classes/Engine/PackageMapClient.h b/Engine/Source/Runtime/Engine/Classes/Engine/PackageMapClient.h
--- a/Engine/Source/Runtime/Engine/Classes/Engine/PackageMapClient.h
+++ b/Engine/Source/Runtime/Engine/Classes/Engine/PackageMapClient.h
@@ -68,7 +68,7 @@ public:
 	bool			IsGUIDLoaded( const FNetworkGUID& NetGUID ) const;
 	bool			IsGUIDBroken( const FNetworkGUID& NetGUID, const bool bMustBeRegistered ) const;
 
-	void			AsyncPackageCallback( const FString& PackageName, UPackage * Package );
+	void			AsyncPackageCallback( const FName& PackageName, UPackage * Package );
 	
 	TMap< FNetworkGUID, FNetGuidCacheObject >		ObjectLookup;
 	TMap< TWeakObjectPtr< UObject >, FNetworkGUID >	NetGUIDLookup;
diff --git a/Engine/Source/Runtime/Engine/Classes/Engine/StreamableManager.h b/Engine/Source/Runtime/Engine/Classes/Engine/StreamableManager.h
--- a/Engine/Source/Runtime/Engine/Classes/Engine/StreamableManager.h
+++ b/Engine/Source/Runtime/Engine/Classes/Engine/StreamableManager.h
@@ -68,7 +68,7 @@ private:
 		}
 	};
 
-	static void AsyncLoadCallbackWrapper(const FString& PackageName, UPackage* LevelPackage, FCallback* Handler)
+	static void AsyncLoadCallbackWrapper(const FName& PackageName, UPackage* LevelPackage, FCallback* Handler)
 	{
 		FCallback* Callback = Handler;
 		Callback->Manager->AsyncLoadCallback(Callback->Request);
diff --git a/Engine/Source/Runtime/Engine/Classes/Engine/World.h b/Engine/Source/Runtime/Engine/Classes/Engine/World.h
--- a/Engine/Source/Runtime/Engine/Classes/Engine/World.h
+++ b/Engine/Source/Runtime/Engine/Classes/Engine/World.h
@@ -143,7 +143,7 @@ private:
 	void CopyWorldData();
 
 	/** callback sent to async loading code to inform us when the level package is complete */
-	void SeamlessTravelLoadCallback(const FString& PackageName, UPackage* LevelPackage);
+	void SeamlessTravelLoadCallback(const FName& PackageName, UPackage* LevelPackage);
 
 	void SetHandlerLoadedData(UObject* InLevelPackage, UWorld* InLoadedWorld);
 
diff --git a/Engine/Source/Runtime/Engine/Private/LevelStreaming.cpp b/Engine/Source/Runtime/Engine/Private/LevelStreaming.cpp
--- a/Engine/Source/Runtime/Engine/Private/LevelStreaming.cpp
+++ b/Engine/Source/Runtime/Engine/Private/LevelStreaming.cpp
@@ -496,12 +496,10 @@ bool ULevelStreaming::RequestLevel(UWorld* PersistentWorld, bool bAllowLevelLoad
 	return true;
 }
 
-void ULevelStreaming::AsyncLevelLoadComplete( const FString& InPackageName, UPackage* InLoadedPackage ) 
+void ULevelStreaming::AsyncLevelLoadComplete( const FName& InPackageName, UPackage* InLoadedPackage ) 
 {
 	bHasLoadRequestPending = false;
 
-	const FName PackageFName = FName(*InPackageName);
-	
 	if( InLoadedPackage )
 	{
 		UPackage* LevelPackage = InLoadedPackage;
@@ -544,7 +542,7 @@ void ULevelStreaming::AsyncLevelLoadComplete( const FString& InPackageName, UPac
 			}
 			else
 			{
-				UE_LOG(LogLevelStreaming, Warning, TEXT("Couldn't find ULevel object in package '%s'"), *InPackageName );
+				UE_LOG(LogLevelStreaming, Warning, TEXT("Couldn't find ULevel object in package '%s'"), *InPackageName.ToString() );
 			}
 		}
 		else
@@ -562,7 +560,7 @@ void ULevelStreaming::AsyncLevelLoadComplete( const FString& InPackageName, UPac
 				//    If the package name to load was different...
 				//         ... it means the specified package name was explicit and we will just load from another file.
 
-				FName OldDesiredPackageName = PackageFName;
+				FName OldDesiredPackageName = InPackageName;
 				UWorld** OwningWorldPtr = ULevel::StreamedLevelsOwningWorld.Find(OldDesiredPackageName);
 				UWorld* OwningWorld = OwningWorldPtr ? *OwningWorldPtr : NULL;
 				ULevel::StreamedLevelsOwningWorld.Remove(OldDesiredPackageName);
@@ -625,12 +623,12 @@ void ULevelStreaming::AsyncLevelLoadComplete( const FString& InPackageName, UPac
 	}
 	else
 	{
-		UE_LOG(LogLevelStreaming, Warning, TEXT("Failed to load package '%s'"), *InPackageName );
+		UE_LOG(LogLevelStreaming, Warning, TEXT("Failed to load package '%s'"), *InPackageName.ToString() );
 	}
 
 	// Clean up the world type list and owning world list now that PostLoad has occurred
-	UWorld::WorldTypePreLoadMap.Remove(PackageFName);
-	ULevel::StreamedLevelsOwningWorld.Remove(PackageFName);
+	UWorld::WorldTypePreLoadMap.Remove(InPackageName);
+	ULevel::StreamedLevelsOwningWorld.Remove(InPackageName);
 }
 
 bool ULevelStreaming::IsLevelVisible() const
diff --git a/Engine/Source/Runtime/Engine/Private/PackageMapClient.cpp b/Engine/Source/Runtime/Engine/Private/PackageMapClient.cpp
--- a/Engine/Source/Runtime/Engine/Private/PackageMapClient.cpp
+++ b/Engine/Source/Runtime/Engine/Private/PackageMapClient.cpp
@@ -1628,19 +1628,17 @@ void FNetGUIDCache::RegisterNetGUIDFromPath_Client( const FNetworkGUID& NetGUID,
 	RegisterNetGUID_Internal( NetGUID, CacheObject );
 }
 
-void FNetGUIDCache::AsyncPackageCallback( const FString& PackageName, UPackage * Package )
+void FNetGUIDCache::AsyncPackageCallback( const FName& PackageName, UPackage * Package )
 {
 	check( Package == NULL || Package->IsFullyLoaded() );
 
-	const FName PackageFName( *PackageName );
+	FNetworkGUID NetGUID = PendingAsyncPackages.FindRef(PackageName);
 
-	FNetworkGUID NetGUID = PendingAsyncPackages.FindRef( PackageFName );
-
-	PendingAsyncPackages.Remove( PackageFName );
+	PendingAsyncPackages.Remove(PackageName);
 
 	if ( !NetGUID.IsValid() )
 	{
-		UE_LOG( LogNetPackageMap, Error, TEXT( "AsyncPackageCallback: Could not find package. Path: %s" ), *PackageName );
+		UE_LOG( LogNetPackageMap, Error, TEXT( "AsyncPackageCallback: Could not find package. Path: %s" ), *PackageName.ToString() );
 		return;
 	}
 
@@ -1648,13 +1646,13 @@ void FNetGUIDCache::AsyncPackageCallback( const FString& PackageName, UPackage *
 
 	if ( CacheObject == NULL )
 	{
-		UE_LOG( LogNetPackageMap, Error, TEXT( "AsyncPackageCallback: Could not find net guid. Path: %s, NetGUID: %s" ), *PackageName, *NetGUID.ToString() );
+		UE_LOG( LogNetPackageMap, Error, TEXT( "AsyncPackageCallback: Could not find net guid. Path: %s, NetGUID: %s" ), *PackageName.ToString(), *NetGUID.ToString() );
 		return;
 	}
 
 	if ( !CacheObject->bIsPending )
 	{
-		UE_LOG( LogNetPackageMap, Error, TEXT( "AsyncPackageCallback: Package wasn't pending. Path: %s, NetGUID: %s" ), *PackageName, *NetGUID.ToString() );
+		UE_LOG( LogNetPackageMap, Error, TEXT( "AsyncPackageCallback: Package wasn't pending. Path: %s, NetGUID: %s" ), *PackageName.ToString(), *NetGUID.ToString() );
 	}
 
 	CacheObject->bIsPending = false;
@@ -1662,12 +1660,12 @@ void FNetGUIDCache::AsyncPackageCallback( const FString& PackageName, UPackage *
 	if ( Package == NULL )
 	{
 		CacheObject->bIsBroken = true;
-		UE_LOG( LogNetPackageMap, Error, TEXT( "AsyncPackageCallback: Package FAILED to load. Path: %s, NetGUID: %s" ), *PackageName, *NetGUID.ToString() );
+		UE_LOG( LogNetPackageMap, Error, TEXT( "AsyncPackageCallback: Package FAILED to load. Path: %s, NetGUID: %s" ), *PackageName.ToString(), *NetGUID.ToString() );
 	}
 	else if ( Package->GetGuid() != CacheObject->PackageGuid && CVarIgnorePackageMismatch.GetValueOnGameThread() == 0 )
 	{
 		CacheObject->bIsBroken = true;
-		UE_LOG( LogNetPackageMap, Error, TEXT( "AsyncPackageCallback: Package GUID mismatch! Path: %s, NetGUID: %s, GUID1: %s, GUID2: %s" ), *PackageName, *NetGUID.ToString(), *Package->GetGuid().ToString(), *CacheObject->PackageGuid.ToString() );
+		UE_LOG( LogNetPackageMap, Error, TEXT( "AsyncPackageCallback: Package GUID mismatch! Path: %s, NetGUID: %s, GUID1: %s, GUID2: %s" ), *PackageName.ToString(), *NetGUID.ToString(), *Package->GetGuid().ToString(), *CacheObject->PackageGuid.ToString() );
 	}
 }
 
diff --git a/Engine/Source/Runtime/Engine/Private/UnrealEngine.cpp b/Engine/Source/Runtime/Engine/Private/UnrealEngine.cpp
--- a/Engine/Source/Runtime/Engine/Private/UnrealEngine.cpp
+++ b/Engine/Source/Runtime/Engine/Private/UnrealEngine.cpp
@@ -6375,9 +6375,9 @@ bool UEngine::HandleLogoutStatLevelsCommand( const TCHAR* Cmd, FOutputDevice& Ar
 		{
 			DisplayName += FString::Printf(TEXT(" - %4.1f sec"), LevelPackage->GetLoadTime());
 		}
-		else if( GetAsyncLoadPercentage( *LevelStatus.PackageName.ToString() ) >= 0 )
+		else if( GetAsyncLoadPercentage( LevelStatus.PackageName ) >= 0 )
 		{
-			const int32 Percentage = FMath::TruncToInt( GetAsyncLoadPercentage( *LevelStatus.PackageName.ToString() ) );
+			const int32 Percentage = FMath::TruncToInt( GetAsyncLoadPercentage( LevelStatus.PackageName ) );
 			DisplayName += FString::Printf(TEXT(" - %3i %%"), Percentage ); 
 		}
 
@@ -8409,7 +8409,7 @@ static void SetGametypeContentObjectReferencers(UObject* GametypeContentPackage,
  * @param	ContentPackage		The package that was loaded.
  * @param	GameEngine			The GameEngine.
  */
-static void AsyncLoadLocalizedMapGameTypeContentCallback(const FString& PackageName, UPackage* ContentPackage, FName InContextHandle)
+static void AsyncLoadLocalizedMapGameTypeContentCallback(const FName& PackageName, UPackage* ContentPackage, FName InContextHandle)
 {
 	SetGametypeContentObjectReferencers(ContentPackage, InContextHandle, GametypeContent_LocalizedReferencerIndex);
 }
@@ -8421,7 +8421,7 @@ static void AsyncLoadLocalizedMapGameTypeContentCallback(const FString& PackageN
  * @param	ContentPackage		The package that was loaded.
  * @param	GameEngine			The GameEngine.
  */
-static void AsyncLoadMapGameTypeContentCallback(const FString& PackageName, UPackage* ContentPackage, FName InContextHandle)
+static void AsyncLoadMapGameTypeContentCallback(const FName& PackageName, UPackage* ContentPackage, FName InContextHandle)
 {
 	SetGametypeContentObjectReferencers(ContentPackage, InContextHandle, GametypeContent_ReferencerIndex);
 }
@@ -9436,7 +9436,7 @@ void UEngine::VerifyLoadMapWorldCleanup()
  * @param	LevelPackage	level package that finished async loading
  * @param	InGameEngine	pointer to game engine object to associated loaded level with so it won't be GC'ed
  */
-static void AsyncMapChangeLevelLoadCompletionCallback(const FString& PackageName, UPackage* LevelPackage, FName InWorldHandle )
+static void AsyncMapChangeLevelLoadCompletionCallback(const FName& PackageName, UPackage* LevelPackage, FName InWorldHandle )
 {
 	FWorldContext &Context = GEngine->GetWorldContextFromHandleChecked( InWorldHandle );
 
diff --git a/Engine/Source/Runtime/Engine/Private/World.cpp b/Engine/Source/Runtime/Engine/Private/World.cpp
--- a/Engine/Source/Runtime/Engine/Private/World.cpp
+++ b/Engine/Source/Runtime/Engine/Private/World.cpp
@@ -4216,7 +4216,7 @@ void FSeamlessTravelHandler::SetHandlerLoadedData(UObject* InLevelPackage, UWorl
 }
 
 /** callback sent to async loading code to inform us when the level package is complete */
-void FSeamlessTravelHandler::SeamlessTravelLoadCallback(const FString& PackageName, UPackage* LevelPackage)
+void FSeamlessTravelHandler::SeamlessTravelLoadCallback(const FName& PackageName, UPackage* LevelPackage)
 {
 	// make sure we remove the name, even if travel was cancelled.
 	const FName URLMapFName = FName(*PendingTravelURL.Map);
