[UE-4878] Fixing it so we don't clear out properties that reference the class that owns them (when replacing the class).

diff --git a/Engine/Source/Runtime/CoreUObject/Private/UObject/PropertyBaseObject.cpp b/Engine/Source/Runtime/CoreUObject/Private/UObject/PropertyBaseObject.cpp
--- a/Engine/Source/Runtime/CoreUObject/Private/UObject/PropertyBaseObject.cpp
+++ b/Engine/Source/Runtime/CoreUObject/Private/UObject/PropertyBaseObject.cpp
@@ -371,10 +371,21 @@ void UObjectPropertyBase::CheckValidObject(void* Value) const
 	UObject *Object = GetObjectPropertyValue(Value);
 	if (Object)
 	{
-		// If the PropertyClass has CLASS_NewerVersionExists set then we are likely in the middle of an
-		// FArchiveReplaceObjectRef pass post-compilation.  Depending on the order of events, we might
-		// serialize the object property before the property class, and need to ignore this check.
-		if (PropertyClass != NULL && !Object->GetClass()->IsChildOf(PropertyClass) && !(PropertyClass->HasAnyClassFlags(CLASS_NewerVersionExists)))
+		//
+		// here we want to make sure the the object value still matches the 
+		// object type expected by the property...
+
+		UClass* ObjectClass = Object->GetClass();
+		// we could be in the middle of replacing references to the 
+		// PropertyClass itself (in the middle of an FArchiveReplaceObjectRef 
+		// pass)... if this is the case, then we might have already replaced 
+		// the object's class, but not the PropertyClass yet (or vise-versa)... 
+		// so we use this to ensure in that situation that we don't clear the 
+		// object value (if CLASS_NewerVersionExists is set, then we are likely 
+		// in the middle of an FArchiveReplaceObjectRef pass)
+		bool bIsReplacingClassRefs = PropertyClass->HasAnyClassFlags(CLASS_NewerVersionExists) != ObjectClass->HasAnyClassFlags(CLASS_NewerVersionExists);
+		
+		if ((PropertyClass != nullptr) && !ObjectClass->IsChildOf(PropertyClass) && !bIsReplacingClassRefs)
 		{
 			UE_LOG(LogProperty, Warning,
 				TEXT("Serialized %s for a property of %s. Reference will be NULLed.\n    Property = %s\n    Item = %s"),
