Added FObjectKey which provides a unique indentifier for UObject types, persistent for the lifetime of the process.

FObjectKey uses the same index/serial number combination as FWeakObjectPtr, but crucially never expires. While an expired FWeakObjectPtr will compare equal with null and any other expired weak object pointer, FObjectKey is immutable, and will only ever compare equal with other keys that point to the same object (even if that object has expired).

diff --git a/Engine/Source/Runtime/CoreUObject/Public/UObject/ObjectKey.h b/Engine/Source/Runtime/CoreUObject/Public/UObject/ObjectKey.h
new file mode 100644
--- /dev/null
+++ b/Engine/Source/Runtime/CoreUObject/Public/UObject/ObjectKey.h
@@ -0,0 +1,54 @@
+// Copyright 1998-2014 Epic Games, Inc. All Rights Reserved.
+
+#pragma once
+
+#include "WeakObjectPtrTemplates.h"
+#include "WeakObjectPtr.h"
+
+/** FObjectKey is an immutable, copyable key which can be used to uniquely identify an object for the lifetime of the application */
+struct FObjectKey
+{
+public:
+	/** Deafult constructor */
+	FORCEINLINE FObjectKey()
+		: ObjectIndex(INDEX_NONE)
+		, ObjectSerialNumber(0)
+	{
+	}
+
+	/** Construct from an object pointer */
+	FORCEINLINE FObjectKey(const UObject *Object)
+	    : ObjectIndex(INDEX_NONE)
+	    , ObjectSerialNumber(0)
+	{
+	    if (Object)
+	    {
+	        FWeakObjectPtr Weak(Object);
+	        ObjectIndex = Weak.ObjectIndex;
+	        ObjectSerialNumber = Weak.ObjectSerialNumber;
+	    }
+	}
+
+	/** Compare this key with another */
+	FORCEINLINE bool operator==(const FObjectKey &Other) const
+	{
+		return ObjectIndex == Other.ObjectIndex && ObjectSerialNumber == Other.ObjectSerialNumber;
+	}
+
+	/** Compare this key with another */
+	FORCEINLINE bool operator!=(const FObjectKey &Other) const
+	{
+		return ObjectIndex != Other.ObjectIndex || ObjectSerialNumber != Other.ObjectSerialNumber;
+	}
+
+	/** Hash function */
+	friend uint32 GetTypeHash(const FObjectKey& Key)
+	{
+		return HashCombine(Key.ObjectIndex, Key.ObjectSerialNumber);
+	}
+
+private:
+
+	int32		ObjectIndex;
+	int32		ObjectSerialNumber;
+};
diff --git a/Engine/Source/Runtime/CoreUObject/Public/UObject/WeakObjectPtr.h b/Engine/Source/Runtime/CoreUObject/Public/UObject/WeakObjectPtr.h
--- a/Engine/Source/Runtime/CoreUObject/Public/UObject/WeakObjectPtr.h
+++ b/Engine/Source/Runtime/CoreUObject/Public/UObject/WeakObjectPtr.h
@@ -147,6 +147,8 @@ protected:
 	}
 
 private:
+	friend struct FObjectKey;
+	
 	/**  
 	 * internal function to test for serial number matches
 	 * @return true if the serial number in this matches the central table
