FStructOnScope is an interface to access a struct instance, that is not directly owned by any UObject. Specialized implementations are allowed.

diff --git a/Engine/Source/Editor/PropertyEditor/Private/SStructureDetailsView.h b/Engine/Source/Editor/PropertyEditor/Private/SStructureDetailsView.h
--- a/Engine/Source/Editor/PropertyEditor/Private/SStructureDetailsView.h
+++ b/Engine/Source/Editor/PropertyEditor/Private/SStructureDetailsView.h
@@ -6,7 +6,7 @@
 
 class SStructureDetailsView : public SDetailsViewBase, public IStructureDetailsView
 {
-	TSharedPtr<struct FStructOnScope> StructData;
+	TSharedPtr<class FStructOnScope> StructData;
 	TSharedPtr<class FStructurePropertyNode> RootNode;
 	FString CustomName;
 public:
diff --git a/Engine/Source/Editor/PropertyEditor/Public/IStructureDetailsView.h b/Engine/Source/Editor/PropertyEditor/Public/IStructureDetailsView.h
--- a/Engine/Source/Editor/PropertyEditor/Public/IStructureDetailsView.h
+++ b/Engine/Source/Editor/PropertyEditor/Public/IStructureDetailsView.h
@@ -12,6 +12,6 @@ class IStructureDetailsView
 public:
 	virtual TSharedPtr<class SWidget> GetWidget() = 0;
 
-	virtual void SetStructureData(TSharedPtr<struct FStructOnScope> StructData) = 0;
+	virtual void SetStructureData(TSharedPtr<class FStructOnScope> StructData) = 0;
 	virtual FOnFinishedChangingProperties& GetOnFinishedChangingPropertiesDelegate() = 0;
 };
diff --git a/Engine/Source/Editor/PropertyEditor/Public/PropertyEditorModule.h b/Engine/Source/Editor/PropertyEditor/Public/PropertyEditorModule.h
--- a/Engine/Source/Editor/PropertyEditor/Public/PropertyEditorModule.h
+++ b/Engine/Source/Editor/PropertyEditor/Public/PropertyEditorModule.h
@@ -232,7 +232,7 @@ public:
 	 */
 	virtual TSharedPtr<class ISinglePropertyView> CreateSingleProperty( UObject* InObject, FName InPropertyName, const struct FSinglePropertyParams& InitParams );
 
-	virtual TSharedRef<class IStructureDetailsView> CreateStructureDetailView(const struct FDetailsViewArgs& DetailsViewArgs, TSharedPtr<struct FStructOnScope> StructData, bool bShowObjects, const FString& CustomName = FString());
+	virtual TSharedRef<class IStructureDetailsView> CreateStructureDetailView(const struct FDetailsViewArgs& DetailsViewArgs, TSharedPtr<class FStructOnScope> StructData, bool bShowObjects, const FString& CustomName = FString());
 
 	/**
 	 * Creates a property change listener that notifies users via a  delegate when a property on an object changes
diff --git a/Engine/Source/Runtime/CoreUObject/Public/UObject/Class.h b/Engine/Source/Runtime/CoreUObject/Public/UObject/Class.h
--- a/Engine/Source/Runtime/CoreUObject/Public/UObject/Class.h
+++ b/Engine/Source/Runtime/CoreUObject/Public/UObject/Class.h
@@ -1090,15 +1090,13 @@ public:
 	virtual COREUOBJECT_API void RecursivelyPreload();
 };
 
-struct FStructOnScope
+class FStructOnScope
 {
-private:
+protected:
 	TWeakObjectPtr<const UScriptStruct> ScriptStruct;
 	uint8* SampleStructMemory;
 
-	FStructOnScope(const FStructOnScope&);
-	FStructOnScope& operator=(const FStructOnScope&);
-
+	FStructOnScope() : SampleStructMemory(NULL) {}
 public:
 	FStructOnScope(const UScriptStruct* InScriptStruct)
 		: ScriptStruct(InScriptStruct)
@@ -1111,15 +1109,15 @@ public:
 		}
 	}
 
-	uint8* GetStructMemory() { return SampleStructMemory; }
+	virtual uint8* GetStructMemory() { return SampleStructMemory; }
 
-	const uint8* GetStructMemory() const { return SampleStructMemory; }
+	virtual const uint8* GetStructMemory() const { return SampleStructMemory; }
 
-	const UScriptStruct* GetStruct() const { return ScriptStruct.Get(); }
+	virtual const UScriptStruct* GetStruct() const { return ScriptStruct.Get(); }
 
-	bool IsValid() const { return ScriptStruct.IsValid() && SampleStructMemory; }
+	virtual bool IsValid() const { return ScriptStruct.IsValid() && SampleStructMemory; }
 
-	void Destroy()
+	virtual void Destroy()
 	{
 		if (ScriptStruct.IsValid() && SampleStructMemory)
 		{
@@ -1134,10 +1132,14 @@ public:
 		}
 	}
 
-	~FStructOnScope()
+	virtual ~FStructOnScope()
 	{
 		Destroy();
 	}
+
+private:
+	FStructOnScope(const FStructOnScope&);
+	FStructOnScope& operator=(const FStructOnScope&);
 };
 
 /*-----------------------------------------------------------------------------
