Add subobject creation functions to UObject with the same semantics as ones in FObjectInitializer.

diff --git a/Engine/Source/Runtime/CoreUObject/Private/UObject/Obj.cpp b/Engine/Source/Runtime/CoreUObject/Private/UObject/Obj.cpp
--- a/Engine/Source/Runtime/CoreUObject/Private/UObject/Obj.cpp
+++ b/Engine/Source/Runtime/CoreUObject/Private/UObject/Obj.cpp
@@ -13,10 +13,10 @@
 #include "Serialization/ArchiveDescribeReference.h"
 #include "FindStronglyConnected.h"
 #include "HotReloadInterface.h"
+#include "UObject/TlsObjectInitializers.h"
 
 DEFINE_LOG_CATEGORY(LogObj);
 
-
 /** Stat group for dynamic objects cycle counters*/
 
 
@@ -50,6 +50,18 @@ UObject::UObject( EStaticConstructor, EObjectFlags InFlags )
 {
 }
 
+UObject* UObject::CreateDefaultSubobject(FName SubobjectFName, UClass* ReturnType, UClass* ClassToCreateByDefault, bool bIsRequired, bool bAbstract, bool bIsTransient)
+{
+	auto CurrentInitializer = FTlsObjectInitializers::Top();
+	return CurrentInitializer->CreateDefaultSubobject(this, SubobjectFName, ReturnType, ClassToCreateByDefault, bIsRequired, bAbstract, bIsTransient);
+}
+
+UObject* UObject::CreateEditorOnlyDefaultSubobjectImpl(FName SubobjectName, UClass* ReturnType, bool bTransient)
+{
+	auto CurrentInitializer = FTlsObjectInitializers::Top();
+	return CurrentInitializer->CreateEditorOnlyDefaultSubobject(this, SubobjectName, ReturnType, bTransient);
+}
+
 bool UObject::Rename( const TCHAR* InName, UObject* NewOuter, ERenameFlags Flags )
 {
 #if WITH_EDITOR
diff --git a/Engine/Source/Runtime/CoreUObject/Private/UObject/TlsObjectInitializers.cpp b/Engine/Source/Runtime/CoreUObject/Private/UObject/TlsObjectInitializers.cpp
new file mode 100644
--- /dev/null
+++ b/Engine/Source/Runtime/CoreUObject/Private/UObject/TlsObjectInitializers.cpp
@@ -0,0 +1,38 @@
+// Copyright 1998-2014 Epic Games, Inc. All Rights Reserved.
+#pragma once
+
+#include "CoreUObjectPrivate.h"
+#include "UObject/TlsObjectInitializers.h"
+
+TArray<FObjectInitializer*>* FTlsObjectInitializers::GetTlsObjectInitializers()
+{
+	/** Thread local slot for component overrides. */
+	static uint32 ComponentOverridesTlsId = FPlatformTLS::AllocTlsSlot();
+
+	auto TLS = static_cast<TArray<FObjectInitializer*>*>(FPlatformTLS::GetTlsValue(ComponentOverridesTlsId));
+	if (!TLS)
+	{
+		TLS = new TArray<FObjectInitializer*>();
+		FPlatformTLS::SetTlsValue(ComponentOverridesTlsId, TLS);
+	}
+
+	return TLS;
+}
+
+FObjectInitializer* FTlsObjectInitializers::Top()
+{
+	auto ObjectInitializers = GetTlsObjectInitializers();
+	return ObjectInitializers->Last();
+}
+
+void FTlsObjectInitializers::Pop()
+{
+	auto ObjectInitializers = GetTlsObjectInitializers();
+	ObjectInitializers->RemoveAt(ObjectInitializers->Num() - 1);
+}
+
+void FTlsObjectInitializers::Push(FObjectInitializer* Initializer)
+{
+	auto ObjectInitializers = GetTlsObjectInitializers();
+	ObjectInitializers->Push(Initializer);
+}
\ No newline at end of file
diff --git a/Engine/Source/Runtime/CoreUObject/Private/UObject/UObjectGlobals.cpp b/Engine/Source/Runtime/CoreUObject/Private/UObject/UObjectGlobals.cpp
--- a/Engine/Source/Runtime/CoreUObject/Private/UObject/UObjectGlobals.cpp
+++ b/Engine/Source/Runtime/CoreUObject/Private/UObject/UObjectGlobals.cpp
@@ -8,6 +8,7 @@
 #include "Misc/SecureHash.h"
 #include "UObject/ConstructorHelpers.h"
 #include "UObject/UObjectAnnotation.h"
+#include "UObject/TlsObjectInitializers.h"
 
 DEFINE_LOG_CATEGORY(LogUObjectGlobals);
 
@@ -1865,6 +1866,7 @@ FObjectInitializer::FObjectInitializer() :
 	// Mark we're in the constructor now.	
 	GIsInConstructor++;
 	GConstructedObject = Obj;
+	FTlsObjectInitializers::Push(this);
 }	
 
 FObjectInitializer::FObjectInitializer(UObject* InObj, UObject* InObjectArchetype, bool bInCopyTransientsFromClassDefaults, bool bInShouldIntializeProps, struct FObjectInstancingGraph* InInstanceGraph) :
@@ -1880,6 +1882,7 @@ FObjectInitializer::FObjectInitializer(UObject* InObj, UObject* InObjectArchetyp
 	// Mark we're in the constructor now.
 	GIsInConstructor++;
 	GConstructedObject = Obj;
+	FTlsObjectInitializers::Push(this);
 }
 
 /**
@@ -1887,6 +1890,7 @@ FObjectInitializer::FObjectInitializer(UObject* InObj, UObject* InObjectArchetyp
  **/
 FObjectInitializer::~FObjectInitializer()
 {
+	FTlsObjectInitializers::Pop();
 	// Let the FObjectFinders know we left the constructor.
 	GIsInConstructor--;
 	check(GIsInConstructor >= 0);
@@ -2599,4 +2603,20 @@ UObject* FObjectInitializer::CreateDefaultSubobject(UObject* Outer, FName Subobj
 		}
 	}
 	return Result;
-}
\ No newline at end of file
+}
+
+UObject* FObjectInitializer::CreateEditorOnlyDefaultSubobject(UObject* Outer, FName SubobjectName, UClass* ReturnType, bool bTransient /*= false*/) const
+{
+#if WITH_EDITOR
+	if (GIsEditor)
+	{
+		UObject* EditorSubobject = CreateDefaultSubobject(Outer, SubobjectName, ReturnType, ReturnType, /*bIsRequired =*/ false, /*bIsAbstract =*/ false, bTransient);
+		if (EditorSubobject)
+		{
+			EditorSubobject->MarkAsEditorOnlySubobject();
+		}
+		return EditorSubobject;
+	}
+#endif
+	return nullptr;
+}
diff --git a/Engine/Source/Runtime/CoreUObject/Public/UObject/TlsObjectInitializers.h b/Engine/Source/Runtime/CoreUObject/Public/UObject/TlsObjectInitializers.h
new file mode 100644
--- /dev/null
+++ b/Engine/Source/Runtime/CoreUObject/Public/UObject/TlsObjectInitializers.h
@@ -0,0 +1,37 @@
+// Copyright 1998-2014 Epic Games, Inc. All Rights Reserved.
+
+#pragma once
+
+class FObjectInitializer;
+
+/* 
+ * Helper class to manage thread local stack of FObjectInitializers
+ * relevant to currently created objects.
+ */
+class COREUOBJECT_API FTlsObjectInitializers
+{
+public:
+	/**
+	* Remove top element from the stack.
+	*/
+	static void Pop();
+
+	/**
+	* Push new FObjectInitializer on stack.
+	* @param	Initializer			Object initializer to push.
+	*/
+	static void Push(FObjectInitializer* Initializer);
+
+	/**
+	* Retrieve current FObjectInitializer for current thread.
+	* @return Current FObjectInitializer.
+	*/
+	static FObjectInitializer* Top();
+
+private:
+	/**
+	* Retrieve thread local stack of FObjectInitializers.
+	* @return The stack.
+	*/
+	static TArray<FObjectInitializer*>* GetTlsObjectInitializers();
+};
diff --git a/Engine/Source/Runtime/CoreUObject/Public/UObject/UObject.h b/Engine/Source/Runtime/CoreUObject/Public/UObject/UObject.h
--- a/Engine/Source/Runtime/CoreUObject/Public/UObject/UObject.h
+++ b/Engine/Source/Runtime/CoreUObject/Public/UObject/UObject.h
@@ -51,6 +51,80 @@ class COREUOBJECT_API UObject : public UObjectBaseUtility
 	{
 	}
 
+	UObject* CreateDefaultSubobject(FName SubobjectFName, UClass* ReturnType, UClass* ClassToCreateByDefault, bool bIsRequired, bool bAbstract, bool bIsTransient);
+
+	/**
+	* Create a component or subobject only to be used with the editor.
+	* @param	TReturnType					class of return type, all overrides must be of this type
+	* @param	Outer						outer to construct the subobject in
+	* @param	SubobjectName				name of the new component
+	* @param bTransient		true if the component is being assigned to a transient property
+	*/
+	template<class TReturnType>
+	TReturnType* CreateEditorOnlyDefaultSubobject(FName SubobjectName, bool bTransient = false)
+	{
+		UClass* ReturnType = TReturnType::StaticClass();
+		return static_cast<TReturnType*>(CreateEditorOnlyDefaultSubobjectImpl(SubobjectName, ReturnType, bTransient));
+	}
+
+	/**
+	* Create a component or subobject
+	* @param	TReturnType					class of return type, all overrides must be of this type
+	* @param	Outer						outer to construct the subobject in
+	* @param	SubobjectName				name of the new component
+	* @param bTransient		true if the component is being assigned to a transient property
+	*/
+	template<class TReturnType>
+	TReturnType* CreateDefaultSubobject(FName SubobjectName, bool bTransient = false)
+	{
+		UClass* ReturnType = TReturnType::StaticClass();
+		return static_cast<TReturnType*>(CreateDefaultSubobject(SubobjectName, ReturnType, ReturnType, /*bIsRequired =*/ true, /*bIsAbstract =*/ false, bTransient));
+	}
+	
+	/**
+	* Create a component or subobject
+	* @param TReturnType class of return type, all overrides must be of this type
+	* @param TClassToConstructByDefault class to construct by default
+	* @param Outer outer to construct the subobject in
+	* @param SubobjectName name of the new component
+	* @param bTransient		true if the component is being assigned to a transient property
+	*/
+	template<class TReturnType, class TClassToConstructByDefault>
+	TReturnType* CreateDefaultSubobject(FName SubobjectName, bool bTransient = false)
+	{
+		return static_cast<TReturnType*>(CreateDefaultSubobject(SubobjectName, TReturnType::StaticClass(), TClassToConstructByDefault::StaticClass(), /*bIsRequired =*/ true, /*bIsAbstract =*/ false, bTransient));
+	}
+	
+	/**
+	* Create optional component or subobject. Optional subobjects may not get created
+	* when a derived class specified DoNotCreateDefaultSubobject with the subobject's name.
+	* @param	TReturnType					class of return type, all overrides must be of this type
+	* @param	Outer						outer to construct the subobject in
+	* @param	SubobjectName				name of the new component
+	* @param bTransient		true if the component is being assigned to a transient property
+	*/
+	template<class TReturnType>
+	TReturnType* CreateOptionalDefaultSubobject(FName SubobjectName, bool bTransient = false)
+	{
+		UClass* ReturnType = TReturnType::StaticClass();
+		return static_cast<TReturnType*>(CreateDefaultSubobject(SubobjectName, ReturnType, ReturnType, /*bIsRequired =*/ false, /*bIsAbstract =*/ false, bTransient));
+	}
+	
+	/**
+	* Create optional component or subobject. Optional subobjects may not get created
+	* when a derived class specified DoNotCreateDefaultSubobject with the subobject's name.
+	* @param	TReturnType					class of return type, all overrides must be of this type
+	* @param	Outer						outer to construct the subobject in
+	* @param	SubobjectName				name of the new component
+	* @param bTransient		true if the component is being assigned to a transient property
+	*/
+	template<class TReturnType>
+	TReturnType* CreateAbstractDefaultSubobject(FName SubobjectName, bool bTransient = false)
+	{
+		UClass* ReturnType = TReturnType::StaticClass();
+		return static_cast<TReturnType*>(CreateDefaultSubobject(SubobjectName, ReturnType, ReturnType, /*bIsRequired =*/ true, /*bIsAbstract =*/ true, bTransient));
+	}
+
 	//==========================================
 	// UObject interface.
 	//==========================================	
@@ -813,6 +887,8 @@ public:
 	 */
 	void DestroyNonNativeProperties();
 
+	virtual void MarkAsEditorOnlySubobject() { }
+
 	// UnrealScript intrinsics.
 
 	// Undefined native handler
@@ -1000,6 +1076,15 @@ protected:
 
 private:
 	void ProcessContextOpcode(FFrame& Stack, RESULT_DECL, bool bCanFailSilent);
+
+	/**
+	* Create a component or subobject only to be used with the editor.
+	* @param	Outer						outer to construct the subobject in
+	* @param	ReturnType					type of the new component
+	* @param	SubobjectName				name of the new component
+	* @param	bTransient					true if the component is being assigned to a transient property
+	*/
+	UObject* CreateEditorOnlyDefaultSubobjectImpl(FName SubobjectName, UClass* ReturnType, bool bTransient = false);
 };
 
 #endif	// __UNOBJUOBJECT_H__
diff --git a/Engine/Source/Runtime/CoreUObject/Public/UObject/UObjectGlobals.h b/Engine/Source/Runtime/CoreUObject/Public/UObject/UObjectGlobals.h
--- a/Engine/Source/Runtime/CoreUObject/Public/UObject/UObjectGlobals.h
+++ b/Engine/Source/Runtime/CoreUObject/Public/UObject/UObjectGlobals.h
@@ -691,27 +691,25 @@ public:
 	 * @param	TReturnType					class of return type, all overrides must be of this type
 	 * @param	Outer						outer to construct the subobject in
 	 * @param	SubobjectName				name of the new component
-	 * @param bTransient		true if the component is being assigned to a transient property
+	 * @param	bTransient					true if the component is being assigned to a transient property
 	 */
 	template<class TReturnType>
 	TReturnType* CreateEditorOnlyDefaultSubobject(UObject* Outer, FName SubobjectName, bool bTransient = false) const
 	{
-#if WITH_EDITOR
-		if (GIsEditor)
-		{
-			UClass* ReturnType = TReturnType::StaticClass();
-			TReturnType* EditorSubobject = static_cast<TReturnType*>(CreateDefaultSubobject(Outer, SubobjectName, ReturnType, ReturnType, /*bIsRequired =*/ false, /*bIsAbstract =*/ false, bTransient));
-			if (EditorSubobject)
-			{
-				EditorSubobject->AlwaysLoadOnClient = false;
-				EditorSubobject->AlwaysLoadOnServer = false;
-			}
-			return EditorSubobject;	
-		}
-#endif
-		return NULL;
+		UClass* ReturnType = TReturnType::StaticClass();
+		return static_cast<TReturnType*>(CreateEditorOnlyDefaultSubobject(Outer, SubobjectName, ReturnType, bTransient));
 	}
 
+	/**
+	* Create a component or subobject only to be used with the editor.
+	* @param	TReturnType					class of return type, all overrides must be of this type
+	* @param	Outer						outer to construct the subobject in
+	* @param	ReturnType					type of the new component 
+	* @param	SubobjectName				name of the new component
+	* @param	bTransient					true if the component is being assigned to a transient property
+	*/
+	UObject* CreateEditorOnlyDefaultSubobject(UObject* Outer, FName SubobjectName, UClass* ReturnType, bool bTransient = false) const;
+
 	/**
 	 * Create a component or subobject
 	 * @param	TReturnType					class of return type, all overrides must be of this type
diff --git a/Engine/Source/Runtime/Engine/Classes/Components/PrimitiveComponent.h b/Engine/Source/Runtime/Engine/Classes/Components/PrimitiveComponent.h
--- a/Engine/Source/Runtime/Engine/Classes/Components/PrimitiveComponent.h
+++ b/Engine/Source/Runtime/Engine/Classes/Components/PrimitiveComponent.h
@@ -1244,6 +1244,12 @@ public:
 	virtual void PostDuplicate(bool bDuplicateForPIE) override;
 	virtual SIZE_T GetResourceSize(EResourceSizeMode::Type Mode) override;
 
+	virtual void MarkAsEditorOnlySubobject() override
+	{
+		AlwaysLoadOnClient = false;
+		AlwaysLoadOnServer = false;
+	}
+
 #if WITH_EDITOR
 	/**
 	 * Called after importing property values for this object (paste, duplicate or .t3d import)
