Non unity CIS fix.

diff --git a/Engine/Source/Runtime/CoreUObject/Private/UObject/TlsObjectInitializers.cpp b/Engine/Source/Runtime/CoreUObject/Private/UObject/TlsObjectInitializers.cpp
--- a/Engine/Source/Runtime/CoreUObject/Private/UObject/TlsObjectInitializers.cpp
+++ b/Engine/Source/Runtime/CoreUObject/Private/UObject/TlsObjectInitializers.cpp
@@ -1,5 +1,4 @@
 // Copyright 1998-2014 Epic Games, Inc. All Rights Reserved.
-#pragma once
 
 #include "CoreUObjectPrivate.h"
 #include "UObject/TlsObjectInitializers.h"
