Declare FObjectInitializer outside of the constructor declaration; it prevents Doxygen matching up the signature with the function definition.

diff --git a/Engine/Source/Runtime/CoreUObject/Public/UObject/ObjectBase.h b/Engine/Source/Runtime/CoreUObject/Public/UObject/ObjectBase.h
--- a/Engine/Source/Runtime/CoreUObject/Public/UObject/ObjectBase.h
+++ b/Engine/Source/Runtime/CoreUObject/Public/UObject/ObjectBase.h
@@ -437,6 +437,7 @@ FORCEINLINE void operator|=(EObjectFlags& Dest,EObjectFlags Arg)
 // forward declarations
 class UObject;
 class UProperty;
+class FObjectInitializer; 
 
 struct COREUOBJECT_API FReferencerInformation 
 {
@@ -1182,10 +1183,10 @@ public: \
 	static_assert(false, "You have to define " #TClass "::" #TClass "() or " #TClass "::" #TClass "(const FObjectInitializer&). This is required by UObject system to work correctly.");
 
 #define DEFINE_DEFAULT_CONSTRUCTOR_CALL(TClass) \
-	static void __DefaultConstructor(const class FObjectInitializer& X) { const_cast<FObjectInitializer&>(X).FinalizeSubobjectClassInitialization(); new((EInternal*)X.GetObj())TClass(); }
+	static void __DefaultConstructor(const FObjectInitializer& X) { const_cast<FObjectInitializer&>(X).FinalizeSubobjectClassInitialization(); new((EInternal*)X.GetObj())TClass(); }
 
 #define DEFINE_DEFAULT_OBJECT_INITIALIZER_CONSTRUCTOR_CALL(TClass) \
-	static void __DefaultConstructor(const class FObjectInitializer& X) { new((EInternal*)X.GetObj())TClass(X); }
+	static void __DefaultConstructor(const FObjectInitializer& X) { new((EInternal*)X.GetObj())TClass(X); }
 
 #define DECLARE_CLASS_INTRINSIC(TClass,TSuperClass,TStaticFlags,TPackage) \
 	DECLARE_CLASS(TClass,TSuperClass,TStaticFlags|CLASS_Intrinsic,CASTCLASS_None,TPackage,NO_API ) \
