Add checks to forbid improper subobject initialization.

diff --git a/Engine/Source/Runtime/CoreUObject/Private/UObject/Obj.cpp b/Engine/Source/Runtime/CoreUObject/Private/UObject/Obj.cpp
--- a/Engine/Source/Runtime/CoreUObject/Private/UObject/Obj.cpp
+++ b/Engine/Source/Runtime/CoreUObject/Private/UObject/Obj.cpp
@@ -17,6 +17,8 @@
 
 DEFINE_LOG_CATEGORY(LogObj);
 
+extern int32 GIsInConstructor;
+
 /** Stat group for dynamic objects cycle counters*/
 
 
@@ -52,7 +54,10 @@ UObject::UObject( EStaticConstructor, EObjectFlags InFlags )
 
 UObject* UObject::CreateDefaultSubobject(FName SubobjectFName, UClass* ReturnType, UClass* ClassToCreateByDefault, bool bIsRequired, bool bAbstract, bool bIsTransient)
 {
+	UE_CLOG(!GIsInConstructor, LogObj, Fatal, TEXT("CreateDefultSubobject can only be used inside of UObject constructors"));
 	auto CurrentInitializer = FTlsObjectInitializers::Top();
+	UE_CLOG(!CurrentInitializer, LogObj, Fatal, TEXT("No object initializer found during construction."));
+	UE_CLOG(CurrentInitializer->Obj != this, LogObj, Fatal, TEXT("Using incorrect object initializer."));
 	return CurrentInitializer->CreateDefaultSubobject(this, SubobjectFName, ReturnType, ClassToCreateByDefault, bIsRequired, bAbstract, bIsTransient);
 }
 
diff --git a/Engine/Source/Runtime/CoreUObject/Private/UObject/UObjectGlobals.cpp b/Engine/Source/Runtime/CoreUObject/Private/UObject/UObjectGlobals.cpp
--- a/Engine/Source/Runtime/CoreUObject/Private/UObject/UObjectGlobals.cpp
+++ b/Engine/Source/Runtime/CoreUObject/Private/UObject/UObjectGlobals.cpp
@@ -1850,7 +1850,7 @@ UObject::UObject(const FObjectInitializer& ObjectInitializer)
 }
 
 /* Global flag so that FObjectFinders know if they are called from inside the UObject constructors or not. */
-static int32 GIsInConstructor = 0;
+int32 GIsInConstructor = 0;
 /* Object that is currently being constructed with ObjectInitializer */
 static UObject* GConstructedObject = NULL;
 
