Add a check if the passed in name to CreatePackage is incompatible.

diff --git a/Engine/Source/Runtime/CoreUObject/Private/UObject/UObjectGlobals.cpp b/Engine/Source/Runtime/CoreUObject/Private/UObject/UObjectGlobals.cpp
--- a/Engine/Source/Runtime/CoreUObject/Private/UObject/UObjectGlobals.cpp
+++ b/Engine/Source/Runtime/CoreUObject/Private/UObject/UObjectGlobals.cpp
@@ -450,6 +450,12 @@ UPackage* CreatePackage( UObject* InOuter, const TCHAR* PackageName )
 	{
 		InName = PackageName;
 	}
+
+	if (InName.Contains(TEXT("//")))
+	{
+		UE_LOG(LogUObjectGlobals, Fatal, TEXT("Attempted to create a package with name containing double slashes. PackageName: %s"), PackageName);
+	}
+
 	if( InName.EndsWith( TEXT( "." ) ) )
 	{
 		FString InName2 = InName.Left( InName.Len() - 1 );
