Replace dynamic_cast/check with CastChecked

diff --git a/Engine/Source/Runtime/CoreUObject/Private/UObject/Class.cpp b/Engine/Source/Runtime/CoreUObject/Private/UObject/Class.cpp
--- a/Engine/Source/Runtime/CoreUObject/Private/UObject/Class.cpp
+++ b/Engine/Source/Runtime/CoreUObject/Private/UObject/Class.cpp
@@ -612,8 +612,7 @@ void UStruct::Link(FArchive& Ar, bool bRelinkExistingProperties)
 
 				for( UField* Field=Children; Field && Field->GetOuter()==this; Field=Field->Next )
 				{
-					UProperty* Property = dynamic_cast<UProperty*>( Field );
-					check(Property);
+					UProperty* Property = CastChecked<UProperty*>( Field );
 					ColorComponentEntries[ColorComponentIndex++] = Property;
 				}
 				check( ColorComponentIndex == 4 );
@@ -1023,9 +1022,9 @@ void UStruct::SerializeTaggedProperties(FArchive& Ar, uint8* Data, UStruct* Defa
 				AdvanceProperty = true;
 				continue; 
 			}
-			else if( dynamic_cast<UStructProperty*>(Property) && dynamic_cast<UStructProperty*>(Property)->Struct && (Tag.Type != Property->GetID() || (Tag.Type == NAME_StructProperty && Tag.StructName != dynamic_cast<UStructProperty*>(Property)->Struct->GetFName())) && (dynamic_cast<UStructProperty*>(Property)->Struct->StructFlags & STRUCT_SerializeFromMismatchedTag))
+			else if( dynamic_cast<UStructProperty*>(Property) && static_cast<UStructProperty*>(Property)->Struct && (Tag.Type != Property->GetID() || (Tag.Type == NAME_StructProperty && Tag.StructName != static_cast<UStructProperty*>(Property)->Struct->GetFName())) && (static_cast<UStructProperty*>(Property)->Struct->StructFlags & STRUCT_SerializeFromMismatchedTag))
 			{
-				UScriptStruct::ICppStructOps* CppStructOps = dynamic_cast<UStructProperty*>(Property)->Struct->GetCppStructOps();
+				UScriptStruct::ICppStructOps* CppStructOps = static_cast<UStructProperty*>(Property)->Struct->GetCppStructOps();
 				check(CppStructOps && CppStructOps->HasSerializeFromMismatchedTag()); // else should not have STRUCT_SerializeFromMismatchedTag
 				void* DestAddress = Property->ContainerPtrToValuePtr<void>(Data, Tag.ArrayIndex);  
 				if (CppStructOps->SerializeFromMismatchedTag(Tag, Ar, DestAddress))
@@ -1042,9 +1041,9 @@ void UStruct::SerializeTaggedProperties(FArchive& Ar, uint8* Data, UStruct* Defa
 			{
 				UE_LOG(LogClass, Warning, TEXT("Type mismatch in %s of %s - Previous (%s) Current(%s) for package:  %s"), *Tag.Name.ToString(), *GetName(), *Tag.Type.ToString(), *Property->GetID().ToString(), *Ar.GetArchiveName() );
 			}
-			else if( Tag.Type == NAME_ArrayProperty && Tag.InnerType != NAME_None && Tag.InnerType != dynamic_cast<UArrayProperty&>(*Property).Inner->GetID() )
+			else if( Tag.Type == NAME_ArrayProperty && Tag.InnerType != NAME_None && Tag.InnerType != CastChecked<UArrayProperty&>(*Property).Inner->GetID() )
 			{
-				UArrayProperty* ArrayProperty = dynamic_cast<UArrayProperty*>(Property);
+				UArrayProperty* ArrayProperty = CastChecked<UArrayProperty*>(Property);
 				void* ArrayPropertyData = ArrayProperty->ContainerPtrToValuePtr<void>(Data);
 
 				int32 ElementCount = 0;
@@ -1073,7 +1072,7 @@ void UStruct::SerializeTaggedProperties(FArchive& Ar, uint8* Data, UStruct* Defa
 						FText Text;  
 						Ar << Text;
 						FString String = FTextInspector::GetSourceString(Text) ? *FTextInspector::GetSourceString(Text) : TEXT("");
-						CastChecked<UStrProperty>(ArrayProperty->Inner)->SetPropertyValue(ScriptArrayHelper.GetRawPtr(i), String);
+						static_cast<UStrProperty>(ArrayProperty->Inner)->SetPropertyValue(ScriptArrayHelper.GetRawPtr(i), String);
 						AdvanceProperty = true;
 					}
 					continue; 
diff --git a/Engine/Source/Runtime/CoreUObject/Private/UObject/UObjectBaseUtility.cpp b/Engine/Source/Runtime/CoreUObject/Private/UObject/UObjectBaseUtility.cpp
--- a/Engine/Source/Runtime/CoreUObject/Private/UObject/UObjectBaseUtility.cpp
+++ b/Engine/Source/Runtime/CoreUObject/Private/UObject/UObjectBaseUtility.cpp
@@ -116,9 +116,7 @@ UPackage* UObjectBaseUtility::GetOutermost() const
 		UObject* Outer = Top->GetOuter();
 		if (!Outer)
 		{
-			UPackage* Result = dynamic_cast<UPackage*>(Top);
-			check(Result);
-			return Result;
+			return CastChecked<UPackage*>(Top);
 		}
 		Top = Outer;
 	}
diff --git a/Engine/Source/Runtime/Engine/Private/LightMap.cpp b/Engine/Source/Runtime/Engine/Private/LightMap.cpp
--- a/Engine/Source/Runtime/Engine/Private/LightMap.cpp
+++ b/Engine/Source/Runtime/Engine/Private/LightMap.cpp
@@ -248,8 +248,7 @@ struct FLightMapAllocation
 	{
 		if (InstanceIndex >= 0)
 		{
-			UInstancedStaticMeshComponent* Component = dynamic_cast<UInstancedStaticMeshComponent*>(Primitive);
-			check(Component);
+			UInstancedStaticMeshComponent* Component = CastChecked<UInstancedStaticMeshComponent*>(Primitive);
 
 			// TODO: We currently only support one LOD of static lighting in foliage
 			// Need to create per-LOD instance data to fix that
diff --git a/Engine/Source/Runtime/Engine/Private/ShadowMap.cpp b/Engine/Source/Runtime/Engine/Private/ShadowMap.cpp
--- a/Engine/Source/Runtime/Engine/Private/ShadowMap.cpp
+++ b/Engine/Source/Runtime/Engine/Private/ShadowMap.cpp
@@ -91,8 +91,7 @@ struct FShadowMapAllocation
 	{
 		if (InstanceIndex >= 0)
 		{
-			UInstancedStaticMeshComponent* Component = dynamic_cast<UInstancedStaticMeshComponent*>(Primitive);
-			check(Component);
+			UInstancedStaticMeshComponent* Component = CastChecked<UInstancedStaticMeshComponent*>(Primitive);
 
 			// TODO: We currently only support one LOD of static lighting in foliage
 			// Need to create per-LOD instance data to fix that
