CastChecked no longer Casts in Shipping builds.

diff --git a/Engine/Source/Runtime/CoreUObject/Public/Templates/Casts.h b/Engine/Source/Runtime/CoreUObject/Public/Templates/Casts.h
--- a/Engine/Source/Runtime/CoreUObject/Public/Templates/Casts.h
+++ b/Engine/Source/Runtime/CoreUObject/Public/Templates/Casts.h
@@ -144,27 +144,34 @@ FORCEINLINE T* ExactCast( UObject* Src )
 template <typename To, typename From>
 To* CastChecked(From* Src, ECastCheckedType::Type CheckType = ECastCheckedType::NullChecked)
 {
-	To* Result = Cast<To>(Src);
-
 	#if !(UE_BUILD_SHIPPING || UE_BUILD_TEST)
 
 		if (Src)
 		{
+			To* Result = Cast<To>(Src);
 			if (!Result)
 			{
 				CastLogError(*Cast<UObject>(Src)->GetFullName(), *GetTypeName<To>());
 			}
+
+			return Result;
 		}
-		else if (CheckType == ECastCheckedType::NullChecked)
+
+		if (CheckType == ECastCheckedType::NullChecked)
 		{
 			CastLogError(TEXT("nullptr"), *GetTypeName<To>());
 		}
 
-	#endif
+		return nullptr;
 
-	return Result;
+	#else
+
+		return (To*)Src;
+
+	#endif
 }
 
+
 template <typename InterfaceType>
 DEPRECATED(4.6, "InterfaceCast is deprecated, use Cast or dynamic_cast instead.")
 FORCEINLINE InterfaceType* InterfaceCast(UObject* Src)
