Write correct serialized header for unversioned case (was still writing a zero for UE3 version)
#ue4

diff --git a/Engine/Source/Runtime/CoreUObject/Private/UObject/PackageFileSummary.cpp b/Engine/Source/Runtime/CoreUObject/Private/UObject/PackageFileSummary.cpp
--- a/Engine/Source/Runtime/CoreUObject/Private/UObject/PackageFileSummary.cpp
+++ b/Engine/Source/Runtime/CoreUObject/Private/UObject/PackageFileSummary.cpp
@@ -84,7 +84,6 @@ FArchive& operator<<( FArchive& Ar, FPackageFileSummary& Sum )
 				int32 Zero = 0;
 				Ar << Zero;
 				Ar << Zero;
-				Ar << Zero;
 
 				FCustomVersionContainer NoCustomVersions;
 				NoCustomVersions.Serialize(Ar);
