More async loading stats

diff --git a/Engine/Source/Runtime/CoreUObject/Private/Serialization/AsyncLoading.cpp b/Engine/Source/Runtime/CoreUObject/Private/Serialization/AsyncLoading.cpp
--- a/Engine/Source/Runtime/CoreUObject/Private/Serialization/AsyncLoading.cpp
+++ b/Engine/Source/Runtime/CoreUObject/Private/Serialization/AsyncLoading.cpp
@@ -16,6 +16,7 @@ DECLARE_MEMORY_STAT(TEXT("Streaming Memory Used"),STAT_StreamingAllocSize,STATGR
 DECLARE_STATS_GROUP_VERBOSE(TEXT("Async Load"), STATGROUP_AsyncLoad, STATCAT_Advanced);
 
 DECLARE_CYCLE_STAT(TEXT("Tick AsyncPackage"),STAT_FAsyncPackage_Tick,STATGROUP_AsyncLoad);
+DECLARE_FLOAT_ACCUMULATOR_STAT(TEXT("Tick AsyncPackage Time"), STAT_FAsyncPackage_TickTime, STATGROUP_AsyncLoad);
 DECLARE_CYCLE_STAT(TEXT("CreateLinker AsyncPackage"),STAT_FAsyncPackage_CreateLinker,STATGROUP_AsyncLoad);
 DECLARE_CYCLE_STAT(TEXT("FinishLinker AsyncPackage"),STAT_FAsyncPackage_FinishLinker,STATGROUP_AsyncLoad);
 DECLARE_CYCLE_STAT(TEXT("LoadImports AsyncPackage"),STAT_FAsyncPackage_LoadImports,STATGROUP_AsyncLoad);
@@ -23,8 +24,10 @@ DECLARE_CYCLE_STAT(TEXT("CreateImports AsyncPackage"),STAT_FAsyncPackage_CreateI
 DECLARE_CYCLE_STAT(TEXT("FinishTextureAllocations AsyncPackage"),STAT_FAsyncPackage_FinishTextureAllocations,STATGROUP_AsyncLoad);
 DECLARE_CYCLE_STAT(TEXT("CreateExports AsyncPackage"),STAT_FAsyncPackage_CreateExports,STATGROUP_AsyncLoad);
 DECLARE_CYCLE_STAT(TEXT("PreLoadObjects AsyncPackage"),STAT_FAsyncPackage_PreLoadObjects,STATGROUP_AsyncLoad);
+DECLARE_FLOAT_ACCUMULATOR_STAT(TEXT("PostLoadObjects AsyncPackage Time"), STAT_FAsyncPackage_PostLoadObjectsTime, STATGROUP_AsyncLoad);
 DECLARE_CYCLE_STAT(TEXT("PostLoadObjects AsyncPackage"),STAT_FAsyncPackage_PostLoadObjects,STATGROUP_AsyncLoad);
 DECLARE_CYCLE_STAT(TEXT("FinishObjects AsyncPackage"),STAT_FAsyncPackage_FinishObjects,STATGROUP_AsyncLoad);
+DECLARE_FLOAT_ACCUMULATOR_STAT(TEXT("Flush Async Loading Time"), STAT_FAsyncPackage_FlushAsyncLoadingTime, STATGROUP_AsyncLoad);
 
 DECLARE_CYCLE_STAT(TEXT("Async Loading Time"),STAT_AsyncLoadingTime,STATGROUP_AsyncLoad);
 
@@ -283,89 +286,96 @@ EAsyncPackageState::Type FAsyncPackage::Tick( bool InbUseTimeLimit, bool InbUseF
 	// Whether we should execute the next step.
 	EAsyncPackageState::Type LoadingState	= EAsyncPackageState::Complete;
 
-	// Make sure we finish our work if there's no time limit. The loop is required as PostLoad
-	// might cause more objects to be loaded in which case we need to Preload them again.
-	do
+	STAT(double ThisTime = 0);
 	{
-		// Reset value to true at beginning of loop.
-		LoadingState	= EAsyncPackageState::Complete;
+		SCOPE_SECONDS_COUNTER(ThisTime);
 
-		// Begin async loading, simulates BeginLoad and sets GIsAsyncLoading to true.
-		BeginAsyncLoad();
-
-		// Create raw linker. Needs to be async created via ticking before it can be used.
-		if( LoadingState == EAsyncPackageState::Complete )
+		// Make sure we finish our work if there's no time limit. The loop is required as PostLoad
+		// might cause more objects to be loaded in which case we need to Preload them again.
+		do
 		{
-			LoadingState = CreateLinker();
-		}
+			// Reset value to true at beginning of loop.
+			LoadingState	= EAsyncPackageState::Complete;
 
-		// Async create linker.
-		if( LoadingState == EAsyncPackageState::Complete )
-		{
-			LoadingState = FinishLinker();
-		}
+			// Begin async loading, simulates BeginLoad and sets GIsAsyncLoading to true.
+			BeginAsyncLoad();
 
-		// Load imports from linker import table asynchronously.
-		if( LoadingState == EAsyncPackageState::Complete )
-		{
-			LoadingState = LoadImports();
-		}
+			// Create raw linker. Needs to be async created via ticking before it can be used.
+			if( LoadingState == EAsyncPackageState::Complete )
+			{
+				LoadingState = CreateLinker();
+			}
 
-		// Create imports from linker import table.
-		if( LoadingState == EAsyncPackageState::Complete )
-		{
-			LoadingState = CreateImports();
-		}
+			// Async create linker.
+			if( LoadingState == EAsyncPackageState::Complete )
+			{
+				LoadingState = FinishLinker();
+			}
 
-		// Finish all async texture allocations.
-		if( LoadingState == EAsyncPackageState::Complete )
-		{
-			LoadingState = FinishTextureAllocations();
-		}
+			// Load imports from linker import table asynchronously.
+			if( LoadingState == EAsyncPackageState::Complete )
+			{
+				LoadingState = LoadImports();
+			}
 
-		// Create exports from linker export table and also preload them.
-		if( LoadingState == EAsyncPackageState::Complete )
-		{
-			LoadingState = CreateExports();
-		}
+			// Create imports from linker import table.
+			if( LoadingState == EAsyncPackageState::Complete )
+			{
+				LoadingState = CreateImports();
+			}
 
-		// Call Preload on the linker for all loaded objects which causes actual serialization.
-		if( LoadingState == EAsyncPackageState::Complete )
-		{
-			LoadingState = PreLoadObjects();
-		}
+			// Finish all async texture allocations.
+			if( LoadingState == EAsyncPackageState::Complete )
+			{
+				LoadingState = FinishTextureAllocations();
+			}
 
-		// Call PostLoad on objects, this could cause new objects to be loaded that require
-		// another iteration of the PreLoad loop.
-		if (LoadingState == EAsyncPackageState::Complete)
-		{
-			LoadingState = PostLoadObjects();
-		}
+			// Create exports from linker export table and also preload them.
+			if( LoadingState == EAsyncPackageState::Complete )
+			{
+				LoadingState = CreateExports();
+			}
 
-		// End async loading, simulates EndLoad and sets GIsAsyncLoading to false.
-		EndAsyncLoad();
+			// Call Preload on the linker for all loaded objects which causes actual serialization.
+			if( LoadingState == EAsyncPackageState::Complete )
+			{
+				LoadingState = PreLoadObjects();
+			}
 
-		// Finish objects (removing RF_AsyncLoading, dissociate imports and forced exports, 
-		// call completion callback, ...
-		// If the load has failed, perform completion callbacks and then quit
-		if( LoadingState == EAsyncPackageState::Complete || bLoadHasFailed )
-		{
-			LoadingState = FinishObjects();
-		}
-	// Only exit loop if we're either done or the time limit has been exceeded.
-	} while( !IsTimeLimitExceeded() && LoadingState == EAsyncPackageState::TimeOut );	
+			// Call PostLoad on objects, this could cause new objects to be loaded that require
+			// another iteration of the PreLoad loop.
+			if (LoadingState == EAsyncPackageState::Complete)
+			{
+				LoadingState = PostLoadObjects();
+			}
 
-	check( bUseTimeLimit || LoadingState != EAsyncPackageState::TimeOut );
+			// End async loading, simulates EndLoad and sets GIsAsyncLoading to false.
+			EndAsyncLoad();
 
-	// We can't have a reference to a UObject.
-	LastObjectWorkWasPerformedOn	= nullptr;
-	// Reset type of work performed.
-	LastTypeOfWorkPerformed			= nullptr;
-	// Mark this package as loaded if everything completed.
-	bLoadHasFinished = (LoadingState == EAsyncPackageState::Complete);
-	// Subtract the time it took to load this package from the global limit.
-	InOutTimeLimit = (float)FMath::Max(0.0, InOutTimeLimit - (FPlatformTime::Seconds() - TickStartTime));
+			// Finish objects (removing RF_AsyncLoading, dissociate imports and forced exports, 
+			// call completion callback, ...
+			// If the load has failed, perform completion callbacks and then quit
+			if( LoadingState == EAsyncPackageState::Complete || bLoadHasFailed )
+			{
+				LoadingState = FinishObjects();
+			}
+		// Only exit loop if we're either done or the time limit has been exceeded.
+		} while( !IsTimeLimitExceeded() && LoadingState == EAsyncPackageState::TimeOut );	
+
+		check( bUseTimeLimit || LoadingState != EAsyncPackageState::TimeOut );
+
+		// We can't have a reference to a UObject.
+		LastObjectWorkWasPerformedOn	= nullptr;
+		// Reset type of work performed.
+		LastTypeOfWorkPerformed			= nullptr;
+		// Mark this package as loaded if everything completed.
+		bLoadHasFinished = (LoadingState == EAsyncPackageState::Complete);
+		// Subtract the time it took to load this package from the global limit.
+		InOutTimeLimit = (float)FMath::Max(0.0, InOutTimeLimit - (FPlatformTime::Seconds() - TickStartTime));
 	
+	}
+	INC_FLOAT_STAT_BY(STAT_FAsyncPackage_TickTime, (float)ThisTime);
+
 	// true means that we're done loading this package.
 	return LoadingState;
 }
@@ -816,28 +826,29 @@ EAsyncPackageState::Type FAsyncPackage::PostLoadObjects()
 {
 	SCOPE_CYCLE_COUNTER(STAT_FAsyncPackage_PostLoadObjects);
 	
-	TGuardValue<bool> GuardIsRoutingPostLoad(GIsRoutingPostLoad, true);
-
-	// PostLoad objects.
-	while( PostLoadIndex < GObjLoaded.Num() && PostLoadIndex < PreLoadIndex && !IsTimeLimitExceeded() )
+	EAsyncPackageState::Type Result = EAsyncPackageState::Complete;
+	STAT(double ThisTime = 0);
 	{
-		UObject* Object	= GObjLoaded[ PostLoadIndex++ ];
-		check(Object);
-		FScopeCycleCounterUObject ConstructorScope(Object, GET_STATID(STAT_FAsyncPackage_PostLoadObjects));
+		SCOPE_SECONDS_COUNTER(ThisTime);
+		TGuardValue<bool> GuardIsRoutingPostLoad(GIsRoutingPostLoad, true);
 
-		Object->ConditionalPostLoad();
+		// PostLoad objects.
+		while( PostLoadIndex < GObjLoaded.Num() && PostLoadIndex < PreLoadIndex && !IsTimeLimitExceeded() )
+		{
+			UObject* Object	= GObjLoaded[ PostLoadIndex++ ];
+			check(Object);
+			FScopeCycleCounterUObject ConstructorScope(Object, GET_STATID(STAT_FAsyncPackage_PostLoadObjects));
 
-		LastObjectWorkWasPerformedOn	= Object;
-		LastTypeOfWorkPerformed			= TEXT("postloading");
-	}
+			Object->ConditionalPostLoad();
 
-	// New objects might have been loaded during PostLoad.
-	EAsyncPackageState::Type Result = (PreLoadIndex == GObjLoaded.Num()) && (PostLoadIndex == GObjLoaded.Num()) ? EAsyncPackageState::Complete : EAsyncPackageState::TimeOut;
-	if (PreLoadIndex < GObjLoaded.Num())
-	{
-		static volatile int32 xx = 0;
-		xx++;
+			LastObjectWorkWasPerformedOn	= Object;
+			LastTypeOfWorkPerformed			= TEXT("postloading");
+		}
+
+		// New objects might have been loaded during PostLoad.
+		Result = (PreLoadIndex == GObjLoaded.Num()) && (PostLoadIndex == GObjLoaded.Num()) ? EAsyncPackageState::Complete : EAsyncPackageState::TimeOut;
 	}
+	INC_FLOAT_STAT_BY(STAT_FAsyncPackage_PostLoadObjectsTime, (float)ThisTime);
 
 	return Result;
 }
@@ -1025,7 +1036,12 @@ void FlushAsyncLoading(FName ExcludeType/*=NAME_None*/)
 
 		// Flush async loaders by not using a time limit. Needed for e.g. garbage collection.
 		UE_LOG(LogStreaming, Log,  TEXT("Flushing async loaders.") );
-		while (ProcessAsyncLoading( false, false, 0, ExcludeType ) == EAsyncPackageState::PendingImports) {}
+		STAT(double ThisTime = 0);
+		{
+			SCOPE_SECONDS_COUNTER(ThisTime);
+			while (ProcessAsyncLoading(false, false, 0, ExcludeType) == EAsyncPackageState::PendingImports) {}
+		}
+		INC_FLOAT_STAT_BY(STAT_FAsyncPackage_FlushAsyncLoadingTime, (float)ThisTime);
 		UE_LOG(LogStreaming, Log,  TEXT("Flushed async loaders.") );
 
 		if (ExcludeType == NAME_None)
