CPF_PersistentInstance comment

diff --git a/Engine/Source/Runtime/CoreUObject/Public/UObject/ObjectBase.h b/Engine/Source/Runtime/CoreUObject/Public/UObject/ObjectBase.h
--- a/Engine/Source/Runtime/CoreUObject/Public/UObject/ObjectBase.h
+++ b/Engine/Source/Runtime/CoreUObject/Public/UObject/ObjectBase.h
@@ -328,7 +328,7 @@ typedef uint64 EClassCastFlags;
 #define CPF_TextExportTransient				DECLARE_UINT64(0x0000400000000000)		// Property shouldn't be exported to text format (e.g. copy/paste)
 #define CPF_NonPIEDuplicateTransient		DECLARE_UINT64(0x0000800000000000)		// Property should only be copied in PIE
 #define CPF_ExposeOnSpawn					DECLARE_UINT64(0x0001000000000000)		// Property is exposed on spawn
-#define CPF_PersistentInstance				DECLARE_UINT64(0x0002000000000000)		//
+#define CPF_PersistentInstance				DECLARE_UINT64(0x0002000000000000)		// A object referenced by the property is duplicated like a component. (Each actor should have an own instance.)
 
 #define CPF_NonPIETransient \
 	EMIT_DEPRECATED_WARNING_MESSAGE("CPF_NonPIETransient is deprecated. Please use CPF_NonPIEDuplicateTransient instead.") \
