ConstructorHelpers warning messages will now also display information about the class that caused the warning.

UE-4590

diff --git a/Engine/Source/Runtime/CoreUObject/Private/UObject/UObjectGlobals.cpp b/Engine/Source/Runtime/CoreUObject/Private/UObject/UObjectGlobals.cpp
--- a/Engine/Source/Runtime/CoreUObject/Private/UObject/UObjectGlobals.cpp
+++ b/Engine/Source/Runtime/CoreUObject/Private/UObject/UObjectGlobals.cpp
@@ -2244,8 +2244,12 @@ void FScopedObjectFlagMarker::RestoreObjectFlags()
 
 void ConstructorHelpers::FailedToFind(const TCHAR* ObjectToFind)
 {
-	FPlatformMisc::LowLevelOutputDebugStringf(TEXT("CDO Constructor: Failed to find %s\n"), ObjectToFind);
-	UClass::GetDefaultPropertiesFeedbackContext().Logf(ELogVerbosity::Error, TEXT("CDO Constructor: Failed to find %s"), ObjectToFind);
+	auto CurrentInitializer = FTlsObjectInitializers::Top();
+	const FString Message = FString::Printf(TEXT("CDO Constructor (%s): Failed to find %s\n"),
+		(CurrentInitializer && CurrentInitializer->GetClass()) ? *CurrentInitializer->GetClass()->GetName() : TEXT("Unknown"),
+		ObjectToFind);
+	FPlatformMisc::LowLevelOutputDebugString(*Message);
+	UClass::GetDefaultPropertiesFeedbackContext().Log(ELogVerbosity::Error, *Message);
 }
 
 void ConstructorHelpers::CheckFoundViaRedirect(UObject *Object, const FString& PathName, const TCHAR* ObjectToFind)
@@ -2256,8 +2260,14 @@ void ConstructorHelpers::CheckFoundViaRedirect(UObject *Object, const FString& P
 		FString NewString = Object->GetFullName();
 		NewString.ReplaceInline(TEXT(" "), TEXT("'"));
 		NewString += TEXT("'");
-		FPlatformMisc::LowLevelOutputDebugStringf(TEXT("CDO Constructor: Followed redirector (%s), change code to new path (%s)\n"), ObjectToFind, *NewString);
-		UClass::GetDefaultPropertiesFeedbackContext().Logf(ELogVerbosity::Warning, TEXT("CDO Warning: Followed redirector (%s), change code to new path (%s)\n"), ObjectToFind, *NewString);
+
+		auto CurrentInitializer = FTlsObjectInitializers::Top();
+		const FString Message = FString::Printf(TEXT("CDO Constructor (%s): Followed redirector (%s), change code to new path (%s)\n"),
+			(CurrentInitializer && CurrentInitializer->GetClass()) ? *CurrentInitializer->GetClass()->GetName() : TEXT("Unknown"),
+			ObjectToFind, *NewString);
+
+		FPlatformMisc::LowLevelOutputDebugString(*Message);
+		UClass::GetDefaultPropertiesFeedbackContext().Log(ELogVerbosity::Warning, *Message);
 	}
 }
 
