[UE-6279] Allowing for interface pins to be left unconnected (without causing compile errors).

diff --git a/Engine/Source/Editor/KismetCompiler/Private/KismetCompilerVMBackend.cpp b/Engine/Source/Editor/KismetCompiler/Private/KismetCompilerVMBackend.cpp
--- a/Engine/Source/Editor/KismetCompiler/Private/KismetCompilerVMBackend.cpp
+++ b/Engine/Source/Editor/KismetCompiler/Private/KismetCompilerVMBackend.cpp
@@ -560,7 +560,11 @@ public:
 				{
 					Writer << EX_Self;
 				}
-				else 
+				else if (Term->ObjectLiteral == nullptr)
+				{
+					Writer << EX_NoInterface;
+				}
+				else
 				{
 					ensureMsgf(false, TEXT("It is not possible to express this interface property as a literal value! (%s)"), *CoerceProperty->GetFullName());
 				}
diff --git a/Engine/Source/Editor/UnrealEd/Private/ScriptDisassembler.cpp b/Engine/Source/Editor/UnrealEd/Private/ScriptDisassembler.cpp
--- a/Engine/Source/Editor/UnrealEd/Private/ScriptDisassembler.cpp
+++ b/Engine/Source/Editor/UnrealEd/Private/ScriptDisassembler.cpp
@@ -451,6 +451,11 @@ void FKismetBytecodeDisassembler::ProcessCommon(int32& ScriptIndex, EExprToken O
 			Ar.Logf(TEXT("%s $%X: EX_NoObject"), *Indents, (int32)Opcode);
 			break;
 		}
+	case EX_NoInterface:
+		{
+			Ar.Logf(TEXT("%s $%X: EX_NoObject"), *Indents, (int32)Opcode);
+			break;
+		}
 	case EX_Self:
 		{
 			Ar.Logf(TEXT("%s $%X: EX_Self"), *Indents, (int32)Opcode);
diff --git a/Engine/Source/Runtime/CoreUObject/Private/UObject/ScriptCore.cpp b/Engine/Source/Runtime/CoreUObject/Private/UObject/ScriptCore.cpp
--- a/Engine/Source/Runtime/CoreUObject/Private/UObject/ScriptCore.cpp
+++ b/Engine/Source/Runtime/CoreUObject/Private/UObject/ScriptCore.cpp
@@ -1861,6 +1861,13 @@ void UObject::execNoObject( FFrame& Stack, RESULT_DECL )
 }
 IMPLEMENT_VM_FUNCTION( EX_NoObject, execNoObject );
 
+void UObject::execNullInterface(FFrame& Stack, RESULT_DECL)
+{
+	FScriptInterface& InterfaceValue = *(FScriptInterface*)Result;
+	InterfaceValue.SetObject(nullptr);
+}
+IMPLEMENT_VM_FUNCTION( EX_NoInterface, execNullInterface );
+
 void UObject::execIntConstByte( FFrame& Stack, RESULT_DECL )
 {
 	*(int32*)Result = *Stack.Code++;
diff --git a/Engine/Source/Runtime/CoreUObject/Public/UObject/Script.h b/Engine/Source/Runtime/CoreUObject/Public/UObject/Script.h
--- a/Engine/Source/Runtime/CoreUObject/Public/UObject/Script.h
+++ b/Engine/Source/Runtime/CoreUObject/Public/UObject/Script.h
@@ -179,7 +179,7 @@ enum EExprToken
 	EX_NoObject				= 0x2A,	// NoObject.
 	EX_TransformConst		= 0x2B, // A transform constant
 	EX_IntConstByte			= 0x2C,	// Int constant that requires 1 byte.
-	//						= 0x2D,
+	EX_NoInterface			= 0x2D, // A null interface (similar to EX_NoObject, but for interfaces)
 	EX_DynamicCast			= 0x2E,	// Safe dynamic class casting.
 	EX_StructConst			= 0x2F, // An arbitrary UStruct constant
 	EX_EndStructConst		= 0x30, // End of UStruct constant
diff --git a/Engine/Source/Runtime/CoreUObject/Public/UObject/ScriptSerialization.h b/Engine/Source/Runtime/CoreUObject/Public/UObject/ScriptSerialization.h
--- a/Engine/Source/Runtime/CoreUObject/Public/UObject/ScriptSerialization.h
+++ b/Engine/Source/Runtime/CoreUObject/Public/UObject/ScriptSerialization.h
@@ -179,6 +179,7 @@
 		case EX_True:
 		case EX_False:
 		case EX_NoObject:
+		case EX_NoInterface:
 		case EX_Self:
 		case EX_EndParmValue:
 		case EX_PopExecutionFlow:
diff --git a/Engine/Source/Runtime/CoreUObject/Public/UObject/UObject.h b/Engine/Source/Runtime/CoreUObject/Public/UObject/UObject.h
--- a/Engine/Source/Runtime/CoreUObject/Public/UObject/UObject.h
+++ b/Engine/Source/Runtime/CoreUObject/Public/UObject/UObject.h
@@ -1004,6 +1004,7 @@ public:
 	DECLARE_FUNCTION(execTrue);
 	DECLARE_FUNCTION(execFalse);
 	DECLARE_FUNCTION(execNoObject);
+	DECLARE_FUNCTION(execNullInterface);
 	DECLARE_FUNCTION(execIntConstByte);
 	DECLARE_FUNCTION(execRotationConst);
 	DECLARE_FUNCTION(execVectorConst);
