XBoxOne CIS fix.

diff --git a/Engine/Source/Programs/UnrealHeaderTool/Private/CodeGenerator.cpp b/Engine/Source/Programs/UnrealHeaderTool/Private/CodeGenerator.cpp
--- a/Engine/Source/Programs/UnrealHeaderTool/Private/CodeGenerator.cpp
+++ b/Engine/Source/Programs/UnrealHeaderTool/Private/CodeGenerator.cpp
@@ -3886,7 +3886,7 @@ FString FNativeClassHeaderGenerator::GetFunctionParameterString(UFunction* Funct
 		}
 
 		auto Dim = GArrayDimensions.Find(Property);
-		Property->ExportCppDeclaration(PropertyText, EExportedDeclaration::Parameter, Dim ? **Dim : NULL);
+		Property->ExportCppDeclaration(PropertyText, EExportedDeclaration::Parameter, Dim ? **Dim : NULL, 0, true);
 		ApplyAlternatePropertyExportText(Property, PropertyText);
 
 		ParameterList += PropertyText;
diff --git a/Engine/Source/Runtime/CoreUObject/Private/UObject/Property.cpp b/Engine/Source/Runtime/CoreUObject/Private/UObject/Property.cpp
--- a/Engine/Source/Runtime/CoreUObject/Private/UObject/Property.cpp
+++ b/Engine/Source/Runtime/CoreUObject/Private/UObject/Property.cpp
@@ -463,7 +463,7 @@ FString UProperty::GetCPPMacroType( FString& ExtendedTypeText ) const
 	return TEXT("PROPERTY");
 }
 
-void UProperty::ExportCppDeclaration(FOutputDevice& Out, EExportedDeclaration::Type DeclarationType, const TCHAR* ArrayDimOverride, uint32 AdditionalExportCPPFlags) const
+void UProperty::ExportCppDeclaration(FOutputDevice& Out, EExportedDeclaration::Type DeclarationType, const TCHAR* ArrayDimOverride, uint32 AdditionalExportCPPFlags, bool bSkipParameterName) const
 {
 	const bool bIsParameter = (DeclarationType == EExportedDeclaration::Parameter) || (DeclarationType == EExportedDeclaration::MacroParameter);
 	const bool bIsInterfaceProp = dynamic_cast<const UInterfaceProperty*>(this) != nullptr;
@@ -495,7 +495,7 @@ void UProperty::ExportCppDeclaration(FOutputDevice& Out, EExportedDeclaration::T
 		}
 	}
 
-	FString NameCpp = GetNameCPP();
+	FString NameCpp = bSkipParameterName ? FString() : GetNameCPP();
 	if (DeclarationType == EExportedDeclaration::MacroParameter)
 	{
 		NameCpp = FString(TEXT(", ")) + NameCpp;
diff --git a/Engine/Source/Runtime/CoreUObject/Public/UObject/UnrealType.h b/Engine/Source/Runtime/CoreUObject/Public/UObject/UnrealType.h
--- a/Engine/Source/Runtime/CoreUObject/Public/UObject/UnrealType.h
+++ b/Engine/Source/Runtime/CoreUObject/Public/UObject/UnrealType.h
@@ -182,7 +182,7 @@ public:
 											FOutputDevice* Warn, TArray<struct FDefinedProperty>& DefinedProperties );
 
 	// UHT interface
-	void ExportCppDeclaration(FOutputDevice& Out, EExportedDeclaration::Type DeclarationType, const TCHAR* ArrayDimOverride = NULL, uint32 AdditionalExportCPPFlags = 0) const;
+	void ExportCppDeclaration(FOutputDevice& Out, EExportedDeclaration::Type DeclarationType, const TCHAR* ArrayDimOverride = NULL, uint32 AdditionalExportCPPFlags = 0, bool bSkipParameterName = false) const;
 	virtual FString GetCPPMacroType( FString& ExtendedTypeText ) const;
 	virtual bool PassCPPArgsByRef() const { return false; }
 
