#stats - Add stats counter for UObjectBaseUtility::IsA

diff --git a/Engine/Source/Runtime/CoreUObject/Private/UObject/UObjectBaseUtility.cpp b/Engine/Source/Runtime/CoreUObject/Private/UObject/UObjectBaseUtility.cpp
--- a/Engine/Source/Runtime/CoreUObject/Private/UObject/UObjectBaseUtility.cpp
+++ b/Engine/Source/Runtime/CoreUObject/Private/UObject/UObjectBaseUtility.cpp
@@ -7,6 +7,8 @@
 #include "CoreUObjectPrivate.h"
 #include "Interface.h"
 
+DEFINE_STAT(STAT_UObjectIsA);
+
 /***********************/
 /******** Names ********/
 /***********************/
@@ -249,6 +251,7 @@ bool UObjectBaseUtility::RootPackageHasAnyFlags( uint32 CheckFlagMask ) const
  */
 bool UObjectBaseUtility::IsA( const UClass* SomeBase ) const
 {
+	SCOPE_CYCLE_COUNTER(STAT_UObjectIsA);
 	UE_CLOG(!SomeBase, LogObj, Fatal, TEXT("IsA(NULL) cannot yield meaningful results"));
 
 	for ( UClass* TempClass=GetClass(); TempClass; TempClass=TempClass->GetSuperClass() )
diff --git a/Engine/Source/Runtime/CoreUObject/Public/UObject/UObjectBaseUtility.h b/Engine/Source/Runtime/CoreUObject/Public/UObject/UObjectBaseUtility.h
--- a/Engine/Source/Runtime/CoreUObject/Public/UObject/UObjectBaseUtility.h
+++ b/Engine/Source/Runtime/CoreUObject/Public/UObject/UObjectBaseUtility.h
@@ -6,6 +6,8 @@
 
 #pragma once
 
+DECLARE_CYCLE_STAT_EXTERN(TEXT("UObjectBaseUtility::IsA"), STAT_UObjectIsA, STATGROUP_Engine, COREUOBJECT_API);
+
 class COREUOBJECT_API UObjectBaseUtility : public UObjectBase
 {
 public:
