Adding support for cancelling async loading

diff --git a/Engine/Source/Editor/UnrealEd/Private/Tests/EditorAutomationTests.cpp b/Engine/Source/Editor/UnrealEd/Private/Tests/EditorAutomationTests.cpp
--- a/Engine/Source/Editor/UnrealEd/Private/Tests/EditorAutomationTests.cpp
+++ b/Engine/Source/Editor/UnrealEd/Private/Tests/EditorAutomationTests.cpp
@@ -1883,7 +1883,7 @@ namespace StaticMeshUVTest
 		/** 
 		 * Callback when a package has finished loading
 		 */
-		void PackageLoadCallback(const FName& PackageName, UPackage* LoadedPackage)
+		void PackageLoadCallback(const FName& PackageName, UPackage* LoadedPackage, EAsyncLoadingResult::Type Result)
 		{
 			CurrentPackage = LoadedPackage;
 			CurrentState = EStaticMeshUVTestState::PackageLoaded;
diff --git a/Engine/Source/Runtime/CoreUObject/Private/Serialization/AsyncLoading.cpp b/Engine/Source/Runtime/CoreUObject/Private/Serialization/AsyncLoading.cpp
--- a/Engine/Source/Runtime/CoreUObject/Private/Serialization/AsyncLoading.cpp
+++ b/Engine/Source/Runtime/CoreUObject/Private/Serialization/AsyncLoading.cpp
@@ -64,6 +64,17 @@ public:
 	{
 		ReferencedObjects.Empty(ReferencedObjects.Num());
 	}
+	/** Removes all referenced objects and markes them for GC */
+	void EmptyReferencedObjectsAndCancelLoading()
+	{
+		// All of the referenced objects have been created by async loading code and may be in an invalid state so mark them for GC
+		for (auto Object : ReferencedObjects)
+		{
+			Object->ClearFlags(RF_AsyncLoading | RF_NeedLoad | RF_NeedPostLoad | RF_NeedPostLoadSubobjects);
+			Object->SetFlags(RF_PendingKill);
+		}
+		ReferencedObjects.Empty(ReferencedObjects.Num());
+	}
 };
 FAsyncObjectsReferencer& FAsyncObjectsReferencer::Get()
 {
@@ -443,8 +454,9 @@ EAsyncPackageState::Type FAsyncPackage::FinishLinker()
 		LastObjectWorkWasPerformedOn	= Linker->LinkerRoot;
 		LastTypeOfWorkPerformed			= TEXT("ticking linker");
 	
+		const float RemainingTimeLimit = (float)(FPlatformTime::Seconds() - TickStartTime);
 		// Operation still pending if Tick returns false
-		if( Linker->Tick( TimeLimit, bUseTimeLimit, bUseFullTimeLimit ) != ULinkerLoad::LINKER_Loaded)
+		if (Linker->Tick(RemainingTimeLimit, bUseTimeLimit, bUseFullTimeLimit) != ULinkerLoad::LINKER_Loaded)
 		{
 			// Give up remainder of timeslice if there is one to give up.
 			GiveUpTimeSlice();
@@ -664,14 +676,17 @@ EAsyncPackageState::Type FAsyncPackage::LoadImports()
 /**
  * Function called when pending import package has been fully loaded.
  */
-void FAsyncPackage::ImportFullyLoadedCallback(const FName& InPackageName, UPackage* LoadedPackage)
+void FAsyncPackage::ImportFullyLoadedCallback(const FName& InPackageName, UPackage* LoadedPackage, EAsyncLoadingResult::Type Result)
 {
-	UE_LOG(LogStreaming, Verbose, TEXT("FAsyncPackage::LoadImports for %s: Loaded %s"), *PackageNameToLoad.ToString(), *InPackageName.ToString());
-	int32 Index = ContainsDependencyPackage(PendingImportedPackages, InPackageName);
-	check(Index != INDEX_NONE);
-	// Keep a reference to this package so that its linker doesn't go away too soon
-	ReferencedImports.Add(PendingImportedPackages[Index]);
-	PendingImportedPackages.RemoveAt(Index);
+	if (Result != EAsyncLoadingResult::Canceled)
+	{
+		UE_LOG(LogStreaming, Verbose, TEXT("FAsyncPackage::LoadImports for %s: Loaded %s"), *PackageNameToLoad.ToString(), *InPackageName.ToString());
+		int32 Index = ContainsDependencyPackage(PendingImportedPackages, InPackageName);
+		check(Index != INDEX_NONE);
+		// Keep a reference to this package so that its linker doesn't go away too soon
+		ReferencedImports.Add(PendingImportedPackages[Index]);
+		PendingImportedPackages.RemoveAt(Index);
+	}
 }
 
 /** 
@@ -879,6 +894,7 @@ EAsyncPackageState::Type FAsyncPackage::FinishObjects()
 	PreLoadIndex = 0;
 	PostLoadIndex = 0;
 
+	EAsyncLoadingResult::Type LoadingResult = bLoadHasFailed ? EAsyncLoadingResult::Failed : EAsyncLoadingResult::Succeeded;
 
 	// If we successfully loaded
 	if (!bLoadHasFailed)
@@ -893,7 +909,7 @@ EAsyncPackageState::Type FAsyncPackage::FinishObjects()
 		// Call any completion callbacks specified.
 		for (int32 i = 0; i < CompletionCallbacks.Num(); i++)
 		{
-			CompletionCallbacks[i].ExecuteIfBound(PackageName, Linker->LinkerRoot);
+			CompletionCallbacks[i].ExecuteIfBound(PackageName, Linker->LinkerRoot, LoadingResult);
 		}
 
 		// give a hint to the IO system that we are done with this file for now
@@ -917,7 +933,7 @@ EAsyncPackageState::Type FAsyncPackage::FinishObjects()
 		// Call any completion callbacks specified.
 		for (int32 i = 0; i < CompletionCallbacks.Num(); i++)
 		{
-			CompletionCallbacks[i].ExecuteIfBound(PackageName, nullptr);
+			CompletionCallbacks[i].ExecuteIfBound(PackageName, nullptr, LoadingResult);
 		}
 	}
 
@@ -935,6 +951,26 @@ EAsyncPackageState::Type FAsyncPackage::FinishObjects()
 	return EAsyncPackageState::Complete;
 }
 
+void FAsyncPackage::Cancel()
+{
+	// Call any completion callbacks specified.
+	const EAsyncLoadingResult::Type Result = EAsyncLoadingResult::Canceled;
+	for (int32 CallbackIndex = 0; CallbackIndex < CompletionCallbacks.Num(); CallbackIndex++)
+	{
+		CompletionCallbacks[CallbackIndex].ExecuteIfBound(PackageName, nullptr, Result);
+	}
+	if (Linker)
+	{
+		// give a hint to the IO system that we are done with this file for now
+		FIOSystem::Get().HintDoneWithFile(Linker->Filename);
+		Linker->FlushCache();
+		if (Linker->LinkerRoot)
+		{
+			Linker->LinkerRoot->ClearFlags(RF_WasLoaded);
+			Linker->LinkerRoot->bHasBeenFullyLoaded = false;
+		}
+	}
+}
 
 /*-----------------------------------------------------------------------------
 	UObject async (pre)loading.
@@ -999,6 +1035,30 @@ FAsyncPackage& LoadPackageAsync( const FString& InPackageName, FLoadPackageAsync
 	return AsyncPackage;
 }
 
+void CancelAsyncLoading()
+{
+	// Can't cancel from within async loading code
+	check(!GIsAsyncLoading);
+
+	GObjLoaded.Empty();
+
+	for (int32 PackageIndex = 0; PackageIndex < GObjAsyncPackages.Num(); PackageIndex++)
+	{
+		FAsyncPackage& PendingPackage = GObjAsyncPackages[PackageIndex];
+		PendingPackage.Cancel();
+	}
+	GObjAsyncPackages.Empty();
+	for (int32 ObjectIndex = 0; ObjectIndex < GObjConstructedDuringAsyncLoading.Num(); ObjectIndex++)
+	{
+		UObject* Object = GObjConstructedDuringAsyncLoading[ObjectIndex];
+		Object->ClearFlags(RF_AsyncLoading | RF_NeedLoad | RF_NeedPostLoad | RF_NeedPostLoadSubobjects);
+	}
+	GObjConstructedDuringAsyncLoading.Empty();
+
+	// Not that streaming has finished, we can stop force-referencing all objects that were created during async load.
+	FAsyncObjectsReferencer::Get().EmptyReferencedObjectsAndCancelLoading();
+}
+
 /**
  * Returns the async load percentage for a package in flight with the passed in name or -1 if there isn't one.
  *
diff --git a/Engine/Source/Runtime/CoreUObject/Public/Serialization/AsyncLoading.h b/Engine/Source/Runtime/CoreUObject/Public/Serialization/AsyncLoading.h
--- a/Engine/Source/Runtime/CoreUObject/Public/Serialization/AsyncLoading.h
+++ b/Engine/Source/Runtime/CoreUObject/Public/Serialization/AsyncLoading.h
@@ -100,6 +100,11 @@ struct FAsyncPackage : public FGCObject
 		return bLoadHasFinished;
 	}
 
+	/**
+	* Cancel loading this package.
+	*/
+	void Cancel();
+
 private:
 	/** Name of the UPackage to create.																	*/
 	FName						PackageName;
@@ -286,7 +291,7 @@ private:
 	/**
 	 * Function called when pending import package has been loaded.
 	 */
-	void ImportFullyLoadedCallback(const FName& PackageName, UPackage* LoadedPackage);
+	void ImportFullyLoadedCallback(const FName& PackageName, UPackage* LoadedPackage, EAsyncLoadingResult::Type Result);
 	/**
 	 * Adds dependency tree to the list if packages to wait for until their linkers have been created.
 	 *
diff --git a/Engine/Source/Runtime/CoreUObject/Public/UObject/UObjectGlobals.h b/Engine/Source/Runtime/CoreUObject/Public/UObject/UObjectGlobals.h
--- a/Engine/Source/Runtime/CoreUObject/Public/UObject/UObjectGlobals.h
+++ b/Engine/Source/Runtime/CoreUObject/Public/UObject/UObjectGlobals.h
@@ -248,12 +248,27 @@ COREUOBJECT_API void StaticTick( float DeltaTime, bool bUseFullTimeLimit = true,
  */
 COREUOBJECT_API UPackage* LoadPackage( UPackage* InOuter, const TCHAR* InLongPackageName, uint32 LoadFlags );
 
+/* Async package loading result */
+namespace EAsyncLoadingResult
+{
+	enum Type
+	{
+		/** Package failed to load */
+		Failed,
+		/** Package loaded successfully */
+		Succeeded,
+		/** Async loading was canceled */
+		Canceled
+	};
+}
+
 /**
  * Delegate called on completion of async package loading
  * @param	PackageName			Package name we were trying to load
  * @param	LoadedPackage		Loaded package if successful, NULL otherwise	
+ * @param	Result		Result of async loading.
  */
-DECLARE_DELEGATE_TwoParams(FLoadPackageAsyncDelegate, const FName& /*PackageName*/, UPackage* /*LoadedPackage*/)
+DECLARE_DELEGATE_ThreeParams(FLoadPackageAsyncDelegate, const FName& /*PackageName*/, UPackage* /*LoadedPackage*/, EAsyncLoadingResult::Type /*Result*/)
 struct FAsyncPackage;
 /**
  * Asynchronously load a package and all contained objects that match context flags. Non- blocking.
@@ -278,6 +293,11 @@ COREUOBJECT_API FAsyncPackage& LoadPackageAsync( const FString& PackageName, FLo
  */
 COREUOBJECT_API FAsyncPackage& LoadPackageAsync( const FString& PackageName, const FGuid* RequiredGuid = NULL, FName PackageType = NAME_None, const TCHAR* PackageToLoadFrom = NULL );
 
+/**
+* Cancels all async package loading requests.
+*/
+COREUOBJECT_API void CancelAsyncLoading();
+
 /**
  * Returns the async load percentage for a package in flight with the passed in name or -1 if there isn't one.
  *
@@ -386,7 +406,7 @@ namespace EAsyncPackageState
 		/** Package has pending import packages that need to be streamed in. */
 		PendingImports,
 		/** Package has finished loading. */
-		Complete
+		Complete,
 	};
 }
 /**
diff --git a/Engine/Source/Runtime/Engine/Classes/Engine/LevelStreaming.h b/Engine/Source/Runtime/Engine/Classes/Engine/LevelStreaming.h
--- a/Engine/Source/Runtime/Engine/Classes/Engine/LevelStreaming.h
+++ b/Engine/Source/Runtime/Engine/Classes/Engine/LevelStreaming.h
@@ -122,7 +122,7 @@ class ULevelStreaming : public UObject
 
 	/** Whether this level streaming object's level should be unloaded and the object be removed from the level list.			*/
 	uint32 bIsRequestingUnloadAndRemoval:1;
-
+	
 	/** If true, will be drawn on the 'level streaming status' map (STAT LEVELMAP console command) */
 	UPROPERTY(EditAnywhere, AdvancedDisplay, Category=LevelStreaming)
 	uint32 bDrawOnLevelStatusMap:1;
@@ -312,7 +312,7 @@ private:
 	 *
 	 * @param LevelPackage	Loaded level package
 	 */
-	void AsyncLevelLoadComplete( const FName& PackageName, UPackage* LevelPackage );
+	void AsyncLevelLoadComplete(const FName& PackageName, UPackage* LevelPackage, EAsyncLoadingResult::Type Result);
 
 	/** Pointer to Level object if currently loaded/ streamed in.																*/
 	UPROPERTY(transient)
diff --git a/Engine/Source/Runtime/Engine/Classes/Engine/PackageMapClient.h b/Engine/Source/Runtime/Engine/Classes/Engine/PackageMapClient.h
--- a/Engine/Source/Runtime/Engine/Classes/Engine/PackageMapClient.h
+++ b/Engine/Source/Runtime/Engine/Classes/Engine/PackageMapClient.h
@@ -69,7 +69,7 @@ public:
 	FString			FullNetGUIDPath( const FNetworkGUID& NetGUID ) const;
 	void			GenerateFullNetGUIDPath_r( const FNetworkGUID& NetGUID, FString& FullPath ) const;
 
-	void			AsyncPackageCallback( const FName& PackageName, UPackage * Package );
+	void			AsyncPackageCallback(const FName& PackageName, UPackage * Package, EAsyncLoadingResult::Type Result);
 	
 	TMap< FNetworkGUID, FNetGuidCacheObject >		ObjectLookup;
 	TMap< TWeakObjectPtr< UObject >, FNetworkGUID >	NetGUIDLookup;
diff --git a/Engine/Source/Runtime/Engine/Classes/Engine/StreamableManager.h b/Engine/Source/Runtime/Engine/Classes/Engine/StreamableManager.h
--- a/Engine/Source/Runtime/Engine/Classes/Engine/StreamableManager.h
+++ b/Engine/Source/Runtime/Engine/Classes/Engine/StreamableManager.h
@@ -68,7 +68,7 @@ private:
 		}
 	};
 
-	static void AsyncLoadCallbackWrapper(const FName& PackageName, UPackage* LevelPackage, FCallback* Handler)
+	static void AsyncLoadCallbackWrapper(const FName& PackageName, UPackage* LevelPackage, EAsyncLoadingResult::Type Result, FCallback* Handler)
 	{
 		FCallback* Callback = Handler;
 		Callback->Manager->AsyncLoadCallback(Callback->Request);
diff --git a/Engine/Source/Runtime/Engine/Classes/Engine/World.h b/Engine/Source/Runtime/Engine/Classes/Engine/World.h
--- a/Engine/Source/Runtime/Engine/Classes/Engine/World.h
+++ b/Engine/Source/Runtime/Engine/Classes/Engine/World.h
@@ -144,7 +144,7 @@ private:
 	void CopyWorldData();
 
 	/** callback sent to async loading code to inform us when the level package is complete */
-	void SeamlessTravelLoadCallback(const FName& PackageName, UPackage* LevelPackage);
+	void SeamlessTravelLoadCallback(const FName& PackageName, UPackage* LevelPackage, EAsyncLoadingResult::Type Result);
 
 	void SetHandlerLoadedData(UObject* InLevelPackage, UWorld* InLoadedWorld);
 
diff --git a/Engine/Source/Runtime/Engine/Classes/Kismet/GameplayStatics.h b/Engine/Source/Runtime/Engine/Classes/Kismet/GameplayStatics.h
--- a/Engine/Source/Runtime/Engine/Classes/Kismet/GameplayStatics.h
+++ b/Engine/Source/Runtime/Engine/Classes/Kismet/GameplayStatics.h
@@ -115,6 +115,11 @@ class ENGINE_API UGameplayStatics : public UBlueprintFunctionLibrary
 	UFUNCTION(BlueprintPure, meta=(WorldContext="WorldContextObject"), Category="Game")
 	static class ULevelStreaming* GetStreamingLevel(UObject* WorldContextObject, FName PackageName);
 	
+	/** Cancels all currently queued streaming packages */
+	UFUNCTION(BlueprintCallable, Category = "Game")
+	static void CancelAsyncLoading();
+
+
 	/**
 	 * Travel to another level
 	 *
diff --git a/Engine/Source/Runtime/Engine/Private/GameplayStatics.cpp b/Engine/Source/Runtime/Engine/Private/GameplayStatics.cpp
--- a/Engine/Source/Runtime/Engine/Private/GameplayStatics.cpp
+++ b/Engine/Source/Runtime/Engine/Private/GameplayStatics.cpp
@@ -376,6 +376,11 @@ ULevelStreaming* UGameplayStatics::GetStreamingLevel(UObject* WorldContextObject
 	return NULL;
 }
 
+void UGameplayStatics::CancelAsyncLoading()
+{
+	::CancelAsyncLoading();
+}
+
 void UGameplayStatics::OpenLevel(UObject* WorldContextObject, FName LevelName, bool bAbsolute, FString Options)
 {
 	UWorld* World = GEngine->GetWorldFromContextObject(WorldContextObject);
diff --git a/Engine/Source/Runtime/Engine/Private/LevelStreaming.cpp b/Engine/Source/Runtime/Engine/Private/LevelStreaming.cpp
--- a/Engine/Source/Runtime/Engine/Private/LevelStreaming.cpp
+++ b/Engine/Source/Runtime/Engine/Private/LevelStreaming.cpp
@@ -495,7 +495,7 @@ bool ULevelStreaming::RequestLevel(UWorld* PersistentWorld, bool bAllowLevelLoad
 	return true;
 }
 
-void ULevelStreaming::AsyncLevelLoadComplete( const FName& InPackageName, UPackage* InLoadedPackage ) 
+void ULevelStreaming::AsyncLevelLoadComplete(const FName& InPackageName, UPackage* InLoadedPackage, EAsyncLoadingResult::Type Result)
 {
 	bHasLoadRequestPending = false;
 
@@ -620,6 +620,12 @@ void ULevelStreaming::AsyncLevelLoadComplete( const FName& InPackageName, UPacka
 			}
 		}
 	}
+	else if (Result == EAsyncLoadingResult::Canceled)
+	{
+		// Cancel level streaming
+		bHasLoadRequestPending = false;
+		bShouldBeLoaded = false;
+	}
 	else
 	{
 		UE_LOG(LogLevelStreaming, Warning, TEXT("Failed to load package '%s'"), *InPackageName.ToString() );
diff --git a/Engine/Source/Runtime/Engine/Private/PackageMapClient.cpp b/Engine/Source/Runtime/Engine/Private/PackageMapClient.cpp
--- a/Engine/Source/Runtime/Engine/Private/PackageMapClient.cpp
+++ b/Engine/Source/Runtime/Engine/Private/PackageMapClient.cpp
@@ -1593,7 +1593,7 @@ void FNetGUIDCache::RegisterNetGUIDFromPath_Client( const FNetworkGUID& NetGUID,
 	RegisterNetGUID_Internal( NetGUID, CacheObject );
 }
 
-void FNetGUIDCache::AsyncPackageCallback( const FName& PackageName, UPackage * Package )
+void FNetGUIDCache::AsyncPackageCallback(const FName& PackageName, UPackage * Package, EAsyncLoadingResult::Type Result)
 {
 	check( Package == NULL || Package->IsFullyLoaded() );
 
diff --git a/Engine/Source/Runtime/Engine/Private/UnrealEngine.cpp b/Engine/Source/Runtime/Engine/Private/UnrealEngine.cpp
--- a/Engine/Source/Runtime/Engine/Private/UnrealEngine.cpp
+++ b/Engine/Source/Runtime/Engine/Private/UnrealEngine.cpp
@@ -2480,6 +2480,11 @@ bool UEngine::Exec( UWorld* InWorld, const TCHAR* Cmd, FOutputDevice& Ar )
 	{
 		return HandleDumpTicksCommand( InWorld, Cmd, Ar );
 	}
+	else if (FParse::Command(&Cmd, TEXT("CANCELASYNCLOAD")))
+	{
+		CancelAsyncLoading();
+		return true;
+	}
 #if USE_NETWORK_PROFILER
 	else if( FParse::Command(&Cmd,TEXT("NETPROFILE")) )
 	{
@@ -8468,7 +8473,7 @@ static void SetGametypeContentObjectReferencers(UObject* GametypeContentPackage,
  * @param	ContentPackage		The package that was loaded.
  * @param	GameEngine			The GameEngine.
  */
-static void AsyncLoadLocalizedMapGameTypeContentCallback(const FName& PackageName, UPackage* ContentPackage, FName InContextHandle)
+static void AsyncLoadLocalizedMapGameTypeContentCallback(const FName& PackageName, UPackage* ContentPackage, EAsyncLoadingResult::Type Result, FName InContextHandle)
 {
 	SetGametypeContentObjectReferencers(ContentPackage, InContextHandle, GametypeContent_LocalizedReferencerIndex);
 }
@@ -8480,7 +8485,7 @@ static void AsyncLoadLocalizedMapGameTypeContentCallback(const FName& PackageNam
  * @param	ContentPackage		The package that was loaded.
  * @param	GameEngine			The GameEngine.
  */
-static void AsyncLoadMapGameTypeContentCallback(const FName& PackageName, UPackage* ContentPackage, FName InContextHandle)
+static void AsyncLoadMapGameTypeContentCallback(const FName& PackageName, UPackage* ContentPackage, EAsyncLoadingResult::Type Result, FName InContextHandle)
 {
 	SetGametypeContentObjectReferencers(ContentPackage, InContextHandle, GametypeContent_ReferencerIndex);
 }
@@ -9661,7 +9666,7 @@ void UEngine::VerifyLoadMapWorldCleanup()
  * @param	LevelPackage	level package that finished async loading
  * @param	InGameEngine	pointer to game engine object to associated loaded level with so it won't be GC'ed
  */
-static void AsyncMapChangeLevelLoadCompletionCallback(const FName& PackageName, UPackage* LevelPackage, FName InWorldHandle )
+static void AsyncMapChangeLevelLoadCompletionCallback(const FName& PackageName, UPackage* LevelPackage, EAsyncLoadingResult::Type Result, FName InWorldHandle)
 {
 	FWorldContext &Context = GEngine->GetWorldContextFromHandleChecked( InWorldHandle );
 
diff --git a/Engine/Source/Runtime/Engine/Private/World.cpp b/Engine/Source/Runtime/Engine/Private/World.cpp
--- a/Engine/Source/Runtime/Engine/Private/World.cpp
+++ b/Engine/Source/Runtime/Engine/Private/World.cpp
@@ -4157,7 +4157,7 @@ void FSeamlessTravelHandler::SetHandlerLoadedData(UObject* InLevelPackage, UWorl
 }
 
 /** callback sent to async loading code to inform us when the level package is complete */
-void FSeamlessTravelHandler::SeamlessTravelLoadCallback(const FName& PackageName, UPackage* LevelPackage)
+void FSeamlessTravelHandler::SeamlessTravelLoadCallback(const FName& PackageName, UPackage* LevelPackage, EAsyncLoadingResult::Type Result)
 {
 	// make sure we remove the name, even if travel was cancelled.
 	const FName URLMapFName = FName(*PendingTravelURL.Map);
