Case insensitive comparison is inordinately expensive at the moment, but this comparison doesn't need to be case insensitive at all.

diff --git a/Engine/Source/Runtime/CoreUObject/Private/Misc/PackageName.cpp b/Engine/Source/Runtime/CoreUObject/Private/Misc/PackageName.cpp
--- a/Engine/Source/Runtime/CoreUObject/Private/Misc/PackageName.cpp
+++ b/Engine/Source/Runtime/CoreUObject/Private/Misc/PackageName.cpp
@@ -399,7 +399,7 @@ bool FPackageName::DoesPackageNameContainInvalidCharacters(const FString& InLong
 	for (const TCHAR* InvalidCharacters = INVALID_LONGPACKAGE_CHARACTERS; *InvalidCharacters; ++InvalidCharacters)
 	{
 		CharString[0] = *InvalidCharacters;
-		if (InLongPackageName.Contains(CharString))
+		if (InLongPackageName.Contains(CharString, ESearchCase::CaseSensitive))
 		{
 			MatchedInvalidChars += *InvalidCharacters;
 		}
