#ue4 - Cast and dynamic_cast with a const class type (eg Cast<const UClass> or dynamic_cast<const UClass*>) now use the optimized class cast flags if possible.

diff --git a/Engine/Source/Runtime/CoreUObject/Public/Templates/Casts.h b/Engine/Source/Runtime/CoreUObject/Public/Templates/Casts.h
--- a/Engine/Source/Runtime/CoreUObject/Public/Templates/Casts.h
+++ b/Engine/Source/Runtime/CoreUObject/Public/Templates/Casts.h
@@ -209,6 +209,11 @@ template< class T, class U > FORCEINLINE const T* CastChecked( const U      * Sr
 #define DECLARE_CAST_BY_FLAG_CAST(ClassName) \
 	template <> \
 	struct TCastFlags<ClassName> \
+	{ \
+		static const uint64 Value = CASTCLASS_##ClassName; \
+	}; \
+	template <> \
+	struct TCastFlags<const ClassName> \
 	{ \
 		static const uint64 Value = CASTCLASS_##ClassName; \
 	};
