#ue4 - Optimize Cast<> usage within script execution code (for non-shipping builds where the call counts were extremely high in these functions).

Removed ~4k IsA() calls per frame in a test scene.

diff --git a/Engine/Source/Runtime/CoreUObject/Public/UObject/UnrealType.h b/Engine/Source/Runtime/CoreUObject/Public/UObject/UnrealType.h
--- a/Engine/Source/Runtime/CoreUObject/Public/UObject/UnrealType.h
+++ b/Engine/Source/Runtime/CoreUObject/Public/UObject/UnrealType.h
@@ -347,10 +347,10 @@ private:
 			// in the future, these checks will be tested if the property is supposed be from a UClass
 			// need something for networking, since those are NOT live uobjects, just memory blocks
 			check(((UObject*)ContainerPtr)->IsValidLowLevel()); // Check its a valid UObject that was passed in
-			check(GetOuter()->IsA(UClass::StaticClass())); // Check that the outer of this property is a UClass (not another property)
+			check(Cast<UClass>(GetOuter()) != nullptr); // Check that the outer of this property is a UClass (not another property)
 
 			// Check that the object we are accessing is of the class that contains this property
-			checkf(((UObject*)ContainerPtr)->IsA((UClass*)GetOuter()), TEXT("'%s' is of class '%s' however property '%s' belongs to class '%s'") 
+			checkfSlow(((UObject*)ContainerPtr)->IsA((UClass*)GetOuter()), TEXT("'%s' is of class '%s' however property '%s' belongs to class '%s'") 
 																	 , *((UObject*)ContainerPtr)->GetName()
 																	 , *((UObject*)ContainerPtr)->GetClass()->GetName()
 																	 , *GetName()
@@ -359,7 +359,7 @@ private:
 		if (0)
 		{
 			// in the future, these checks will be tested if the property is NOT relative to a UClass
-			check(!GetOuter()->IsA(UClass::StaticClass())); // Check we are _not_ calling this on a direct child property of a UClass, you should pass in a UObject* in that case
+			check(!Cast<UClass>(GetOuter())); // Check we are _not_ calling this on a direct child property of a UClass, you should pass in a UObject* in that case
 		}
 		check(ArrayIndex < ArrayDim);
 		check(ContainerPtr);
