#ue4 - Cast flags can now be declared for non-intrinsic classes.

- Added cast flag support to a handful of the most commonly used classes in the engine, determined by tracking all calls to IsA() over typical workflows.
- Result: reduced IsA() function calls from 21k to 1k instances per frame in a large test scene (1.0ms to 0.06ms per frame).

diff --git a/Engine/Source/Programs/UnrealHeaderTool/Private/HeaderParser.cpp b/Engine/Source/Programs/UnrealHeaderTool/Private/HeaderParser.cpp
--- a/Engine/Source/Programs/UnrealHeaderTool/Private/HeaderParser.cpp
+++ b/Engine/Source/Programs/UnrealHeaderTool/Private/HeaderParser.cpp
@@ -4140,6 +4140,9 @@ void FHeaderParser::ParseClassNameDeclaration(FClasses& AllClasses, FString& Dec
 		bSpecifiesParentClass = true;
 	}
 
+	// Add class cast flag
+	Class->ClassCastFlags |= ClassCastFlagMap::Get().GetCastFlag(DeclaredClassName);
+
 	if (bSpecifiesParentClass)
 	{
 		// Set the base class.
diff --git a/Engine/Source/Runtime/CoreUObject/Private/Templates/Casts.cpp b/Engine/Source/Runtime/CoreUObject/Private/Templates/Casts.cpp
--- a/Engine/Source/Runtime/CoreUObject/Private/Templates/Casts.cpp
+++ b/Engine/Source/Runtime/CoreUObject/Private/Templates/Casts.cpp
@@ -10,3 +10,34 @@ void CastLogError(const TCHAR* FromType, const TCHAR* ToType)
 {
 	UE_LOG(LogCasts, Fatal, TEXT("Cast of %s to %s failed"), FromType, ToType);
 }
+
+
+#if HACK_HEADER_GENERATOR
+
+//////////////////////////////////////////////////////////////////////////
+// ClassCastFlagMap
+
+ClassCastFlagMap::ClassCastFlagMap()
+{
+	// Define macro to be applied to all cast class declarations.
+	#define DECLARE_CAST_BY_FLAG(ClassName) CastFlagMap.Add(FString(#ClassName), CASTCLASS_##ClassName);
+
+	// Now apply the macro to the whole list.
+	DECLARE_ALL_CAST_FLAGS;
+
+	#undef DECLARE_CAST_BY_FLAG
+}
+
+ClassCastFlagMap& ClassCastFlagMap::Get()
+{
+	static ClassCastFlagMap TheInstance;
+	return TheInstance;
+}
+
+EClassCastFlags ClassCastFlagMap::GetCastFlag(const FString& ClassName) const
+{
+	const EClassCastFlags* ValuePtr = CastFlagMap.Find(ClassName);
+	return ValuePtr ? *ValuePtr : CASTCLASS_None;
+}
+
+#endif //HACK_HEADER_GENERATOR
diff --git a/Engine/Source/Runtime/CoreUObject/Public/Templates/Casts.h b/Engine/Source/Runtime/CoreUObject/Public/Templates/Casts.h
--- a/Engine/Source/Runtime/CoreUObject/Public/Templates/Casts.h
+++ b/Engine/Source/Runtime/CoreUObject/Public/Templates/Casts.h
@@ -222,39 +222,75 @@ template< class T, class U > FORCEINLINE const T* CastChecked( const U      * Sr
 	DECLARE_CAST_BY_FLAG_FWD(ClassName) \
 	DECLARE_CAST_BY_FLAG_CAST(ClassName)
 
-DECLARE_CAST_BY_FLAG(UField)
-DECLARE_CAST_BY_FLAG(UEnum)
-DECLARE_CAST_BY_FLAG(UStruct)
-DECLARE_CAST_BY_FLAG(UScriptStruct)
-DECLARE_CAST_BY_FLAG(UClass)
-DECLARE_CAST_BY_FLAG(UProperty)
-DECLARE_CAST_BY_FLAG(UObjectPropertyBase)
-DECLARE_CAST_BY_FLAG(UObjectProperty)
-DECLARE_CAST_BY_FLAG(UWeakObjectProperty)
-DECLARE_CAST_BY_FLAG(ULazyObjectProperty)
-DECLARE_CAST_BY_FLAG(UAssetObjectProperty)
-DECLARE_CAST_BY_FLAG(UAssetClassProperty)
-DECLARE_CAST_BY_FLAG(UBoolProperty)
-DECLARE_CAST_BY_FLAG(UFunction)
-DECLARE_CAST_BY_FLAG(UStructProperty)
-DECLARE_CAST_BY_FLAG(UByteProperty)
-DECLARE_CAST_BY_FLAG(UIntProperty)
-DECLARE_CAST_BY_FLAG(UFloatProperty)
-DECLARE_CAST_BY_FLAG(UDoubleProperty)
-DECLARE_CAST_BY_FLAG(UClassProperty)
-DECLARE_CAST_BY_FLAG(UInterfaceProperty)
-DECLARE_CAST_BY_FLAG(UNameProperty)
-DECLARE_CAST_BY_FLAG(UStrProperty)
-DECLARE_CAST_BY_FLAG(UTextProperty)
-DECLARE_CAST_BY_FLAG(UArrayProperty)
-DECLARE_CAST_BY_FLAG(UDelegateProperty)
-DECLARE_CAST_BY_FLAG(UMulticastDelegateProperty)
+#define FINISH_DECLARING_CAST_FLAGS // intentionally defined to do nothing.
+
+// Define a macro that declares all the cast flags.
+// This allows us to reuse these declarations elsewhere to define other properties for these classes.
+#define DECLARE_ALL_CAST_FLAGS \
+DECLARE_CAST_BY_FLAG(UField)							\
+DECLARE_CAST_BY_FLAG(UEnum)								\
+DECLARE_CAST_BY_FLAG(UStruct)							\
+DECLARE_CAST_BY_FLAG(UScriptStruct)						\
+DECLARE_CAST_BY_FLAG(UClass)							\
+DECLARE_CAST_BY_FLAG(UProperty)							\
+DECLARE_CAST_BY_FLAG(UObjectPropertyBase)				\
+DECLARE_CAST_BY_FLAG(UObjectProperty)					\
+DECLARE_CAST_BY_FLAG(UWeakObjectProperty)				\
+DECLARE_CAST_BY_FLAG(ULazyObjectProperty)				\
+DECLARE_CAST_BY_FLAG(UAssetObjectProperty)				\
+DECLARE_CAST_BY_FLAG(UAssetClassProperty)				\
+DECLARE_CAST_BY_FLAG(UBoolProperty)						\
+DECLARE_CAST_BY_FLAG(UFunction)							\
+DECLARE_CAST_BY_FLAG(UStructProperty)					\
+DECLARE_CAST_BY_FLAG(UByteProperty)						\
+DECLARE_CAST_BY_FLAG(UIntProperty)						\
+DECLARE_CAST_BY_FLAG(UFloatProperty)					\
+DECLARE_CAST_BY_FLAG(UDoubleProperty)					\
+DECLARE_CAST_BY_FLAG(UClassProperty)					\
+DECLARE_CAST_BY_FLAG(UInterfaceProperty)				\
+DECLARE_CAST_BY_FLAG(UNameProperty)						\
+DECLARE_CAST_BY_FLAG(UStrProperty)						\
+DECLARE_CAST_BY_FLAG(UTextProperty)						\
+DECLARE_CAST_BY_FLAG(UArrayProperty)					\
+DECLARE_CAST_BY_FLAG(UDelegateProperty)					\
+DECLARE_CAST_BY_FLAG(UMulticastDelegateProperty)		\
+DECLARE_CAST_BY_FLAG(UPackage)							\
+DECLARE_CAST_BY_FLAG(ULevel)							\
+DECLARE_CAST_BY_FLAG(AActor)							\
+DECLARE_CAST_BY_FLAG(APlayerController)					\
+DECLARE_CAST_BY_FLAG(APawn)								\
+DECLARE_CAST_BY_FLAG(USceneComponent)					\
+DECLARE_CAST_BY_FLAG(UPrimitiveComponent)				\
+DECLARE_CAST_BY_FLAG(USkinnedMeshComponent)				\
+DECLARE_CAST_BY_FLAG(USkeletalMeshComponent)			\
+DECLARE_CAST_BY_FLAG(UBlueprint)						\
+FINISH_DECLARING_CAST_FLAGS		// This is here to hopefully remind people to include the "\" in all declarations above, especially when copy/pasting the final line.
+
+// Now actually declare the flags
+DECLARE_ALL_CAST_FLAGS
 
 #undef DECLARE_CAST_BY_FLAG
-#undef DECLARE_CAST_BY_FLAG_CASTCHECKED
 #undef DECLARE_CAST_BY_FLAG_CAST
 #undef DECLARE_CAST_BY_FLAG_FWD
 
+
+#if HACK_HEADER_GENERATOR
+// Singleton class to get the cast flag for a given class name.
+struct COREUOBJECT_API ClassCastFlagMap
+{
+	static ClassCastFlagMap& Get();
+
+	// Get Mapped name -> cast flag. Returns CASTCLASS_None if name is not found.
+	EClassCastFlags GetCastFlag(const FString& ClassName) const;
+
+private:
+	ClassCastFlagMap();
+	TMap<FString, EClassCastFlags> CastFlagMap;
+};
+#endif // HACK_HEADER_GENERATOR
+
+
+
 namespace UE4Casts_Private
 {
 	template <typename To, typename From>
diff --git a/Engine/Source/Runtime/CoreUObject/Public/UObject/CoreObject.h b/Engine/Source/Runtime/CoreUObject/Public/UObject/CoreObject.h
--- a/Engine/Source/Runtime/CoreUObject/Public/UObject/CoreObject.h
+++ b/Engine/Source/Runtime/CoreUObject/Public/UObject/CoreObject.h
@@ -22,7 +22,7 @@ struct FLevelGuids
  */
 class COREUOBJECT_API UPackage : public UObject
 {
-	DECLARE_CLASS_INTRINSIC(UPackage,UObject,0,CoreUObject)
+	DECLARE_CASTED_CLASS_INTRINSIC(UPackage,UObject,0,CoreUObject, CASTCLASS_UPackage)
 
 public:
 	/** delegate type for package dirty state events.  ( Params: UPackage* ModifiedPackage ) */
diff --git a/Engine/Source/Runtime/CoreUObject/Public/UObject/ObjectBase.h b/Engine/Source/Runtime/CoreUObject/Public/UObject/ObjectBase.h
--- a/Engine/Source/Runtime/CoreUObject/Public/UObject/ObjectBase.h
+++ b/Engine/Source/Runtime/CoreUObject/Public/UObject/ObjectBase.h
@@ -269,6 +269,16 @@ typedef uint64 EClassCastFlags;
 #define CASTCLASS_UInt16Property				DECLARE_UINT64(0x0000000080000000)
 #define CASTCLASS_UDoubleProperty				DECLARE_UINT64(0x0000000100000000)
 #define CASTCLASS_UAssetClassProperty			DECLARE_UINT64(0x0000000200000000)
+#define CASTCLASS_UPackage						DECLARE_UINT64(0x0000000400000000)
+#define CASTCLASS_ULevel						DECLARE_UINT64(0x0000000800000000)
+#define CASTCLASS_AActor						DECLARE_UINT64(0x0000001000000000)
+#define CASTCLASS_APlayerController				DECLARE_UINT64(0x0000002000000000)
+#define CASTCLASS_APawn							DECLARE_UINT64(0x0000004000000000)
+#define CASTCLASS_USceneComponent				DECLARE_UINT64(0x0000008000000000)
+#define CASTCLASS_UPrimitiveComponent			DECLARE_UINT64(0x0000010000000000)
+#define CASTCLASS_USkinnedMeshComponent			DECLARE_UINT64(0x0000020000000000)
+#define CASTCLASS_USkeletalMeshComponent		DECLARE_UINT64(0x0000040000000000)
+#define CASTCLASS_UBlueprint					DECLARE_UINT64(0x0000080000000000)
 
 #define CASTCLASS_AllFlags						DECLARE_UINT64(0xFFFFFFFFFFFFFFFF)
 
