Fixed a case where packages that were loaded as part of another package loading would fail to have the flag for requiring localization gathering set. This was causing false negatives for gathering text from assets.

diff --git a/Engine/Source/Runtime/CoreUObject/Private/UObject/UObjectGlobals.cpp b/Engine/Source/Runtime/CoreUObject/Private/UObject/UObjectGlobals.cpp
--- a/Engine/Source/Runtime/CoreUObject/Private/UObject/UObjectGlobals.cpp
+++ b/Engine/Source/Runtime/CoreUObject/Private/UObject/UObjectGlobals.cpp
@@ -837,16 +837,23 @@ UPackage* LoadPackageInternal(UPackage* InOuter, const TCHAR* InLongPackageName,
 			return nullptr;
 		}
 
-		if (Linker->LinkerRoot && Linker->LinkerRoot->HasAnyFlags(RF_WasLoaded))
+		Result = Linker->LinkerRoot;
+
+		auto EndLoadAndCopyLocalizationGatherFlag = [&]
+		{
+			EndLoad();
+			// Set package-requires-localization flags from archive after loading. This reinforces flagging of packages that haven't yet been resaved.
+			Result->ThisRequiresLocalizationGather(Linker->RequiresLocalizationGather());
+		};
+
+		if (Result && Result->HasAnyFlags(RF_WasLoaded))
 		{
 			// The linker is associated with a package that has already been loaded.
 			// Loading packages that have already been loaded is unsupported.
-			EndLoad();
-			return Linker->LinkerRoot;
+			EndLoadAndCopyLocalizationGatherFlag();
+			return Result;
 		}
 
-		Result = Linker->LinkerRoot;
-
 		// If we are loading a package for diff'ing, set the package flag
 		if(LoadFlags & LOAD_ForDiff)
 		{
@@ -886,15 +893,12 @@ UPackage* LoadPackageInternal(UPackage* InOuter, const TCHAR* InLongPackageName,
 
 		SlowTask.EnterProgressFrame(30);
 
-		EndLoad();
+		EndLoadAndCopyLocalizationGatherFlag();
 
 #if WITH_EDITOR
 		GIsEditorLoadingPackage = *IsEditorLoadingPackage;
 #endif
 
-		// Set package-requires-localization flags from archive after loading. This reinforces flagging of packages that haven't yet been resaved.
-		Result->ThisRequiresLocalizationGather(Linker->RequiresLocalizationGather());
-
 #if WITH_ENGINE
 		// Cancel all texture allocations that haven't been claimed yet.
 		Linker->Summary.TextureAllocations.CancelRemainingAllocations( true );
