Fix issues when outputting not fully qualified paths

diff --git a/Engine/Source/Runtime/CoreUObject/Private/UObject/PropertyBaseObject.cpp b/Engine/Source/Runtime/CoreUObject/Private/UObject/PropertyBaseObject.cpp
--- a/Engine/Source/Runtime/CoreUObject/Private/UObject/PropertyBaseObject.cpp
+++ b/Engine/Source/Runtime/CoreUObject/Private/UObject/PropertyBaseObject.cpp
@@ -115,30 +115,16 @@ void UObjectPropertyBase::ExportTextItem( FString& ValueStr, const void* Propert
 				{
 					StopOuter = Parent->GetOutermost();
 				}
-
-				FString TempName = Temp->GetPathName(StopOuter);
-				if ( (PortFlags & PPF_Delimited) && (!Temp->GetFName().IsValidXName(INVALID_OBJECTNAME_CHARACTERS)) )
-				{
-					TempName = FString::Printf(TEXT("\"%s\""), *TempName.ReplaceQuotesWithEscapedQuotes());
-				}
-				ValueStr += FString::Printf( TEXT("%s'%s'"), *Temp->GetClass()->GetName(), *TempName );
-			}
-			else if (Parent != NULL && (Temp->IsIn(Parent) || Temp->IsIn(Parent->GetOuter())) )
-			{
-				FString TempName = Temp->GetName();
-				if ( (PortFlags & PPF_Delimited) && (!Temp->GetFName().IsValidXName(INVALID_OBJECTNAME_CHARACTERS)) )
-				{
-					TempName = FString::Printf(TEXT("\"%s\""), *TempName.ReplaceQuotesWithEscapedQuotes());
-				}
-				ValueStr += FString::Printf( TEXT("%s'%s'"), *Temp->GetClass()->GetName(), *TempName ); 
 			}
-			else
+			else if (Parent != NULL && Temp->IsIn(Parent))
 			{
-				// Take the path name relative to the stopping point outermost ptr.
-				// This is so that cases like a component referencing a component in another actor work correctly when pasted
-				FString PathName = Temp->GetPathName(StopOuter);
-				ValueStr += FString::Printf( TEXT("%s'\"%s\"'"), *Temp->GetClass()->GetName(), *PathName );
+				StopOuter = Parent;
 			}
+
+			// Take the path name relative to the stopping point outermost ptr.
+			// This is so that cases like a component referencing a component in another actor work correctly when pasted
+			const FString PathName = Temp->GetPathName(StopOuter);
+			ValueStr += FString::Printf( TEXT("%s'\"%s\"'"), *Temp->GetClass()->GetName(), *PathName );
 		}
 	}
 	else
