Fix issues caused with object property paths due to always quoting

diff --git a/Engine/Source/Runtime/CoreUObject/Private/UObject/PropertyBaseObject.cpp b/Engine/Source/Runtime/CoreUObject/Private/UObject/PropertyBaseObject.cpp
--- a/Engine/Source/Runtime/CoreUObject/Private/UObject/PropertyBaseObject.cpp
+++ b/Engine/Source/Runtime/CoreUObject/Private/UObject/PropertyBaseObject.cpp
@@ -123,8 +123,12 @@ void UObjectPropertyBase::ExportTextItem( FString& ValueStr, const void* Propert
 
 			// Take the path name relative to the stopping point outermost ptr.
 			// This is so that cases like a component referencing a component in another actor work correctly when pasted
-			const FString PathName = Temp->GetPathName(StopOuter);
-			ValueStr += FString::Printf( TEXT("%s'\"%s\"'"), *Temp->GetClass()->GetName(), *PathName );
+			FString PathName = Temp->GetPathName(StopOuter);
+			if ( (PortFlags & PPF_Delimited) && (!Temp->GetFName().IsValidXName(INVALID_OBJECTNAME_CHARACTERS)) )
+			{
+				PathName = FString::Printf(TEXT("\"%s\""), *PathName.ReplaceQuotesWithEscapedQuotes());
+			}
+			ValueStr += FString::Printf( TEXT("%s'%s'"), *Temp->GetClass()->GetName(), *PathName );
 		}
 	}
 	else
