Core - More stats for memory profiling

diff --git a/Engine/Source/Runtime/CoreUObject/Private/Serialization/BulkData.cpp b/Engine/Source/Runtime/CoreUObject/Private/Serialization/BulkData.cpp
--- a/Engine/Source/Runtime/CoreUObject/Private/Serialization/BulkData.cpp
+++ b/Engine/Source/Runtime/CoreUObject/Private/Serialization/BulkData.cpp
@@ -93,6 +93,8 @@ FUntypedBulkData::FUntypedBulkData()
  */
 FUntypedBulkData::FUntypedBulkData( const FUntypedBulkData& Other )
 {
+	DECLARE_SCOPE_CYCLE_COUNTER(TEXT("FUntypedBulkData::FUntypedBulkData"), STAT_UBD_Constructor, STATGROUP_Memory);
+
 	InitializeMemberVariables();
 
 	// Prepare bulk data pointer. Can't call any functions that would call virtual GetElementSize on "this" as
@@ -114,6 +116,8 @@ FUntypedBulkData::FUntypedBulkData( const FUntypedBulkData& Other )
  */
 FUntypedBulkData::~FUntypedBulkData()
 {
+	DECLARE_SCOPE_CYCLE_COUNTER(TEXT("FUntypedBulkData::~FUntypedBulkData"), STAT_UBD_Destructor, STATGROUP_Memory);
+
 	check( LockStatus == LOCKSTATUS_Unlocked );
 	// Free memory.
 	if( bShouldFreeOnEmpty )
@@ -349,6 +353,8 @@ bool FUntypedBulkData::IsAvailableForUse() const
  */
 void FUntypedBulkData::GetCopy( void** Dest, bool bDiscardInternalCopy )
 {
+	DECLARE_SCOPE_CYCLE_COUNTER(TEXT("FUntypedBulkData::GetCopy"), STAT_UBD_GetCopy, STATGROUP_Memory);
+
 	check( LockStatus == LOCKSTATUS_Unlocked );
 	check( Dest );
 
@@ -476,6 +482,8 @@ const void* FUntypedBulkData::LockReadOnly() const
  */
 void* FUntypedBulkData::Realloc( int32 InElementCount )
 {
+	DECLARE_SCOPE_CYCLE_COUNTER(TEXT("FUntypedBulkData::Realloc"), STAT_UBD_Realloc, STATGROUP_Memory);
+
 	check( LockStatus == LOCKSTATUS_ReadWriteLock );
 	// Progate element count and reallocate data based on new size.
 	ElementCount	= InElementCount;
@@ -489,6 +497,8 @@ void* FUntypedBulkData::Realloc( int32 InElementCount )
  */
 void FUntypedBulkData::Unlock() const
 {
+	DECLARE_SCOPE_CYCLE_COUNTER(TEXT("FUntypedBulkData::Unlock"), STAT_UBD_Unlock, STATGROUP_Memory);
+
 	check(LockStatus != LOCKSTATUS_Unlocked);
 
 	FUntypedBulkData* mutable_this = const_cast<FUntypedBulkData*>(this);
@@ -511,6 +521,8 @@ void FUntypedBulkData::Unlock() const
  */
 void FUntypedBulkData::RemoveBulkData()
 {
+	DECLARE_SCOPE_CYCLE_COUNTER(TEXT("FUntypedBulkData::RemoveBulkData"), STAT_UBD_RemoveBulkData, STATGROUP_Memory);
+
 	check( LockStatus == LOCKSTATUS_Unlocked );
 
 #if WITH_EDITOR
@@ -604,6 +616,8 @@ bool FUntypedBulkData::ShouldFreeOnEmpty() const
 */
 void FUntypedBulkData::Serialize( FArchive& Ar, UObject* Owner, int32 Idx )
 {
+	DECLARE_SCOPE_CYCLE_COUNTER(TEXT("FUntypedBulkData::Serialize"), STAT_UBD_Serialize, STATGROUP_Memory);
+
 	check( LockStatus == LOCKSTATUS_Unlocked );
 
 	if(Ar.IsTransacting())
@@ -1070,6 +1084,8 @@ void FUntypedBulkData::SerializeBulkData( FArchive& Ar, void* Data )
  */
 void FUntypedBulkData::MakeSureBulkDataIsLoaded()
 {
+	DECLARE_SCOPE_CYCLE_COUNTER(TEXT("FUntypedBulkData::MakeSureBulkDataIsLoaded"), STAT_UBD_MakeSureBulkDataIsLoaded, STATGROUP_Memory);
+
 	// Nothing to do if data is already loaded.
 	if( !BulkData )
 	{
