AreInstancedObjectsIdentical - critical fix

diff --git a/Engine/Source/Runtime/CoreUObject/Private/UObject/PropertyHelper.h b/Engine/Source/Runtime/CoreUObject/Private/UObject/PropertyHelper.h
--- a/Engine/Source/Runtime/CoreUObject/Private/UObject/PropertyHelper.h
+++ b/Engine/Source/Runtime/CoreUObject/Private/UObject/PropertyHelper.h
@@ -51,7 +51,7 @@ static bool AreInstancedObjectsIdentical( UObject* ObjectA, UObject* ObjectB, ui
 	if (!RecursionCheck.Contains(Test))
 	{
 		RecursionCheck.Push(Test);
-		for ( UProperty* Prop = ObjectA->GetClass()->PropertyLink; Prop && !Result; Prop = Prop->PropertyLinkNext )
+		for ( UProperty* Prop = ObjectA->GetClass()->PropertyLink; Prop && Result; Prop = Prop->PropertyLinkNext )
 		{
 			// only the properties that could have been modified in the editor should be compared
 			// (skipping the name and archetype properties, since name will almost always be different)
@@ -63,7 +63,7 @@ static bool AreInstancedObjectsIdentical( UObject* ObjectA, UObject* ObjectB, ui
 
 			if ( bConsiderProperty )
 			{
-				for ( int32 i = 0; i < Prop->ArrayDim && !Result; i++ )
+				for ( int32 i = 0; i < Prop->ArrayDim && Result; i++ )
 				{
 					if ( !Prop->Identical_InContainer(ObjectA, ObjectB, i, PortFlags) )
 					{
