Use not fully qualified paths for specifying object archetypes when appropriate

diff --git a/Engine/Source/Runtime/CoreUObject/Private/UObject/PropertyBaseObject.cpp b/Engine/Source/Runtime/CoreUObject/Private/UObject/PropertyBaseObject.cpp
--- a/Engine/Source/Runtime/CoreUObject/Private/UObject/PropertyBaseObject.cpp
+++ b/Engine/Source/Runtime/CoreUObject/Private/UObject/PropertyBaseObject.cpp
@@ -81,6 +81,44 @@ void UObjectPropertyBase::AddReferencedObjects(UObject* InThis, FReferenceCollec
 	Collector.AddReferencedObject( This->PropertyClass, This );
 	Super::AddReferencedObjects( This, Collector );
 }
+
+FString UObjectPropertyBase::GetExportPath(const UObject* Object, const UObject* Parent, const UObject* ExportRootScope, const uint32 PortFlags)
+{
+	bool bExportFullyQualified = true;
+
+	// When exporting from one package or graph to another package or graph, we don't want to fully qualify the name, as it may refer
+	// to a level or graph that doesn't exist or cause a linkage to a node in a different graph
+	const UObject* StopOuter = nullptr;
+	if (PortFlags & PPF_ExportsNotFullyQualified)
+	{
+		StopOuter = (ExportRootScope || (Parent == nullptr)) ? ExportRootScope : Parent->GetOutermost();
+		bExportFullyQualified = !Object->IsIn(StopOuter);
+	}
+
+	// if we want a full qualified object reference, use the pathname, otherwise, use just the object name
+	if (bExportFullyQualified)
+	{
+		StopOuter = nullptr;
+		if ( (PortFlags&PPF_SimpleObjectText) != 0 && Parent != nullptr )
+		{
+			StopOuter = Parent->GetOutermost();
+		}
+	}
+	else if (Parent != NULL && Object->IsIn(Parent))
+	{
+		StopOuter = Parent;
+	}
+
+	// Take the path name relative to the stopping point outermost ptr.
+	// This is so that cases like a component referencing a component in another actor work correctly when pasted
+	FString PathName = Object->GetPathName(StopOuter);
+	if ( (PortFlags & PPF_Delimited) && (!Object->GetFName().IsValidXName(INVALID_OBJECTNAME_CHARACTERS)) )
+	{
+		PathName = FString::Printf(TEXT("\"%s\""), *PathName.ReplaceQuotesWithEscapedQuotes());
+	}
+	return FString::Printf( TEXT("%s'%s'"), *Object->GetClass()->GetName(), *PathName );
+}
+
 void UObjectPropertyBase::ExportTextItem( FString& ValueStr, const void* PropertyValue, const void* DefaultValue, UObject* Parent, int32 PortFlags, UObject* ExportRootScope ) const
 {
 	UObject* Temp = GetObjectPropertyValue(PropertyValue);
@@ -96,39 +134,7 @@ void UObjectPropertyBase::ExportTextItem( FString& ValueStr, const void* Propert
 		}
 		else
 		{
-			bool bExportFullyQualified = true;
-
-			// When exporting from one package or graph to another package or graph, we don't want to fully qualify the name, as it may refer
-			// to a level or graph that doesn't exist or cause a linkage to a node in a different graph
-			UObject* StopOuter = NULL;
-			if (PortFlags & PPF_ExportsNotFullyQualified)
-			{
-				StopOuter = (ExportRootScope || (Parent == NULL)) ? ExportRootScope : Parent->GetOutermost();
-				bExportFullyQualified = !Temp->IsIn(StopOuter);
-			}
-
-			// if we want a full qualified object reference, use the pathname, otherwise, use just the object name
-			if (bExportFullyQualified)
-			{
-				StopOuter = NULL;
-				if ( (PortFlags&PPF_SimpleObjectText) != 0 && Parent != NULL )
-				{
-					StopOuter = Parent->GetOutermost();
-				}
-			}
-			else if (Parent != NULL && Temp->IsIn(Parent))
-			{
-				StopOuter = Parent;
-			}
-
-			// Take the path name relative to the stopping point outermost ptr.
-			// This is so that cases like a component referencing a component in another actor work correctly when pasted
-			FString PathName = Temp->GetPathName(StopOuter);
-			if ( (PortFlags & PPF_Delimited) && (!Temp->GetFName().IsValidXName(INVALID_OBJECTNAME_CHARACTERS)) )
-			{
-				PathName = FString::Printf(TEXT("\"%s\""), *PathName.ReplaceQuotesWithEscapedQuotes());
-			}
-			ValueStr += FString::Printf( TEXT("%s'%s'"), *Temp->GetClass()->GetName(), *PathName );
+			ValueStr += GetExportPath(Temp, Parent, ExportRootScope, PortFlags);
 		}
 	}
 	else
diff --git a/Engine/Source/Runtime/CoreUObject/Public/UObject/UnrealType.h b/Engine/Source/Runtime/CoreUObject/Public/UObject/UnrealType.h
--- a/Engine/Source/Runtime/CoreUObject/Public/UObject/UnrealType.h
+++ b/Engine/Source/Runtime/CoreUObject/Public/UObject/UnrealType.h
@@ -1614,7 +1614,7 @@ class COREUOBJECT_API UObjectPropertyBase : public UProperty
 	 *									SIZE = the ElementSize of this UProperty
 	 * @param	Src					the address of the value to copy from. should be evaluated the same way as Dest
 	 */
-	virtual void CopySingleValueToScriptVM( void* Dest, void const* Src ) const
+	virtual void CopySingleValueToScriptVM( void* Dest, void const* Src ) const override
 	{
 		*(UObject**)Dest = GetObjectPropertyValue(Src);
 	}
@@ -1626,7 +1626,7 @@ class COREUOBJECT_API UObjectPropertyBase : public UProperty
 	 *									OFFSET = the Offset of this UProperty
 	 * @param	Src					the address of the value to copy from. should be evaluated the same way as Dest
 	 */
-	virtual void CopyCompleteValueToScriptVM( void* Dest, void const* Src ) const
+	virtual void CopyCompleteValueToScriptVM( void* Dest, void const* Src ) const override
 	{
 		for (int32 Index = 0; Index < ArrayDim; Index++)
 		{
@@ -1644,7 +1644,7 @@ class COREUOBJECT_API UObjectPropertyBase : public UProperty
 	 *									SIZE = the ElementSize of this UProperty
 	 * @param	Src					the address of the value to copy from. should be evaluated the same way as Dest
 	 */
-	virtual void CopySingleValueFromScriptVM( void* Dest, void const* Src ) const
+	virtual void CopySingleValueFromScriptVM( void* Dest, void const* Src ) const override
 	{
 		SetObjectPropertyValue(Dest, *(UObject**)Src);
 	}
@@ -1656,7 +1656,7 @@ class COREUOBJECT_API UObjectPropertyBase : public UProperty
 	 *									OFFSET = the Offset of this UProperty
 	 * @param	Src					the address of the value to copy from. should be evaluated the same way as Dest
 	 */
-	virtual void CopyCompleteValueFromScriptVM( void* Dest, void const* Src ) const
+	virtual void CopyCompleteValueFromScriptVM( void* Dest, void const* Src ) const override
 	{
 		checkSlow(ElementSize == sizeof(UObject*)); // the idea that script pointers are the same size as weak pointers is maybe required, maybe not
 		for (int32 Index = 0; Index < ArrayDim; Index++)
@@ -1686,6 +1686,9 @@ public:
 	 */
 	static bool ParseObjectPropertyValue( const UProperty* Property, UObject* OwnerObject, UClass* RequiredMetaClass, uint32 PortFlags, const TCHAR*& Buffer, UObject*& out_ResolvedValue );
 	static UObject* FindImportedObject( const UProperty* Property, UObject* OwnerObject, UClass* ObjectClass, UClass* RequiredMetaClass, const TCHAR* Text, uint32 PortFlags = 0);
+	
+	// Returns the qualified export path for a given object, parent, and export root scope
+	static FString GetExportPath(const UObject* Object, const UObject* Parent, const UObject* ExportRootScope, const uint32 PortFlags);
 
 	virtual UObject* GetObjectPropertyValue(const void* PropertyValueAddress) const
 	{
diff --git a/Engine/Source/Runtime/Engine/Private/UnrealExporter.cpp b/Engine/Source/Runtime/Engine/Private/UnrealExporter.cpp
--- a/Engine/Source/Runtime/Engine/Private/UnrealExporter.cpp
+++ b/Engine/Source/Runtime/Engine/Private/UnrealExporter.cpp
@@ -491,7 +491,8 @@ void UExporter::EmitBeginObject( FOutputDevice& Ar, UObject* Obj, uint32 PortFla
 		// do we want the archetype string?
 		if (!bIsExportingDefaultObject)
 		{
-			Ar.Logf(TEXT(" Archetype=%s'%s'"), *Obj->GetArchetype()->GetClass()->GetName(), *Obj->GetArchetype()->GetPathName());
+			UObject* Archetype = Obj->GetArchetype();
+			Ar.Logf(TEXT(" Archetype=%s"), *UObjectPropertyBase::GetExportPath(Archetype, Archetype->GetOutermost(), Archetype->GetOuter(), PortFlags));
 		}
 	}
 
