Core - Misc UObject stats tweaks

diff --git a/Engine/Source/Runtime/CoreUObject/Private/UObject/UObjectGlobals.cpp b/Engine/Source/Runtime/CoreUObject/Private/UObject/UObjectGlobals.cpp
--- a/Engine/Source/Runtime/CoreUObject/Private/UObject/UObjectGlobals.cpp
+++ b/Engine/Source/Runtime/CoreUObject/Private/UObject/UObjectGlobals.cpp
@@ -782,6 +782,8 @@ UClass* StaticLoadClass( UClass* BaseClass, UObject* InOuter, const TCHAR* InNam
 */
 UPackage* LoadPackageInternal(UPackage* InOuter, const TCHAR* InLongPackageName, uint32 LoadFlags, ULinkerLoad* ImportLinker)
 {
+	DECLARE_SCOPE_CYCLE_COUNTER(TEXT("LoadPackageInternal"), STAT_LoadPackageInternal, STATGROUP_ObjectVerbose);
+
 	UPackage* Result = NULL;
 
 	FString FileToLoad;
@@ -1663,7 +1665,7 @@ UObject* StaticAllocateObject
 	bool* bOutRecycledSubobject
 )
 {
-//	SCOPE_CYCLE_COUNTER(STAT_AllocateObject);
+	SCOPE_CYCLE_COUNTER(STAT_AllocateObject);
 	checkSlow(InOuter != INVALID_OBJECT); // not legal
 	check(!InClass || (InClass->ClassWithin && InClass->ClassConstructor));
 #if WITH_EDITOR
@@ -1906,7 +1908,7 @@ FObjectInitializer::~FObjectInitializer()
 	check(GIsInConstructor >= 0);
 	GConstructedObject = LastConstructedObject;
 
-//	SCOPE_CYCLE_COUNTER(STAT_PostConstructInitializeProperties);
+	SCOPE_CYCLE_COUNTER(STAT_PostConstructInitializeProperties);
 	check(Obj);
 	const bool bIsCDO = Obj->HasAnyFlags(RF_ClassDefaultObject);
 	UClass* Class = Obj->GetClass();
@@ -2196,7 +2198,7 @@ UObject* StaticConstructObject
 	// Don't call the constructor on recycled subobjects, they haven't been destroyed.
 	if (!bRecycledSubobject)
 	{		
-		FScopeCycleCounterUObject ConstructorScope(Result);
+		FScopeCycleCounterUObject ConstructorScope(Result, GET_STATID(STAT_ConstructObject));
 		(*InClass->ClassConstructor)( FObjectInitializer(Result, InTemplate, bCopyTransientsFromClassDefaults, true, InInstanceGraph) );
 	}
 	
diff --git a/Engine/Source/Runtime/CoreUObject/Public/UObject/UObjectGlobals.h b/Engine/Source/Runtime/CoreUObject/Public/UObject/UObjectGlobals.h
--- a/Engine/Source/Runtime/CoreUObject/Public/UObject/UObjectGlobals.h
+++ b/Engine/Source/Runtime/CoreUObject/Public/UObject/UObjectGlobals.h
@@ -12,8 +12,8 @@
 COREUOBJECT_API DECLARE_LOG_CATEGORY_EXTERN(LogUObjectGlobals, Log, All);
 
 DECLARE_CYCLE_STAT_EXTERN(TEXT("ConstructObject"),STAT_ConstructObject,STATGROUP_Object, );
-DECLARE_CYCLE_STAT_EXTERN(TEXT("AllocateObject"),STAT_AllocateObject,STATGROUP_Object, );
-DECLARE_CYCLE_STAT_EXTERN(TEXT("PostConstructInitializeProperties"),STAT_PostConstructInitializeProperties,STATGROUP_Object, );
+DECLARE_CYCLE_STAT_EXTERN(TEXT("AllocateObject"),STAT_AllocateObject,STATGROUP_ObjectVerbose, );
+DECLARE_CYCLE_STAT_EXTERN(TEXT("PostConstructInitializeProperties"),STAT_PostConstructInitializeProperties,STATGROUP_ObjectVerbose, );
 DECLARE_CYCLE_STAT_EXTERN(TEXT("LoadConfig"),STAT_LoadConfig,STATGROUP_Object, );
 DECLARE_CYCLE_STAT_EXTERN(TEXT("LoadObject"),STAT_LoadObject,STATGROUP_Object, );
 DECLARE_CYCLE_STAT_EXTERN(TEXT("InitProperties"),STAT_InitProperties,STATGROUP_Object, );
diff --git a/Engine/Source/Runtime/Engine/Private/Particles/ParticleComponents.cpp b/Engine/Source/Runtime/Engine/Private/Particles/ParticleComponents.cpp
--- a/Engine/Source/Runtime/Engine/Private/Particles/ParticleComponents.cpp
+++ b/Engine/Source/Runtime/Engine/Private/Particles/ParticleComponents.cpp
@@ -3945,7 +3945,7 @@ int32 UParticleSystemComponent::GetCurrentDetailMode() const
 void UParticleSystemComponent::ComputeTickComponent_Concurrent()
 {
 	SCOPE_CYCLE_COUNTER(STAT_ParticleComputeTickTime);
-	FScopeCycleCounterUObject AdditionalScope(AdditionalStatObject());
+	FScopeCycleCounterUObject AdditionalScope(AdditionalStatObject(), GET_STATID(STAT_ParticleComputeTickTime));
 	// Tick Subemitters.
 	int32 EmitterIndex;
 	for (EmitterIndex = 0; EmitterIndex < EmitterInstances.Num(); EmitterIndex++)
