[UE-7296] In progress work to help combat Blueprint circular dependency issues (currently disabled).

diff --git a/Engine/Config/BaseEngine.ini b/Engine/Config/BaseEngine.ini
--- a/Engine/Config/BaseEngine.ini
+++ b/Engine/Config/BaseEngine.ini
@@ -1484,6 +1484,10 @@ bNativeCodeGenerationTool=false
 bOptimizeExecutionFlowStack=true
 bOptimizeAdjacentStates=true
 bEnableInheritableComponents=false
+bDeferDependencyLoads=false
+bDeferCDOLoadSerialization=true
+bPropagateDeferredLoadsForStructs=true
+bCheckDeferredDependencies=false
 
 [/Script/Engine.Blueprint]
 bRecompileOnLoad=true
diff --git a/Engine/Source/Runtime/Core/Public/Misc/Build.h b/Engine/Source/Runtime/Core/Public/Misc/Build.h
--- a/Engine/Source/Runtime/Core/Public/Misc/Build.h
+++ b/Engine/Source/Runtime/Core/Public/Misc/Build.h
@@ -240,3 +240,6 @@
 
 /** Enable UberGraphPersistentFrame feature. It can speed up BP compilation (re-instancing) in editor, but introduce an unnecessary overhead in runtime. */
 #define USE_UBER_GRAPH_PERSISTENT_FRAME 1
+
+#define USE_CIRCULAR_DEPENDENCY_LOAD_DEFERRING (WITH_EDITOR && 0)
+#define TEST_CHECK_DEPENDENCY_LOAD_DEFERRING   (USE_CIRCULAR_DEPENDENCY_LOAD_DEFERRING && 1)
diff --git a/Engine/Source/Runtime/CoreUObject/Private/Blueprint/BlueprintSupport.cpp b/Engine/Source/Runtime/CoreUObject/Private/Blueprint/BlueprintSupport.cpp
--- a/Engine/Source/Runtime/CoreUObject/Private/Blueprint/BlueprintSupport.cpp
+++ b/Engine/Source/Runtime/CoreUObject/Private/Blueprint/BlueprintSupport.cpp
@@ -1,7 +1,148 @@
 // Copyright 1998-2015 Epic Games, Inc. All Rights Reserved.
 #include "CoreUObjectPrivate.h"
+#include "UObject/LinkerPlaceholderClass.h"
+#include "Linker.h"
+#include "PropertyTag.h"
 
 
+/** 
+ * Defined in BlueprintSupport.cpp
+ * Duplicates all fields of a class in depth-first order. It makes sure that everything contained
+ * in a class is duplicated before the class itself, as well as all function parameters before the
+ * function itself.
+ *
+ * @param	StructToDuplicate			Instance of the struct that is about to be duplicated
+ * @param	Writer						duplicate writer instance to write the duplicated data to
+ */
+void FBlueprintSupport::DuplicateAllFields(UStruct* StructToDuplicate, FDuplicateDataWriter& Writer)
+{
+	// This is a very simple fake topological-sort to make sure everything contained in the class
+	// is processed before the class itself is, and each function parameter is processed before the function
+	if (StructToDuplicate)
+	{
+		// Make sure each field gets allocated into the array
+		for (TFieldIterator<UField> FieldIt(StructToDuplicate, EFieldIteratorFlags::ExcludeSuper); FieldIt; ++FieldIt)
+		{
+			UField* Field = *FieldIt;
+
+			// Make sure functions also do their parameters and children first
+			if (UFunction* Function = dynamic_cast<UFunction*>(Field))
+			{
+				for (TFieldIterator<UField> FunctionFieldIt(Function, EFieldIteratorFlags::ExcludeSuper); FunctionFieldIt; ++FunctionFieldIt)
+				{
+					UField* InnerField = *FunctionFieldIt;
+					Writer.GetDuplicatedObject(InnerField);
+				}
+			}
+
+			Writer.GetDuplicatedObject(Field);
+		}
+	}
+}
+
+bool FBlueprintSupport::UseDeferredDependencyLoading()
+{
+#if USE_CIRCULAR_DEPENDENCY_LOAD_DEFERRING
+	static const FBoolConfigValueHelper DeferDependencyLoads(TEXT("Kismet"), TEXT("bDeferDependencyLoads"), GEngineIni);
+	return DeferDependencyLoads;
+#else
+	return false;
+#endif
+}
+
+bool FBlueprintSupport::UseDeferredCDOSerialization()
+{
+#if USE_CIRCULAR_DEPENDENCY_LOAD_DEFERRING
+	static const FBoolConfigValueHelper DeferCDOSerializations(TEXT("Kismet"), TEXT("bDeferCDOLoadSerialization"), GEngineIni);
+	return UseDeferredDependencyLoading() && DeferCDOSerializations;
+#else 
+	return false;
+#endif
+}
+
+bool FBlueprintSupport::UseDeferredLoadingForSubsequentLoads()
+{
+#if USE_CIRCULAR_DEPENDENCY_LOAD_DEFERRING
+	static const FBoolConfigValueHelper PropagateDeferredFlagForStructs(TEXT("Kismet"), TEXT("bPropagateDeferredLoadsForStructs"), GEngineIni);
+	return UseDeferredDependencyLoading() && PropagateDeferredFlagForStructs;
+#else 
+	return false;
+#endif
+}
+
+bool FBlueprintSupport::UseDeferredDependencyVerificationChecks()
+{
+#if TEST_CHECK_DEPENDENCY_LOAD_DEFERRING
+	static const FBoolConfigValueHelper TestCheckDeferredDependencies(TEXT("Kismet"), TEXT("bCheckDeferredDependencies"), GEngineIni);
+	return UseDeferredDependencyLoading() && TestCheckDeferredDependencies;
+#else 
+	return false;
+#endif
+}
+
+/*******************************************************************************
+ * FScopedClassDependencyGather
+ ******************************************************************************/
+
+#if WITH_EDITOR
+UClass* FScopedClassDependencyGather::BatchMasterClass = NULL;
+TArray<UClass*> FScopedClassDependencyGather::BatchClassDependencies;
+
+FScopedClassDependencyGather::FScopedClassDependencyGather(UClass* ClassToGather)
+	: bMasterClass(false)
+{
+	// Do NOT track duplication dependencies, as these are intermediate products that we don't care about
+	if( !GIsDuplicatingClassForReinstancing )
+	{
+		if( BatchMasterClass == NULL )
+		{
+			// If there is no current dependency master, register this class as the master, and reset the array
+			BatchMasterClass = ClassToGather;
+			BatchClassDependencies.Empty();
+			bMasterClass = true;
+		}
+		else
+		{
+			// This class was instantiated while another class was gathering dependencies, so record it as a dependency
+			BatchClassDependencies.AddUnique(ClassToGather);
+		}
+	}
+}
+
+FScopedClassDependencyGather::~FScopedClassDependencyGather()
+{
+	// If this gatherer was the initial gatherer for the current scope, process 
+	// dependencies (unless compiling on load is explicitly disabled)
+	if( bMasterClass && !GForceDisableBlueprintCompileOnLoad )
+	{
+		for( auto DependencyIter = BatchClassDependencies.CreateIterator(); DependencyIter; ++DependencyIter )
+		{
+			UClass* Dependency = *DependencyIter;
+			if( Dependency->ClassGeneratedBy != BatchMasterClass->ClassGeneratedBy )
+			{
+				Dependency->ConditionalRecompileClass(&GObjLoaded);
+			}
+		}
+
+		// Finally, recompile the master class to make sure it gets updated too
+		BatchMasterClass->ConditionalRecompileClass(&GObjLoaded);
+		
+		BatchMasterClass = NULL;
+	}
+}
+
+TArray<UClass*> const& FScopedClassDependencyGather::GetCachedDependencies()
+{
+	return BatchClassDependencies;
+}
+#endif //WITH_EDITOR
+
+
+
+/*******************************************************************************
+ * ULinkerLoad
+ ******************************************************************************/
+ 
 struct FPreloadMembersHelper
 {
 	static void PreloadMembers(UObject* InObject)
@@ -49,6 +190,12 @@
  */
 bool ULinkerLoad::RegenerateBlueprintClass(UClass* LoadClass, UObject* ExportObject)
 {
+#if TEST_CHECK_DEPENDENCY_LOAD_DEFERRING
+	static ULinkerLoad* RegeneratingLinker = nullptr;
+	checkSlow(!FBlueprintSupport::UseDeferredDependencyVerificationChecks() || (RegeneratingLinker == nullptr) || (RegeneratingLinker == this));
+	TGuardValue<ULinkerLoad*> ReentrantGuard(RegeneratingLinker, this);
+#endif // TEST_CHECK_DEPENDENCY_LOAD_DEFERRING
+
 	// determine if somewhere further down the callstack, we're already in this
 	// function for this class
 	bool const bAlreadyRegenerating = LoadClass->ClassGeneratedBy->HasAnyFlags(RF_BeingRegenerated);
@@ -141,6 +288,193 @@
 	return bSuccessfulRegeneration;
 }
 
+bool ULinkerLoad::DeferPotentialCircularImport(const int32 Index)
+{
+#if USE_CIRCULAR_DEPENDENCY_LOAD_DEFERRING
+	if (!FBlueprintSupport::UseDeferredDependencyLoading())
+	{
+		return false;
+	}
+
+	//--------------------------------------
+	// Phase 1: Stub in Dependencies
+	//--------------------------------------
+
+	FObjectImport& Import = ImportMap[Index];
+	if (Import.XObject != nullptr)
+	{
+		// this import has already been created
+		return Import.XObject->IsA<ULinkerPlaceholderClass>();
+	}
+
+	if ((LoadFlags & LOAD_DeferDependencyLoads) && !IsImportNative(Index))
+	{
+		if (UObject* ClassPackage = FindObject<UPackage>(/*Outer =*/nullptr, *Import.ClassPackage.ToString()))
+		{
+			if (const UClass* ImportClass = FindObject<UClass>(ClassPackage, *Import.ClassName.ToString()))
+			{
+				bool const bIsBlueprintClass = ImportClass->IsChildOf<UClass>();
+				if (bIsBlueprintClass)
+				{
+					UPackage* PlaceholderOuter = LinkerRoot;
+					UClass*   PlaceholderType  = ULinkerPlaceholderClass::StaticClass();
+
+					FName PlaceholderName(*FString::Printf(TEXT("PLACEHOLDER-CLASS_%s"), *Import.ObjectName.ToString()));
+					PlaceholderName = MakeUniqueObjectName(PlaceholderOuter, PlaceholderType, PlaceholderName);
+
+					UClass* Placeholder = ConstructObject<ULinkerPlaceholderClass>(PlaceholderType, PlaceholderOuter, PlaceholderName, RF_Public | RF_Transient);
+					// make sure the class is fully formed (has its 
+					// ClassAddReferencedObjects/ClassConstructor members set)
+					Placeholder->Bind();
+					Placeholder->StaticLink(/*bRelinkExistingProperties =*/true);
+
+					Import.XObject = Placeholder;
+				}
+			}
+		}
+	}
+	return (Import.XObject != nullptr);
+#else // !USE_CIRCULAR_DEPENDENCY_LOAD_DEFERRING
+	return false;
+#endif // USE_CIRCULAR_DEPENDENCY_LOAD_DEFERRING
+}
+
+void ULinkerLoad::ResolveDeferredDependencies(UClass* LoadClass)
+{
+#if USE_CIRCULAR_DEPENDENCY_LOAD_DEFERRING
+	if (!FBlueprintSupport::UseDeferredDependencyLoading())
+	{
+		return;
+	}
+
+	//--------------------------------------
+	// Phase 2: Resolve Dependency Stubs
+	//--------------------------------------
+	TGuardValue<uint32> LoadFlagsGuard(LoadFlags, LoadFlags & ~LOAD_DeferDependencyLoads);
+
+	for (int32 ImportIndex = 0; ImportIndex < ImportMap.Num(); ++ImportIndex)
+	{
+		FObjectImport& Import = ImportMap[ImportIndex];
+
+		if (ULinkerPlaceholderClass* PlaceholderClass = Cast<ULinkerPlaceholderClass>(Import.XObject))
+		{
+			Import.XObject = nullptr;
+			UClass* RealClassObj = CastChecked<UClass>(CreateImport(ImportIndex), ECastCheckedType::NullAllowed);
+
+			for (UObjectPropertyBase* ObjProperty : PlaceholderClass->ReferencingProperties)
+			{
+				ObjProperty->PropertyClass = RealClassObj;
+			}
+			PlaceholderClass->ReferencingProperties.Empty();
+			PlaceholderClass->MarkPendingKill(); // @TODO: ensure these are properly GC'd
+		}
+	}
+
+	// @TODO: don't know if we need this, but could be good to have (as class 
+	//        regeneration is about to force load a lot of this).
+	LoadAllObjects();
+
+	//--------------------------------------
+	// Phase 3: Finalize (serialize CDO & regenerate class)
+	//--------------------------------------
+
+	// since this function is called after class serialization, then we can be
+	// certain that the CDO export object exists (since class serialization 
+	// would call CreateExport() for the CDO)
+
+#if TEST_CHECK_DEPENDENCY_LOAD_DEFERRING
+	bool const bCircumventValidationChecks = !FBlueprintSupport::UseDeferredDependencyVerificationChecks();
+	checkSlow(bCircumventValidationChecks || (LoadClass->ClassDefaultObject != nullptr));
+#endif // TEST_CHECK_DEPENDENCY_LOAD_DEFERRING
+
+	// this should assert if the index hasn't been set (as we expect)
+	int32 CDOExportIndex = LoadClass->ClassDefaultObject->GetLinkerIndex();
+	if (CDOExportIndex == INDEX_NONE)
+	{
+		CDOExportIndex = DeferredExportIndex;
+	}
+	else if (DeferredExportIndex != INDEX_NONE)
+	{
+		check(CDOExportIndex == DeferredExportIndex);
+	}
+	// else, the CDO hasn't been touched (which is expected, since the CDO needs the class first)
+	// @TODO: this comment here ^ is in direct conflict with the assumption 
+	//        stated above (that the CDO should be loaded)... figure out which is correct!
+
+	{
+		check(CDOExportIndex != INDEX_NONE);
+		FObjectExport& CDOExport = ExportMap[CDOExportIndex];
+
+		UObject* CDO = CDOExport.Object;
+		if (FBlueprintSupport::UseDeferredCDOSerialization())
+		{
+			// have to prematurely set the CDO's linker so we can force a Preload()/
+			// Serialization of the CDO before we regenerate the Blueprint class
+			{
+				const EObjectFlags OldFlags = CDO->GetFlags();
+				CDO->ClearFlags(RF_NeedLoad | RF_NeedPostLoad);
+				CDO->SetLinker(this, CDOExportIndex, /*bShouldDetatchExisting =*/false);
+				CDO->SetFlags(OldFlags);
+			}
+			// should load the CDO (ensuring that it has been serialized in by the 
+			// time we get to class regeneration)
+			//
+			// NOTE: this is point where circular dependencies could reveal 
+			//       themselves, as the CDO could depend on a class not listed in 
+			//       the package's imports
+			//
+			// @TODO: should we block/detect reentrant behavior? would this CDO ever
+			//        need another CDO?
+			Preload(CDO);
+		}
+
+		UClass* BlueprintClass = Cast<UClass>(IndexToObject(CDOExport.ClassIndex));
+#if TEST_CHECK_DEPENDENCY_LOAD_DEFERRING
+		checkSlow(bCircumventValidationChecks || CDO->HasAnyFlags(RF_LoadCompleted));
+		checkSlow(bCircumventValidationChecks || BlueprintClass->HasAnyClassFlags(CLASS_CompiledFromBlueprint));
+
+		// at this point there should not be any instances of the Blueprint (else,
+		// we'd have to reinstance and that is too expensive an operation to 
+		// have at load time)
+		TArray<UObject*> ClassInstances;
+		GetObjectsOfClass(BlueprintClass, ClassInstances, /*bIncludeDerivedClasses =*/true);
+		checkSlow(bCircumventValidationChecks || ClassInstances.Num() == 0);
+
+		for (TObjectIterator<ULinkerPlaceholderClass> PlaceholderIt; PlaceholderIt; ++PlaceholderIt)
+		{
+			ULinkerPlaceholderClass* PlaceholderClass = *PlaceholderIt;
+			// there shouldn't be any deferred dependencies that need to be 
+			// resolved by this point
+			checkSlow(bCircumventValidationChecks || (PlaceholderClass->GetOuter() != LinkerRoot) || PlaceholderClass->IsPendingKill());
+		}
+
+		// @TODO: Test verify there are no instances of the class yet (except for the CDO)
+
+#endif // TEST_CHECK_DEPENDENCY_LOAD_DEFERRING
+
+		UObject* OldCDO = BlueprintClass->ClassDefaultObject;
+		if (RegenerateBlueprintClass(BlueprintClass, CDO))
+		{
+			// emulate class CDO serialization (RegenerateBlueprintClass() could 
+			// have a side-effect where it overwrites the class's CDO; so we 
+			// want to make sure that we don't overwrite that new CDO with a 
+			// stale one)
+			if (OldCDO == BlueprintClass->ClassDefaultObject)
+			{
+				BlueprintClass->ClassDefaultObject = CDOExport.Object;
+			}
+		}
+		
+		// 
+		DeferredExportIndex = INDEX_NONE;
+	}
+#endif // USE_CIRCULAR_DEPENDENCY_LOAD_DEFERRING
+}
+
+/*******************************************************************************
+ * UObject
+ ******************************************************************************/
+
 /** 
  * Returns whether this object is contained in or part of a blueprint object
  */
@@ -188,6 +522,10 @@
 	}
 }
 
+/*******************************************************************************
+ * FObjectInitializer
+ ******************************************************************************/
+
 /** 
  * Initializes a non-native property, according to the initialization rules. If the property is non-native
  * and does not have a zero constructor, it is initialized with the default value.
@@ -211,92 +549,3 @@
 		return false;
 	}
 }
-
-/** 
- * Defined in BlueprintSupport.cpp
- * Duplicates all fields of a class in depth-first order. It makes sure that everything contained
- * in a class is duplicated before the class itself, as well as all function parameters before the
- * function itself.
- *
- * @param	StructToDuplicate			Instance of the struct that is about to be duplicated
- * @param	Writer						duplicate writer instance to write the duplicated data to
- */
-void FBlueprintSupport::DuplicateAllFields(UStruct* StructToDuplicate, FDuplicateDataWriter& Writer)
-{
-	// This is a very simple fake topological-sort to make sure everything contained in the class
-	// is processed before the class itself is, and each function parameter is processed before the function
-	if (StructToDuplicate)
-	{
-		// Make sure each field gets allocated into the array
-		for (TFieldIterator<UField> FieldIt(StructToDuplicate, EFieldIteratorFlags::ExcludeSuper); FieldIt; ++FieldIt)
-		{
-			UField* Field = *FieldIt;
-
-			// Make sure functions also do their parameters and children first
-			if (UFunction* Function = dynamic_cast<UFunction*>(Field))
-			{
-				for (TFieldIterator<UField> FunctionFieldIt(Function, EFieldIteratorFlags::ExcludeSuper); FunctionFieldIt; ++FunctionFieldIt)
-				{
-					UField* InnerField = *FunctionFieldIt;
-					Writer.GetDuplicatedObject(InnerField);
-				}
-			}
-
-			Writer.GetDuplicatedObject(Field);
-		}
-	}
-}
-
-#if WITH_EDITOR
-UClass* FScopedClassDependencyGather::BatchMasterClass = NULL;
-TArray<UClass*> FScopedClassDependencyGather::BatchClassDependencies;
-
-FScopedClassDependencyGather::FScopedClassDependencyGather(UClass* ClassToGather)
-	: bMasterClass(false)
-{
-	// Do NOT track duplication dependencies, as these are intermediate products that we don't care about
-	if( !GIsDuplicatingClassForReinstancing )
-	{
-		if( BatchMasterClass == NULL )
-		{
-			// If there is no current dependency master, register this class as the master, and reset the array
-			BatchMasterClass = ClassToGather;
-			BatchClassDependencies.Empty();
-			bMasterClass = true;
-		}
-		else
-		{
-			// This class was instantiated while another class was gathering dependencies, so record it as a dependency
-			BatchClassDependencies.AddUnique(ClassToGather);
-		}
-	}
-}
-
-FScopedClassDependencyGather::~FScopedClassDependencyGather()
-{
-	// If this gatherer was the initial gatherer for the current scope, process 
-	// dependencies (unless compiling on load is explicitly disabled)
-	if( bMasterClass && !GForceDisableBlueprintCompileOnLoad )
-	{
-		for( auto DependencyIter = BatchClassDependencies.CreateIterator(); DependencyIter; ++DependencyIter )
-		{
-			UClass* Dependency = *DependencyIter;
-			if( Dependency->ClassGeneratedBy != BatchMasterClass->ClassGeneratedBy )
-			{
-				Dependency->ConditionalRecompileClass(&GObjLoaded);
-			}
-		}
-
-		// Finally, recompile the master class to make sure it gets updated too
-		BatchMasterClass->ConditionalRecompileClass(&GObjLoaded);
-		
-		BatchMasterClass = NULL;
-	}
-}
-
-TArray<UClass*> const& FScopedClassDependencyGather::GetCachedDependencies()
-{
-	return BatchClassDependencies;
-}
-
-#endif //WITH_EDITOR
diff --git a/Engine/Source/Runtime/CoreUObject/Private/UObject/Class.cpp b/Engine/Source/Runtime/CoreUObject/Private/UObject/Class.cpp
--- a/Engine/Source/Runtime/CoreUObject/Private/UObject/Class.cpp
+++ b/Engine/Source/Runtime/CoreUObject/Private/UObject/Class.cpp
@@ -3024,19 +3024,41 @@ void UClass::Serialize( FArchive& Ar )
 	{
 		check((Ar.GetPortFlags() & PPF_Duplicate) || (GetStructureSize() >= sizeof(UObject)));
 		check(!GetSuperClass() || !GetSuperClass()->HasAnyFlags(RF_NeedLoad));
+		
+		// record the current CDO, as it stands, so we can compare against it 
+		// after we've serialized in the new CDO (to detect if, as a side-effect
+		// of the serialization, a different CDO was generated)
 		UObject* const OldCDO = ClassDefaultObject;
-		UObject* TempClassDefaultObject = NULL;
-		Ar << TempClassDefaultObject;
-		// we need to avoid a case, when while class regeneration a different CDO is set, and now we set it to an stale one (ttp#343166)
+
+		// serialize in the CDO, but first store it here (in a temporary var) so
+		// we can check to see if it should be the authoritative CDO (a newer 
+		// CDO could be generated as a side-effect of this serialization)
+		//
+		// @TODO: for USE_CIRCULAR_DEPENDENCY_LOAD_DEFERRING, do we need to 
+		//        defer this serialization (should we just save off the tagged
+		//        serialization data for later use)?
+		UObject* PerspectiveNewCDO = NULL;
+		Ar << PerspectiveNewCDO;
+
+		// Blueprint class regeneration could cause the class's CDO to be set.
+		// The CDO (<<) serialization call (above) probably will invoke class 
+		// regeneration, and as a side-effect the CDO could already be set by 
+		// the time it returns. So we only want to set the CDO here (to what was 
+		// serialized in) if it hasn't already changed (else, the serialized
+		// version could be stale). See: TTP #343166
 		if (ClassDefaultObject == OldCDO)
 		{
-			ClassDefaultObject = TempClassDefaultObject;
+			ClassDefaultObject = PerspectiveNewCDO;
 		}
-		else if (TempClassDefaultObject != ClassDefaultObject)
+		// if we reach this point, then the CDO was regenerated as a side-effect
+		// of the serialization... let's log if the regenerated CDO (what's 
+		// already been set) is not the same as what was returned from the 
+		// serialization (could mean the CDO was regenerated multiple times?)
+		else if (PerspectiveNewCDO != ClassDefaultObject)
 		{
 			UE_LOG(LogClass, Log, TEXT("CDO was changed while class serialization.\n\tOld: '%s'\n\tSerialized: '%s'\n\tActual: '%s'")
 				, OldCDO ? *OldCDO->GetFullName() : TEXT("NULL")
-				, TempClassDefaultObject ? *TempClassDefaultObject->GetFullName() : TEXT("NULL")
+				, PerspectiveNewCDO ? *PerspectiveNewCDO->GetFullName() : TEXT("NULL")
 				, ClassDefaultObject ? *ClassDefaultObject->GetFullName() : TEXT("NULL"));
 		}
 		ClassUnique = 0;
diff --git a/Engine/Source/Runtime/CoreUObject/Private/UObject/GarbageCollection.cpp b/Engine/Source/Runtime/CoreUObject/Private/UObject/GarbageCollection.cpp
--- a/Engine/Source/Runtime/CoreUObject/Private/UObject/GarbageCollection.cpp
+++ b/Engine/Source/Runtime/CoreUObject/Private/UObject/GarbageCollection.cpp
@@ -7,6 +7,7 @@
 #include "CoreUObjectPrivate.h"
 #include "TaskGraphInterfaces.h"
 #include "IConsoleManager.h"
+#include "LinkerPlaceholderClass.h"
 
 /*-----------------------------------------------------------------------------
    Garbage collection.
@@ -1214,7 +1215,8 @@ void UObject::AddReferencedObjects(UObject* This, FReferenceCollector& Collector
 
 		// Serialize object properties which are defined in the class.
 		// Note: This check is intentionally excluding UClass objects but including subclasses like UBlueprintGeneratedClass
-		if (Class != UClass::StaticClass())
+		// @TODO: is it right to also exclude ULinkerPlaceholderClass here (it was causing a crash in here).
+		if (Class != UClass::StaticClass() && Class != ULinkerPlaceholderClass::StaticClass())
 		{
 			// Script properties
 			FSimpleObjectReferenceCollectorArchive ObjectReferenceCollector( This, Collector );
diff --git a/Engine/Source/Runtime/CoreUObject/Private/UObject/LinkerLoad.cpp b/Engine/Source/Runtime/CoreUObject/Private/UObject/LinkerLoad.cpp
--- a/Engine/Source/Runtime/CoreUObject/Private/UObject/LinkerLoad.cpp
+++ b/Engine/Source/Runtime/CoreUObject/Private/UObject/LinkerLoad.cpp
@@ -6,6 +6,7 @@
 #include "MessageLog.h"
 #include "UObjectToken.h"
 #include "EngineVersion.h"
+#include "LinkerPlaceholderClass.h"
 
 #define LOCTEXT_NAMESPACE "LinkerLoad"
 
@@ -14,6 +15,7 @@ DECLARE_STATS_GROUP_VERBOSE(TEXT("Linker Load"), STATGROUP_LinkerLoad, STATCAT_A
 DECLARE_CYCLE_STAT(TEXT("Linker Preload"),STAT_LinkerPreload,STATGROUP_LinkerLoad);
 DECLARE_CYCLE_STAT(TEXT("Linker Precache"),STAT_LinkerPrecache,STATGROUP_LinkerLoad);
 DECLARE_CYCLE_STAT(TEXT("Linker Serialize"),STAT_LinkerSerialize,STATGROUP_LinkerLoad);
+DECLARE_CYCLE_STAT(TEXT("Linker Load Deferred"), STAT_LinkerLoadDeferred, STATGROUP_LinkerLoad);
 DECLARE_DWORD_ACCUMULATOR_STAT(TEXT("Linker Count"), STAT_LinkerCount, STATGROUP_LinkerLoad);
 DECLARE_DWORD_ACCUMULATOR_STAT(TEXT("Live Linker Count"), STAT_LiveLinkerCount, STATGROUP_LinkerLoad);
 
@@ -348,8 +350,33 @@ ULinkerLoad* ULinkerLoad::CreateLinker( UPackage* Parent, const TCHAR* Filename,
 	{
 		UE_LOG(LogLinker, Warning, TEXT("ULinkerLoad::CreateLinker(%s) blocking async loading!"), Filename);
 	}
+
+#if USE_CIRCULAR_DEPENDENCY_LOAD_DEFERRING
+	// we don't want the linker permanently created with the 
+	// DeferDependencyLoads flag (we also want to be able to determine if the 
+	// linker already exists with that flag), so clear it before we attempt 
+	// CreateLinkerAsync()
+	// 
+	// if this flag is present here, then we're most likely in a nested load and a 
+	// blueprint up the load chain needed an asset (most likely a user-defined 
+	// struct) loaded (we expect calls with LOAD_DeferDependencyLoads to be 
+	// coming from LoadPackageInternal)
+	uint32 const DeferredLoadFlag = (LoadFlags & LOAD_DeferDependencyLoads);
+	LoadFlags &= ~LOAD_DeferDependencyLoads;
+#endif // USE_CIRCULAR_DEPENDENCY_LOAD_DEFERRING
+
 	ULinkerLoad* Linker = CreateLinkerAsync( Parent, Filename, LoadFlags );
 	{
+#if USE_CIRCULAR_DEPENDENCY_LOAD_DEFERRING
+		// the linker could already have the DeferDependencyLoads flag present 
+		// (if this linker was already created further up the load chain, and 
+		// we're re-entering this to further finalize its creation)... we want 
+		// to make sure the DeferDependencyLoads flag is supplied (if it was 
+		// specified) for the duration of the Tick() below, because its call to 
+		// FinalizeCreation() could invoke further dependency loads
+		TGuardValue<uint32> LinkerLoadFlagGuard(Linker->LoadFlags, Linker->LoadFlags | DeferredLoadFlag);
+#endif // USE_CIRCULAR_DEPENDENCY_LOAD_DEFERRING
+
 		FSerializedPackageLinkerGuard Guard;	
 		GSerializedPackageLinker = Linker;
 		if (Linker->Tick( 0.f, false, false ) == LINKER_Failed)
@@ -600,7 +627,10 @@ ULinkerLoad::ULinkerLoad( const FObjectInitializer& ObjectInitializer, UPackage*
 ,	bHaveImportsBeenVerified( false )
 #if WITH_EDITOR
 ,	LoadProgressScope( nullptr )
-#endif
+#endif // WITH_EDITOR
+#if USE_CIRCULAR_DEPENDENCY_LOAD_DEFERRING
+,	DeferredExportIndex(INDEX_NONE)
+#endif // USE_CIRCULAR_DEPENDENCY_LOAD_DEFERRING
 {
 	check(!HasAnyFlags(RF_ClassDefaultObject));
 
@@ -2123,6 +2153,26 @@ bool ULinkerLoad::VerifyImportInner(const int32 ImportIndex, FString& WarningSuf
 
 		if (!bWasFullyLoaded)
 		{
+
+#if USE_CIRCULAR_DEPENDENCY_LOAD_DEFERRING
+			// when LOAD_DeferDependencyLoads is in play, we usually head off 
+			// dependency loads before we get to this point, but there are two 
+			// cases where we can reach here intentionally: 
+			//
+			//   1) the package we're attempting to load is native (and thusly,
+			//      LoadPackageInternal() should fail, and retrun null)
+			//
+			//   2) the package we're attempting to load is a user defined 
+			//      struct asset, which we need to load because the blueprint 
+			//      class's layout depends on the struct's size... in this case, 
+			//      we choke off circular loads by propagating this flag along 
+			//      to the struct linker (so it doesn't load any blueprints)
+			if (FBlueprintSupport::UseDeferredLoadingForSubsequentLoads())
+			{
+				InternalLoadFlags |= (LoadFlags & LOAD_DeferDependencyLoads);
+			}
+#endif // USE_CIRCULAR_DEPENDENCY_LOAD_DEFERRING
+
 			// we now fully load the package that we need a single export from - however, we still use CreatePackage below as it handles all cases when the package
 			// didn't exist (native only), etc		
 			TmpPkg = LoadPackageInternal(NULL, *Import.ObjectName.ToString(), InternalLoadFlags, this);
@@ -2488,6 +2538,15 @@ void ULinkerLoad::LoadAllObjects( bool bForcePreload )
 	SlowTask.bVisibleOnUI = false;
 #endif
 
+#if TEST_CHECK_DEPENDENCY_LOAD_DEFERRING
+	// if we're re-entering a call to LoadAllObjects() while DeferDependencyLoads
+	// is set, then we're not doing our job (we're risking an export needing 
+	// another external asset)... if this is hit, then we're most likely already
+	// in this function (for this linker) further up the load chain; it should 
+	// finish the loads there
+	checkSlow( ((LoadFlags & LOAD_DeferDependencyLoads) == 0) || !FBlueprintSupport::UseDeferredDependencyVerificationChecks() );
+#endif // TEST_CHECK_DEPENDENCY_LOAD_DEFERRING
+
 	if ( (LoadFlags&LOAD_SeekFree) != 0 )
 	{
 		bForcePreload = true;
@@ -2706,6 +2765,23 @@ void ULinkerLoad::Preload( UObject* Object )
 {
 	check(IsValidLowLevel());
 	check(Object);
+
+#if USE_CIRCULAR_DEPENDENCY_LOAD_DEFERRING
+	UClass* ObjectAsClass = Cast<UClass>(Object);
+	// we can determine that this is a blueprint class, if it is indeed a class
+	// AND it is not native (blueprint classes are the only asset package 
+	// classes we have)
+	bool const bIsBlueprintClass = (ObjectAsClass != nullptr) && !(ObjectAsClass->GetOutermost()->PackageFlags & PKG_CompiledIn);
+
+#if TEST_CHECK_DEPENDENCY_LOAD_DEFERRING
+	// we should NEVER be pre-loading another blueprint class when the 
+	// DeferDependencyLoads flag is set (some other blueprint class is already 
+	// being loaded further up the load chain, and this could introduce a 
+	// circular load)
+	checkSlow(!bIsBlueprintClass || !Object->HasAnyFlags(RF_NeedLoad) || !(LoadFlags & LOAD_DeferDependencyLoads) || !FBlueprintSupport::UseDeferredDependencyVerificationChecks());
+#endif // TEST_CHECK_DEPENDENCY_LOAD_DEFERRING
+#endif // USE_CIRCULAR_DEPENDENCY_LOAD_DEFERRING
+	
 	// Preload the object if necessary.
 	if( Object->HasAnyFlags(RF_NeedLoad) )
 	{
@@ -2725,11 +2801,20 @@ void ULinkerLoad::Preload( UObject* Object )
 				}
 			}
 
+#if USE_CIRCULAR_DEPENDENCY_LOAD_DEFERRING
+			TGuardValue<uint32> LoadFlagsGuard(LoadFlags, LoadFlags);			
+			if (FBlueprintSupport::UseDeferredDependencyLoading() && bIsBlueprintClass)
+			{
+				LoadFlags |= LOAD_DeferDependencyLoads;
+			}
+#endif // USE_CIRCULAR_DEPENDENCY_LOAD_DEFERRING
+
 			// make sure this object didn't get loaded in the above Preload call
 			if (Object->HasAnyFlags(RF_NeedLoad))
 			{
 				// grab the resource for this Object
-				FObjectExport& Export = ExportMap[ Object->GetLinkerIndex() ];
+				int32 const ExportIndex = Object->GetLinkerIndex();
+				FObjectExport& Export = ExportMap[ExportIndex];
 				check(Export.Object==Object);
 
 				const int64 SavedPos = Loader->Tell();
@@ -2751,8 +2836,32 @@ void ULinkerLoad::Preload( UObject* Object )
 					SCOPE_CYCLE_COUNTER(STAT_LinkerSerialize);
 					if ( Object->HasAnyFlags(RF_ClassDefaultObject) )
 					{
-						Object->GetClass()->SerializeDefaultObject(Object, *this);
-						Object->SetFlags(RF_LoadCompleted);
+#if USE_CIRCULAR_DEPENDENCY_LOAD_DEFERRING
+						if ( ((LoadFlags & LOAD_DeferDependencyLoads) != 0) && FBlueprintSupport::UseDeferredCDOSerialization() )
+						{
+#if TEST_CHECK_DEPENDENCY_LOAD_DEFERRING
+							checkSlow(!FBlueprintSupport::UseDeferredDependencyVerificationChecks() || (DeferredExportIndex == INDEX_NONE) || (DeferredExportIndex == ExportIndex));
+#endif // TEST_CHECK_DEPENDENCY_LOAD_DEFERRING
+
+							// since serializing the CDO can introduce circular 
+							// dependencies, we want to stave that off until 
+							// we're ready to handle those 
+							DeferredExportIndex = ExportIndex;
+							// don't need to actually "consume" the data through
+							// serialization though (since we seek back to 
+							// SavedPos later on)
+
+							// reset the flag and return (don't worry, we make
+							// sure to force load this later)
+							Object->SetFlags(RF_NeedLoad);
+							return;
+						}
+						else
+#endif // USE_CIRCULAR_DEPENDENCY_LOAD_DEFERRING
+						{
+							Object->GetClass()->SerializeDefaultObject(Object, *this);
+							Object->SetFlags(RF_LoadCompleted);
+						}
 					}
 					else
 					{
@@ -2765,6 +2874,16 @@ void ULinkerLoad::Preload( UObject* Object )
 					}
 				}
 
+				{
+#if USE_CIRCULAR_DEPENDENCY_LOAD_DEFERRING
+					SCOPE_CYCLE_COUNTER(STAT_LinkerLoadDeferred);
+					if ((LoadFlags & LOAD_DeferDependencyLoads) != (*LoadFlagsGuard & LOAD_DeferDependencyLoads))
+					{
+						ResolveDeferredDependencies(ObjectAsClass);
+					}
+#endif // USE_CIRCULAR_DEPENDENCY_LOAD_DEFERRING
+				}
+
 				// Make sure we serialized the right amount of stuff.
 				if( Tell()-Export.SerialOffset != Export.SerialSize )
 				{
@@ -2772,7 +2891,7 @@ void ULinkerLoad::Preload( UObject* Object )
 					{
 						UE_LOG(LogLinker, Warning, TEXT("%s"), *FString::Printf( TEXT("%s: Serial size mismatch: Got %d, Expected %d"), *Object->GetFullName(), (int32)(Tell()-Export.SerialOffset), Export.SerialSize ) );
 					}
-					else
+					else //if (!Object->HasAnyFlags(RF_ClassDefaultObject) || (LoadFlags & LOAD_DeferDependencyLoads) == 0)
 					{
 						UE_LOG(LogLinker, Fatal, TEXT("%s"), *FString::Printf( TEXT("%s: Serial size mismatch: Got %d, Expected %d"), *Object->GetFullName(), (int32)(Tell()-Export.SerialOffset), Export.SerialSize ) );
 					}
@@ -3222,6 +3341,16 @@ UObject* ULinkerLoad::CreateExport( int32 Index )
 				&& (LoadClass->GetOutermost() != GetTransientPackage()) 
 				&& ((Export.ObjectFlags&RF_ClassDefaultObject) != 0) )
 			{
+#if USE_CIRCULAR_DEPENDENCY_LOAD_DEFERRING
+				if ((LoadFlags & LOAD_DeferDependencyLoads) != 0)
+				{
+					// if LOAD_DeferDependencyLoads is set, then we're already
+					// serializing the blueprint's class somewhere up the chain,
+					DeferredExportIndex = Index;
+					return Export.Object;
+				}
+#endif // USE_CIRCULAR_DEPENDENCY_LOAD_DEFERRING
+
 				{
 					// For classes that are about to be regenerated, make sure we register them with the linker, so future references to this linker index will be valid
 					const EObjectFlags OldFlags = Export.Object->GetFlags();
@@ -3289,15 +3418,42 @@ UObject* ULinkerLoad::CreateExport( int32 Index )
 	return Export.Object;
 }
 
+bool ULinkerLoad::IsImportNative(const int32 Index) const
+{
+	const FObjectImport& Import = ImportMap[Index];
+
+	bool bIsImportNative = false;
+	// if this import has a linker, then it belongs to some (non-native) asset package 
+	if (Import.SourceLinker == nullptr)
+	{
+		if (!Import.OuterIndex.IsNull())
+		{
+			// need to check the package that this import belongs to, so recurse
+			// up then import's outer chain
+			bIsImportNative = IsImportNative(Import.OuterIndex.ToImport());
+		}
+		else if (UPackage* ExistingPackage = FindObject<UPackage>(/*Outer =*/nullptr, *Import.ObjectName.ToString()))
+		{
+			// @TODO: what if the package's outer isn't null... what does that mean?
+			bIsImportNative = !ExistingPackage->GetOuter() && (ExistingPackage->PackageFlags & PKG_CompiledIn);
+		}
+	}
+
+	return bIsImportNative;
+}
+
 // Return the loaded object corresponding to an import index; any errors are fatal.
 UObject* ULinkerLoad::CreateImport( int32 Index )
 {
 	FScopedCreateImportCounter ScopedCounter( this, Index );
 	FObjectImport& Import = ImportMap[ Index ];
+	
+#if USE_CIRCULAR_DEPENDENCY_LOAD_DEFERRING
+	DeferPotentialCircularImport(Index);
+#endif // USE_CIRCULAR_DEPENDENCY_LOAD_DEFERRING
 
 	if( Import.XObject == NULL )
 	{
-		// Look in memory first.
 		if (!GIsEditor && !IsRunningCommandlet())
 		{
 			// Try to find existing version in memory first.
diff --git a/Engine/Source/Runtime/CoreUObject/Private/UObject/LinkerPlaceholderClass.cpp b/Engine/Source/Runtime/CoreUObject/Private/UObject/LinkerPlaceholderClass.cpp
new file mode 100644
--- /dev/null
+++ b/Engine/Source/Runtime/CoreUObject/Private/UObject/LinkerPlaceholderClass.cpp
@@ -0,0 +1,53 @@
+// Copyright 1998-2015 Epic Games, Inc. All Rights Reserved.
+
+#include "CoreUObjectPrivate.h"
+#include "LinkerPlaceholderClass.h"
+
+ULinkerPlaceholderClass::ULinkerPlaceholderClass(const FObjectInitializer& ObjectInitializer)
+	: Super(ObjectInitializer)
+{
+}
+
+ULinkerPlaceholderClass::~ULinkerPlaceholderClass()
+{
+	for (UObjectPropertyBase* ObjProperty : ReferencingProperties)
+	{
+		check(ObjProperty->PropertyClass == this);
+		ObjProperty->PropertyClass = nullptr;
+	}
+}
+
+void ULinkerPlaceholderClass::PostInitProperties()
+{
+	Super::PostInitProperties();
+	if (HasAnyFlags(RF_ClassDefaultObject))
+	{
+		// Default__LinkerPlaceholderClass uses its own AddReferencedObjects function.
+		ClassAddReferencedObjects = &ULinkerPlaceholderClass::AddReferencedObjects;
+	}
+}
+
+void ULinkerPlaceholderClass::Bind()
+{
+	ClassConstructor = InternalConstructor<ULinkerPlaceholderClass>;
+	Super::Bind();
+
+	ClassAddReferencedObjects = &ULinkerPlaceholderClass::AddReferencedObjects;
+}
+
+IMPLEMENT_CORE_INTRINSIC_CLASS(ULinkerPlaceholderClass, UClass, 
+	{
+		Class->ClassAddReferencedObjects = &ULinkerPlaceholderClass::AddReferencedObjects;
+		
+		// @TODO: use the Class->Emit...() functions here to aid garbage 
+		//        collection, so it has information on what class variables 
+		//        hold onto UObject references
+	}
+);
+
+void ULinkerPlaceholderClass::AddReferencedObjects(UObject* InThis, FReferenceCollector& Collector)
+{
+	ULinkerPlaceholderClass* This = CastChecked<ULinkerPlaceholderClass>(InThis);
+	//... 
+	Super::AddReferencedObjects(InThis, Collector);
+}
diff --git a/Engine/Source/Runtime/CoreUObject/Private/UObject/LinkerPlaceholderClass.h b/Engine/Source/Runtime/CoreUObject/Private/UObject/LinkerPlaceholderClass.h
new file mode 100644
--- /dev/null
+++ b/Engine/Source/Runtime/CoreUObject/Private/UObject/LinkerPlaceholderClass.h
@@ -0,0 +1,38 @@
+// Copyright 1998-2014 Epic Games, Inc. All Rights Reserved.
+
+#pragma once
+#include "Class.h"
+#include "Set.h"
+
+// Forward declarations
+class FObjectInitializer;
+class UObjectPropertyBase;
+
+/**  
+ * A utility class for the deferred dependency loader, used to stub in temporary
+ * class references so we don't have to load blueprint resources for their class.
+ * Holds on to references where this is currently being utilized, so we can 
+ * easily replace references to it later (once the real class is available).
+ */ 
+class ULinkerPlaceholderClass : public UClass
+{
+public:
+	DECLARE_CASTED_CLASS_INTRINSIC_NO_CTOR(ULinkerPlaceholderClass, UClass, /*TStaticFlags =*/0, CoreUObject, /*TStaticCastFlags =*/0, NO_API)
+
+	ULinkerPlaceholderClass(const FObjectInitializer& ObjectInitializer);
+	virtual ~ULinkerPlaceholderClass();
+
+	static void AddReferencedObjects(UObject* InThis, FReferenceCollector& Collector);
+
+	// UObject interface.
+ 	virtual void PostInitProperties() override;
+	// End of UObject interface.
+
+	// UField interface.
+ 	virtual void Bind() override;
+	// End of UField interface.
+
+public:
+	/** Links to UObject properties that are currently referencing this class */
+	TSet<UObjectPropertyBase*> ReferencingProperties;
+}; 
\ No newline at end of file
diff --git a/Engine/Source/Runtime/CoreUObject/Private/UObject/PropertyBaseObject.cpp b/Engine/Source/Runtime/CoreUObject/Private/UObject/PropertyBaseObject.cpp
--- a/Engine/Source/Runtime/CoreUObject/Private/UObject/PropertyBaseObject.cpp
+++ b/Engine/Source/Runtime/CoreUObject/Private/UObject/PropertyBaseObject.cpp
@@ -2,11 +2,24 @@
 
 #include "CoreUObjectPrivate.h"
 #include "PropertyHelper.h"
+#include "LinkerPlaceholderClass.h"
 
 /*-----------------------------------------------------------------------------
 	UObjectPropertyBase.
 -----------------------------------------------------------------------------*/
 
+UObjectPropertyBase::~UObjectPropertyBase()
+{
+#if USE_CIRCULAR_DEPENDENCY_LOAD_DEFERRING
+	// @TODO: sometimes PropertyClass is freed memory, so operating on it causes 
+	//        a crash
+// 	if (ULinkerPlaceholderClass* PlaceholderClass = Cast<ULinkerPlaceholderClass>(PropertyClass))
+// 	{
+// 		PlaceholderClass->ReferencingProperties.Remove(this);
+// 	}
+#endif // USE_CIRCULAR_DEPENDENCY_LOAD_DEFERRING
+}
+
 void UObjectPropertyBase::InstanceSubobjects(void* Data, void const* DefaultData, UObject* Owner, FObjectInstancingGraph* InstanceGraph )
 {
 	for ( int32 ArrayIndex = 0; ArrayIndex < ArrayDim; ArrayIndex++ )
@@ -74,6 +87,16 @@ void UObjectPropertyBase::Serialize( FArchive& Ar )
 {
 	Super::Serialize( Ar );
 	Ar << PropertyClass;
+
+#if USE_CIRCULAR_DEPENDENCY_LOAD_DEFERRING
+	if (Ar.IsLoading())
+	{
+		if (ULinkerPlaceholderClass* PlaceholderClass = Cast<ULinkerPlaceholderClass>(PropertyClass))
+		{
+			PlaceholderClass->ReferencingProperties.Add(this);
+		}
+	}
+#endif // USE_CIRCULAR_DEPENDENCY_LOAD_DEFERRING
 }
 void UObjectPropertyBase::AddReferencedObjects(UObject* InThis, FReferenceCollector& Collector)
 {	
diff --git a/Engine/Source/Runtime/CoreUObject/Private/UObject/UObjectGlobals.cpp b/Engine/Source/Runtime/CoreUObject/Private/UObject/UObjectGlobals.cpp
--- a/Engine/Source/Runtime/CoreUObject/Private/UObject/UObjectGlobals.cpp
+++ b/Engine/Source/Runtime/CoreUObject/Private/UObject/UObjectGlobals.cpp
@@ -878,7 +878,19 @@ UPackage* LoadPackageInternal(UPackage* InOuter, const TCHAR* InLongPackageName,
 
 		SlowTask.EnterProgressFrame(30);
 
-		if( !(LoadFlags & LOAD_Verify) )
+		uint32 DoNotLoadExportsFlags = LOAD_Verify;
+#if USE_CIRCULAR_DEPENDENCY_LOAD_DEFERRING
+		// if this linker already has the DeferDependencyLoads flag, then we're
+		// already loading it earlier up the load chain (don't let it invoke any
+		// deeper loads that may introduce a circular dependency)
+		if ((Linker->LoadFlags & LOAD_DeferDependencyLoads))
+		{
+			DoNotLoadExportsFlags |= LOAD_DeferDependencyLoads;
+		}
+		// @TODO: what of cases where DeferDependencyLoads was specified, but not already set on the Linker?
+#endif // USE_CIRCULAR_DEPENDENCY_LOAD_DEFERRING
+
+		if ((LoadFlags & DoNotLoadExportsFlags) == 0)
 		{
 			// Make sure we pass the property that's currently being serialized by the linker that owns the import 
 			// that triggered this LoadPackage call
diff --git a/Engine/Source/Runtime/CoreUObject/Public/Blueprint/BlueprintSupport.h b/Engine/Source/Runtime/CoreUObject/Public/Blueprint/BlueprintSupport.h
--- a/Engine/Source/Runtime/CoreUObject/Public/Blueprint/BlueprintSupport.h
+++ b/Engine/Source/Runtime/CoreUObject/Public/Blueprint/BlueprintSupport.h
@@ -17,6 +17,15 @@ struct FBlueprintSupport
 	 * @param	Writer						duplicate writer instance to write the duplicated data to
 	 */
 	static void DuplicateAllFields(class UStruct* StructToDuplicate, class FDuplicateDataWriter& Writer);
+
+	/** 
+	 * A series of query functions that we can use to easily gate-off/disable 
+	 * aspects of the deferred loading (mostly for testing purposes). 
+	 */
+	static bool UseDeferredDependencyLoading();
+	static bool UseDeferredCDOSerialization();
+	static bool UseDeferredLoadingForSubsequentLoads();
+	static bool UseDeferredDependencyVerificationChecks();
 };
 
 #if WITH_EDITOR
diff --git a/Engine/Source/Runtime/CoreUObject/Public/UObject/Linker.h b/Engine/Source/Runtime/CoreUObject/Public/UObject/Linker.h
--- a/Engine/Source/Runtime/CoreUObject/Public/UObject/Linker.h
+++ b/Engine/Source/Runtime/CoreUObject/Public/UObject/Linker.h
@@ -1625,6 +1625,17 @@ private:
 	int32 LoadMetaDataFromExportMap(bool bForcePreload = false);
 
 	UObject* CreateImport( int32 Index );
+
+	/**
+	 * Determines if the specified import belongs to a native "compiled in"
+	 * package (as opposed to an asset-file package). Recursive if the
+	 * specified import is not a package itself.
+	 * 
+	 * @param  ImportIndex    An index into the ImportMap, defining the import you wish to check.
+	 * @return True if the specified import comes from (or is) a "compiled in" package, otherwise false (it is an asset import).
+	 */
+	bool IsImportNative(const int32 ImportIndex) const;
+
 	UObject* IndexToObject( FPackageIndex Index );
 
 	void DetachExport( int32 i );
@@ -1814,6 +1825,35 @@ private:
 	 */
 	bool RegenerateBlueprintClass(UClass* LoadClass, UObject* ExportObject);
 
+	/**
+	 * Determines if the specified import should be deferred. If so, it will 
+	 * instantiate a placeholder object in its place.
+	 * 
+	 * @param  ImportIndex    An index into this linker's ImportMap, specifying which import to check.
+	 * @return True if the specified import was deferred, other wise false (it is ok to load it).
+	 */
+	bool DeferPotentialCircularImport(const int32 ImportIndex);
+
+	/**
+	 * Combs the ImportMap for any imports that were deferred, and then loads 
+	 * them. After that, deferred CDOs are serialized in and the Blueprint class
+	 * is regenerated.
+	 * 
+	 * @param  LoadClass    The (Blueprint) class that was loading, while we deferred dependencies.
+	 */
+	void ResolveDeferredDependencies(UClass* LoadClass);
+
+#if	USE_CIRCULAR_DEPENDENCY_LOAD_DEFERRING
+	/** 
+	 * For deferring dependency loads, we block CDO serialization until the 
+	 * class if complete. If we attempt to serialize the CDO whil that is 
+	 * happening, we instead defer it and record the export's index here (so we 
+	 * can return to it later).
+	 */
+	int32 DeferredExportIndex;
+#endif // USE_CIRCULAR_DEPENDENCY_LOAD_DEFERRING
+
+
 	/** 
 	 * Creates the export hash.
 	 */
diff --git a/Engine/Source/Runtime/CoreUObject/Public/UObject/ObjectBase.h b/Engine/Source/Runtime/CoreUObject/Public/UObject/ObjectBase.h
--- a/Engine/Source/Runtime/CoreUObject/Public/UObject/ObjectBase.h
+++ b/Engine/Source/Runtime/CoreUObject/Public/UObject/ObjectBase.h
@@ -44,6 +44,7 @@ enum ELoadFlags
 	LOAD_ForDiff					= 0x00020000,	// Loading for diffing.
 	LOAD_NoSeekFreeLinkerDetatch	= 0x00040000,	// Do not detach linkers for this package when seek-free loading
 	LOAD_PackageForPIE				= 0x00080000,   // This package is being loaded for PIE, it must be flagged as such immediately
+	LOAD_DeferDependencyLoads       = 0x00100000,   // Do not load external (blueprint) dependencies (instead, track them for deferred loading)
 };
 
 //
diff --git a/Engine/Source/Runtime/CoreUObject/Public/UObject/UnrealType.h b/Engine/Source/Runtime/CoreUObject/Public/UObject/UnrealType.h
--- a/Engine/Source/Runtime/CoreUObject/Public/UObject/UnrealType.h
+++ b/Engine/Source/Runtime/CoreUObject/Public/UObject/UnrealType.h
@@ -1591,6 +1591,8 @@ class COREUOBJECT_API UObjectPropertyBase : public UProperty
 	,	PropertyClass( InClass )
 	{}
 
+	virtual ~UObjectPropertyBase();
+
 	// UObject interface
 	virtual void Serialize( FArchive& Ar ) override;
 	static void AddReferencedObjects(UObject* InThis, FReferenceCollector& Collector);
