Remove incorrect remaining time calculation in AsyncPackage::FinishLinker()

diff --git a/Engine/Source/Runtime/CoreUObject/Private/Serialization/AsyncLoading.cpp b/Engine/Source/Runtime/CoreUObject/Private/Serialization/AsyncLoading.cpp
--- a/Engine/Source/Runtime/CoreUObject/Private/Serialization/AsyncLoading.cpp
+++ b/Engine/Source/Runtime/CoreUObject/Private/Serialization/AsyncLoading.cpp
@@ -454,9 +454,8 @@ EAsyncPackageState::Type FAsyncPackage::FinishLinker()
 		LastObjectWorkWasPerformedOn	= Linker->LinkerRoot;
 		LastTypeOfWorkPerformed			= TEXT("ticking linker");
 	
-		const float RemainingTimeLimit = (float)(FPlatformTime::Seconds() - TickStartTime);
 		// Operation still pending if Tick returns false
-		if (Linker->Tick(RemainingTimeLimit, bUseTimeLimit, bUseFullTimeLimit) != ULinkerLoad::LINKER_Loaded)
+		if (Linker->Tick(TimeLimit, bUseTimeLimit, bUseFullTimeLimit) != ULinkerLoad::LINKER_Loaded)
 		{
 			// Give up remainder of timeslice if there is one to give up.
 			GiveUpTimeSlice();
