[UE-7296] WIP... properly defering struct blueprint dependencies (left disabled).

diff --git a/Engine/Source/Editor/UnrealEd/Private/Tests/CircularDependencyTestActor.cpp b/Engine/Source/Editor/UnrealEd/Private/Tests/CircularDependencyTestActor.cpp
--- a/Engine/Source/Editor/UnrealEd/Private/Tests/CircularDependencyTestActor.cpp
+++ b/Engine/Source/Editor/UnrealEd/Private/Tests/CircularDependencyTestActor.cpp
@@ -108,6 +108,13 @@ bool ACircularDependencyTestActor::TestVerifyIsBlueprintTypeVar(FName VarName, b
 			return true;
 		}
 	}
+	else if (UInterfaceProperty* InterfaceProp = Cast<UInterfaceProperty>(VarProperty))
+	{
+		if (bCheckPropertyType && Cast<UBlueprintGeneratedClass>(InterfaceProp->InterfaceClass) != nullptr)
+		{
+			return true;
+		}
+	}
 	return !bCheckPropertyType;
 }
 
@@ -118,6 +125,28 @@ bool ACircularDependencyTestActor::RunVerificationTests_Implementation()
 
 bool ACircularDependencyTestActor::TestVerifyProperty(UProperty* Property, uint8* Container, bool bCheckPropertyType)
 {
+	auto TestVerifyClass = [](UClass* ClassToVerify)->bool
+	{
+		return (ClassToVerify != nullptr) && !IsPlaceholderClass(ClassToVerify);
+	};
+
+	bool bBelongsToThisBlueprintClass = false;
+	if (UClass* OwnerClass = Property->GetOwnerClass())
+	{
+		bBelongsToThisBlueprintClass = (OwnerClass == GetClass()) && (Cast<UBlueprintGeneratedClass>(OwnerClass) != nullptr);
+		if (UClass* SuperClass = OwnerClass->GetSuperClass())
+		{
+			FMemberReference VarRef;
+			VarRef.SetExternalMember(Property->GetFName(), SuperClass);
+			// make sure the parent doesn't have this variable
+			bBelongsToThisBlueprintClass &= (VarRef.ResolveMember<UProperty>(GetClass()) == nullptr);
+		}
+	}
+	else if (UStruct* OwnerStruct = Property->GetOwnerStruct())
+	{
+
+	}
+
 	if (UObjectProperty* ObjProp = Cast<UObjectProperty>(Property))
 	{
 		UClass* ClassToVerify = ObjProp->PropertyClass;
@@ -126,36 +155,30 @@ bool ACircularDependencyTestActor::TestVerifyProperty(UProperty* Property, uint8
 			ClassToVerify = Cast<UClass>(ClassProp->GetPropertyValue_InContainer(Container));
 		}
 
-		if (ClassToVerify == nullptr || IsPlaceholderClass(ClassToVerify))
+		if (!TestVerifyClass(ClassToVerify))
 		{
 			return false;
-		}
+		}		
 
-		//if (bVerifyObjTypes)
+		if (bBelongsToThisBlueprintClass)
 		{
-			bool bIsBlueprintSubClassProperty = false;
-			if (UClass* OwnerClass = Property->GetOwnerClass())
-			{
-				bIsBlueprintSubClassProperty = (OwnerClass == GetClass()) && (Cast<UBlueprintGeneratedClass>(OwnerClass) != nullptr);
-				if (UClass* SuperClass = OwnerClass->GetSuperClass())
-				{
-					FMemberReference VarRef;
-					VarRef.SetExternalMember(Property->GetFName(), SuperClass);
-					// make sure the parent doesn't have this variable
-					bIsBlueprintSubClassProperty &= (VarRef.ResolveMember<UProperty>(GetClass()) == nullptr);
-				}
-			}
-			else if (UStruct* OwnerStruct = Property->GetOwnerStruct())
-			{
-
-			}
+			// assume any object properties that we've added to the BP 
+			// directly should be circular refs (and thus, pointers to other blueprint types) 
+			return TestVerifyIsBlueprintTypeVar(Property->GetFName(), bCheckPropertyType);
+		}
+	}
+	else if (UInterfaceProperty* InterfaceProp = Cast<UInterfaceProperty>(Property))
+	{
+		if (!TestVerifyClass(InterfaceProp->InterfaceClass))
+		{
+			return false;
+		}
 
-			if (bIsBlueprintSubClassProperty)
-			{
-				// assume any object properties that we've added to the BP 
-				// directly should be circular refs (and thus, pointers to other blueprint types) 
-				return TestVerifyIsBlueprintTypeVar(Property->GetFName(), bCheckPropertyType);
-			}
+		if (bBelongsToThisBlueprintClass)
+		{
+			// assume any object properties that we've added to the BP 
+			// directly should be circular refs (and thus, pointers to other blueprint types) 
+			return TestVerifyIsBlueprintTypeVar(Property->GetFName(), bCheckPropertyType);
 		}
 	}
 
diff --git a/Engine/Source/Runtime/CoreUObject/Private/Blueprint/BlueprintSupport.cpp b/Engine/Source/Runtime/CoreUObject/Private/Blueprint/BlueprintSupport.cpp
--- a/Engine/Source/Runtime/CoreUObject/Private/Blueprint/BlueprintSupport.cpp
+++ b/Engine/Source/Runtime/CoreUObject/Private/Blueprint/BlueprintSupport.cpp
@@ -191,9 +191,12 @@ struct FPreloadMembersHelper
 bool ULinkerLoad::RegenerateBlueprintClass(UClass* LoadClass, UObject* ExportObject)
 {
 #if TEST_CHECK_DEPENDENCY_LOAD_DEFERRING
-	static ULinkerLoad* RegeneratingLinker = nullptr;
-	checkSlow(!FBlueprintSupport::UseDeferredDependencyVerificationChecks() || (RegeneratingLinker == nullptr) || (RegeneratingLinker == this));
-	TGuardValue<ULinkerLoad*> ReentrantGuard(RegeneratingLinker, this);
+	// @TODO: we can't rely on this check, because once we're in class 
+	//        regeneration, graphs/nodes are force loaded and that could cause
+	//        subsequent class regeneration (for other packages)
+	//static ULinkerLoad* RegeneratingLinker = nullptr;
+	//checkSlow(!FBlueprintSupport::UseDeferredDependencyVerificationChecks() || (RegeneratingLinker == nullptr) || (RegeneratingLinker == this));
+	//TGuardValue<ULinkerLoad*> ReentrantGuard(RegeneratingLinker, this);
 #endif // TEST_CHECK_DEPENDENCY_LOAD_DEFERRING
 
 	// determine if somewhere further down the callstack, we're already in this
@@ -339,18 +342,12 @@ bool ULinkerLoad::DeferPotentialCircularImport(const int32 Index)
 #endif // USE_CIRCULAR_DEPENDENCY_LOAD_DEFERRING
 }
 
-void ULinkerLoad::ResolveDeferredDependencies(UClass* LoadClass)
+void ULinkerLoad::ResolveDeferredDependencies()
 {
-#if USE_CIRCULAR_DEPENDENCY_LOAD_DEFERRING
-	if (!FBlueprintSupport::UseDeferredDependencyLoading())
-	{
-		return;
-	}
-
 	//--------------------------------------
 	// Phase 2: Resolve Dependency Stubs
 	//--------------------------------------
-	TGuardValue<uint32> LoadFlagsGuard(LoadFlags, LoadFlags & ~LOAD_DeferDependencyLoads);
+	TGuardValue<uint32> LoadFlagsGuard(LoadFlags, (LoadFlags & ~LOAD_DeferDependencyLoads));
 
 	for (int32 ImportIndex = 0; ImportIndex < ImportMap.Num(); ++ImportIndex)
 	{
@@ -361,12 +358,30 @@ void ULinkerLoad::ResolveDeferredDependencies(UClass* LoadClass)
 			Import.XObject = nullptr;
 			UClass* RealClassObj = CastChecked<UClass>(CreateImport(ImportIndex), ECastCheckedType::NullAllowed);
 
-			for (UObjectPropertyBase* ObjProperty : PlaceholderClass->ReferencingProperties)
+#if TEST_CHECK_DEPENDENCY_LOAD_DEFERRING
+			bool const bCircumventValidationChecks = !FBlueprintSupport::UseDeferredDependencyVerificationChecks();
+			// make sure that we know what utilized this placeholder class... 
+			// right now we only expect UObjectProperties/UClassProperties/
+			// UInterfaceProperties to, but something else could have requested 
+			// the class without logging itself with the placeholder... if  
+			// ReferencingProperties is empty, then there is something out there 
+			// still using a placeholder class
+			checkSlow(bCircumventValidationChecks || PlaceholderClass->HasReferences());
+#endif // TEST_CHECK_DEPENDENCY_LOAD_DEFERRING
+
+			PlaceholderClass->ReplaceTrackedReferences(RealClassObj);
+			PlaceholderClass->MarkPendingKill(); // @TODO: ensure these are properly GC'd
+		}
+		else if (UScriptStruct* StructObj = Cast<UScriptStruct>(Import.XObject))
+		{
+			// in case this is a user defined struct, we have to resolve any 
+			// deferred dependencies in the struct 
+			// 
+			// @TODO: could this result in an endless ResolveDeferredDependencies() loop? 
+			if (Import.SourceLinker != nullptr)
 			{
-				ObjProperty->PropertyClass = RealClassObj;
+				Import.SourceLinker->ResolveDeferredDependencies();
 			}
-			PlaceholderClass->ReferencingProperties.Empty();
-			PlaceholderClass->MarkPendingKill(); // @TODO: ensure these are properly GC'd
 		}
 	}
 
@@ -374,9 +389,33 @@ void ULinkerLoad::ResolveDeferredDependencies(UClass* LoadClass)
 	//        regeneration is about to force load a lot of this).
 	LoadAllObjects();
 
+#if TEST_CHECK_DEPENDENCY_LOAD_DEFERRING
+	bool const bCircumventValidationChecks = !FBlueprintSupport::UseDeferredDependencyVerificationChecks();
+	for (TObjectIterator<ULinkerPlaceholderClass> PlaceholderIt; PlaceholderIt; ++PlaceholderIt)
+	{
+		ULinkerPlaceholderClass* PlaceholderClass = *PlaceholderIt;
+		if (PlaceholderClass->GetOuter() == LinkerRoot)
+		{
+			// there shouldn't be any deferred dependencies (belonging to this 
+			// linker) that need to be resolved by this point
+			checkSlow( bCircumventValidationChecks || (!PlaceholderClass->HasReferences() && PlaceholderClass->IsPendingKill()) );
+		}
+	}
+#endif // TEST_CHECK_DEPENDENCY_LOAD_DEFERRING
+}
+
+void ULinkerLoad::FinalizeBlueprint(UClass* LoadClass)
+{
+#if USE_CIRCULAR_DEPENDENCY_LOAD_DEFERRING
+	if (!FBlueprintSupport::UseDeferredDependencyLoading())
+	{
+		return;
+	}
+
 	//--------------------------------------
 	// Phase 3: Finalize (serialize CDO & regenerate class)
 	//--------------------------------------
+	TGuardValue<uint32> LoadFlagsGuard(LoadFlags, LoadFlags & ~LOAD_DeferDependencyLoads);
 
 	// since this function is called after class serialization, then we can be
 	// certain that the CDO export object exists (since class serialization 
@@ -400,6 +439,7 @@ void ULinkerLoad::ResolveDeferredDependencies(UClass* LoadClass)
 	// else, the CDO hasn't been touched (which is expected, since the CDO needs the class first)
 	// @TODO: this comment here ^ is in direct conflict with the assumption 
 	//        stated above (that the CDO should be loaded)... figure out which is correct!
+	DeferredExportIndex = INDEX_NONE;
 
 	{
 		check(CDOExportIndex != INDEX_NONE);
@@ -439,17 +479,6 @@ void ULinkerLoad::ResolveDeferredDependencies(UClass* LoadClass)
 		TArray<UObject*> ClassInstances;
 		GetObjectsOfClass(BlueprintClass, ClassInstances, /*bIncludeDerivedClasses =*/true);
 		checkSlow(bCircumventValidationChecks || ClassInstances.Num() == 0);
-
-		for (TObjectIterator<ULinkerPlaceholderClass> PlaceholderIt; PlaceholderIt; ++PlaceholderIt)
-		{
-			ULinkerPlaceholderClass* PlaceholderClass = *PlaceholderIt;
-			// there shouldn't be any deferred dependencies that need to be 
-			// resolved by this point
-			checkSlow(bCircumventValidationChecks || (PlaceholderClass->GetOuter() != LinkerRoot) || PlaceholderClass->IsPendingKill());
-		}
-
-		// @TODO: Test verify there are no instances of the class yet (except for the CDO)
-
 #endif // TEST_CHECK_DEPENDENCY_LOAD_DEFERRING
 
 		UObject* OldCDO = BlueprintClass->ClassDefaultObject;
@@ -464,9 +493,6 @@ void ULinkerLoad::ResolveDeferredDependencies(UClass* LoadClass)
 				BlueprintClass->ClassDefaultObject = CDOExport.Object;
 			}
 		}
-		
-		// 
-		DeferredExportIndex = INDEX_NONE;
 	}
 #endif // USE_CIRCULAR_DEPENDENCY_LOAD_DEFERRING
 }
diff --git a/Engine/Source/Runtime/CoreUObject/Private/UObject/LinkerLoad.cpp b/Engine/Source/Runtime/CoreUObject/Private/UObject/LinkerLoad.cpp
--- a/Engine/Source/Runtime/CoreUObject/Private/UObject/LinkerLoad.cpp
+++ b/Engine/Source/Runtime/CoreUObject/Private/UObject/LinkerLoad.cpp
@@ -2879,7 +2879,8 @@ void ULinkerLoad::Preload( UObject* Object )
 					SCOPE_CYCLE_COUNTER(STAT_LinkerLoadDeferred);
 					if ((LoadFlags & LOAD_DeferDependencyLoads) != (*LoadFlagsGuard & LOAD_DeferDependencyLoads))
 					{
-						ResolveDeferredDependencies(ObjectAsClass);
+						ResolveDeferredDependencies();
+						FinalizeBlueprint(ObjectAsClass);
 					}
 #endif // USE_CIRCULAR_DEPENDENCY_LOAD_DEFERRING
 				}
@@ -3449,6 +3450,10 @@ UObject* ULinkerLoad::CreateImport( int32 Index )
 	FObjectImport& Import = ImportMap[ Index ];
 	
 #if USE_CIRCULAR_DEPENDENCY_LOAD_DEFERRING
+	// if this Import could possibly introduce a circular load (and we're 
+	// actively trying to avoid that at this point in the load process), then 
+	// this will stub in the Import with a placeholder object, to be replace 
+	// later on (this will return true if the import was actually deferred)
 	DeferPotentialCircularImport(Index);
 #endif // USE_CIRCULAR_DEPENDENCY_LOAD_DEFERRING
 
diff --git a/Engine/Source/Runtime/CoreUObject/Private/UObject/LinkerPlaceholderClass.cpp b/Engine/Source/Runtime/CoreUObject/Private/UObject/LinkerPlaceholderClass.cpp
--- a/Engine/Source/Runtime/CoreUObject/Private/UObject/LinkerPlaceholderClass.cpp
+++ b/Engine/Source/Runtime/CoreUObject/Private/UObject/LinkerPlaceholderClass.cpp
@@ -2,21 +2,50 @@
 
 #include "CoreUObjectPrivate.h"
 #include "LinkerPlaceholderClass.h"
+#include "Blueprint/BlueprintSupport.h"
 
+//------------------------------------------------------------------------------
 ULinkerPlaceholderClass::ULinkerPlaceholderClass(const FObjectInitializer& ObjectInitializer)
 	: Super(ObjectInitializer)
 {
 }
 
+//------------------------------------------------------------------------------
 ULinkerPlaceholderClass::~ULinkerPlaceholderClass()
 {
-	for (UObjectPropertyBase* ObjProperty : ReferencingProperties)
+	// @TODO: should we assert when this happens, so we can figure out what 
+	//        placeholder references are still around?
+#if TEST_CHECK_DEPENDENCY_LOAD_DEFERRING
+	bool const bCircumventValidationChecks = !FBlueprintSupport::UseDeferredDependencyVerificationChecks();
+	checkSlow(bCircumventValidationChecks || !HasReferences());
+#endif // TEST_CHECK_DEPENDENCY_LOAD_DEFERRING
+
+	// by this point, we really shouldn't have any properties left (they should 
+	// have all got replaced), but just in case (so things don't blow up with a
+	// missing class)...
+	ReplaceTrackedReferences(UObject::StaticClass());
+}
+
+//------------------------------------------------------------------------------
+IMPLEMENT_CORE_INTRINSIC_CLASS(ULinkerPlaceholderClass, UClass, 
 	{
-		check(ObjProperty->PropertyClass == this);
-		ObjProperty->PropertyClass = nullptr;
+		Class->ClassAddReferencedObjects = &ULinkerPlaceholderClass::AddReferencedObjects;
+		
+		// @TODO: use the Class->Emit...() functions here to aid garbage 
+		//        collection, so it has information on what class variables 
+		//        hold onto UObject references
 	}
+);
+
+//------------------------------------------------------------------------------
+void ULinkerPlaceholderClass::AddReferencedObjects(UObject* InThis, FReferenceCollector& Collector)
+{
+	ULinkerPlaceholderClass* This = CastChecked<ULinkerPlaceholderClass>(InThis);
+	//... 
+	Super::AddReferencedObjects(InThis, Collector);
 }
 
+//------------------------------------------------------------------------------
 void ULinkerPlaceholderClass::PostInitProperties()
 {
 	Super::PostInitProperties();
@@ -27,6 +56,7 @@ void ULinkerPlaceholderClass::PostInitProperties()
 	}
 }
 
+//------------------------------------------------------------------------------
 void ULinkerPlaceholderClass::Bind()
 {
 	ClassConstructor = InternalConstructor<ULinkerPlaceholderClass>;
@@ -35,19 +65,52 @@ void ULinkerPlaceholderClass::Bind()
 	ClassAddReferencedObjects = &ULinkerPlaceholderClass::AddReferencedObjects;
 }
 
-IMPLEMENT_CORE_INTRINSIC_CLASS(ULinkerPlaceholderClass, UClass, 
+//------------------------------------------------------------------------------
+void ULinkerPlaceholderClass::AddTrackedReference(UProperty* ReferencingProperty)
+{
+	ReferencingProperties.Add(ReferencingProperty);
+}
+
+//------------------------------------------------------------------------------
+bool ULinkerPlaceholderClass::HasReferences() const
+{
+	return (ReferencingProperties.Num() > 0);
+}
+
+//------------------------------------------------------------------------------
+void ULinkerPlaceholderClass::RemoveTrackedReference(UProperty* ReferencingProperty)
+{
+	//if (ReferencingProperty->IsValidLowLevel())
 	{
-		Class->ClassAddReferencedObjects = &ULinkerPlaceholderClass::AddReferencedObjects;
-		
-		// @TODO: use the Class->Emit...() functions here to aid garbage 
-		//        collection, so it has information on what class variables 
-		//        hold onto UObject references
+		ReferencingProperties.Remove(ReferencingProperty);
 	}
-);
+}
 
-void ULinkerPlaceholderClass::AddReferencedObjects(UObject* InThis, FReferenceCollector& Collector)
+//------------------------------------------------------------------------------
+void ULinkerPlaceholderClass::ReplaceTrackedReferences(UClass* ReplacementClass)
 {
-	ULinkerPlaceholderClass* This = CastChecked<ULinkerPlaceholderClass>(InThis);
-	//... 
-	Super::AddReferencedObjects(InThis, Collector);
+	for (UProperty* Property : ReferencingProperties)
+	{
+		if (UObjectPropertyBase* ObjProperty = Cast<UObjectPropertyBase>(Property))
+		{
+			if (ObjProperty->PropertyClass == this)
+			{
+				ObjProperty->PropertyClass = ReplacementClass;
+			}
+
+			UClassProperty* ClassProperty = Cast<UClassProperty>(ObjProperty);
+			if ((ClassProperty != nullptr) && (ClassProperty->MetaClass == this))
+			{
+				ClassProperty->MetaClass = ReplacementClass;
+			}
+		}	
+		else if (UInterfaceProperty* InterfaceProp = Cast<UInterfaceProperty>(Property))
+		{
+			if (InterfaceProp->InterfaceClass == this)
+			{
+				InterfaceProp->InterfaceClass = ReplacementClass;
+			}
+		}
+	}
+	ReferencingProperties.Empty();
 }
diff --git a/Engine/Source/Runtime/CoreUObject/Private/UObject/LinkerPlaceholderClass.h b/Engine/Source/Runtime/CoreUObject/Private/UObject/LinkerPlaceholderClass.h
--- a/Engine/Source/Runtime/CoreUObject/Private/UObject/LinkerPlaceholderClass.h
+++ b/Engine/Source/Runtime/CoreUObject/Private/UObject/LinkerPlaceholderClass.h
@@ -32,7 +32,41 @@ public:
  	virtual void Bind() override;
 	// End of UField interface.
 
-public:
-	/** Links to UObject properties that are currently referencing this class */
-	TSet<UObjectPropertyBase*> ReferencingProperties;
-}; 
\ No newline at end of file
+	/**
+	 * Caches off the supplied property so that we can later replace it's use of
+	 * this class with another (real) class.
+	 * 
+	 * @param  ReferencingProperty    A property that uses and stores this class.
+	 */
+	void AddTrackedReference(UProperty* ReferencingProperty);
+
+	/**
+	 * A query method that let's us check to see if this class is currently 
+	 * being referenced by anything (if this returns false, then a referencing 
+	 * property could have forgotten to add itself... or, we've replaced all
+	 * references).
+	 * 
+	 * @return True if this has anything stored in its ReferencingProperties container, otherwise false.
+	 */
+	bool HasReferences() const;
+
+	/**
+	 * Removes the specified property from this class's internal tracking list 
+	 * (which aims to keep track of properties utilizing this class).
+	 * 
+	 * @param  ReferencingProperty    A property that used to use this class, and now no longer does.
+	 */
+	void RemoveTrackedReference(UProperty* ReferencingProperty);
+
+	/**
+	 * Iterates over all referencing properties and attempts to replace their 
+	 * references to this class with a new (hopefully proper) class.
+	 * 
+	 * @param  ReplacementClass    The class that you want all references to this class replaced with.
+	 */
+	void ReplaceTrackedReferences(UClass* ReplacementClass);
+
+private:
+	/** Links to UProperties that are currently using this class */
+	TSet<UProperty*> ReferencingProperties;
+}; 
diff --git a/Engine/Source/Runtime/CoreUObject/Private/UObject/PropertyBaseObject.cpp b/Engine/Source/Runtime/CoreUObject/Private/UObject/PropertyBaseObject.cpp
--- a/Engine/Source/Runtime/CoreUObject/Private/UObject/PropertyBaseObject.cpp
+++ b/Engine/Source/Runtime/CoreUObject/Private/UObject/PropertyBaseObject.cpp
@@ -11,12 +11,13 @@
 UObjectPropertyBase::~UObjectPropertyBase()
 {
 #if USE_CIRCULAR_DEPENDENCY_LOAD_DEFERRING
-	// @TODO: sometimes PropertyClass is freed memory, so operating on it causes 
-	//        a crash
-// 	if (ULinkerPlaceholderClass* PlaceholderClass = Cast<ULinkerPlaceholderClass>(PropertyClass))
-// 	{
-// 		PlaceholderClass->ReferencingProperties.Remove(this);
-// 	}
+	if (PropertyClass->IsValidLowLevelFast(/*bRecursive =*/false))
+	{
+		if (ULinkerPlaceholderClass* PlaceholderClass = Cast<ULinkerPlaceholderClass>(PropertyClass))
+		{
+			PlaceholderClass->RemoveTrackedReference(this);
+		}
+	}
 #endif // USE_CIRCULAR_DEPENDENCY_LOAD_DEFERRING
 }
 
@@ -93,7 +94,7 @@ void UObjectPropertyBase::Serialize( FArchive& Ar )
 	{
 		if (ULinkerPlaceholderClass* PlaceholderClass = Cast<ULinkerPlaceholderClass>(PropertyClass))
 		{
-			PlaceholderClass->ReferencingProperties.Add(this);
+			PlaceholderClass->AddTrackedReference(this);
 		}
 	}
 #endif // USE_CIRCULAR_DEPENDENCY_LOAD_DEFERRING
diff --git a/Engine/Source/Runtime/CoreUObject/Private/UObject/PropertyClass.cpp b/Engine/Source/Runtime/CoreUObject/Private/UObject/PropertyClass.cpp
--- a/Engine/Source/Runtime/CoreUObject/Private/UObject/PropertyClass.cpp
+++ b/Engine/Source/Runtime/CoreUObject/Private/UObject/PropertyClass.cpp
@@ -2,16 +2,40 @@
 
 #include "CoreUObjectPrivate.h"
 #include "PropertyHelper.h"
+#include "LinkerPlaceholderClass.h"
 
 /*-----------------------------------------------------------------------------
 	UClassProperty.
 -----------------------------------------------------------------------------*/
 
+UClassProperty::~UClassProperty()
+{
+#if USE_CIRCULAR_DEPENDENCY_LOAD_DEFERRING
+	if (MetaClass->IsValidLowLevelFast(/*bRecursive =*/false))
+	{
+		if (ULinkerPlaceholderClass* PlaceholderClass = Cast<ULinkerPlaceholderClass>(MetaClass))
+		{
+			PlaceholderClass->RemoveTrackedReference(this);
+		}
+	}
+#endif // USE_CIRCULAR_DEPENDENCY_LOAD_DEFERRING
+}
+
 void UClassProperty::Serialize( FArchive& Ar )
 {
 	Super::Serialize( Ar );
 	Ar << MetaClass;
 
+#if USE_CIRCULAR_DEPENDENCY_LOAD_DEFERRING
+	if (Ar.IsLoading())
+	{
+		if (ULinkerPlaceholderClass* PlaceholderClass = Cast<ULinkerPlaceholderClass>(MetaClass))
+		{
+			PlaceholderClass->AddTrackedReference(this);
+		}
+	}
+#endif // USE_CIRCULAR_DEPENDENCY_LOAD_DEFERRING
+
 	if( !(MetaClass||HasAnyFlags(RF_ClassDefaultObject)) )
 	{
 		// If we failed to load the MetaClass and we're not a CDO, that means we relied on a class that has been removed or doesn't exist.
diff --git a/Engine/Source/Runtime/CoreUObject/Private/UObject/PropertyInterface.cpp b/Engine/Source/Runtime/CoreUObject/Private/UObject/PropertyInterface.cpp
--- a/Engine/Source/Runtime/CoreUObject/Private/UObject/PropertyInterface.cpp
+++ b/Engine/Source/Runtime/CoreUObject/Private/UObject/PropertyInterface.cpp
@@ -1,10 +1,25 @@
 // Copyright 1998-2015 Epic Games, Inc. All Rights Reserved.
 
 #include "CoreUObjectPrivate.h"
+#include "LinkerPlaceholderClass.h"
 
 /*-----------------------------------------------------------------------------
 	UInterfaceProperty.
 -----------------------------------------------------------------------------*/
+
+UInterfaceProperty::~UInterfaceProperty()
+{
+#if USE_CIRCULAR_DEPENDENCY_LOAD_DEFERRING
+	if (InterfaceClass->IsValidLowLevelFast(/*bRecursive =*/false))
+	{
+		if (ULinkerPlaceholderClass* PlaceholderClass = Cast<ULinkerPlaceholderClass>(InterfaceClass))
+		{
+			PlaceholderClass->RemoveTrackedReference(this);
+		}
+	}
+#endif // USE_CIRCULAR_DEPENDENCY_LOAD_DEFERRING
+}
+
 /**
  * Returns the text to use for exporting this property to header file.
  *
@@ -198,6 +213,16 @@ void UInterfaceProperty::Serialize( FArchive& Ar )
 
 	Ar << InterfaceClass;
 
+#if USE_CIRCULAR_DEPENDENCY_LOAD_DEFERRING
+	if (Ar.IsLoading())
+	{
+		if (ULinkerPlaceholderClass* PlaceholderClass = Cast<ULinkerPlaceholderClass>(InterfaceClass))
+		{
+			PlaceholderClass->AddTrackedReference(this);
+		}
+	}
+#endif // USE_CIRCULAR_DEPENDENCY_LOAD_DEFERRING
+
 	if ( !InterfaceClass && !HasAnyFlags(RF_ClassDefaultObject) )
  	{
 		// If we failed to load the InterfaceClass and we're not a CDO, that means we relied on a class that has been removed or doesn't exist.
diff --git a/Engine/Source/Runtime/CoreUObject/Private/UObject/PropertyStruct.cpp b/Engine/Source/Runtime/CoreUObject/Private/UObject/PropertyStruct.cpp
--- a/Engine/Source/Runtime/CoreUObject/Private/UObject/PropertyStruct.cpp
+++ b/Engine/Source/Runtime/CoreUObject/Private/UObject/PropertyStruct.cpp
@@ -4,6 +4,28 @@
 #include "Archive.h"
 #include "PropertyHelper.h"
 
+static inline void PreloadInnerStructMembers(UStructProperty* StructProperty)
+{
+#if USE_CIRCULAR_DEPENDENCY_LOAD_DEFERRING
+	uint32 PropagatedLoadFlags = 0;
+	if (ULinkerLoad* Linker = StructProperty->GetLinker())
+	{
+		PropagatedLoadFlags |= (Linker->LoadFlags & LOAD_DeferDependencyLoads);
+	}
+
+	if (UScriptStruct* Struct = StructProperty->Struct)
+	{
+		if (ULinkerLoad* StructLinker = Struct->GetLinker())
+		{
+			TGuardValue<uint32> LoadFlagGuard(StructLinker->LoadFlags, StructLinker->LoadFlags | PropagatedLoadFlags);
+			Struct->RecursivelyPreload();
+		}
+	}
+#else // USE_CIRCULAR_DEPENDENCY_LOAD_DEFERRING
+	StructProperty->Struct->RecursivelyPreload();
+#endif // USE_CIRCULAR_DEPENDENCY_LOAD_DEFERRING
+}
+
 /*-----------------------------------------------------------------------------
 	UStructProperty.
 -----------------------------------------------------------------------------*/
@@ -41,7 +63,7 @@ void UStructProperty::LinkInternal(FArchive& Ar)
 	{
 		Struct = GetFallbackStruct();
 	}
-	Struct->RecursivelyPreload();
+	PreloadInnerStructMembers(this);
 	
 	ElementSize = Align(Struct->PropertiesSize, Struct->GetMinAlignment());
 	if (Struct->StructFlags & STRUCT_IsPlainOldData) // if there is nothing to construct or the struct is known to be memcpy-able, then allow memcpy
@@ -184,7 +206,7 @@ void UStructProperty::Serialize( FArchive& Ar )
 #endif // WITH_EDITOR
 	if (Struct)
 	{
-		Struct->RecursivelyPreload();
+		PreloadInnerStructMembers(this);
 	}
 	else
 	{
diff --git a/Engine/Source/Runtime/CoreUObject/Public/UObject/Linker.h b/Engine/Source/Runtime/CoreUObject/Public/UObject/Linker.h
--- a/Engine/Source/Runtime/CoreUObject/Public/UObject/Linker.h
+++ b/Engine/Source/Runtime/CoreUObject/Public/UObject/Linker.h
@@ -1841,7 +1841,15 @@ private:
 	 * 
 	 * @param  LoadClass    The (Blueprint) class that was loading, while we deferred dependencies.
 	 */
-	void ResolveDeferredDependencies(UClass* LoadClass);
+	void ResolveDeferredDependencies();
+
+	/**
+	 * Takes the supplied serialized class and serializes in its CDO, then 
+	 * regenerates both.
+	 * 
+	 * @param  LoadClass    The loaded blueprint class (assumes that it has been fully loaded/serialized).
+	 */
+	void FinalizeBlueprint(UClass* LoadClass);
 
 #if	USE_CIRCULAR_DEPENDENCY_LOAD_DEFERRING
 	/** 
diff --git a/Engine/Source/Runtime/CoreUObject/Public/UObject/UnrealType.h b/Engine/Source/Runtime/CoreUObject/Public/UObject/UnrealType.h
--- a/Engine/Source/Runtime/CoreUObject/Public/UObject/UnrealType.h
+++ b/Engine/Source/Runtime/CoreUObject/Public/UObject/UnrealType.h
@@ -1919,6 +1919,8 @@ public:
 	{
 	}
 
+	virtual ~UClassProperty();
+
 	// UObject interface
 	virtual void Serialize( FArchive& Ar ) override;
 	static void AddReferencedObjects(UObject* InThis, FReferenceCollector& Collector);
@@ -2001,6 +2003,8 @@ public:
 	{
 	}
 
+	virtual ~UInterfaceProperty();
+
 	// UHT interface
 	virtual FString GetCPPMacroType( FString& ExtendedTypeText ) const  override;
 	virtual FString GetCPPType( FString* ExtendedTypeText, uint32 CPPExportFlags ) const override;
