MaxReadBytes removed from SerializeItem, SerializeBin and SerializeTaggedProperty, as it was unused.

diff --git a/Engine/Source/Runtime/CoreUObject/Private/UObject/Class.cpp b/Engine/Source/Runtime/CoreUObject/Private/UObject/Class.cpp
--- a/Engine/Source/Runtime/CoreUObject/Private/UObject/Class.cpp
+++ b/Engine/Source/Runtime/CoreUObject/Private/UObject/Class.cpp
@@ -726,7 +726,7 @@ void UStruct::DestroyStruct(void* Dest, int32 ArrayDim) const
 // Serialize all of the class's data that belongs in a particular
 // bin and resides in Data.
 //
-void UStruct::SerializeBin( FArchive& Ar, void* Data, int32 MaxReadBytes ) const
+void UStruct::SerializeBin( FArchive& Ar, void* Data ) const
 {
 	if( Ar.IsObjectReferenceCollector() )
 	{
@@ -748,7 +748,7 @@ void UStruct::SerializeBinEx( FArchive& Ar, void* Data, void const* DefaultData,
 {
 	if ( !DefaultData || !DefaultStruct )
 	{
-		SerializeBin(Ar, Data, 0);
+		SerializeBin(Ar, Data);
 		return;
 	}
 
@@ -1209,7 +1209,7 @@ void UStruct::SerializeTaggedProperties(FArchive& Ar, uint8* Data, UStruct* Defa
 				uint8* DestAddress = Property->ContainerPtrToValuePtr<uint8>(Data, Tag.ArrayIndex);  
 
 				// This property is ok.			
-				Tag.SerializeTaggedProperty( Ar, Property, DestAddress, Tag.Size, NULL );
+				Tag.SerializeTaggedProperty( Ar, Property, DestAddress, NULL );
 
 				AdvanceProperty = true;
 				continue;
@@ -1259,7 +1259,7 @@ void UStruct::SerializeTaggedProperties(FArchive& Ar, uint8* Data, UStruct* Defa
 						// need to know how much data this call to SerializeTaggedProperty consumes, so mark where we are
 						int32 DataOffset = Ar.Tell();
 
-						Tag.SerializeTaggedProperty( Ar, Property, DataPtr, 0, DefaultValue );
+						Tag.SerializeTaggedProperty( Ar, Property, DataPtr, DefaultValue );
 
 						// set the tag's size
 						Tag.Size = Ar.Tell() - DataOffset;
@@ -3147,7 +3147,7 @@ void UClass::SerializeDefaultObject(UObject* Object, FArchive& Ar)
 	}
 	else
 	{
-		SerializeBin(Ar, Object, 0);
+		SerializeBin(Ar, Object);
 	}
 	Ar.StopSerializingDefaults();
 }
diff --git a/Engine/Source/Runtime/CoreUObject/Private/UObject/Obj.cpp b/Engine/Source/Runtime/CoreUObject/Private/UObject/Obj.cpp
--- a/Engine/Source/Runtime/CoreUObject/Private/UObject/Obj.cpp
+++ b/Engine/Source/Runtime/CoreUObject/Private/UObject/Obj.cpp
@@ -907,7 +907,7 @@ void UObject::SerializeScriptProperties( FArchive& Ar ) const
 	}
 	else
 	{
-		Class->SerializeBin( Ar, const_cast<UObject *>(this), 0 );
+		Class->SerializeBin( Ar, const_cast<UObject *>(this) );
 	}
 
 	if( HasAnyFlags(RF_ClassDefaultObject) )
diff --git a/Engine/Source/Runtime/CoreUObject/Private/UObject/Property.cpp b/Engine/Source/Runtime/CoreUObject/Private/UObject/Property.cpp
--- a/Engine/Source/Runtime/CoreUObject/Private/UObject/Property.cpp
+++ b/Engine/Source/Runtime/CoreUObject/Private/UObject/Property.cpp
@@ -634,7 +634,7 @@ bool UProperty::ShouldSerializeValue( FArchive& Ar ) const
 //
 bool UProperty::NetSerializeItem( FArchive& Ar, UPackageMap* Map, void* Data, TArray<uint8> * MetaData ) const
 {
-	SerializeItem( Ar, Data, 0, NULL );
+	SerializeItem( Ar, Data, NULL );
 	return 1;
 }
 
diff --git a/Engine/Source/Runtime/CoreUObject/Private/UObject/PropertyArray.cpp b/Engine/Source/Runtime/CoreUObject/Private/UObject/PropertyArray.cpp
--- a/Engine/Source/Runtime/CoreUObject/Private/UObject/PropertyArray.cpp
+++ b/Engine/Source/Runtime/CoreUObject/Private/UObject/PropertyArray.cpp
@@ -54,7 +54,7 @@ bool UArrayProperty::Identical( const void* A, const void* B, uint32 PortFlags )
 	}
 	return true;
 }
-void UArrayProperty::SerializeItem( FArchive& Ar, void* Value, int32 MaxReadBytes, void const* Defaults ) const
+void UArrayProperty::SerializeItem( FArchive& Ar, void* Value, void const* Defaults ) const
 {
 	checkSlow(Inner);
 
@@ -70,11 +70,9 @@ void UArrayProperty::SerializeItem( FArchive& Ar, void* Value, int32 MaxReadByte
 	}
 	ArrayHelper.CountBytes( Ar );
 
-	const int32 ArrayMaxReadBytes = MaxReadBytes > 0 ? ( MaxReadBytes - sizeof( n ) ) : 0;
 	for( int32 i=0; i<n; i++ )
 	{
-		const int32 ItemMaxReadBytes = ArrayMaxReadBytes > 0 ? ArrayMaxReadBytes / n: 0;
-		Inner->SerializeItem( Ar, ArrayHelper.GetRawPtr(i), ItemMaxReadBytes );
+		Inner->SerializeItem( Ar, ArrayHelper.GetRawPtr(i) );
 	}
 }
 
diff --git a/Engine/Source/Runtime/CoreUObject/Private/UObject/PropertyAssetObject.cpp b/Engine/Source/Runtime/CoreUObject/Private/UObject/PropertyAssetObject.cpp
--- a/Engine/Source/Runtime/CoreUObject/Private/UObject/PropertyAssetObject.cpp
+++ b/Engine/Source/Runtime/CoreUObject/Private/UObject/PropertyAssetObject.cpp
@@ -30,7 +30,7 @@ bool UAssetObjectProperty::Identical( const void* A, const void* B, uint32 PortF
 	return ObjectA.GetUniqueID() == ObjectB.GetUniqueID();
 }
 
-void UAssetObjectProperty::SerializeItem( FArchive& Ar, void* Value, int32 MaxReadBytes, void const* Defaults ) const
+void UAssetObjectProperty::SerializeItem( FArchive& Ar, void* Value, void const* Defaults ) const
 {
 	// We never serialize our reference while the garbage collector is harvesting references
 	// to objects, because we don't want asset pointers to keep objects from being garbage collected
diff --git a/Engine/Source/Runtime/CoreUObject/Private/UObject/PropertyBool.cpp b/Engine/Source/Runtime/CoreUObject/Private/UObject/PropertyBool.cpp
--- a/Engine/Source/Runtime/CoreUObject/Private/UObject/PropertyBool.cpp
+++ b/Engine/Source/Runtime/CoreUObject/Private/UObject/PropertyBool.cpp
@@ -216,7 +216,7 @@ bool UBoolProperty::Identical( const void* A, const void* B, uint32 PortFlags )
 	return ((*ByteValueA ^ (B ? *ByteValueB : 0)) & FieldMask) == 0;
 }
 
-void UBoolProperty::SerializeItem( FArchive& Ar, void* Value, int32 MaxReadBytes, void const* Defaults ) const
+void UBoolProperty::SerializeItem( FArchive& Ar, void* Value, void const* Defaults ) const
 {
 	check(FieldSize != 0);
 	uint8* ByteValue = (uint8*)Value + ByteOffset;
diff --git a/Engine/Source/Runtime/CoreUObject/Private/UObject/PropertyByte.cpp b/Engine/Source/Runtime/CoreUObject/Private/UObject/PropertyByte.cpp
--- a/Engine/Source/Runtime/CoreUObject/Private/UObject/PropertyByte.cpp
+++ b/Engine/Source/Runtime/CoreUObject/Private/UObject/PropertyByte.cpp
@@ -5,7 +5,7 @@
 /*-----------------------------------------------------------------------------
 	UByteProperty.
 -----------------------------------------------------------------------------*/
-void UByteProperty::SerializeItem( FArchive& Ar, void* Value, int32 MaxReadBytes, void const* Defaults ) const
+void UByteProperty::SerializeItem( FArchive& Ar, void* Value, void const* Defaults ) const
 {
 	if(Enum && Ar.UseToResolveEnumerators())
 	{
@@ -18,7 +18,7 @@ void UByteProperty::SerializeItem( FArchive& Ar, void* Value, int32 MaxReadBytes
 	const bool bUseBinarySerialization = (Enum == NULL) || (!Ar.IsLoading() && !Ar.IsSaving());
 	if( bUseBinarySerialization )
 	{
-		Super::SerializeItem(Ar, Value, MaxReadBytes, Defaults);
+		Super::SerializeItem(Ar, Value, Defaults);
 	}
 	// Loading
 	else if (Ar.IsLoading())
diff --git a/Engine/Source/Runtime/CoreUObject/Private/UObject/PropertyDelegate.cpp b/Engine/Source/Runtime/CoreUObject/Private/UObject/PropertyDelegate.cpp
--- a/Engine/Source/Runtime/CoreUObject/Private/UObject/PropertyDelegate.cpp
+++ b/Engine/Source/Runtime/CoreUObject/Private/UObject/PropertyDelegate.cpp
@@ -58,7 +58,7 @@ bool UDelegateProperty::Identical( const void* A, const void* B, uint32 PortFlag
 }
 
 
-void UDelegateProperty::SerializeItem( FArchive& Ar, void* Value, int32 MaxReadBytes, void const* Defaults ) const
+void UDelegateProperty::SerializeItem( FArchive& Ar, void* Value, void const* Defaults ) const
 {
 	Ar << *GetPropertyValuePtr(Value);
 }
diff --git a/Engine/Source/Runtime/CoreUObject/Private/UObject/PropertyInterface.cpp b/Engine/Source/Runtime/CoreUObject/Private/UObject/PropertyInterface.cpp
--- a/Engine/Source/Runtime/CoreUObject/Private/UObject/PropertyInterface.cpp
+++ b/Engine/Source/Runtime/CoreUObject/Private/UObject/PropertyInterface.cpp
@@ -102,7 +102,7 @@ bool UInterfaceProperty::Identical( const void* A, const void* B, uint32 PortFla
 	return (InterfaceA->GetObject() == InterfaceB->GetObject() && InterfaceA->GetInterface() == InterfaceB->GetInterface());
 }
 
-void UInterfaceProperty::SerializeItem( FArchive& Ar, void* Value, int32 MaxReadBytes, void const* Defaults ) const
+void UInterfaceProperty::SerializeItem( FArchive& Ar, void* Value, void const* Defaults ) const
 {
 	FScriptInterface* InterfaceValue = (FScriptInterface*)Value;
 
diff --git a/Engine/Source/Runtime/CoreUObject/Private/UObject/PropertyLazyObjectPtr.cpp b/Engine/Source/Runtime/CoreUObject/Private/UObject/PropertyLazyObjectPtr.cpp
--- a/Engine/Source/Runtime/CoreUObject/Private/UObject/PropertyLazyObjectPtr.cpp
+++ b/Engine/Source/Runtime/CoreUObject/Private/UObject/PropertyLazyObjectPtr.cpp
@@ -22,7 +22,7 @@ FName ULazyObjectProperty::GetID() const
 	return NAME_LazyObjectProperty;
 }
 
-void ULazyObjectProperty::SerializeItem( FArchive& Ar, void* Value, int32 MaxReadBytes, void const* Defaults ) const
+void ULazyObjectProperty::SerializeItem( FArchive& Ar, void* Value, void const* Defaults ) const
 {
 	// We never serialize our reference while the garbage collector is harvesting references
 	// to objects, because we don't want lazy pointers to keep objects from being garbage collected
diff --git a/Engine/Source/Runtime/CoreUObject/Private/UObject/PropertyMulticastDelegate.cpp b/Engine/Source/Runtime/CoreUObject/Private/UObject/PropertyMulticastDelegate.cpp
--- a/Engine/Source/Runtime/CoreUObject/Private/UObject/PropertyMulticastDelegate.cpp
+++ b/Engine/Source/Runtime/CoreUObject/Private/UObject/PropertyMulticastDelegate.cpp
@@ -98,7 +98,7 @@ bool UMulticastDelegateProperty::Identical( const void* A, const void* B, uint32
 	return bResult;
 }
 
-void UMulticastDelegateProperty::SerializeItem( FArchive& Ar, void* Value, int32 MaxReadBytes, void const* Defaults ) const
+void UMulticastDelegateProperty::SerializeItem( FArchive& Ar, void* Value, void const* Defaults ) const
 {
 	Ar << *GetPropertyValuePtr(Value);
 }
diff --git a/Engine/Source/Runtime/CoreUObject/Private/UObject/PropertyObject.cpp b/Engine/Source/Runtime/CoreUObject/Private/UObject/PropertyObject.cpp
--- a/Engine/Source/Runtime/CoreUObject/Private/UObject/PropertyObject.cpp
+++ b/Engine/Source/Runtime/CoreUObject/Private/UObject/PropertyObject.cpp
@@ -22,7 +22,7 @@ FString UObjectProperty::GetCPPMacroType( FString& ExtendedTypeText ) const
 	return TEXT("OBJECT");
 }
 
-void UObjectProperty::SerializeItem( FArchive& Ar, void* Value, int32 MaxReadBytes, void const* Defaults ) const
+void UObjectProperty::SerializeItem( FArchive& Ar, void* Value, void const* Defaults ) const
 {
 	UObject* ObjectValue = GetObjectPropertyValue(Value);
 	Ar << ObjectValue;
diff --git a/Engine/Source/Runtime/CoreUObject/Private/UObject/PropertyStruct.cpp b/Engine/Source/Runtime/CoreUObject/Private/UObject/PropertyStruct.cpp
--- a/Engine/Source/Runtime/CoreUObject/Private/UObject/PropertyStruct.cpp
+++ b/Engine/Source/Runtime/CoreUObject/Private/UObject/PropertyStruct.cpp
@@ -104,7 +104,7 @@ bool UStructProperty::UseBinaryOrNativeSerialization(const FArchive& Ar) const
 	return bUseBinarySerialization || bUseNativeSerialization;
 }
 
-void UStructProperty::SerializeItem( FArchive& Ar, void* Value, int32 MaxReadBytes, void const* Defaults ) const
+void UStructProperty::SerializeItem( FArchive& Ar, void* Value, void const* Defaults ) const
 {
 	const bool bUseBinarySerialization = UseBinarySerialization(Ar);
 	const bool bUseNativeSerialization = UseNativeSerialization();
@@ -135,7 +135,7 @@ void UStructProperty::SerializeItem( FArchive& Ar, void* Value, int32 MaxReadByt
 			}
 			else
 			{
-				Struct->SerializeBin( Ar, Value, MaxReadBytes );
+				Struct->SerializeBin( Ar, Value );
 			}
 		}
 		else
diff --git a/Engine/Source/Runtime/CoreUObject/Private/UObject/PropertyWeakObjectPtr.cpp b/Engine/Source/Runtime/CoreUObject/Private/UObject/PropertyWeakObjectPtr.cpp
--- a/Engine/Source/Runtime/CoreUObject/Private/UObject/PropertyWeakObjectPtr.cpp
+++ b/Engine/Source/Runtime/CoreUObject/Private/UObject/PropertyWeakObjectPtr.cpp
@@ -31,7 +31,7 @@ FString UWeakObjectProperty::GetCPPMacroType( FString& ExtendedTypeText ) const
 	return TEXT("WEAKOBJECT");
 }
 
-void UWeakObjectProperty::SerializeItem( FArchive& Ar, void* Value, int32 MaxReadBytes, void const* Defaults ) const
+void UWeakObjectProperty::SerializeItem( FArchive& Ar, void* Value, void const* Defaults ) const
 {
 	UObject* ObjectValue = GetObjectPropertyValue(Value);
 	Ar << *(FWeakObjectPtr*)Value;
diff --git a/Engine/Source/Runtime/CoreUObject/Private/UObject/UTextProperty.cpp b/Engine/Source/Runtime/CoreUObject/Private/UObject/UTextProperty.cpp
--- a/Engine/Source/Runtime/CoreUObject/Private/UObject/UTextProperty.cpp
+++ b/Engine/Source/Runtime/CoreUObject/Private/UObject/UTextProperty.cpp
@@ -39,7 +39,7 @@ bool UTextProperty::Identical( const void* A, const void* B, uint32 PortFlags )
 	return FTextInspector::GetDisplayString(ValueA).IsEmpty();
 }
 
-void UTextProperty::SerializeItem( FArchive& Ar, void* Value, int32 MaxReadBytes, void const* Defaults ) const
+void UTextProperty::SerializeItem( FArchive& Ar, void* Value, void const* Defaults ) const
 {
 	const TCppType PropertyValue = GetPropertyValue(Value);
 	if ( Ar.IsSaving() && Ar.IsPersistent() && PropertyValue.IsTransient() )
diff --git a/Engine/Source/Runtime/CoreUObject/Public/Serialization/ArchiveUObject.h b/Engine/Source/Runtime/CoreUObject/Public/Serialization/ArchiveUObject.h
--- a/Engine/Source/Runtime/CoreUObject/Public/Serialization/ArchiveUObject.h
+++ b/Engine/Source/Runtime/CoreUObject/Public/Serialization/ArchiveUObject.h
@@ -1092,7 +1092,7 @@ public:
 				}
 				else
 				{
-					ObjectClass->SerializeBin(*this, SearchObject, 0);
+					ObjectClass->SerializeBin(*this, SearchObject);
 				}
 				StopSerializingDefaults();
 			}
@@ -1161,7 +1161,7 @@ public:
 						}
 						else
 						{
-							ObjectClass->SerializeBin(*this, Obj, 0);
+							ObjectClass->SerializeBin(*this, Obj);
 						}
 						StopSerializingDefaults();
 					}
diff --git a/Engine/Source/Runtime/CoreUObject/Public/UObject/Class.h b/Engine/Source/Runtime/CoreUObject/Public/UObject/Class.h
--- a/Engine/Source/Runtime/CoreUObject/Public/UObject/Class.h
+++ b/Engine/Source/Runtime/CoreUObject/Public/UObject/Class.h
@@ -289,7 +289,7 @@ public:
 	//
 	virtual void Link(FArchive& Ar, bool bRelinkExistingProperties);
 
-	virtual void SerializeBin( FArchive& Ar, void* Data, int32 MaxReadBytes ) const;
+	virtual void SerializeBin( FArchive& Ar, void* Data ) const;
 
 	/**
 	 * Serializes the class properties that reside in Data if they differ from the corresponding values in DefaultData
diff --git a/Engine/Source/Runtime/CoreUObject/Public/UObject/PropertyTag.h b/Engine/Source/Runtime/CoreUObject/Public/UObject/PropertyTag.h
--- a/Engine/Source/Runtime/CoreUObject/Public/UObject/PropertyTag.h
+++ b/Engine/Source/Runtime/CoreUObject/Public/UObject/PropertyTag.h
@@ -109,7 +109,7 @@ struct FPropertyTag
 	}
 
 	// Property serializer.
-	void SerializeTaggedProperty( FArchive& Ar, UProperty* Property, uint8* Value, int32 MaxReadBytes, uint8* Defaults )
+	void SerializeTaggedProperty( FArchive& Ar, UProperty* Property, uint8* Value, uint8* Defaults )
 	{
 		if (Property->GetClass() == UBoolProperty::StaticClass())
 		{
@@ -124,7 +124,7 @@ struct FPropertyTag
 			UProperty* OldSerializedProperty = Ar.GetSerializedProperty();
 			Ar.SetSerializedProperty(Property);
 
-			Property->SerializeItem( Ar, Value, MaxReadBytes, Defaults );
+			Property->SerializeItem( Ar, Value, Defaults );
 
 			Ar.SetSerializedProperty(OldSerializedProperty);
 		}
diff --git a/Engine/Source/Runtime/CoreUObject/Public/UObject/UTextProperty.h b/Engine/Source/Runtime/CoreUObject/Public/UObject/UTextProperty.h
--- a/Engine/Source/Runtime/CoreUObject/Public/UObject/UTextProperty.h
+++ b/Engine/Source/Runtime/CoreUObject/Public/UObject/UTextProperty.h
@@ -22,7 +22,7 @@ public:
 
 	// UProperty interface
 	virtual bool Identical( const void* A, const void* B, uint32 PortFlags ) const override;
-	virtual void SerializeItem( FArchive& Ar, void* Value, int32 MaxReadBytes, void const* Defaults ) const override;
+	virtual void SerializeItem( FArchive& Ar, void* Value, void const* Defaults ) const override;
 	virtual void ExportTextItem( FString& ValueStr, const void* PropertyValue, const void* DefaultValue, UObject* Parent, int32 PortFlags, UObject* ExportRootScope ) const override;
 	virtual const TCHAR* ImportText_Internal( const TCHAR* Buffer, void* Data, int32 PortFlags, UObject* OwnerObject, FOutputDevice* ErrorText ) const override;
 	virtual FString GetCPPTypeForwardDeclaration() const override
diff --git a/Engine/Source/Runtime/CoreUObject/Public/UObject/UnrealType.h b/Engine/Source/Runtime/CoreUObject/Public/UObject/UnrealType.h
--- a/Engine/Source/Runtime/CoreUObject/Public/UObject/UnrealType.h
+++ b/Engine/Source/Runtime/CoreUObject/Public/UObject/UnrealType.h
@@ -285,7 +285,7 @@ public:
 			for (int32 Idx = 0; Idx < ArrayDim; Idx++)
 			{
 				Ar.SetSerializedProperty(this);
-				SerializeItem( Ar, ContainerPtrToValuePtr<void>(Data, Idx), 0 );
+				SerializeItem( Ar, ContainerPtrToValuePtr<void>(Data, Idx) );
 			}
 			Ar.SetSerializedProperty(OldSerializedProperty);
 		}
@@ -310,14 +310,14 @@ public:
 				{
 					UProperty* OldSerializedProperty = Ar.GetSerializedProperty();
 					Ar.SetSerializedProperty(this);
-					SerializeItem( Ar, Target, 0, Default );
+					SerializeItem( Ar, Target, Default );
 					Ar.SetSerializedProperty(OldSerializedProperty);
 				}
 			}
 		}
 	}
 
-	virtual void SerializeItem( FArchive& Ar, void* Value, int32 MaxReadBytes=0, void const* Defaults=NULL ) const PURE_VIRTUAL(UProperty::SerializeItem,);
+	virtual void SerializeItem( FArchive& Ar, void* Value, void const* Defaults=NULL ) const PURE_VIRTUAL(UProperty::SerializeItem,);
 	virtual bool NetSerializeItem( FArchive& Ar, UPackageMap* Map, void* Data, TArray<uint8> * MetaData = NULL ) const;
 	virtual void ExportTextItem( FString& ValueStr, const void* PropertyValue, const void* DefaultValue, UObject* Parent, int32 PortFlags, UObject* ExportRootScope = NULL ) const PURE_VIRTUAL(UProperty::ExportTextItem,);
 	const TCHAR* ImportText( const TCHAR* Buffer, void* Data, int32 PortFlags, UObject* OwnerObject, FOutputDevice* ErrorText = (FOutputDevice*)GWarn ) const
@@ -1042,7 +1042,7 @@ public:
 	{
 		return TTypeFundamentals::GetPropertyValue(A) == TTypeFundamentals::GetOptionalPropertyValue(B);
 	}
-	virtual void SerializeItem( FArchive& Ar, void* Value, int32 MaxReadBytes, void const* Defaults ) const override
+	virtual void SerializeItem( FArchive& Ar, void* Value, void const* Defaults ) const override
 	{
 		Ar << *TTypeFundamentals::GetPropertyValuePtr(Value);
 	}
@@ -1282,7 +1282,7 @@ class COREUOBJECT_API UByteProperty : public TProperty_Numeric<uint8>
 	// End of UHT interface
 
 	// UProperty interface.
-	virtual void SerializeItem( FArchive& Ar, void* Value, int32 MaxReadBytes, void const* Defaults ) const override;
+	virtual void SerializeItem( FArchive& Ar, void* Value, void const* Defaults ) const override;
 	virtual bool NetSerializeItem( FArchive& Ar, UPackageMap* Map, void* Data, TArray<uint8> * MetaData = NULL ) const override;
 	virtual void ExportTextItem( FString& ValueStr, const void* PropertyValue, const void* DefaultValue, UObject* Parent, int32 PortFlags, UObject* ExportRootScope ) const override;
 	virtual const TCHAR* ImportText_Internal( const TCHAR* Buffer, void* Data, int32 PortFlags, UObject* Parent, FOutputDevice* ErrorText ) const override;
@@ -1507,7 +1507,7 @@ public:
 	// UProperty interface.
 	virtual void LinkInternal(FArchive& Ar) override;
 	virtual bool Identical( const void* A, const void* B, uint32 PortFlags ) const override;
-	virtual void SerializeItem( FArchive& Ar, void* Value, int32 MaxReadBytes, void const* Defaults ) const override;
+	virtual void SerializeItem( FArchive& Ar, void* Value, void const* Defaults ) const override;
 	virtual bool NetSerializeItem( FArchive& Ar, UPackageMap* Map, void* Data, TArray<uint8> * MetaData = NULL ) const override;
 	virtual void ExportTextItem( FString& ValueStr, const void* PropertyValue, const void* DefaultValue, UObject* Parent, int32 PortFlags, UObject* ExportRootScope ) const override;
 	virtual const TCHAR* ImportText_Internal( const TCHAR* Buffer, void* Data, int32 PortFlags, UObject* Parent, FOutputDevice* ErrorText ) const override;
@@ -1771,7 +1771,7 @@ class COREUOBJECT_API UObjectProperty : public TUObjectPropertyBase<UObject*>
 	// End of UHT interface
 
 	// UProperty interface
-	virtual void SerializeItem( FArchive& Ar, void* Value, int32 MaxReadBytes, void const* Defaults ) const override;
+	virtual void SerializeItem( FArchive& Ar, void* Value, void const* Defaults ) const override;
 	virtual void EmitReferenceInfo(UClass& OwnerClass, int32 BaseOffset) override;
 	// End of UProperty interface
 	// UObjectPropertyBase interface
@@ -1805,7 +1805,7 @@ class COREUOBJECT_API UWeakObjectProperty : public TUObjectPropertyBase<FWeakObj
 	// End of UHT interface
 
 	// UProperty interface
-	virtual void SerializeItem( FArchive& Ar, void* Value, int32 MaxReadBytes, void const* Defaults ) const override;
+	virtual void SerializeItem( FArchive& Ar, void* Value, void const* Defaults ) const override;
 	// End of UProperty interface
 
 	// UObjectProperty interface
@@ -1840,7 +1840,7 @@ class COREUOBJECT_API ULazyObjectProperty : public TUObjectPropertyBase<FLazyObj
 	// UProperty interface
 	virtual FName GetID() const override;
 	virtual bool Identical( const void* A, const void* B, uint32 PortFlags ) const override;
-	virtual void SerializeItem( FArchive& Ar, void* Value, int32 MaxReadBytes, void const* Defaults ) const override;
+	virtual void SerializeItem( FArchive& Ar, void* Value, void const* Defaults ) const override;
 	// End of UProperty interface
 
 	// UObjectProperty interface
@@ -1878,7 +1878,7 @@ class COREUOBJECT_API UAssetObjectProperty : public TUObjectPropertyBase<FAssetP
 	// UProperty interface
 	virtual FName GetID() const override;
 	virtual bool Identical( const void* A, const void* B, uint32 PortFlags ) const override;
-	virtual void SerializeItem( FArchive& Ar, void* Value, int32 MaxReadBytes, void const* Defaults ) const override;
+	virtual void SerializeItem( FArchive& Ar, void* Value, void const* Defaults ) const override;
 	virtual void ExportTextItem( FString& ValueStr, const void* PropertyValue, const void* DefaultValue, UObject* Parent, int32 PortFlags, UObject* ExportRootScope ) const override;
 	virtual const TCHAR* ImportText_Internal( const TCHAR* Buffer, void* Data, int32 PortFlags, UObject* OwnerObject, FOutputDevice* ErrorText ) const override;
 	// End of UProperty interface
@@ -2014,7 +2014,7 @@ public:
 	// UProperty interface
 	virtual void LinkInternal(FArchive& Ar) override;
 	virtual bool Identical( const void* A, const void* B, uint32 PortFlags ) const override;
-	virtual void SerializeItem( FArchive& Ar, void* Value, int32 MaxReadBytes, void const* Defaults ) const override;
+	virtual void SerializeItem( FArchive& Ar, void* Value, void const* Defaults ) const override;
 	virtual bool NetSerializeItem( FArchive& Ar, UPackageMap* Map, void* Data, TArray<uint8> * MetaData = NULL ) const override;
 	virtual void ExportTextItem( FString& ValueStr, const void* PropertyValue, const void* DefaultValue, UObject* Parent, int32 PortFlags, UObject* ExportRootScope ) const override;
 	virtual const TCHAR* ImportText_Internal( const TCHAR* Buffer, void* Data, int32 PortFlags, UObject* OwnerObject, FOutputDevice* ErrorText ) const override;
@@ -2136,7 +2136,7 @@ public:
 	virtual FString GetCPPTypeForwardDeclaration() const override;
 	virtual void LinkInternal(FArchive& Ar) override;
 	virtual bool Identical( const void* A, const void* B, uint32 PortFlags ) const override;
-	virtual void SerializeItem( FArchive& Ar, void* Value, int32 MaxReadBytes, void const* Defaults ) const override;
+	virtual void SerializeItem( FArchive& Ar, void* Value, void const* Defaults ) const override;
 	virtual bool NetSerializeItem( FArchive& Ar, UPackageMap* Map, void* Data, TArray<uint8> * MetaData = NULL ) const override;
 	virtual void ExportTextItem( FString& ValueStr, const void* PropertyValue, const void* DefaultValue, UObject* Parent, int32 PortFlags, UObject* ExportRootScope ) const override;
 	virtual const TCHAR* ImportText_Internal( const TCHAR* Buffer, void* Data, int32 PortFlags, UObject* OwnerObject, FOutputDevice* ErrorText ) const override;
@@ -2502,7 +2502,7 @@ public:
 	virtual FString GetCPPTypeForwardDeclaration() const override;
 	virtual void LinkInternal(FArchive& Ar) override;
 	virtual bool Identical( const void* A, const void* B, uint32 PortFlags ) const override;
-	virtual void SerializeItem( FArchive& Ar, void* Value, int32 MaxReadBytes, void const* Defaults ) const override;
+	virtual void SerializeItem( FArchive& Ar, void* Value, void const* Defaults ) const override;
 	virtual bool NetSerializeItem( FArchive& Ar, UPackageMap* Map, void* Data, TArray<uint8> * MetaData = NULL ) const override;
 	virtual void ExportTextItem( FString& ValueStr, const void* PropertyValue, const void* DefaultValue, UObject* Parent, int32 PortFlags, UObject* ExportRootScope ) const override;
 	virtual const TCHAR* ImportText_Internal( const TCHAR* Buffer, void* Data, int32 PortFlags, UObject* OwnerObject, FOutputDevice* ErrorText ) const override;
@@ -2574,7 +2574,7 @@ public:
 	virtual FString GetCPPType( FString* ExtendedTypeText, uint32 CPPExportFlags ) const override;
 	virtual FString GetCPPTypeForwardDeclaration() const override;
 	virtual bool Identical( const void* A, const void* B, uint32 PortFlags ) const override;
-	virtual void SerializeItem( FArchive& Ar, void* Value, int32 MaxReadBytes, void const* Defaults ) const override;
+	virtual void SerializeItem( FArchive& Ar, void* Value, void const* Defaults ) const override;
 	virtual bool NetSerializeItem( FArchive& Ar, UPackageMap* Map, void* Data, TArray<uint8> * MetaData = NULL ) const override;
 	virtual void ExportTextItem( FString& ValueStr, const void* PropertyValue, const void* DefaultValue, UObject* Parent, int32 PortFlags, UObject* ExportRootScope ) const override;
 	virtual const TCHAR* ImportText_Internal( const TCHAR* Buffer, void* Data, int32 PortFlags, UObject* OwnerObject, FOutputDevice* ErrorText ) const override;
@@ -2619,7 +2619,7 @@ public:
 	// UProperty interface
 	virtual FString GetCPPType( FString* ExtendedTypeText, uint32 CPPExportFlags ) const override;
 	virtual bool Identical( const void* A, const void* B, uint32 PortFlags ) const override;
-	virtual void SerializeItem( FArchive& Ar, void* Value, int32 MaxReadBytes, void const* Defaults ) const override;
+	virtual void SerializeItem( FArchive& Ar, void* Value, void const* Defaults ) const override;
 	virtual bool NetSerializeItem( FArchive& Ar, UPackageMap* Map, void* Data, TArray<uint8> * MetaData = NULL ) const override;
 	virtual void ExportTextItem( FString& ValueStr, const void* PropertyValue, const void* DefaultValue, UObject* Parent, int32 PortFlags, UObject* ExportRootScope ) const override;
 	virtual const TCHAR* ImportText_Internal( const TCHAR* Buffer, void* Data, int32 PortFlags, UObject* OwnerObject, FOutputDevice* ErrorText ) const override;
diff --git a/Engine/Source/Runtime/Engine/Private/BlueprintGeneratedClass.cpp b/Engine/Source/Runtime/Engine/Private/BlueprintGeneratedClass.cpp
--- a/Engine/Source/Runtime/Engine/Private/BlueprintGeneratedClass.cpp
+++ b/Engine/Source/Runtime/Engine/Private/BlueprintGeneratedClass.cpp
@@ -559,7 +559,7 @@ void UBlueprintGeneratedClass::AddReferencedObjectsInUbergraphFrame(UObject* InT
 				if (PointerToUberGraphFrame->RawPointer)
 				{
 					FSimpleObjectReferenceCollectorArchive ObjectReferenceCollector(InThis, Collector);
-					BPGC->UberGraphFunction->SerializeBin(ObjectReferenceCollector, PointerToUberGraphFrame->RawPointer, 0);
+					BPGC->UberGraphFunction->SerializeBin(ObjectReferenceCollector, PointerToUberGraphFrame->RawPointer);
 				}
 			}
 		}
diff --git a/Engine/Source/Runtime/Engine/Private/DataTable.cpp b/Engine/Source/Runtime/Engine/Private/DataTable.cpp
--- a/Engine/Source/Runtime/Engine/Private/DataTable.cpp
+++ b/Engine/Source/Runtime/Engine/Private/DataTable.cpp
@@ -186,7 +186,7 @@ void UDataTable::AddReferencedObjects(UObject* InThis, FReferenceCollector& Coll
 			{
 				// Serialize all of the properties to make sure they get in the collector
 				FSimpleObjectReferenceCollectorArchive ObjectReferenceCollector( This, Collector );
-				This->RowStruct->SerializeBin(ObjectReferenceCollector, RowData, 0);
+				This->RowStruct->SerializeBin(ObjectReferenceCollector, RowData);
 			}
 		}
 	}
