UObject::GetArchetypeFromRequiredInfo  - fix for InheritableComponents. RF_InheritableComponentTemplate flag was added.

diff --git a/Engine/Source/Runtime/CoreUObject/Private/UObject/LinkerLoad.cpp b/Engine/Source/Runtime/CoreUObject/Private/UObject/LinkerLoad.cpp
--- a/Engine/Source/Runtime/CoreUObject/Private/UObject/LinkerLoad.cpp
+++ b/Engine/Source/Runtime/CoreUObject/Private/UObject/LinkerLoad.cpp
@@ -3213,7 +3213,7 @@ UObject* ULinkerLoad::CreateExport( int32 Index )
 		}
 
 		// Find the Archetype object for the one we are loading.
-		UObject* Template = GetArchetypeFromRequiredInfo(LoadClass, ThisParent, Export.ObjectName, !!(Export.ObjectFlags&RF_ClassDefaultObject));
+		UObject* Template = GetArchetypeFromRequiredInfo(LoadClass, ThisParent, Export.ObjectName, Export.ObjectFlags);
 
 		check(Template);
 		checkfSlow(((Export.ObjectFlags&RF_ClassDefaultObject)!=0 || Template->IsA(LoadClass)), TEXT("Mismatch between template %s and load class %s.  If this is a legacy blueprint or map, it may need to be resaved with bRecompileOnLoad turned off."), *Template->GetPathName(), *LoadClass->GetPathName());
diff --git a/Engine/Source/Runtime/CoreUObject/Private/UObject/UObjectArchetype.cpp b/Engine/Source/Runtime/CoreUObject/Private/UObject/UObjectArchetype.cpp
--- a/Engine/Source/Runtime/CoreUObject/Private/UObject/UObjectArchetype.cpp
+++ b/Engine/Source/Runtime/CoreUObject/Private/UObject/UObjectArchetype.cpp
@@ -7,9 +7,10 @@
 #include "CoreUObjectPrivate.h"
 
 
-UObject* UObject::GetArchetypeFromRequiredInfo(UClass* Class, UObject* Outer, FName Name, bool bIsCDO)
+UObject* UObject::GetArchetypeFromRequiredInfo(UClass* Class, UObject* Outer, FName Name, EObjectFlags ObjectFlags)
 {
 	UObject* Result = NULL;
+	const bool bIsCDO = !!(ObjectFlags&RF_ClassDefaultObject);
 	if (bIsCDO)
 	{
 		static UClass* UObjectStaticClass(UObject::StaticClass());
@@ -31,6 +32,23 @@ UObject* UObject::GetArchetypeFromRequiredInfo(UClass* Class, UObject* Outer, FN
 				Result = MyArchetype; // found that my outers archetype had a matching component, that must be my archetype
 			}
 		}
+
+		if (!Result 
+			&& !!(ObjectFlags&RF_InheritableComponentTemplate)
+			&& Outer && Outer->IsA<UClass>())
+		{
+			for (auto SuperClassArchetype = static_cast<UClass*>(Outer)->GetSuperClass();
+				SuperClassArchetype && SuperClassArchetype->HasAllClassFlags(CLASS_CompiledFromBlueprint);
+				SuperClassArchetype = SuperClassArchetype->GetSuperClass())
+			{
+				Result = static_cast<UObject*>(FindObjectWithOuter(SuperClassArchetype, Class, Name));
+				if (Result)
+				{
+					break;
+				}
+			}
+		}
+
 		if (!Result)
 		{
 			// nothing found, I am not a CDO, so this is just the class CDO
diff --git a/Engine/Source/Runtime/CoreUObject/Private/UObject/UObjectGlobals.cpp b/Engine/Source/Runtime/CoreUObject/Private/UObject/UObjectGlobals.cpp
--- a/Engine/Source/Runtime/CoreUObject/Private/UObject/UObjectGlobals.cpp
+++ b/Engine/Source/Runtime/CoreUObject/Private/UObject/UObjectGlobals.cpp
@@ -2202,7 +2202,7 @@ UObject* StaticConstructObject
 	const bool bIsNativeFromCDO = InClass->HasAnyClassFlags(CLASS_Native | CLASS_Intrinsic) && 
 		(
 			!InTemplate || 
-			(InName != NAME_None && InTemplate == UObject::GetArchetypeFromRequiredInfo(InClass, InOuter, InName, !!(InFlags & RF_ClassDefaultObject)))
+			(InName != NAME_None && InTemplate == UObject::GetArchetypeFromRequiredInfo(InClass, InOuter, InName, InFlags))
 		);
 	bool bRecycledSubobject = false;
 	Result = StaticAllocateObject(InClass, InOuter, InName, InFlags, bIsNativeFromCDO, &bRecycledSubobject);
diff --git a/Engine/Source/Runtime/CoreUObject/Public/UObject/ObjectBase.h b/Engine/Source/Runtime/CoreUObject/Public/UObject/ObjectBase.h
--- a/Engine/Source/Runtime/CoreUObject/Public/UObject/ObjectBase.h
+++ b/Engine/Source/Runtime/CoreUObject/Public/UObject/ObjectBase.h
@@ -401,13 +401,14 @@ enum EObjectFlags
 	RF_WasLoaded				=0x00080000,	///< Flagged on UObjects that were loaded
 	RF_TextExportTransient		=0x00100000,	///< Do not export object to text form (e.g. copy/paste). Generally used for sub-objects that can be regenerated from data in their parent object.
 	RF_LoadCompleted			=0x00200000,	///< Object has been completely serialized by linkerload at least once. DO NOT USE THIS FLAG, It should be replaced with RF_WasLoaded.
+	RF_InheritableComponentTemplate = 0x00400000, ///< Archetype of the object can be in its super class
 
 	// Special all and none masks
-	RF_AllFlags					=0x003fffff,	///< All flags, used mainly for error checking
+	RF_AllFlags					=0x00ffffff,	///< All flags, used mainly for error checking
 	RF_NoFlags					=0x00000000,	///< No flags, used to avoid a cast
 
 	// Predefined groups of the above
-	RF_Load						= RF_Public | RF_Standalone | RF_Native | RF_Transactional | RF_ClassDefaultObject | RF_ArchetypeObject | RF_DefaultSubObject | RF_TextExportTransient, // Flags to load from Unrealfiles.
+	RF_Load						= RF_Public | RF_Standalone | RF_Native | RF_Transactional | RF_ClassDefaultObject | RF_ArchetypeObject | RF_DefaultSubObject | RF_TextExportTransient | RF_InheritableComponentTemplate, // Flags to load from Unrealfiles.
 	RF_PropagateToSubObjects	= RF_Public | RF_ArchetypeObject | RF_Transactional | RF_Transient,		// Sub-objects will inherit these flags from their SuperObject.
 
 	
diff --git a/Engine/Source/Runtime/CoreUObject/Public/UObject/UObject.h b/Engine/Source/Runtime/CoreUObject/Public/UObject/UObject.h
--- a/Engine/Source/Runtime/CoreUObject/Public/UObject/UObject.h
+++ b/Engine/Source/Runtime/CoreUObject/Public/UObject/UObject.h
@@ -821,7 +821,7 @@ public:
 	 * 
 	 * @return the archetype for this object
 	 */
-	static UObject* GetArchetypeFromRequiredInfo(UClass* Class, UObject* Outer, FName Name, bool bIsCDO);
+	static UObject* GetArchetypeFromRequiredInfo(UClass* Class, UObject* Outer, FName Name, EObjectFlags ObjectFlags);
 
 	/**
 	 * Return the template this object is based on. 
@@ -830,7 +830,7 @@ public:
 	 */
 	UObject* GetArchetype() const
 	{
-		return GetArchetypeFromRequiredInfo(GetClass(), GetOuter(), GetFName(), HasAnyFlags(RF_ClassDefaultObject));
+		return GetArchetypeFromRequiredInfo(GetClass(), GetOuter(), GetFName(), GetFlags());
 	}
 
 	/**
diff --git a/Engine/Source/Runtime/Engine/Private/InheritableComponentHandler.cpp b/Engine/Source/Runtime/Engine/Private/InheritableComponentHandler.cpp
--- a/Engine/Source/Runtime/Engine/Private/InheritableComponentHandler.cpp
+++ b/Engine/Source/Runtime/Engine/Private/InheritableComponentHandler.cpp
@@ -31,7 +31,7 @@ UActorComponent* UInheritableComponentHandler::CreateOverridenComponentTemplate(
 	}
 	ensure(Cast<UBlueprintGeneratedClass>(GetOuter()));
 	auto NewComponentTemplate = ConstructObject<UActorComponent>(
-		BestArchetype->GetClass(), GetOuter(), BestArchetype->GetFName(), RF_ArchetypeObject | RF_Public, BestArchetype);
+		BestArchetype->GetClass(), GetOuter(), BestArchetype->GetFName(), RF_ArchetypeObject | RF_Public | RF_InheritableComponentTemplate, BestArchetype);
 
 	FComponentOverrideRecord NewRecord;
 	NewRecord.ComponentKey = Key;
