This change fixes some issues with previous change ( CL 2410401 ) which added the ability in the blueprint details panel to set blueprint variable values using the owning config file.
#UE-6467 - Blueprints need to be able to make properties config (.ini)
#Branch UE4
#Proj CoreUObject

#info Minor change, but just wanted to get it checked as it is in CoreUObject and also possibly discuss if we should feed a certain ini file string to the LoadConfig call.

diff --git a/Engine/Source/Runtime/CoreUObject/Private/Blueprint/BlueprintSupport.cpp b/Engine/Source/Runtime/CoreUObject/Private/Blueprint/BlueprintSupport.cpp
--- a/Engine/Source/Runtime/CoreUObject/Private/Blueprint/BlueprintSupport.cpp
+++ b/Engine/Source/Runtime/CoreUObject/Private/Blueprint/BlueprintSupport.cpp
@@ -237,6 +237,12 @@ bool ULinkerLoad::RegenerateBlueprintClass(UClass* LoadClass, UObject* ExportObj
 		LoadClass->StaticLink(true);
 		Preload(CurrentCDO);
 
+		// Load the class config values
+		if( LoadClass->HasAnyClassFlags( CLASS_Config ))
+		{
+			CurrentCDO->LoadConfig( LoadClass );
+		}
+
 		// Make sure that we regenerate any parent classes first before attempting to build a child
 		TArray<UClass*> ClassChainOrdered;
 		{
