Remove obsolete FNetworkGUID USTRUCT, FNetworkGUIDs should not be UPROPERTY()s

diff --git a/Engine/Source/Runtime/Core/Private/Misc/NetworkGuid.cpp b/Engine/Source/Runtime/Core/Private/Misc/NetworkGuid.cpp
deleted file mode 100644
diff --git a/Engine/Source/Runtime/CoreUObject/Classes/Object.h b/Engine/Source/Runtime/CoreUObject/Classes/Object.h
--- a/Engine/Source/Runtime/CoreUObject/Classes/Object.h
+++ b/Engine/Source/Runtime/CoreUObject/Classes/Object.h
@@ -223,14 +223,6 @@ struct FGuid
 	int32 D;
 };
 
-// A unique identifier for networking objects
-USTRUCT(immutable, noexport)
-struct FNetworkGUID
-{
-	UPROPERTY(SaveGame)
-	uint32 Value;
-};
-
 /**
  * A point or direction FVector in 3d space.
  * The full C++ class is located here: Engine\Source\Runtime\Core\Public\Math\Vector.h
diff --git a/Engine/Source/Runtime/CoreUObject/Private/UObject/Property.cpp b/Engine/Source/Runtime/CoreUObject/Private/UObject/Property.cpp
--- a/Engine/Source/Runtime/CoreUObject/Private/UObject/Property.cpp
+++ b/Engine/Source/Runtime/CoreUObject/Private/UObject/Property.cpp
@@ -199,16 +199,6 @@ struct TStructOpsTypeTraits<FGuid> : public TStructOpsTypeTraitsBase
 };
 IMPLEMENT_STRUCT(Guid);
 
-template<>
-struct TStructOpsTypeTraits<FNetworkGUID> : public TStructOpsTypeTraitsBase
-{
-	enum 
-	{
-		WithNetSerializer = true,
-	};
-};
-IMPLEMENT_STRUCT(NetworkGUID);
-
 template<>
 struct TStructOpsTypeTraits<FTransform> : public TStructOpsTypeTraitsBase
 {
