[UE-7296][UE-7784][UE-7843] Catching and handling more dangling placeholder class references (WIP: deferred dependency loading, for circular dependency issues).

diff --git a/Engine/Source/Editor/KismetCompiler/Private/KismetCompilerMisc.cpp b/Engine/Source/Editor/KismetCompiler/Private/KismetCompilerMisc.cpp
--- a/Engine/Source/Editor/KismetCompiler/Private/KismetCompilerMisc.cpp
+++ b/Engine/Source/Editor/KismetCompiler/Private/KismetCompilerMisc.cpp
@@ -700,7 +700,10 @@ UProperty* FKismetCompilerUtilities::CreatePropertyOnScope(UStruct* Scope, const
 			if (SubType->HasAnyClassFlags(CLASS_Interface))
 			{
 				UInterfaceProperty* NewPropertyObj = NewNamedObject<UInterfaceProperty>(PropertyScope, ValidatedPropertyName, ObjectFlags);
-				NewPropertyObj->InterfaceClass = SubType;
+				// we want to use this setter function instead of setting the 
+				// InterfaceClass member directly, because it properly handles  
+				// placeholder classes (classes that are stubbed in during load)
+				NewPropertyObj->SetInterfaceClass(SubType);
 				NewProperty = NewPropertyObj;
 			}
 			else
@@ -715,7 +718,10 @@ UProperty* FKismetCompilerUtilities::CreatePropertyOnScope(UStruct* Scope, const
 				{
 					NewPropertyObj = NewNamedObject<UObjectProperty>(PropertyScope, ValidatedPropertyName, ObjectFlags);
 				}
-				NewPropertyObj->PropertyClass = SubType;
+				// we want to use this setter function instead of setting the 
+				// PropertyClass member directly, because it properly handles  
+				// placeholder classes (classes that are stubbed in during load)
+				NewPropertyObj->SetPropertyClass(SubType);
 				NewProperty = NewPropertyObj;
 			}
 		}
@@ -763,7 +769,10 @@ UProperty* FKismetCompilerUtilities::CreatePropertyOnScope(UStruct* Scope, const
 		if (SubType != NULL)
 		{
 			UClassProperty* NewPropertyClass = NewNamedObject<UClassProperty>(PropertyScope, ValidatedPropertyName, ObjectFlags);
-			NewPropertyClass->MetaClass = SubType;
+			// we want to use this setter function instead of setting the 
+			// MetaClass member directly, because it properly handles  
+			// placeholder classes (classes that are stubbed in during load)
+			NewPropertyClass->SetMetaClass(SubType);
 			NewPropertyClass->PropertyClass = UClass::StaticClass();
 			NewProperty = NewPropertyClass;
 		}
diff --git a/Engine/Source/Runtime/CoreUObject/Private/Blueprint/BlueprintSupport.cpp b/Engine/Source/Runtime/CoreUObject/Private/Blueprint/BlueprintSupport.cpp
--- a/Engine/Source/Runtime/CoreUObject/Private/Blueprint/BlueprintSupport.cpp
+++ b/Engine/Source/Runtime/CoreUObject/Private/Blueprint/BlueprintSupport.cpp
@@ -348,8 +348,9 @@ bool ULinkerLoad::DeferPotentialCircularImport(const int32 Index)
 #endif // USE_CIRCULAR_DEPENDENCY_LOAD_DEFERRING
 }
 
-void ULinkerLoad::ResolveDeferredDependencies()
+void ULinkerLoad::ResolveDeferredDependencies(UClass* LoadClass)
 {
+#if USE_CIRCULAR_DEPENDENCY_LOAD_DEFERRING
 	//--------------------------------------
 	// Phase 2: Resolve Dependency Stubs
 	//--------------------------------------
@@ -364,19 +365,46 @@ void ULinkerLoad::ResolveDeferredDependencies()
 			Import.XObject = nullptr;
 			UClass* RealClassObj = CastChecked<UClass>(CreateImport(ImportIndex), ECastCheckedType::NullAllowed);
 
+			bool bUsedPlaceholderInterface = false;
+			if (LoadClass != nullptr)
+			{
+				for (FImplementedInterface& Interface : LoadClass->Interfaces)
+				{
+					if (Interface.Class == PlaceholderClass)
+					{
+						bUsedPlaceholderInterface = true;
+						Interface.Class = RealClassObj;
+					}
+				}
+			}
+
 #if TEST_CHECK_DEPENDENCY_LOAD_DEFERRING
 			bool const bCircumventValidationChecks = !FBlueprintSupport::UseDeferredDependencyVerificationChecks();
+			checkSlow(bCircumventValidationChecks || (PlaceholderClass->GetOuter() == LinkerRoot));
 			// make sure that we know what utilized this placeholder class... 
 			// right now we only expect UObjectProperties/UClassProperties/
-			// UInterfaceProperties to, but something else could have requested 
-			// the class without logging itself with the placeholder... if  
-			// ReferencingProperties is empty, then there is something out there 
-			// still using a placeholder class
-			checkSlow(bCircumventValidationChecks || PlaceholderClass->HasReferences());
+			// UInterfaceProperties/FImplementedInterfaces to, but something 
+			// else could have requested the class without logging itself with 
+			// the placeholder... if ReferencingProperties is empty, then there 
+			// is something out there still using a placeholder class
+			checkSlow(bCircumventValidationChecks || bUsedPlaceholderInterface || PlaceholderClass->HasReferences());
+
+			UObject* PlaceholderObj = PlaceholderClass;
+			// 
+			FReferencerInformationList RemainingRefs;
+			//checkSlow(bCircumventValidationChecks || PlaceholderClass->HasReferences() == IsReferenced(PlaceholderObj, RF_NoFlags, /*bCheckSubObjects =*/false, &RemainingRefs));
+			//checkSlow(bCircumventValidationChecks || RemainingRefs.ExternalReferences.Num() == PlaceholderClass->GetRefCount());
 #endif // TEST_CHECK_DEPENDENCY_LOAD_DEFERRING
 
 			PlaceholderClass->ReplaceTrackedReferences(RealClassObj);
 			PlaceholderClass->MarkPendingKill(); // @TODO: ensure these are properly GC'd
+		
+#if TEST_CHECK_DEPENDENCY_LOAD_DEFERRING
+			// there should not be any references left to this placeholder class 
+			// (if there is, then we didn't log that referencer with the placeholder)
+			FReferencerInformationList DanglingReferences;
+			checkSlow(bCircumventValidationChecks || !IsReferenced(PlaceholderObj, RF_NoFlags, /*bCheckSubObjects =*/false, &DanglingReferences));
+#endif // TEST_CHECK_DEPENDENCY_LOAD_DEFERRING
 		}
 		else if (UScriptStruct* StructObj = Cast<UScriptStruct>(Import.XObject))
 		{
@@ -386,14 +414,16 @@ void ULinkerLoad::ResolveDeferredDependencies()
 			// @TODO: could this result in an endless ResolveDeferredDependencies() loop? 
 			if (Import.SourceLinker != nullptr)
 			{
-				Import.SourceLinker->ResolveDeferredDependencies();
+				Import.SourceLinker->ResolveDeferredDependencies(/*LoadClass =*/nullptr);
 			}
 		}
 	}
 
 	// @TODO: don't know if we need this, but could be good to have (as class 
-	//        regeneration is about to force load a lot of this).
-	LoadAllObjects();
+	//        regeneration is about to force load a lot of this), BUT! this 
+	//        doesn't work for map packages (because this would load the level's
+	//        ALevelScriptActor instance BEFORE the class has been regenerated)
+	//LoadAllObjects();
 
 #if TEST_CHECK_DEPENDENCY_LOAD_DEFERRING
 	bool const bCircumventValidationChecks = !FBlueprintSupport::UseDeferredDependencyVerificationChecks();
@@ -408,6 +438,7 @@ void ULinkerLoad::ResolveDeferredDependencies()
 		}
 	}
 #endif // TEST_CHECK_DEPENDENCY_LOAD_DEFERRING
+#endif // USE_CIRCULAR_DEPENDENCY_LOAD_DEFERRING
 }
 
 void ULinkerLoad::FinalizeBlueprint(UClass* LoadClass)
diff --git a/Engine/Source/Runtime/CoreUObject/Private/UObject/LinkerLoad.cpp b/Engine/Source/Runtime/CoreUObject/Private/UObject/LinkerLoad.cpp
--- a/Engine/Source/Runtime/CoreUObject/Private/UObject/LinkerLoad.cpp
+++ b/Engine/Source/Runtime/CoreUObject/Private/UObject/LinkerLoad.cpp
@@ -2879,7 +2879,7 @@ void ULinkerLoad::Preload( UObject* Object )
 					SCOPE_CYCLE_COUNTER(STAT_LinkerLoadDeferred);
 					if ((LoadFlags & LOAD_DeferDependencyLoads) != (*LoadFlagsGuard & LOAD_DeferDependencyLoads))
 					{
-						ResolveDeferredDependencies();
+						ResolveDeferredDependencies(ObjectAsClass);
 						FinalizeBlueprint(ObjectAsClass);
 					}
 #endif // USE_CIRCULAR_DEPENDENCY_LOAD_DEFERRING
diff --git a/Engine/Source/Runtime/CoreUObject/Private/UObject/LinkerPlaceholderClass.cpp b/Engine/Source/Runtime/CoreUObject/Private/UObject/LinkerPlaceholderClass.cpp
--- a/Engine/Source/Runtime/CoreUObject/Private/UObject/LinkerPlaceholderClass.cpp
+++ b/Engine/Source/Runtime/CoreUObject/Private/UObject/LinkerPlaceholderClass.cpp
@@ -7,6 +7,9 @@
 //------------------------------------------------------------------------------
 ULinkerPlaceholderClass::ULinkerPlaceholderClass(const FObjectInitializer& ObjectInitializer)
 	: Super(ObjectInitializer)
+#if TEST_CHECK_DEPENDENCY_LOAD_DEFERRING
+	, bResolvedReferences(false)
+#endif // TEST_CHECK_DEPENDENCY_LOAD_DEFERRING
 {
 }
 
@@ -68,6 +71,30 @@ void ULinkerPlaceholderClass::Bind()
 //------------------------------------------------------------------------------
 void ULinkerPlaceholderClass::AddTrackedReference(UProperty* ReferencingProperty)
 {
+#if TEST_CHECK_DEPENDENCY_LOAD_DEFERRING
+	bool const bCircumventValidationChecks = !FBlueprintSupport::UseDeferredDependencyVerificationChecks();
+	if (!bCircumventValidationChecks)
+	{
+		FObjectImport* PlaceholderImport = nullptr;
+		if (ULinkerLoad* PropertyLinker = ReferencingProperty->GetLinker())
+		{
+			for (FObjectImport& Import : PropertyLinker->ImportMap)
+			{
+				if (Import.XObject == this)
+				{
+					PlaceholderImport = &Import;
+					break;
+				}
+			}
+			check(GetOuter() == PropertyLinker->LinkerRoot);
+			check(PropertyLinker->LoadFlags & LOAD_DeferDependencyLoads);
+		}
+		// if this check hits, then we're adding dependencies after we've 
+		// already resolved the placeholder (it won't be resolved again)
+		check(!bResolvedReferences); 
+	}	
+#endif // TEST_CHECK_DEPENDENCY_LOAD_DEFERRING
+
 	ReferencingProperties.Add(ReferencingProperty);
 }
 
@@ -77,6 +104,12 @@ bool ULinkerPlaceholderClass::HasReferences() const
 	return (ReferencingProperties.Num() > 0);
 }
 
+//------------------------------------------------------------------------------
+int32 ULinkerPlaceholderClass::GetRefCount() const
+{
+	return ReferencingProperties.Num();
+}
+
 //------------------------------------------------------------------------------
 void ULinkerPlaceholderClass::RemoveTrackedReference(UProperty* ReferencingProperty)
 {
@@ -113,4 +146,8 @@ void ULinkerPlaceholderClass::ReplaceTrackedReferences(UClass* ReplacementClass)
 		}
 	}
 	ReferencingProperties.Empty();
+
+#if TEST_CHECK_DEPENDENCY_LOAD_DEFERRING
+	bResolvedReferences = true;
+#endif // TEST_CHECK_DEPENDENCY_LOAD_DEFERRING
 }
diff --git a/Engine/Source/Runtime/CoreUObject/Private/UObject/LinkerPlaceholderClass.h b/Engine/Source/Runtime/CoreUObject/Private/UObject/LinkerPlaceholderClass.h
--- a/Engine/Source/Runtime/CoreUObject/Private/UObject/LinkerPlaceholderClass.h
+++ b/Engine/Source/Runtime/CoreUObject/Private/UObject/LinkerPlaceholderClass.h
@@ -50,6 +50,14 @@ public:
 	 */
 	bool HasReferences() const;
 
+	/**
+	 * Query method that retrieves the current number of KNOWN references to 
+	 * this placeholder class.
+	 * 
+	 * @return The number of references that this class is currently tracking.
+	 */
+	int32 GetRefCount() const;
+
 	/**
 	 * Removes the specified property from this class's internal tracking list 
 	 * (which aims to keep track of properties utilizing this class).
@@ -69,4 +77,9 @@ public:
 private:
 	/** Links to UProperties that are currently using this class */
 	TSet<UProperty*> ReferencingProperties;
+
+#if TEST_CHECK_DEPENDENCY_LOAD_DEFERRING
+	/** Used to catch references that are added after we've already resolved all references */
+	bool bResolvedReferences;
+#endif // TEST_CHECK_DEPENDENCY_LOAD_DEFERRING
 }; 
diff --git a/Engine/Source/Runtime/CoreUObject/Private/UObject/PropertyBaseObject.cpp b/Engine/Source/Runtime/CoreUObject/Private/UObject/PropertyBaseObject.cpp
--- a/Engine/Source/Runtime/CoreUObject/Private/UObject/PropertyBaseObject.cpp
+++ b/Engine/Source/Runtime/CoreUObject/Private/UObject/PropertyBaseObject.cpp
@@ -90,7 +90,7 @@ void UObjectPropertyBase::Serialize( FArchive& Ar )
 	Ar << PropertyClass;
 
 #if USE_CIRCULAR_DEPENDENCY_LOAD_DEFERRING
-	if (Ar.IsLoading())
+	if (Ar.IsLoading() || Ar.IsObjectReferenceCollector())
 	{
 		if (ULinkerPlaceholderClass* PlaceholderClass = Cast<ULinkerPlaceholderClass>(PropertyClass))
 		{
@@ -99,6 +99,23 @@ void UObjectPropertyBase::Serialize( FArchive& Ar )
 	}
 #endif // USE_CIRCULAR_DEPENDENCY_LOAD_DEFERRING
 }
+
+#if USE_CIRCULAR_DEPENDENCY_LOAD_DEFERRING
+void UObjectPropertyBase::SetPropertyClass(UClass* NewPropertyClass)
+{
+	if (ULinkerPlaceholderClass* NewPlaceholderClass = Cast<ULinkerPlaceholderClass>(NewPropertyClass))
+	{
+		NewPlaceholderClass->AddTrackedReference(this);
+	}
+	
+	if (ULinkerPlaceholderClass* OldPlaceholderClass = Cast<ULinkerPlaceholderClass>(PropertyClass))
+	{
+		OldPlaceholderClass->RemoveTrackedReference(this);
+	}
+	PropertyClass = NewPropertyClass;
+}
+#endif // USE_CIRCULAR_DEPENDENCY_LOAD_DEFERRING
+
 void UObjectPropertyBase::AddReferencedObjects(UObject* InThis, FReferenceCollector& Collector)
 {	
 	UObjectPropertyBase* This = CastChecked<UObjectPropertyBase>(InThis);
diff --git a/Engine/Source/Runtime/CoreUObject/Private/UObject/PropertyClass.cpp b/Engine/Source/Runtime/CoreUObject/Private/UObject/PropertyClass.cpp
--- a/Engine/Source/Runtime/CoreUObject/Private/UObject/PropertyClass.cpp
+++ b/Engine/Source/Runtime/CoreUObject/Private/UObject/PropertyClass.cpp
@@ -27,7 +27,7 @@ void UClassProperty::Serialize( FArchive& Ar )
 	Ar << MetaClass;
 
 #if USE_CIRCULAR_DEPENDENCY_LOAD_DEFERRING
-	if (Ar.IsLoading())
+	if (Ar.IsLoading() || Ar.IsObjectReferenceCollector())
 	{
 		if (ULinkerPlaceholderClass* PlaceholderClass = Cast<ULinkerPlaceholderClass>(MetaClass))
 		{
@@ -48,12 +48,30 @@ void UClassProperty::Serialize( FArchive& Ar )
 		}
 	}
 }
+
+#if USE_CIRCULAR_DEPENDENCY_LOAD_DEFERRING
+void UClassProperty::SetMetaClass(UClass* NewMetaClass)
+{
+	if (ULinkerPlaceholderClass* NewPlaceholderClass = Cast<ULinkerPlaceholderClass>(NewMetaClass))
+	{
+		NewPlaceholderClass->AddTrackedReference(this);
+	}
+
+	if (ULinkerPlaceholderClass* OldPlaceholderClass = Cast<ULinkerPlaceholderClass>(NewMetaClass))
+	{
+		OldPlaceholderClass->RemoveTrackedReference(this);
+	}
+	MetaClass = NewMetaClass;
+}
+#endif // USE_CIRCULAR_DEPENDENCY_LOAD_DEFERRING
+
 void UClassProperty::AddReferencedObjects(UObject* InThis, FReferenceCollector& Collector)
 {
 	UClassProperty* This = CastChecked<UClassProperty>(InThis);
 	Collector.AddReferencedObject( This->MetaClass, This );
 	Super::AddReferencedObjects( This, Collector );
 }
+
 const TCHAR* UClassProperty::ImportText_Internal( const TCHAR* Buffer, void* Data, int32 PortFlags, UObject* Parent, FOutputDevice* ErrorText ) const
 {
 	const TCHAR* Result = UObjectProperty::ImportText_Internal( Buffer, Data, PortFlags, Parent, ErrorText );
diff --git a/Engine/Source/Runtime/CoreUObject/Private/UObject/PropertyInterface.cpp b/Engine/Source/Runtime/CoreUObject/Private/UObject/PropertyInterface.cpp
--- a/Engine/Source/Runtime/CoreUObject/Private/UObject/PropertyInterface.cpp
+++ b/Engine/Source/Runtime/CoreUObject/Private/UObject/PropertyInterface.cpp
@@ -214,7 +214,7 @@ void UInterfaceProperty::Serialize( FArchive& Ar )
 	Ar << InterfaceClass;
 
 #if USE_CIRCULAR_DEPENDENCY_LOAD_DEFERRING
-	if (Ar.IsLoading())
+	if (Ar.IsLoading() || Ar.IsObjectReferenceCollector())
 	{
 		if (ULinkerPlaceholderClass* PlaceholderClass = Cast<ULinkerPlaceholderClass>(InterfaceClass))
 		{
@@ -239,6 +239,23 @@ void UInterfaceProperty::Serialize( FArchive& Ar )
  	}
 }
 
+
+#if USE_CIRCULAR_DEPENDENCY_LOAD_DEFERRING
+void UInterfaceProperty::SetInterfaceClass(UClass* NewInterfaceClass)
+{
+	if (ULinkerPlaceholderClass* NewPlaceholderClass = Cast<ULinkerPlaceholderClass>(NewInterfaceClass))
+	{
+		NewPlaceholderClass->AddTrackedReference(this);
+	}
+
+	if (ULinkerPlaceholderClass* OldPlaceholderClass = Cast<ULinkerPlaceholderClass>(InterfaceClass))
+	{
+		OldPlaceholderClass->RemoveTrackedReference(this);
+	}
+	InterfaceClass = NewInterfaceClass;
+}
+#endif // USE_CIRCULAR_DEPENDENCY_LOAD_DEFERRING
+
 bool UInterfaceProperty::SameType(const UProperty* Other) const
 {
 	return Super::SameType(Other) && (InterfaceClass == ((UInterfaceProperty*)Other)->InterfaceClass);
diff --git a/Engine/Source/Runtime/CoreUObject/Public/UObject/Linker.h b/Engine/Source/Runtime/CoreUObject/Public/UObject/Linker.h
--- a/Engine/Source/Runtime/CoreUObject/Public/UObject/Linker.h
+++ b/Engine/Source/Runtime/CoreUObject/Public/UObject/Linker.h
@@ -1841,7 +1841,7 @@ private:
 	 * 
 	 * @param  LoadClass    The (Blueprint) class that was loading, while we deferred dependencies.
 	 */
-	void ResolveDeferredDependencies();
+	void ResolveDeferredDependencies(UClass* LoadClass);
 
 	/**
 	 * Takes the supplied serialized class and serializes in its CDO, then 
diff --git a/Engine/Source/Runtime/CoreUObject/Public/UObject/UnrealType.h b/Engine/Source/Runtime/CoreUObject/Public/UObject/UnrealType.h
--- a/Engine/Source/Runtime/CoreUObject/Public/UObject/UnrealType.h
+++ b/Engine/Source/Runtime/CoreUObject/Public/UObject/UnrealType.h
@@ -1710,6 +1710,21 @@ public:
 		SetObjectPropertyValue(ContainerPtrToValuePtr<void>(PropertyValueAddress, ArrayIndex), Value);
 	}
 
+	/**
+	 * Setter function for this property's PropertyClass member. Favor this 
+	 * function whilst loading (since, to handle circular dependencies, we defer 
+	 * some class loads and use a placeholder class instead). It properly 
+	 * handles deferred loading placeholder classes (so they can properly be 
+	 * replaced later).
+	 *  
+	 * @param  NewPropertyClass    The PropertyClass you want this property set with.
+	 */
+#if USE_CIRCULAR_DEPENDENCY_LOAD_DEFERRING
+	void SetPropertyClass(UClass* NewPropertyClass);
+#else
+	FORCEINLINE void SetPropertyClass(UClass* NewPropertyClass) { PropertyClass = NewPropertyClass; }
+#endif // USE_CIRCULAR_DEPENDENCY_LOAD_DEFERRING
+
 protected:
 	virtual bool AllowCrossLevel() const
 	{
@@ -1937,6 +1952,21 @@ public:
 	virtual bool SameType(const UProperty* Other) const override;
 	// End of UProperty interface
 
+	/**
+	 * Setter function for this property's MetaClass member. Favor this function 
+	 * whilst loading (since, to handle circular dependencies, we defer some 
+	 * class loads and use a placeholder class instead). It properly handles 
+	 * deferred loading placeholder classes (so they can properly be replaced 
+	 * later).
+	 * 
+	 * @param  NewMetaClass    The MetaClass you want this property set with.
+	 */
+#if USE_CIRCULAR_DEPENDENCY_LOAD_DEFERRING
+	void SetMetaClass(UClass* NewMetaClass);
+#else
+	FORCEINLINE void SetMetaClass(UClass* NewMetaClass) { MetaClass = NewMetaClass; }
+#endif // USE_CIRCULAR_DEPENDENCY_LOAD_DEFERRING
+
 protected:
 	virtual void CheckValidObject(void* Value) const override;
 };
@@ -2026,6 +2056,21 @@ public:
 	virtual void Serialize( FArchive& Ar ) override;
 	virtual void EmitReferenceInfo(UClass& OwnerClass, int32 BaseOffset) override;
 	// End of UObject interface
+
+	/**
+	 * Setter function for this property's InterfaceClass member. Favor this 
+	 * function whilst loading (since, to handle circular dependencies, we defer 
+	 * some class loads and use a placeholder class instead). It properly 
+	 * handles deferred loading placeholder classes (so they can properly be 
+	 * replaced later).
+	 *  
+	 * @param  NewInterfaceClass    The InterfaceClass you want this property set with.
+	 */
+#if USE_CIRCULAR_DEPENDENCY_LOAD_DEFERRING
+	void SetInterfaceClass(UClass* NewInterfaceClass);
+#else
+	FORCEINLINE void SetInterfaceClass(UClass* NewInterfaceClass) { InterfaceClass = NewInterfaceClass; }
+#endif // USE_CIRCULAR_DEPENDENCY_LOAD_DEFERRING
 };
 
 /*-----------------------------------------------------------------------------
