[UE-7296][UE-7784][UE-7973] Catching and fixing-up more dangling placeholder refs ending up in generated script-code... making sure a class's super is resolved and finalized before we regenerate a sub-class.

diff --git a/Engine/Source/Runtime/CoreUObject/Private/Blueprint/BlueprintSupport.cpp b/Engine/Source/Runtime/CoreUObject/Private/Blueprint/BlueprintSupport.cpp
--- a/Engine/Source/Runtime/CoreUObject/Private/Blueprint/BlueprintSupport.cpp
+++ b/Engine/Source/Runtime/CoreUObject/Private/Blueprint/BlueprintSupport.cpp
@@ -312,8 +312,8 @@ bool ULinkerLoad::DeferPotentialCircularImport(const int32 Index)
 	FObjectImport& Import = ImportMap[Index];
 	if (Import.XObject != nullptr)
 	{
-		// this import has already been created
-		return Import.XObject->IsA<ULinkerPlaceholderClass>();
+		bool const bIsPlaceHolderClass = Import.XObject->IsA<ULinkerPlaceholderClass>();
+		return bIsPlaceHolderClass;
 	}
 
 	if ((LoadFlags & LOAD_DeferDependencyLoads) && !IsImportNative(Index))
@@ -331,7 +331,12 @@ bool ULinkerLoad::DeferPotentialCircularImport(const int32 Index)
 					FName PlaceholderName(*FString::Printf(TEXT("PLACEHOLDER-CLASS_%s"), *Import.ObjectName.ToString()));
 					PlaceholderName = MakeUniqueObjectName(PlaceholderOuter, PlaceholderType, PlaceholderName);
 
-					UClass* Placeholder = ConstructObject<ULinkerPlaceholderClass>(PlaceholderType, PlaceholderOuter, PlaceholderName, RF_Public | RF_Transient);
+					ULinkerPlaceholderClass* Placeholder = ConstructObject<ULinkerPlaceholderClass>(PlaceholderType, PlaceholderOuter, PlaceholderName, RF_Public | RF_Transient);
+					// store the import index in the placeholder, so we can 
+					// easily look it up in the import map, given the 
+					// placeholder (needed, to find the corresponding import for 
+					// ResolvingDeferredPlaceholder)
+					Placeholder->ImportIndex = Index;
 					// make sure the class is fully formed (has its 
 					// ClassAddReferencedObjects/ClassConstructor members set)
 					Placeholder->Bind();
@@ -356,62 +361,47 @@ void ULinkerLoad::ResolveDeferredDependencies(UClass* LoadClass)
 	//--------------------------------------
 	TGuardValue<uint32> LoadFlagsGuard(LoadFlags, (LoadFlags & ~LOAD_DeferDependencyLoads));
 
+#if TEST_CHECK_DEPENDENCY_LOAD_DEFERRING
+	static int32 RecursiveDepth = 0;
+	int32 const MeasuredDepth = RecursiveDepth;
+	TGuardValue<int32> RecursiveDepthGuard(RecursiveDepth, RecursiveDepth + 1);
+
+	bool const bCircumventValidationChecks = !FBlueprintSupport::UseDeferredDependencyVerificationChecks();
+	checkSlow(bCircumventValidationChecks || !LoadClass || (LoadClass->GetLinker() == this));
+#endif // TEST_CHECK_DEPENDENCY_LOAD_DEFERRING
+
+	// this function (for this linker) could be reentrant (see where we 
+	// recursively call ResolveDeferredDependencies() for super-classes below);  
+	// if that's the case, then we want to finish resolving the pending class 
+	// before we continue on
+	if (ResolvingDeferredPlaceholder != nullptr)
+	{
+		int32 ResolvedRefCount = ResolveDependencyPlaceholder(ResolvingDeferredPlaceholder, LoadClass);
+		// @TODO: can we reliably count on this resolving some dependencies?... 
+		//        if so, check verify that!
+		ResolvingDeferredPlaceholder = nullptr;
+	}
+	
 	for (int32 ImportIndex = 0; ImportIndex < ImportMap.Num(); ++ImportIndex)
 	{
 		FObjectImport& Import = ImportMap[ImportIndex];
 
 		if (ULinkerPlaceholderClass* PlaceholderClass = Cast<ULinkerPlaceholderClass>(Import.XObject))
 		{
-			Import.XObject = nullptr;
-			UClass* RealClassObj = CastChecked<UClass>(CreateImport(ImportIndex), ECastCheckedType::NullAllowed);
-
-			bool bUsedPlaceholderInterface = false;
-			if (LoadClass != nullptr)
-			{
-				for (FImplementedInterface& Interface : LoadClass->Interfaces)
-				{
-					if (Interface.Class == PlaceholderClass)
-					{
-						bUsedPlaceholderInterface = true;
-						Interface.Class = RealClassObj;
-					}
-				}
-			}
-
 #if TEST_CHECK_DEPENDENCY_LOAD_DEFERRING
-			bool const bCircumventValidationChecks = !FBlueprintSupport::UseDeferredDependencyVerificationChecks();
-			checkSlow(bCircumventValidationChecks || (PlaceholderClass->GetOuter() == LinkerRoot));
-			// make sure that we know what utilized this placeholder class... 
-			// right now we only expect UObjectProperties/UClassProperties/
-			// UInterfaceProperties/FImplementedInterfaces to, but something 
-			// else could have requested the class without logging itself with 
-			// the placeholder... if ReferencingProperties is empty, then there 
-			// is something out there still using a placeholder class
-			checkSlow(bCircumventValidationChecks || bUsedPlaceholderInterface || PlaceholderClass->HasReferences());
-
-			UObject* PlaceholderObj = PlaceholderClass;
-			// 
-			FReferencerInformationList RemainingRefs;
-			//checkSlow(bCircumventValidationChecks || PlaceholderClass->HasReferences() == IsReferenced(PlaceholderObj, RF_NoFlags, /*bCheckSubObjects =*/false, &RemainingRefs));
-			//checkSlow(bCircumventValidationChecks || RemainingRefs.ExternalReferences.Num() == PlaceholderClass->GetRefCount());
+			checkSlow(bCircumventValidationChecks || PlaceholderClass->ImportIndex == ImportIndex);
 #endif // TEST_CHECK_DEPENDENCY_LOAD_DEFERRING
 
-			PlaceholderClass->ReplaceTrackedReferences(RealClassObj);
-			PlaceholderClass->MarkPendingKill(); // @TODO: ensure these are properly GC'd
-		
-#if TEST_CHECK_DEPENDENCY_LOAD_DEFERRING
-			// there should not be any references left to this placeholder class 
-			// (if there is, then we didn't log that referencer with the placeholder)
-			FReferencerInformationList DanglingReferences;
-			checkSlow(bCircumventValidationChecks || !IsReferenced(PlaceholderObj, RF_NoFlags, /*bCheckSubObjects =*/false, &DanglingReferences));
-#endif // TEST_CHECK_DEPENDENCY_LOAD_DEFERRING
+			// NOTE: we don't check that this resolve successfully replaced any
+			//       references (by the return value), because this resolve 
+			//       could have been re-entered and completed by a nested call
+			//       to the same function (for the same placeholder)
+			ResolveDependencyPlaceholder(PlaceholderClass, LoadClass);
 		}
 		else if (UScriptStruct* StructObj = Cast<UScriptStruct>(Import.XObject))
 		{
 			// in case this is a user defined struct, we have to resolve any 
 			// deferred dependencies in the struct 
-			// 
-			// @TODO: could this result in an endless ResolveDeferredDependencies() loop? 
 			if (Import.SourceLinker != nullptr)
 			{
 				Import.SourceLinker->ResolveDeferredDependencies(/*LoadClass =*/nullptr);
@@ -419,6 +409,26 @@ void ULinkerLoad::ResolveDeferredDependencies(UClass* LoadClass)
 		}
 	}
 
+	if (LoadClass != nullptr)
+	{
+		UClass* SuperClass = LoadClass->GetSuperClass();
+		ULinkerLoad* SuperLinker = (SuperClass != nullptr) ? SuperClass->GetLinker() : nullptr;
+		// NOTE: there is no harm in calling this when the super is not 
+		//       "actively resolving deferred dependencies"... this condition  
+		//       just saves on wasted ops, looping over the super's ImportMap
+		if ((SuperLinker != nullptr) && SuperLinker->IsActivelyResolvingDeferredDependency())
+		{
+			// a resolve could have already been started up the stack, and in turn  
+			// started loading a different package that resulted in another (this) 
+			// resolve beginning... in that scenario, the original resolve could be 
+			// for this class's super and we want to make sure that finishes before
+			// we regenerate this class (else the generated script code could end up 
+			// with unwanted placeholder references; ones that would have been
+			// resolved by the super's linker)
+			SuperLinker->ResolveDeferredDependencies(SuperClass);
+		}
+	}
+
 	// @TODO: don't know if we need this, but could be good to have (as class 
 	//        regeneration is about to force load a lot of this), BUT! this 
 	//        doesn't work for map packages (because this would load the level's
@@ -426,7 +436,6 @@ void ULinkerLoad::ResolveDeferredDependencies(UClass* LoadClass)
 	//LoadAllObjects();
 
 #if TEST_CHECK_DEPENDENCY_LOAD_DEFERRING
-	bool const bCircumventValidationChecks = !FBlueprintSupport::UseDeferredDependencyVerificationChecks();
 	for (TObjectIterator<ULinkerPlaceholderClass> PlaceholderIt; PlaceholderIt; ++PlaceholderIt)
 	{
 		ULinkerPlaceholderClass* PlaceholderClass = *PlaceholderIt;
@@ -441,6 +450,91 @@ void ULinkerLoad::ResolveDeferredDependencies(UClass* LoadClass)
 #endif // USE_CIRCULAR_DEPENDENCY_LOAD_DEFERRING
 }
 
+bool ULinkerLoad::IsActivelyResolvingDeferredDependency() const
+{
+#if USE_CIRCULAR_DEPENDENCY_LOAD_DEFERRING
+	return (ResolvingDeferredPlaceholder != nullptr);
+#else  // USE_CIRCULAR_DEPENDENCY_LOAD_DEFERRING
+	return false;
+#endif // USE_CIRCULAR_DEPENDENCY_LOAD_DEFERRING
+}
+
+int32 ULinkerLoad::ResolveDependencyPlaceholder(UClass* PlaceholderIn, UClass* ReferencingClass)
+{
+#if USE_CIRCULAR_DEPENDENCY_LOAD_DEFERRING
+	TGuardValue<uint32>  LoadFlagsGuard(LoadFlags, (LoadFlags & ~LOAD_DeferDependencyLoads));
+ 	TGuardValue<UClass*> ResolvingClassGuard(ResolvingDeferredPlaceholder, PlaceholderIn);
+
+#if TEST_CHECK_DEPENDENCY_LOAD_DEFERRING
+ 	bool const bCircumventValidationChecks = !FBlueprintSupport::UseDeferredDependencyVerificationChecks();
+ 	checkSlow(bCircumventValidationChecks || Cast<ULinkerPlaceholderClass>(PlaceholderIn) != nullptr);
+ 	checkSlow(bCircumventValidationChecks || PlaceholderIn->GetOuter() == LinkerRoot);
+#endif // TEST_CHECK_DEPENDENCY_LOAD_DEFERRING
+
+	ULinkerPlaceholderClass* PlaceholderClass = (ULinkerPlaceholderClass*)PlaceholderIn;
+	
+	int32 const ImportIndex = PlaceholderClass->ImportIndex;
+	FObjectImport& Import = ImportMap[ImportIndex];
+
+#if TEST_CHECK_DEPENDENCY_LOAD_DEFERRING
+	checkSlow( bCircumventValidationChecks || (Import.XObject == PlaceholderClass) || (Import.XObject == nullptr) );
+#endif // TEST_CHECK_DEPENDENCY_LOAD_DEFERRING
+	
+	// clear the placeholder from the import, so that a call to CreateImport()
+	// properly fills it in
+	Import.XObject = nullptr;
+	// NOTE: this is a possible point of recursion... CreateImport() could 
+	//       continue to load a package already started up the stack and you 
+	//       could end up in another ResolveDependencyPlaceholder() for some  
+	//       other placeholder before this one has completely finished resolving
+	UClass* RealClassObj = CastChecked<UClass>(CreateImport(ImportIndex), ECastCheckedType::NullAllowed);
+
+	int32 ReplacementCount = 0;
+	if (ReferencingClass != nullptr)
+	{
+		for (FImplementedInterface& Interface : ReferencingClass->Interfaces)
+		{
+			if (Interface.Class == PlaceholderClass)
+			{
+				++ReplacementCount;
+				Interface.Class = RealClassObj;
+			}
+		}
+	}
+
+#if TEST_CHECK_DEPENDENCY_LOAD_DEFERRING
+	// make sure that we know what utilized this placeholder class... right now
+	// we only expect UObjectProperties/UClassProperties/UInterfaceProperties/
+	// FImplementedInterfaces to, but something else could have requested the 
+	// class without logging itself with the placeholder... if the placeholder
+	// doesn't have any known references (and it hasn't already been resolved in
+	// some recursive call), then there is something out there still using this
+	// placeholder class
+	checkSlow(bCircumventValidationChecks || (ReplacementCount > 0) || PlaceholderClass->HasReferences() || PlaceholderClass->HasBeenResolved());
+
+	UObject* PlaceholderObj = PlaceholderClass;
+	// 
+	FReferencerInformationList RemainingRefs;
+	//checkSlow(bCircumventValidationChecks || PlaceholderClass->HasReferences() == IsReferenced(PlaceholderObj, RF_NoFlags, /*bCheckSubObjects =*/false, &RemainingRefs));
+	//checkSlow(bCircumventValidationChecks || RemainingRefs.ExternalReferences.Num() == PlaceholderClass->GetRefCount());
+#endif // TEST_CHECK_DEPENDENCY_LOAD_DEFERRING
+
+	ReplacementCount += PlaceholderClass->ReplaceTrackedReferences(RealClassObj);
+	PlaceholderClass->MarkPendingKill(); // @TODO: ensure these are properly GC'd
+
+#if TEST_CHECK_DEPENDENCY_LOAD_DEFERRING
+	// there should not be any references left to this placeholder class 
+	// (if there is, then we didn't log that referencer with the placeholder)
+	FReferencerInformationList DanglingReferences;
+	checkSlow(bCircumventValidationChecks || !IsReferenced(PlaceholderObj, RF_NoFlags, /*bCheckSubObjects =*/false, &DanglingReferences));
+#endif // TEST_CHECK_DEPENDENCY_LOAD_DEFERRING
+
+	return ReplacementCount;
+#else  // USE_CIRCULAR_DEPENDENCY_LOAD_DEFERRING
+	return 0;
+#endif // USE_CIRCULAR_DEPENDENCY_LOAD_DEFERRING
+}
+
 void ULinkerLoad::FinalizeBlueprint(UClass* LoadClass)
 {
 #if USE_CIRCULAR_DEPENDENCY_LOAD_DEFERRING
@@ -454,33 +548,24 @@ void ULinkerLoad::FinalizeBlueprint(UClass* LoadClass)
 	//--------------------------------------
 	TGuardValue<uint32> LoadFlagsGuard(LoadFlags, LoadFlags & ~LOAD_DeferDependencyLoads);
 
-	// since this function is called after class serialization, then we can be
-	// certain that the CDO export object exists (since class serialization 
-	// would call CreateExport() for the CDO)
-
-#if TEST_CHECK_DEPENDENCY_LOAD_DEFERRING
-	bool const bCircumventValidationChecks = !FBlueprintSupport::UseDeferredDependencyVerificationChecks();
-	checkSlow(bCircumventValidationChecks || (LoadClass->ClassDefaultObject != nullptr));
-#endif // TEST_CHECK_DEPENDENCY_LOAD_DEFERRING
-
-	// this should assert if the index hasn't been set (as we expect)
-	int32 CDOExportIndex = LoadClass->ClassDefaultObject->GetLinkerIndex();
-	if (CDOExportIndex == INDEX_NONE)
-	{
-		CDOExportIndex = DeferredExportIndex;
-	}
-	else if (DeferredExportIndex != INDEX_NONE)
+	// we can get in a state where a sub-class is getting finalized here, before
+	// its super-class has been "finalized" (like when the super's 
+	// ResolveDeferredDependencies() ends up Preloading a sub-class), so we want
+	// to make sure that its properly finalized before this sub-class is (so we
+	// can have a properly formed sub-class)
+	if (UClass* SuperClass = LoadClass->GetSuperClass())
 	{
-		check(CDOExportIndex == DeferredExportIndex);
+		ULinkerLoad* SuperLinker = SuperClass->GetLinker();
+		if ((SuperLinker != nullptr) && SuperLinker->IsBlueprintFinalizationPending())
+		{
+			SuperLinker->FinalizeBlueprint(SuperClass);
+		}
 	}
-	// else, the CDO hasn't been touched (which is expected, since the CDO needs the class first)
-	// @TODO: this comment here ^ is in direct conflict with the assumption 
-	//        stated above (that the CDO should be loaded)... figure out which is correct!
-	DeferredExportIndex = INDEX_NONE;
 
+	if (IsBlueprintFinalizationPending())
 	{
-		check(CDOExportIndex != INDEX_NONE);
-		FObjectExport& CDOExport = ExportMap[CDOExportIndex];
+		check(DeferredExportIndex != INDEX_NONE);
+		FObjectExport& CDOExport = ExportMap[DeferredExportIndex];
 
 		UObject* CDO = CDOExport.Object;
 		if (FBlueprintSupport::UseDeferredCDOSerialization())
@@ -490,7 +575,7 @@ void ULinkerLoad::FinalizeBlueprint(UClass* LoadClass)
 			{
 				const EObjectFlags OldFlags = CDO->GetFlags();
 				CDO->ClearFlags(RF_NeedLoad | RF_NeedPostLoad);
-				CDO->SetLinker(this, CDOExportIndex, /*bShouldDetatchExisting =*/false);
+				CDO->SetLinker(this, DeferredExportIndex, /*bShouldDetatchExisting =*/false);
 				CDO->SetFlags(OldFlags);
 			}
 			// should load the CDO (ensuring that it has been serialized in by the 
@@ -500,13 +585,16 @@ void ULinkerLoad::FinalizeBlueprint(UClass* LoadClass)
 			//       themselves, as the CDO could depend on a class not listed in 
 			//       the package's imports
 			//
-			// @TODO: should we block/detect reentrant behavior? would this CDO ever
-			//        need another CDO?
+			// NOTE: we don't guard against reentrant behavior... if the CDO has
+			//       already been "finalized", then its RF_NeedLoad flag would 
+			//       have been cleared
 			Preload(CDO);
 		}
 
 		UClass* BlueprintClass = Cast<UClass>(IndexToObject(CDOExport.ClassIndex));
 #if TEST_CHECK_DEPENDENCY_LOAD_DEFERRING
+		bool const bCircumventValidationChecks = !FBlueprintSupport::UseDeferredDependencyVerificationChecks();
+		checkSlow(bCircumventValidationChecks || BlueprintClass == LoadClass);
 		checkSlow(bCircumventValidationChecks || CDO->HasAnyFlags(RF_LoadCompleted));
 		checkSlow(bCircumventValidationChecks || BlueprintClass->HasAnyClassFlags(CLASS_CompiledFromBlueprint));
 
@@ -518,6 +606,9 @@ void ULinkerLoad::FinalizeBlueprint(UClass* LoadClass)
 		checkSlow(bCircumventValidationChecks || ClassInstances.Num() == 0);
 #endif // TEST_CHECK_DEPENDENCY_LOAD_DEFERRING
 
+		// clear this so IsBlueprintFinalizationPending() doesn't report true
+		DeferredExportIndex = INDEX_NONE;
+
 		UObject* OldCDO = BlueprintClass->ClassDefaultObject;
 		if (RegenerateBlueprintClass(BlueprintClass, CDO))
 		{
@@ -534,6 +625,15 @@ void ULinkerLoad::FinalizeBlueprint(UClass* LoadClass)
 #endif // USE_CIRCULAR_DEPENDENCY_LOAD_DEFERRING
 }
 
+bool ULinkerLoad::IsBlueprintFinalizationPending() const
+{
+#if USE_CIRCULAR_DEPENDENCY_LOAD_DEFERRING
+	return (DeferredExportIndex != INDEX_NONE);
+#else  // USE_CIRCULAR_DEPENDENCY_LOAD_DEFERRING
+	return false;
+#endif // USE_CIRCULAR_DEPENDENCY_LOAD_DEFERRING
+}
+
 /*******************************************************************************
  * UObject
  ******************************************************************************/
diff --git a/Engine/Source/Runtime/CoreUObject/Private/UObject/LinkerLoad.cpp b/Engine/Source/Runtime/CoreUObject/Private/UObject/LinkerLoad.cpp
--- a/Engine/Source/Runtime/CoreUObject/Private/UObject/LinkerLoad.cpp
+++ b/Engine/Source/Runtime/CoreUObject/Private/UObject/LinkerLoad.cpp
@@ -630,9 +630,10 @@ ULinkerLoad::ULinkerLoad( const FObjectInitializer& ObjectInitializer, UPackage*
 #endif // WITH_EDITOR
 #if USE_CIRCULAR_DEPENDENCY_LOAD_DEFERRING
 ,	DeferredExportIndex(INDEX_NONE)
+,	ResolvingDeferredPlaceholder(nullptr)
 #endif // USE_CIRCULAR_DEPENDENCY_LOAD_DEFERRING
 {
-	check(!HasAnyFlags(RF_ClassDefaultObject));
+	check(!HasAnyFlags(RF_ClassDefaultObject)); 
 
 	INC_DWORD_STAT(STAT_LinkerCount);
 	INC_DWORD_STAT(STAT_LiveLinkerCount);
@@ -2879,6 +2880,17 @@ void ULinkerLoad::Preload( UObject* Object )
 					SCOPE_CYCLE_COUNTER(STAT_LinkerLoadDeferred);
 					if ((LoadFlags & LOAD_DeferDependencyLoads) != (*LoadFlagsGuard & LOAD_DeferDependencyLoads))
 					{
+#if TEST_CHECK_DEPENDENCY_LOAD_DEFERRING
+						bool const bCircumventValidationChecks = !FBlueprintSupport::UseDeferredDependencyVerificationChecks();
+						// since class serialization reads in the class's CDO, 
+						// then we can be certain that the CDO export object 
+						// exists (and DeferredExportIndex should reference it)
+						// ... FinalizeBlueprint() depends on this (and since 
+						// ResolveDeferredDependencies() can be recursive, we 
+						// check it out here, before it is called)
+						checkSlow(bCircumventValidationChecks || (DeferredExportIndex != INDEX_NONE));
+#endif // TEST_CHECK_DEPENDENCY_LOAD_DEFERRING
+
 						ResolveDeferredDependencies(ObjectAsClass);
 						FinalizeBlueprint(ObjectAsClass);
 					}
diff --git a/Engine/Source/Runtime/CoreUObject/Private/UObject/LinkerPlaceholderClass.cpp b/Engine/Source/Runtime/CoreUObject/Private/UObject/LinkerPlaceholderClass.cpp
--- a/Engine/Source/Runtime/CoreUObject/Private/UObject/LinkerPlaceholderClass.cpp
+++ b/Engine/Source/Runtime/CoreUObject/Private/UObject/LinkerPlaceholderClass.cpp
@@ -110,6 +110,12 @@ int32 ULinkerPlaceholderClass::GetRefCount() const
 	return ReferencingProperties.Num();
 }
 
+//------------------------------------------------------------------------------
+bool ULinkerPlaceholderClass::HasBeenResolved() const
+{
+	return !HasReferences() && bResolvedReferences;
+}
+
 //------------------------------------------------------------------------------
 void ULinkerPlaceholderClass::RemoveTrackedReference(UProperty* ReferencingProperty)
 {
@@ -120,8 +126,10 @@ void ULinkerPlaceholderClass::RemoveTrackedReference(UProperty* ReferencingPrope
 }
 
 //------------------------------------------------------------------------------
-void ULinkerPlaceholderClass::ReplaceTrackedReferences(UClass* ReplacementClass)
+int32 ULinkerPlaceholderClass::ReplaceTrackedReferences(UClass* ReplacementClass)
 {
+	int32 ReplacementCount = 0;
+
 	for (UProperty* Property : ReferencingProperties)
 	{
 		if (UObjectPropertyBase* ObjProperty = Cast<UObjectPropertyBase>(Property))
@@ -129,12 +137,14 @@ void ULinkerPlaceholderClass::ReplaceTrackedReferences(UClass* ReplacementClass)
 			if (ObjProperty->PropertyClass == this)
 			{
 				ObjProperty->PropertyClass = ReplacementClass;
+				++ReplacementCount;
 			}
 
 			UClassProperty* ClassProperty = Cast<UClassProperty>(ObjProperty);
 			if ((ClassProperty != nullptr) && (ClassProperty->MetaClass == this))
 			{
 				ClassProperty->MetaClass = ReplacementClass;
+				++ReplacementCount;
 			}
 		}	
 		else if (UInterfaceProperty* InterfaceProp = Cast<UInterfaceProperty>(Property))
@@ -142,12 +152,12 @@ void ULinkerPlaceholderClass::ReplaceTrackedReferences(UClass* ReplacementClass)
 			if (InterfaceProp->InterfaceClass == this)
 			{
 				InterfaceProp->InterfaceClass = ReplacementClass;
+				++ReplacementCount;
 			}
 		}
 	}
 	ReferencingProperties.Empty();
-
-#if TEST_CHECK_DEPENDENCY_LOAD_DEFERRING
 	bResolvedReferences = true;
-#endif // TEST_CHECK_DEPENDENCY_LOAD_DEFERRING
+
+	return ReplacementCount;
 }
diff --git a/Engine/Source/Runtime/CoreUObject/Private/UObject/LinkerPlaceholderClass.h b/Engine/Source/Runtime/CoreUObject/Private/UObject/LinkerPlaceholderClass.h
--- a/Engine/Source/Runtime/CoreUObject/Private/UObject/LinkerPlaceholderClass.h
+++ b/Engine/Source/Runtime/CoreUObject/Private/UObject/LinkerPlaceholderClass.h
@@ -58,6 +58,15 @@ public:
 	 */
 	int32 GetRefCount() const;
 
+	/**
+	 * Checks to see if 1) this placeholder has had RemoveTrackedReference() 
+	 * called on it, and 2) it doesn't have any more references that have since 
+	 * been added.
+	 * 
+	 * @return True if ReplaceTrackedReferences() has been ran, and no KNOWN references have been added.
+	 */
+	bool HasBeenResolved() const;
+
 	/**
 	 * Removes the specified property from this class's internal tracking list 
 	 * (which aims to keep track of properties utilizing this class).
@@ -71,15 +80,18 @@ public:
 	 * references to this class with a new (hopefully proper) class.
 	 * 
 	 * @param  ReplacementClass    The class that you want all references to this class replaced with.
+	 * @return The number of references that were successfully replaced.
 	 */
-	void ReplaceTrackedReferences(UClass* ReplacementClass);
+	int32 ReplaceTrackedReferences(UClass* ReplacementClass);
+
+public:
+	/** Set by the ULinkerLoad that created this instance, tracks what import this was used in place of. */
+	int32 ImportIndex;
 
 private:
 	/** Links to UProperties that are currently using this class */
 	TSet<UProperty*> ReferencingProperties;
 
-#if TEST_CHECK_DEPENDENCY_LOAD_DEFERRING
 	/** Used to catch references that are added after we've already resolved all references */
 	bool bResolvedReferences;
-#endif // TEST_CHECK_DEPENDENCY_LOAD_DEFERRING
 }; 
diff --git a/Engine/Source/Runtime/CoreUObject/Public/UObject/Linker.h b/Engine/Source/Runtime/CoreUObject/Public/UObject/Linker.h
--- a/Engine/Source/Runtime/CoreUObject/Public/UObject/Linker.h
+++ b/Engine/Source/Runtime/CoreUObject/Public/UObject/Linker.h
@@ -1843,6 +1843,27 @@ private:
 	 */
 	void ResolveDeferredDependencies(UClass* LoadClass);
 
+	/**
+	 * Loads the import that the Placeholder was initially stubbed in for (NOTE:
+	 * this could cause recursive behavior), and then replaces all known 
+	 * placeholder references with the proper class.
+	 * 
+	 * @param  Placeholder		A ULinkerPlaceholderClass that was substituted in place of a deferred dependency.
+	 * @param  ReferencingClass	The (Blueprint) class that was loading, while we deferred dependencies (now referencing the placeholder).
+	 * @return The number of placeholder references replaced (could be none, if this was recursively resolved).
+	 */
+	int32 ResolveDependencyPlaceholder(UClass* Placeholder, UClass* ReferencingClass = nullptr);
+
+	/**
+	 * Query method to help catch recursive behavior. When this returns true, a 
+	 * dependency placeholder is in the middle of being resolved by 
+	 * ResolveDependencyPlaceholder(). Used so a nested call would know to 
+	 * complete that placeholder before continuing.
+	 * 
+	 * @return True if ResolveDependencyPlaceholder() is being ran on a placeholder that has yet to be resolved. 
+	 */
+	bool IsActivelyResolvingDeferredDependency() const;
+
 	/**
 	 * Takes the supplied serialized class and serializes in its CDO, then 
 	 * regenerates both.
@@ -1851,14 +1872,33 @@ private:
 	 */
 	void FinalizeBlueprint(UClass* LoadClass);
 
+	/**
+	 * Query method to help handle recursive behavior. When this returns true, 
+	 * this linker is in the middle of, or is about to call FinalizeBlueprint()
+	 * (for a blueprint class somewhere in the current callstack). Needed when 
+	 * we get to finalizing a sub-class before we've finished finalizing its 
+	 * super (so we know we need to finish finalizing the super first).
+	 * 
+	 * @return True if FinalizeBlueprint() is currently being ran (or about to be ran) for an export (Blueprint) class.
+	 */
+	bool IsBlueprintFinalizationPending() const;
+
 #if	USE_CIRCULAR_DEPENDENCY_LOAD_DEFERRING
 	/** 
 	 * For deferring dependency loads, we block CDO serialization until the 
-	 * class if complete. If we attempt to serialize the CDO whil that is 
+	 * class if complete. If we attempt to serialize the CDO while that is 
 	 * happening, we instead defer it and record the export's index here (so we 
 	 * can return to it later).
 	 */
 	int32 DeferredExportIndex;
+
+	/** 
+	 * Used to track dependency placeholders currently being resolved inside of 
+	 * ResolveDependencyPlaceholder()... utilized for nested reentrant behavior, 
+	 * to make sure this placeholder is completely resolved before continuing on 
+	 * to the next.
+	 */
+	UClass* ResolvingDeferredPlaceholder;
 #endif // USE_CIRCULAR_DEPENDENCY_LOAD_DEFERRING
 
 
