#UE4 Quick fixing a crash when attempting to find a dynamic delegate. This prevents immediate recursion, but a loop is still possible.

diff --git a/Engine/Source/Runtime/CoreUObject/Private/UObject/UObjectGlobals.cpp b/Engine/Source/Runtime/CoreUObject/Private/UObject/UObjectGlobals.cpp
--- a/Engine/Source/Runtime/CoreUObject/Private/UObject/UObjectGlobals.cpp
+++ b/Engine/Source/Runtime/CoreUObject/Private/UObject/UObjectGlobals.cpp
@@ -181,7 +181,13 @@ UObject* StaticFindObject( UClass* ObjectClass, UObject* InObjectPackage, const
 			)
 		)
 	{
-		MatchingObject = StaticFindObject(ObjectClass, ObjectPackage->GetOutermost(), *ObjectName.ToString(), ExactClass);
+		UPackage* SearchPackage = ObjectPackage->GetOutermost();
+		const FString SearchName = ObjectName.ToString();
+		const bool bIdenticalSearch = (SearchPackage == InObjectPackage) && (SearchName == FString(OrigInName));
+		if (!bIdenticalSearch)
+		{
+			MatchingObject = StaticFindObject(ObjectClass, SearchPackage, *SearchName, ExactClass);
+		}
 	}
 
 	return MatchingObject;
