[UE-7296][UE-7784][UE-7911] Resolving all dependency packages that were in the middle of a resolve (before  generating script code on load)... WIP: to solve circular dependency load issues.

diff --git a/Engine/Source/Runtime/CoreUObject/Private/Blueprint/BlueprintSupport.cpp b/Engine/Source/Runtime/CoreUObject/Private/Blueprint/BlueprintSupport.cpp
--- a/Engine/Source/Runtime/CoreUObject/Private/Blueprint/BlueprintSupport.cpp
+++ b/Engine/Source/Runtime/CoreUObject/Private/Blueprint/BlueprintSupport.cpp
@@ -353,7 +353,92 @@ bool ULinkerLoad::DeferPotentialCircularImport(const int32 Index)
 #endif // USE_CIRCULAR_DEPENDENCY_LOAD_DEFERRING
 }
 
-void ULinkerLoad::ResolveDeferredDependencies(UClass* LoadClass)
+/**
+ * This utility struct helps track blueprint structs/linkers that are currently 
+ * in the  middle of a call to ResolveDeferredDependencies(). This can be used  
+ * to know if a dependency's resolve needs to be finished (to avoid unwanted 
+ * placeholder references ending up in script-code).
+ */
+struct FUnresolvedStructTracker
+{
+public:
+	/** Marks the specified struct (and its linker) as "resolving" for the lifetime of this instance */
+	FUnresolvedStructTracker(UStruct* LoadStruct)
+		: TrackedStruct(LoadStruct)
+	{
+#if TEST_CHECK_DEPENDENCY_LOAD_DEFERRING
+		bool const bCircumventValidationChecks = !FBlueprintSupport::UseDeferredDependencyVerificationChecks();
+		checkSlow(bCircumventValidationChecks || LoadStruct != nullptr);
+		checkSlow(bCircumventValidationChecks || LoadStruct->GetLinker() != nullptr);
+#endif // TEST_CHECK_DEPENDENCY_LOAD_DEFERRING
+
+		UnresolvedStructs.Add(LoadStruct);
+	}
+
+	/** Removes the wrapped struct from the "resolving" set (it has been fully "resolved") */
+	~FUnresolvedStructTracker()
+	{
+		// even if another FUnresolvedStructTracker added this struct earlier,  
+		// we want the most nested one removing it from the set (because this 
+		// means the struct is fully resolved, even if we're still in the middle  
+		// of a ResolveDeferredDependencies() call further up the stack)
+		UnresolvedStructs.Remove(TrackedStruct);
+	}
+
+	/**
+	 * Checks to see if the specified import object is a blueprint class/struct 
+	 * that is currently in the midst of resolving (and hasn't completed that  
+	 * resolve elsewhere in some nested call).
+	 * 
+	 * @param  ImportObject    The object you wish to check.
+	 * @return True if the specified object is a class/struct that hasn't been fully resolved yet (otherwise false).
+	 */
+	static bool IsImportStructUnresolved(UObject* ImportObject)
+	{
+		return UnresolvedStructs.Contains(ImportObject);
+	}
+
+	/**
+	 * Checks to see if the specified linker is associated with any of the 
+	 * unresolved structs that this is currently tracking.
+	 *
+	 * NOTE: This could return false, even if the linker is in a 
+	 *       ResolveDeferredDependencies() call futher up the callstack... in 
+	 *       that scenario, the associated struct was fully resolved by a 
+	 *       subsequent call to the same function (for the same linker/struct).
+	 * 
+	 * @param  Linker	The linker you want to check.
+	 * @return True if the specified linker is in the midst of an unfinished ResolveDeferredDependencies() call (otherwise false).
+	 */
+	static bool IsAssociatedStructUnresolved(const ULinkerLoad* Linker)
+	{
+		for (UObject* UnresolvedObj : UnresolvedStructs)
+		{
+			// each unresolved struct should have a linker set on it, because 
+			// they would have had to go through Preload()
+			if (UnresolvedObj->GetLinker() == Linker)
+			{
+				return true;
+			}
+		}
+		return false;
+	}
+
+private:
+	/** The struct that is currently being "resolved" */
+	UStruct* TrackedStruct;
+
+	/** 
+	 * A set of blueprint structs & classes that are currently being "resolved"  
+	 * by ResolveDeferredDependencies() (using UObject* instead of UStruct, so
+	 * we don't have to cast import objects before checking for their presence).
+	 */
+	static TSet<UObject*> UnresolvedStructs;
+};
+/** A global set that tracks structs currently being ran through (and unfinished by) ULinkerLoad::ResolveDeferredDependencies() */
+TSet<UObject*> FUnresolvedStructTracker::UnresolvedStructs;
+
+void ULinkerLoad::ResolveDeferredDependencies(UStruct* LoadStruct)
 {
 #if USE_CIRCULAR_DEPENDENCY_LOAD_DEFERRING
 	//--------------------------------------
@@ -367,66 +452,78 @@ void ULinkerLoad::ResolveDeferredDependencies(UClass* LoadClass)
 	TGuardValue<int32> RecursiveDepthGuard(RecursiveDepth, RecursiveDepth + 1);
 
 	bool const bCircumventValidationChecks = !FBlueprintSupport::UseDeferredDependencyVerificationChecks();
-	checkSlow(bCircumventValidationChecks || !LoadClass || (LoadClass->GetLinker() == this));
+	checkSlow( bCircumventValidationChecks || (LoadStruct && (LoadStruct->GetLinker() == this)) );
 #endif // TEST_CHECK_DEPENDENCY_LOAD_DEFERRING
 
-	// this function (for this linker) could be reentrant (see where we 
-	// recursively call ResolveDeferredDependencies() for super-classes below);  
-	// if that's the case, then we want to finish resolving the pending class 
-	// before we continue on
-	if (ResolvingDeferredPlaceholder != nullptr)
-	{
-		int32 ResolvedRefCount = ResolveDependencyPlaceholder(ResolvingDeferredPlaceholder, LoadClass);
-		// @TODO: can we reliably count on this resolving some dependencies?... 
-		//        if so, check verify that!
-		ResolvingDeferredPlaceholder = nullptr;
-	}
-	
-	for (int32 ImportIndex = 0; ImportIndex < ImportMap.Num(); ++ImportIndex)
+	// scoped block to manage the lifetime of ScopedResolveTracker, so that 
+	// this resolve is only tracked for the duration of resolving all its 
+	// placeholder classes, all member struct's placholders, and its parent's
 	{
-		FObjectImport& Import = ImportMap[ImportIndex];
+		FUnresolvedStructTracker ScopedResolveTracker(LoadStruct);
 
-		if (ULinkerPlaceholderClass* PlaceholderClass = Cast<ULinkerPlaceholderClass>(Import.XObject))
+		// this function (for this linker) could be reentrant (see where we 
+		// recursively call ResolveDeferredDependencies() for super-classes below);  
+		// if that's the case, then we want to finish resolving the pending class 
+		// before we continue on
+		if (ResolvingDeferredPlaceholder != nullptr)
 		{
+			int32 ResolvedRefCount = ResolveDependencyPlaceholder(ResolvingDeferredPlaceholder, Cast<UClass>(LoadStruct));
+			// @TODO: can we reliably count on this resolving some dependencies?... 
+			//        if so, check verify that!
+			ResolvingDeferredPlaceholder = nullptr;
+		}
+
+		// because this loop could recurse (and end up finishing all of this for
+		// us), we check HasUnresolvedDependencies() so we can early out  
+		// from this loop in that situation (the loop has been finished elsewhere)
+		for (int32 ImportIndex = 0; ImportIndex < ImportMap.Num() && HasUnresolvedDependencies(); ++ImportIndex)
+		{
+			FObjectImport& Import = ImportMap[ImportIndex];
+
+			if (ULinkerPlaceholderClass* PlaceholderClass = Cast<ULinkerPlaceholderClass>(Import.XObject))
+			{
 #if TEST_CHECK_DEPENDENCY_LOAD_DEFERRING
-			checkSlow(bCircumventValidationChecks || PlaceholderClass->ImportIndex == ImportIndex);
+				checkSlow(bCircumventValidationChecks || PlaceholderClass->ImportIndex == ImportIndex);
 #endif // TEST_CHECK_DEPENDENCY_LOAD_DEFERRING
 
-			// NOTE: we don't check that this resolve successfully replaced any
-			//       references (by the return value), because this resolve 
-			//       could have been re-entered and completed by a nested call
-			//       to the same function (for the same placeholder)
-			ResolveDependencyPlaceholder(PlaceholderClass, LoadClass);
-		}
-		else if (UScriptStruct* StructObj = Cast<UScriptStruct>(Import.XObject))
-		{
-			// in case this is a user defined struct, we have to resolve any 
-			// deferred dependencies in the struct 
-			if (Import.SourceLinker != nullptr)
+				// NOTE: we don't check that this resolve successfully replaced any
+				//       references (by the return value), because this resolve 
+				//       could have been re-entered and completed by a nested call
+				//       to the same function (for the same placeholder)
+				ResolveDependencyPlaceholder(PlaceholderClass, Cast<UClass>(LoadStruct));
+			}
+			else if (UScriptStruct* StructObj = Cast<UScriptStruct>(Import.XObject))
 			{
-				Import.SourceLinker->ResolveDeferredDependencies(/*LoadClass =*/nullptr);
+				// in case this is a user defined struct, we have to resolve any 
+				// deferred dependencies in the struct 
+				if (Import.SourceLinker != nullptr)
+				{
+					Import.SourceLinker->ResolveDeferredDependencies(StructObj);
+				}
 			}
 		}
-	}
 
-	if (LoadClass != nullptr)
-	{
-		UClass* SuperClass = LoadClass->GetSuperClass();
-		ULinkerLoad* SuperLinker = (SuperClass != nullptr) ? SuperClass->GetLinker() : nullptr;
-		// NOTE: there is no harm in calling this when the super is not 
-		//       "actively resolving deferred dependencies"... this condition  
-		//       just saves on wasted ops, looping over the super's ImportMap
-		if ((SuperLinker != nullptr) && SuperLinker->IsActivelyResolvingDeferredDependency())
+		if (UStruct* SuperStruct = LoadStruct->GetSuperStruct())
 		{
-			// a resolve could have already been started up the stack, and in turn  
-			// started loading a different package that resulted in another (this) 
-			// resolve beginning... in that scenario, the original resolve could be 
-			// for this class's super and we want to make sure that finishes before
-			// we regenerate this class (else the generated script code could end up 
-			// with unwanted placeholder references; ones that would have been
-			// resolved by the super's linker)
-			SuperLinker->ResolveDeferredDependencies(SuperClass);
+			ULinkerLoad* SuperLinker = SuperStruct->GetLinker();
+			// NOTE: there is no harm in calling this when the super is not 
+			//       "actively resolving deferred dependencies"... this condition  
+			//       just saves on wasted ops, looping over the super's ImportMap
+			if ((SuperLinker != nullptr) && SuperLinker->HasUnresolvedDependencies())
+			{
+				// a resolve could have already been started up the stack, and in turn  
+				// started loading a different package that resulted in another (this) 
+				// resolve beginning... in that scenario, the original resolve could be 
+				// for this class's super and we want to make sure that finishes before
+				// we regenerate this class (else the generated script code could end up 
+				// with unwanted placeholder references; ones that would have been
+				// resolved by the super's linker)
+				SuperLinker->ResolveDeferredDependencies(SuperStruct);
+			}
 		}
+
+	// close the scope on ScopedResolveTracker (so LoadClass doesn't appear to 
+	// be resolving through the rest of this function)
 	}
 
 	// @TODO: don't know if we need this, but could be good to have (as class 
@@ -450,10 +547,23 @@ void ULinkerLoad::ResolveDeferredDependencies(UClass* LoadClass)
 #endif // USE_CIRCULAR_DEPENDENCY_LOAD_DEFERRING
 }
 
-bool ULinkerLoad::IsActivelyResolvingDeferredDependency() const
+bool ULinkerLoad::HasUnresolvedDependencies() const
 {
 #if USE_CIRCULAR_DEPENDENCY_LOAD_DEFERRING
-	return (ResolvingDeferredPlaceholder != nullptr);
+	// checking (ResolvingDeferredPlaceholder != nullptr) is not sufficient, 
+	// because the linker could be in the midst of a nested resolve (for a 
+	// struct, or super... see ULinkerLoad::ResolveDeferredDependencies)
+	bool bIsClassExportUnresolved = FUnresolvedStructTracker::IsAssociatedStructUnresolved(this);
+
+#if TEST_CHECK_DEPENDENCY_LOAD_DEFERRING
+	bool const bCircumventValidationChecks = !FBlueprintSupport::UseDeferredDependencyVerificationChecks();
+	// (ResolvingDeferredPlaceholder != nullptr) should imply 
+	// bIsClassExportUnresolved is true (but not the other way around)
+	checkSlow(bCircumventValidationChecks || (ResolvingDeferredPlaceholder == nullptr) || bIsClassExportUnresolved);
+#endif // TEST_CHECK_DEPENDENCY_LOAD_DEFERRING
+	
+	return bIsClassExportUnresolved;
+
 #else  // USE_CIRCULAR_DEPENDENCY_LOAD_DEFERRING
 	return false;
 #endif // USE_CIRCULAR_DEPENDENCY_LOAD_DEFERRING
@@ -550,9 +660,9 @@ void ULinkerLoad::FinalizeBlueprint(UClass* LoadClass)
 
 	// we can get in a state where a sub-class is getting finalized here, before
 	// its super-class has been "finalized" (like when the super's 
-	// ResolveDeferredDependencies() ends up Preloading a sub-class), so we want
-	// to make sure that its properly finalized before this sub-class is (so we
-	// can have a properly formed sub-class)
+	// ResolveDeferredDependencies() ends up Preloading a sub-class), so we
+	// want to make sure that its properly finalized before this sub-class is
+	// (so we can have a properly formed sub-class)
 	if (UClass* SuperClass = LoadClass->GetSuperClass())
 	{
 		ULinkerLoad* SuperLinker = SuperClass->GetLinker();
@@ -562,6 +672,64 @@ void ULinkerLoad::FinalizeBlueprint(UClass* LoadClass)
 		}
 	}
 
+	// at this point, we're sure that LoadClass doesn't contain any class 
+	// placeholders (because ResolveDeferredDependencies() was ran on it); 
+	// however, once we get to regenerating/compiling the blueprint, the graph 
+	// (nodes, pins, etc.) could bring in new dependencies that weren't part of 
+	// the class... this would normally be all fine and well, but in complicated 
+	// dependency chains those graph-dependencies could already be in the middle 
+	// of resolving themselves somewhere up the stack... if we just continue  
+	// along and let the blueprint compile, then it could end up with  
+	// placeholder refs in its script code (which it bad); we need to make sure
+	// that all dependencies don't have any placeholder classes left in them
+	// 
+	// we don't want this to be part of ResolveDeferredDependencies() 
+	// because we don't want this to count as a linker's "class resolution 
+	// phase"; at this point, it is ok if other blueprints compile with refs to  
+	// this LoadClass since it doesn't have any placeholders left in it (we also 
+	// don't want this linker externally claiming that it has resolving left to 
+	// do, otherwise other linkers could want to finish this off when they don't
+	// have to)... we do however need it here in FinalizeBlueprint(), because
+	// we need it ran for any super-classes before we regen
+	for (int32 ImportIndex = 0; ImportIndex < ImportMap.Num() && IsBlueprintFinalizationPending(); ++ImportIndex)
+	{
+		FObjectImport& Import = ImportMap[ImportIndex];
+		if (Import.XObject == nullptr)
+		{
+			// first, make sure every import object is available... just because 
+			// it isn't present in the map already, doesn't mean it isn't in the 
+			// middle of a resolve (the CreateImport() brings in an export 
+			// object from another package, which could be resolving itself)... 
+			// 
+			// don't fret, all these imports were bound to get created sooner or 
+			// later (like when the blueprint was regenerated)
+			//
+			// NOTE: this is a possible root point for recursion... accessing a 
+			//       separate package could continue its loading process which
+			//       in turn, could end us back in this function before we ever  
+			//       returned from this
+			CreateImport(ImportIndex);
+		}
+
+		// see if this import is currently being resolved (presumably somewhere 
+		// up the callstack)... if it is, we need to ensure that this dependency 
+		// is fully resolved before we get to regenerating the blueprint (else,
+		// we could end up with placeholder classes in our script-code)
+		if (FUnresolvedStructTracker::IsImportStructUnresolved(Import.XObject))
+		{
+#if TEST_CHECK_DEPENDENCY_LOAD_DEFERRING
+			bool const bCircumventValidationChecks = !FBlueprintSupport::UseDeferredDependencyVerificationChecks();
+			// because it is tracked by FUnresolvedStructTracker, it must be a class
+			checkSlow(bCircumventValidationChecks || Cast<UClass>(Import.XObject) != nullptr);
+#endif // TEST_CHECK_DEPENDENCY_LOAD_DEFERRING
+			Import.SourceLinker->ResolveDeferredDependencies((UClass*)Import.XObject);
+		}
+	}
+
+	// the above loop could have caused some recursion... if it ended up 
+	// finalizing a sub-class (of LoadClass), then that would have finalized
+	// this as well; so, before we continue, make sure that this didn't already 
+	// get fully executed in some nested call
 	if (IsBlueprintFinalizationPending())
 	{
 		check(DeferredExportIndex != INDEX_NONE);
diff --git a/Engine/Source/Runtime/CoreUObject/Public/UObject/Linker.h b/Engine/Source/Runtime/CoreUObject/Public/UObject/Linker.h
--- a/Engine/Source/Runtime/CoreUObject/Public/UObject/Linker.h
+++ b/Engine/Source/Runtime/CoreUObject/Public/UObject/Linker.h
@@ -1835,13 +1835,12 @@ private:
 	bool DeferPotentialCircularImport(const int32 ImportIndex);
 
 	/**
-	 * Combs the ImportMap for any imports that were deferred, and then loads 
-	 * them. After that, deferred CDOs are serialized in and the Blueprint class
-	 * is regenerated.
+	 * Combs the ImportMap for any imports that were deferred, and then creates 
+	 * them (via CreateImport).
 	 * 
-	 * @param  LoadClass    The (Blueprint) class that was loading, while we deferred dependencies.
+	 * @param  LoadStruct    The (Blueprint) class or struct that you want resolved (so that it no longer contains dependency placeholders).
 	 */
-	void ResolveDeferredDependencies(UClass* LoadClass);
+	void ResolveDeferredDependencies(UStruct* LoadStruct);
 
 	/**
 	 * Loads the import that the Placeholder was initially stubbed in for (NOTE:
@@ -1862,7 +1861,7 @@ private:
 	 * 
 	 * @return True if ResolveDependencyPlaceholder() is being ran on a placeholder that has yet to be resolved. 
 	 */
-	bool IsActivelyResolvingDeferredDependency() const;
+	bool HasUnresolvedDependencies() const;
 
 	/**
 	 * Takes the supplied serialized class and serializes in its CDO, then 
