[UE-7296] Cleaning up circular dependency verification tests and test scenarios/hooks.

diff --git a/Engine/Config/BaseEngine.ini b/Engine/Config/BaseEngine.ini
--- a/Engine/Config/BaseEngine.ini
+++ b/Engine/Config/BaseEngine.ini
@@ -1485,9 +1485,6 @@ bOptimizeExecutionFlowStack=true
 bOptimizeAdjacentStates=true
 bEnableInheritableComponents=true
 bDeferDependencyLoads=false
-bDeferCDOLoadSerialization=true
-bPropagateDeferredLoadsForStructs=true
-bCheckDeferredDependencies=false
 
 [/Script/Engine.Blueprint]
 bRecompileOnLoad=true
diff --git a/Engine/Source/Editor/UnrealEd/Classes/Tests/CircularDependencyTestActor.h b/Engine/Source/Editor/UnrealEd/Classes/Tests/CircularDependencyTestActor.h
--- a/Engine/Source/Editor/UnrealEd/Classes/Tests/CircularDependencyTestActor.h
+++ b/Engine/Source/Editor/UnrealEd/Classes/Tests/CircularDependencyTestActor.h
@@ -31,6 +31,9 @@ public:
 	UFUNCTION(BlueprintCallable, Category="Development|Tests")
 	bool TestVerifyClass(bool bCheckPropertyType = true);
 
+	UFUNCTION(BlueprintCallable, Category = "Development|Tests")
+	bool TestVerifyBlueprint();
+
 	//UFUNCTION(BlueprintCallable, Category="Development|Tests")
 	bool TestVerifyStructMember(UScriptStruct* Struct, uint8* Container, bool bCheckPropertyType = true);
 
diff --git a/Engine/Source/Editor/UnrealEd/Private/Tests/CircularDependencyTestActor.cpp b/Engine/Source/Editor/UnrealEd/Private/Tests/CircularDependencyTestActor.cpp
--- a/Engine/Source/Editor/UnrealEd/Private/Tests/CircularDependencyTestActor.cpp
+++ b/Engine/Source/Editor/UnrealEd/Private/Tests/CircularDependencyTestActor.cpp
@@ -2,7 +2,11 @@
 
 #include "UnrealEd.h"
 #include "Tests/CircularDependencyTestActor.h"
+#include "EdGraphSchema_K2.h"
+#include "EdGraph/EdGraph.h"
 #include "K2Node.h"
+#include "EdGraph/EdGraphPin.h"
+#include "TokenizedMessage.h"
 
 ACircularDependencyTestActor::ACircularDependencyTestActor(const FObjectInitializer& ObjectInitializer)
 	: Super(ObjectInitializer)
@@ -28,6 +32,11 @@ static bool IsPlaceholderClass(UClass* TestClass)
 	return TestClass->GetName().Contains(TEXT("PLACEHOLDER"));
 }
 
+static bool IsValidClass(UClass* ClassToVerify)
+{
+	return (ClassToVerify != nullptr) && !IsPlaceholderClass(ClassToVerify);
+};
+
 bool ACircularDependencyTestActor::TestVerifyClass(bool bCheckPropertyType)
 {
 	UClass* BPClass = GetClass();
@@ -53,6 +62,75 @@ bool ACircularDependencyTestActor::TestVerifyClass(bool bCheckPropertyType)
 	return true;
 }
 
+bool ACircularDependencyTestActor::TestVerifyBlueprint()
+{
+#if WITH_EDITOR
+	UClass* BPClass = GetClass();
+	if (UBlueprint* Blueprint = Cast<UBlueprint>(BPClass->ClassGeneratedBy))
+	{
+		TArray<UEdGraph*> BpGraphs;
+		Blueprint->GetAllGraphs(BpGraphs);
+
+		bool bAllValidNodes = true;
+		for (UEdGraph* Graph : BpGraphs)
+		{
+			TArray<UObject*> GraphObjects;
+			GetObjectsWithOuter(Graph, GraphObjects, /*bIncludeNestedObjects =*/true);
+
+			for (UObject* GraphObj : GraphObjects)
+			{
+				if (UProperty* GraphProperty = Cast<UProperty>(GraphObj))
+				{
+					// @TODO: do we even have graph owned properties? no?
+				}
+				else if (UEdGraphNode* EdNode = Cast<UEdGraphNode>(GraphObj))
+				{
+					if (EdNode->bHasCompilerMessage && EdNode->ErrorType >= EMessageSeverity::Warning)
+					{
+						SetTestState(ETestResult::Failed);
+						return false;
+					}
+				}
+				else if (UEdGraphPin* NodePin = Cast<UEdGraphPin>(GraphObj))
+				{
+					FEdGraphPinType& PinType = NodePin->PinType;
+					if (PinType.PinCategory == UEdGraphSchema_K2::PC_Class)
+					{
+						if (!IsValidClass(Cast<UClass>(PinType.PinSubCategoryObject.Get())) || 
+							(NodePin->DefaultObject && !IsValidClass(Cast<UClass>(NodePin->DefaultObject))) )
+						{
+							SetTestState(ETestResult::Failed);
+							return false;
+						}
+					}
+					else if (PinType.PinCategory == UEdGraphSchema_K2::PC_Object)
+					{
+						if ((PinType.PinSubCategory != UEdGraphSchema_K2::PSC_Self) &&
+							!IsValidClass(Cast<UClass>(PinType.PinSubCategoryObject.Get())) )
+						{
+							SetTestState(ETestResult::Failed);
+							return false;
+						}
+						else if ( NodePin->DefaultObject && !IsValidClass(NodePin->DefaultObject->GetClass()) )
+						{
+							SetTestState(ETestResult::Failed);
+							return false;
+						}
+					}
+					else if ( (PinType.PinCategory == UEdGraphSchema_K2::PC_Interface)
+						&& !IsValidClass(Cast<UClass>(PinType.PinSubCategoryObject.Get())) )
+					{
+						SetTestState(ETestResult::Failed);
+						return false;
+					}
+				}
+			}
+		}
+	}
+#endif // WITH_EDITOR
+	return true;
+}
+
 bool ACircularDependencyTestActor::TestVerifyStructMember(UScriptStruct* Struct, uint8* StructInst, bool bCheckPropertyType)
 {
 	for (UProperty* StructProp = Struct->PropertyLink; StructProp != nullptr; StructProp = StructProp->NextRef)
@@ -120,16 +198,11 @@ bool ACircularDependencyTestActor::TestVerifyIsBlueprintTypeVar(FName VarName, b
 
 bool ACircularDependencyTestActor::RunVerificationTests_Implementation()
 {
-	return TestVerifyClass();
+	return TestVerifyClass() && TestVerifyBlueprint();
 }
 
 bool ACircularDependencyTestActor::TestVerifyProperty(UProperty* Property, uint8* Container, bool bCheckPropertyType)
 {
-	auto TestVerifyClass = [](UClass* ClassToVerify)->bool
-	{
-		return (ClassToVerify != nullptr) && !IsPlaceholderClass(ClassToVerify);
-	};
-
 	bool bBelongsToThisBlueprintClass = false;
 	if (UClass* OwnerClass = Property->GetOwnerClass())
 	{
@@ -144,7 +217,7 @@ bool ACircularDependencyTestActor::TestVerifyProperty(UProperty* Property, uint8
 	}
 	else if (UStruct* OwnerStruct = Property->GetOwnerStruct())
 	{
-
+		// @TODO
 	}
 
 	if (UObjectProperty* ObjProp = Cast<UObjectProperty>(Property))
@@ -155,7 +228,7 @@ bool ACircularDependencyTestActor::TestVerifyProperty(UProperty* Property, uint8
 			ClassToVerify = Cast<UClass>(ClassProp->GetPropertyValue_InContainer(Container));
 		}
 
-		if (!TestVerifyClass(ClassToVerify))
+		if (!IsValidClass(ClassToVerify))
 		{
 			return false;
 		}		
@@ -169,7 +242,7 @@ bool ACircularDependencyTestActor::TestVerifyProperty(UProperty* Property, uint8
 	}
 	else if (UInterfaceProperty* InterfaceProp = Cast<UInterfaceProperty>(Property))
 	{
-		if (!TestVerifyClass(InterfaceProp->InterfaceClass))
+		if (!IsValidClass(InterfaceProp->InterfaceClass))
 		{
 			return false;
 		}
diff --git a/Engine/Source/Runtime/Core/Public/Misc/Build.h b/Engine/Source/Runtime/Core/Public/Misc/Build.h
--- a/Engine/Source/Runtime/Core/Public/Misc/Build.h
+++ b/Engine/Source/Runtime/Core/Public/Misc/Build.h
@@ -242,4 +242,4 @@
 #define USE_UBER_GRAPH_PERSISTENT_FRAME 1
 
 #define USE_CIRCULAR_DEPENDENCY_LOAD_DEFERRING (WITH_EDITOR && 0)
-#define TEST_CHECK_DEPENDENCY_LOAD_DEFERRING   (USE_CIRCULAR_DEPENDENCY_LOAD_DEFERRING && 1)
+#define USE_DEFERRED_DEPENDENCY_CHECK_VERIFICATION_TESTS (USE_CIRCULAR_DEPENDENCY_LOAD_DEFERRING && 0)
diff --git a/Engine/Source/Runtime/CoreUObject/Private/Blueprint/BlueprintSupport.cpp b/Engine/Source/Runtime/CoreUObject/Private/Blueprint/BlueprintSupport.cpp
--- a/Engine/Source/Runtime/CoreUObject/Private/Blueprint/BlueprintSupport.cpp
+++ b/Engine/Source/Runtime/CoreUObject/Private/Blueprint/BlueprintSupport.cpp
@@ -50,31 +50,22 @@ bool FBlueprintSupport::UseDeferredDependencyLoading()
 #endif
 }
 
-bool FBlueprintSupport::UseDeferredCDOSerialization()
-{
-#if USE_CIRCULAR_DEPENDENCY_LOAD_DEFERRING
-	static const FBoolConfigValueHelper DeferCDOSerializations(TEXT("Kismet"), TEXT("bDeferCDOLoadSerialization"), GEngineIni);
-	return UseDeferredDependencyLoading() && DeferCDOSerializations;
-#else 
-	return false;
-#endif
-}
 
-bool FBlueprintSupport::UseDeferredLoadingForSubsequentLoads()
+bool FBlueprintSupport::IsResolvingDeferredDependenciesDisabled()
 {
-#if USE_CIRCULAR_DEPENDENCY_LOAD_DEFERRING
-	static const FBoolConfigValueHelper PropagateDeferredFlagForStructs(TEXT("Kismet"), TEXT("bPropagateDeferredLoadsForStructs"), GEngineIni);
-	return UseDeferredDependencyLoading() && PropagateDeferredFlagForStructs;
+#if USE_DEFERRED_DEPENDENCY_CHECK_VERIFICATION_TESTS
+	static const FBoolConfigValueHelper NoDeferredDependencyResolves(TEXT("Kismet"), TEXT("bForceDisableDeferredDependencyResolving"), GEngineIni);
+	return !UseDeferredDependencyLoading() || NoDeferredDependencyResolves;
 #else 
 	return false;
 #endif
 }
 
-bool FBlueprintSupport::UseDeferredDependencyVerificationChecks()
+bool FBlueprintSupport::IsDeferredCDOSerializationDisabled()
 {
-#if TEST_CHECK_DEPENDENCY_LOAD_DEFERRING
-	static const FBoolConfigValueHelper TestCheckDeferredDependencies(TEXT("Kismet"), TEXT("bCheckDeferredDependencies"), GEngineIni);
-	return UseDeferredDependencyLoading() && TestCheckDeferredDependencies;
+#if USE_DEFERRED_DEPENDENCY_CHECK_VERIFICATION_TESTS
+	static const FBoolConfigValueHelper NoDeferredCDOLoading(TEXT("Kismet"), TEXT("bForceDisableDeferredCDOLoading"), GEngineIni);
+	return !UseDeferredDependencyLoading() || NoDeferredCDOLoading;
 #else 
 	return false;
 #endif
@@ -142,6 +133,14 @@ TArray<UClass*> const& FScopedClassDependencyGather::GetCachedDependencies()
 /*******************************************************************************
  * ULinkerLoad
  ******************************************************************************/
+
+// rather than littering the code with USE_DEFERRED_DEPENDENCY_CHECK_VERIFICATION_TESTS
+// checks, let's just define DEFERRED_DEPENDENCY_CHECK for the file
+#if USE_DEFERRED_DEPENDENCY_CHECK_VERIFICATION_TESTS
+#define DEFERRED_DEPENDENCY_CHECK(CheckExpr) checkSlow(CheckExpr)
+#else  // USE_DEFERRED_DEPENDENCY_CHECK_VERIFICATION_TESTS
+#define DEFERRED_DEPENDENCY_CHECK(CheckExpr)
+#endif // USE_DEFERRED_DEPENDENCY_CHECK_VERIFICATION_TESTS
  
 struct FPreloadMembersHelper
 {
@@ -190,15 +189,6 @@ struct FPreloadMembersHelper
  */
 bool ULinkerLoad::RegenerateBlueprintClass(UClass* LoadClass, UObject* ExportObject)
 {
-#if TEST_CHECK_DEPENDENCY_LOAD_DEFERRING
-	// @TODO: we can't rely on this check, because once we're in class 
-	//        regeneration, graphs/nodes are force loaded and that could cause
-	//        subsequent class regeneration (for other packages)
-	//static ULinkerLoad* RegeneratingLinker = nullptr;
-	//checkSlow(!FBlueprintSupport::UseDeferredDependencyVerificationChecks() || (RegeneratingLinker == nullptr) || (RegeneratingLinker == this));
-	//TGuardValue<ULinkerLoad*> ReentrantGuard(RegeneratingLinker, this);
-#endif // TEST_CHECK_DEPENDENCY_LOAD_DEFERRING
-
 	// determine if somewhere further down the callstack, we're already in this
 	// function for this class
 	bool const bAlreadyRegenerating = LoadClass->ClassGeneratedBy->HasAnyFlags(RF_BeingRegenerated);
@@ -366,12 +356,7 @@ public:
 	FUnresolvedStructTracker(UStruct* LoadStruct)
 		: TrackedStruct(LoadStruct)
 	{
-#if TEST_CHECK_DEPENDENCY_LOAD_DEFERRING
-		bool const bCircumventValidationChecks = !FBlueprintSupport::UseDeferredDependencyVerificationChecks();
-		checkSlow(bCircumventValidationChecks || LoadStruct != nullptr);
-		checkSlow(bCircumventValidationChecks || LoadStruct->GetLinker() != nullptr);
-#endif // TEST_CHECK_DEPENDENCY_LOAD_DEFERRING
-
+		DEFERRED_DEPENDENCY_CHECK((LoadStruct != nullptr) && (LoadStruct->GetLinker() != nullptr));
 		UnresolvedStructs.Add(LoadStruct);
 	}
 
@@ -446,14 +431,13 @@ void ULinkerLoad::ResolveDeferredDependencies(UStruct* LoadStruct)
 	//--------------------------------------
 	TGuardValue<uint32> LoadFlagsGuard(LoadFlags, (LoadFlags & ~LOAD_DeferDependencyLoads));
 
-#if TEST_CHECK_DEPENDENCY_LOAD_DEFERRING
+#if USE_DEFERRED_DEPENDENCY_CHECK_VERIFICATION_TESTS
 	static int32 RecursiveDepth = 0;
 	int32 const MeasuredDepth = RecursiveDepth;
 	TGuardValue<int32> RecursiveDepthGuard(RecursiveDepth, RecursiveDepth + 1);
 
-	bool const bCircumventValidationChecks = !FBlueprintSupport::UseDeferredDependencyVerificationChecks();
-	checkSlow( bCircumventValidationChecks || (LoadStruct && (LoadStruct->GetLinker() == this)) );
-#endif // TEST_CHECK_DEPENDENCY_LOAD_DEFERRING
+	DEFERRED_DEPENDENCY_CHECK( (LoadStruct != nullptr) && (LoadStruct->GetLinker() == this) );
+#endif // USE_DEFERRED_DEPENDENCY_CHECK_VERIFICATION_TESTS
 
 	// scoped block to manage the lifetime of ScopedResolveTracker, so that 
 	// this resolve is only tracked for the duration of resolving all its 
@@ -482,9 +466,7 @@ void ULinkerLoad::ResolveDeferredDependencies(UStruct* LoadStruct)
 
 			if (ULinkerPlaceholderClass* PlaceholderClass = Cast<ULinkerPlaceholderClass>(Import.XObject))
 			{
-#if TEST_CHECK_DEPENDENCY_LOAD_DEFERRING
-				checkSlow(bCircumventValidationChecks || PlaceholderClass->ImportIndex == ImportIndex);
-#endif // TEST_CHECK_DEPENDENCY_LOAD_DEFERRING
+				DEFERRED_DEPENDENCY_CHECK(PlaceholderClass->ImportIndex == ImportIndex);
 
 				// NOTE: we don't check that this resolve successfully replaced any
 				//       references (by the return value), because this resolve 
@@ -532,18 +514,18 @@ void ULinkerLoad::ResolveDeferredDependencies(UStruct* LoadStruct)
 	//        ALevelScriptActor instance BEFORE the class has been regenerated)
 	//LoadAllObjects();
 
-#if TEST_CHECK_DEPENDENCY_LOAD_DEFERRING
-	for (TObjectIterator<ULinkerPlaceholderClass> PlaceholderIt; PlaceholderIt; ++PlaceholderIt)
+#if USE_DEFERRED_DEPENDENCY_CHECK_VERIFICATION_TESTS
+	for (TObjectIterator<ULinkerPlaceholderClass> PlaceholderIt; PlaceholderIt && !FBlueprintSupport::IsResolvingDeferredDependenciesDisabled(); ++PlaceholderIt)
 	{
 		ULinkerPlaceholderClass* PlaceholderClass = *PlaceholderIt;
 		if (PlaceholderClass->GetOuter() == LinkerRoot)
 		{
 			// there shouldn't be any deferred dependencies (belonging to this 
 			// linker) that need to be resolved by this point
-			checkSlow( bCircumventValidationChecks || (!PlaceholderClass->HasReferences() && PlaceholderClass->IsPendingKill()) );
+			DEFERRED_DEPENDENCY_CHECK(!PlaceholderClass->HasReferences() && PlaceholderClass->IsPendingKill());
 		}
 	}
-#endif // TEST_CHECK_DEPENDENCY_LOAD_DEFERRING
+#endif // USE_DEFERRED_DEPENDENCY_CHECK_VERIFICATION_TESTS
 #endif // USE_CIRCULAR_DEPENDENCY_LOAD_DEFERRING
 }
 
@@ -555,12 +537,9 @@ bool ULinkerLoad::HasUnresolvedDependencies() const
 	// struct, or super... see ULinkerLoad::ResolveDeferredDependencies)
 	bool bIsClassExportUnresolved = FUnresolvedStructTracker::IsAssociatedStructUnresolved(this);
 
-#if TEST_CHECK_DEPENDENCY_LOAD_DEFERRING
-	bool const bCircumventValidationChecks = !FBlueprintSupport::UseDeferredDependencyVerificationChecks();
 	// (ResolvingDeferredPlaceholder != nullptr) should imply 
 	// bIsClassExportUnresolved is true (but not the other way around)
-	checkSlow(bCircumventValidationChecks || (ResolvingDeferredPlaceholder == nullptr) || bIsClassExportUnresolved);
-#endif // TEST_CHECK_DEPENDENCY_LOAD_DEFERRING
+	DEFERRED_DEPENDENCY_CHECK((ResolvingDeferredPlaceholder == nullptr) || bIsClassExportUnresolved);
 	
 	return bIsClassExportUnresolved;
 
@@ -572,23 +551,25 @@ bool ULinkerLoad::HasUnresolvedDependencies() const
 int32 ULinkerLoad::ResolveDependencyPlaceholder(UClass* PlaceholderIn, UClass* ReferencingClass)
 {
 #if USE_CIRCULAR_DEPENDENCY_LOAD_DEFERRING
+#if USE_DEFERRED_DEPENDENCY_CHECK_VERIFICATION_TESTS
+	if (FBlueprintSupport::IsResolvingDeferredDependenciesDisabled())
+	{
+		return 0;
+	}
+#endif // USE_DEFERRED_DEPENDENCY_CHECK_VERIFICATION_TESTS
+
 	TGuardValue<uint32>  LoadFlagsGuard(LoadFlags, (LoadFlags & ~LOAD_DeferDependencyLoads));
  	TGuardValue<UClass*> ResolvingClassGuard(ResolvingDeferredPlaceholder, PlaceholderIn);
 
-#if TEST_CHECK_DEPENDENCY_LOAD_DEFERRING
- 	bool const bCircumventValidationChecks = !FBlueprintSupport::UseDeferredDependencyVerificationChecks();
- 	checkSlow(bCircumventValidationChecks || Cast<ULinkerPlaceholderClass>(PlaceholderIn) != nullptr);
- 	checkSlow(bCircumventValidationChecks || PlaceholderIn->GetOuter() == LinkerRoot);
-#endif // TEST_CHECK_DEPENDENCY_LOAD_DEFERRING
+	DEFERRED_DEPENDENCY_CHECK(Cast<ULinkerPlaceholderClass>(PlaceholderIn) != nullptr);
+	DEFERRED_DEPENDENCY_CHECK(PlaceholderIn->GetOuter() == LinkerRoot);
 
 	ULinkerPlaceholderClass* PlaceholderClass = (ULinkerPlaceholderClass*)PlaceholderIn;
 	
 	int32 const ImportIndex = PlaceholderClass->ImportIndex;
 	FObjectImport& Import = ImportMap[ImportIndex];
 
-#if TEST_CHECK_DEPENDENCY_LOAD_DEFERRING
-	checkSlow( bCircumventValidationChecks || (Import.XObject == PlaceholderClass) || (Import.XObject == nullptr) );
-#endif // TEST_CHECK_DEPENDENCY_LOAD_DEFERRING
+	DEFERRED_DEPENDENCY_CHECK( (Import.XObject == PlaceholderClass) || (Import.XObject == nullptr) );
 	
 	// clear the placeholder from the import, so that a call to CreateImport()
 	// properly fills it in
@@ -612,7 +593,6 @@ int32 ULinkerLoad::ResolveDependencyPlaceholder(UClass* PlaceholderIn, UClass* R
 		}
 	}
 
-#if TEST_CHECK_DEPENDENCY_LOAD_DEFERRING
 	// make sure that we know what utilized this placeholder class... right now
 	// we only expect UObjectProperties/UClassProperties/UInterfaceProperties/
 	// FImplementedInterfaces to, but something else could have requested the 
@@ -620,24 +600,19 @@ int32 ULinkerLoad::ResolveDependencyPlaceholder(UClass* PlaceholderIn, UClass* R
 	// doesn't have any known references (and it hasn't already been resolved in
 	// some recursive call), then there is something out there still using this
 	// placeholder class
-	checkSlow(bCircumventValidationChecks || (ReplacementCount > 0) || PlaceholderClass->HasReferences() || PlaceholderClass->HasBeenResolved());
-
-	UObject* PlaceholderObj = PlaceholderClass;
-	// 
-	FReferencerInformationList RemainingRefs;
-	//checkSlow(bCircumventValidationChecks || PlaceholderClass->HasReferences() == IsReferenced(PlaceholderObj, RF_NoFlags, /*bCheckSubObjects =*/false, &RemainingRefs));
-	//checkSlow(bCircumventValidationChecks || RemainingRefs.ExternalReferences.Num() == PlaceholderClass->GetRefCount());
-#endif // TEST_CHECK_DEPENDENCY_LOAD_DEFERRING
+	DEFERRED_DEPENDENCY_CHECK( (ReplacementCount > 0) || PlaceholderClass->HasReferences() || PlaceholderClass->HasBeenResolved() );
 
 	ReplacementCount += PlaceholderClass->ReplaceTrackedReferences(RealClassObj);
 	PlaceholderClass->MarkPendingKill(); // @TODO: ensure these are properly GC'd
 
-#if TEST_CHECK_DEPENDENCY_LOAD_DEFERRING
+#if USE_DEFERRED_DEPENDENCY_CHECK_VERIFICATION_TESTS
+	UObject* PlaceholderObj = PlaceholderClass;
 	// there should not be any references left to this placeholder class 
 	// (if there is, then we didn't log that referencer with the placeholder)
-	FReferencerInformationList DanglingReferences;
-	checkSlow(bCircumventValidationChecks || !IsReferenced(PlaceholderObj, RF_NoFlags, /*bCheckSubObjects =*/false, &DanglingReferences));
-#endif // TEST_CHECK_DEPENDENCY_LOAD_DEFERRING
+	FReferencerInformationList UnresolvedReferences;
+	DEFERRED_DEPENDENCY_CHECK(!IsReferenced(PlaceholderObj, RF_NoFlags, /*bCheckSubObjects =*/false, &UnresolvedReferences));
+#endif // USE_DEFERRED_DEPENDENCY_CHECK_VERIFICATION_TESTS
+
 
 	return ReplacementCount;
 #else  // USE_CIRCULAR_DEPENDENCY_LOAD_DEFERRING
@@ -717,12 +692,9 @@ void ULinkerLoad::FinalizeBlueprint(UClass* LoadClass)
 		// we could end up with placeholder classes in our script-code)
 		if (FUnresolvedStructTracker::IsImportStructUnresolved(Import.XObject))
 		{
-#if TEST_CHECK_DEPENDENCY_LOAD_DEFERRING
-			bool const bCircumventValidationChecks = !FBlueprintSupport::UseDeferredDependencyVerificationChecks();
-			// because it is tracked by FUnresolvedStructTracker, it must be a class
-			checkSlow(bCircumventValidationChecks || Cast<UClass>(Import.XObject) != nullptr);
-#endif // TEST_CHECK_DEPENDENCY_LOAD_DEFERRING
-			Import.SourceLinker->ResolveDeferredDependencies((UClass*)Import.XObject);
+			// because it is tracked by FUnresolvedStructTracker, it must be a struct
+			DEFERRED_DEPENDENCY_CHECK(Cast<UStruct>(Import.XObject) != nullptr);
+			Import.SourceLinker->ResolveDeferredDependencies((UStruct*)Import.XObject);
 		}
 	}
 
@@ -732,11 +704,15 @@ void ULinkerLoad::FinalizeBlueprint(UClass* LoadClass)
 	// get fully executed in some nested call
 	if (IsBlueprintFinalizationPending())
 	{
-		check(DeferredExportIndex != INDEX_NONE);
+		DEFERRED_DEPENDENCY_CHECK(DeferredExportIndex != INDEX_NONE);
 		FObjectExport& CDOExport = ExportMap[DeferredExportIndex];
 
 		UObject* CDO = CDOExport.Object;
-		if (FBlueprintSupport::UseDeferredCDOSerialization())
+		DEFERRED_DEPENDENCY_CHECK(CDO != nullptr);
+
+#if USE_DEFERRED_DEPENDENCY_CHECK_VERIFICATION_TESTS
+		if (!FBlueprintSupport::IsDeferredCDOSerializationDisabled())
+#endif // USE_DEFERRED_DEPENDENCY_CHECK_VERIFICATION_TESTS
 		{
 			// have to prematurely set the CDO's linker so we can force a Preload()/
 			// Serialization of the CDO before we regenerate the Blueprint class
@@ -757,22 +733,22 @@ void ULinkerLoad::FinalizeBlueprint(UClass* LoadClass)
 			//       already been "finalized", then its RF_NeedLoad flag would 
 			//       have been cleared
 			Preload(CDO);
+
+			DEFERRED_DEPENDENCY_CHECK(CDO->HasAnyFlags(RF_LoadCompleted));
 		}
 
 		UClass* BlueprintClass = Cast<UClass>(IndexToObject(CDOExport.ClassIndex));
-#if TEST_CHECK_DEPENDENCY_LOAD_DEFERRING
-		bool const bCircumventValidationChecks = !FBlueprintSupport::UseDeferredDependencyVerificationChecks();
-		checkSlow(bCircumventValidationChecks || BlueprintClass == LoadClass);
-		checkSlow(bCircumventValidationChecks || CDO->HasAnyFlags(RF_LoadCompleted));
-		checkSlow(bCircumventValidationChecks || BlueprintClass->HasAnyClassFlags(CLASS_CompiledFromBlueprint));
+		DEFERRED_DEPENDENCY_CHECK(BlueprintClass == LoadClass);
+		DEFERRED_DEPENDENCY_CHECK(BlueprintClass->HasAnyClassFlags(CLASS_CompiledFromBlueprint));
 
+#if USE_DEFERRED_DEPENDENCY_CHECK_VERIFICATION_TESTS
 		// at this point there should not be any instances of the Blueprint (else,
 		// we'd have to reinstance and that is too expensive an operation to 
 		// have at load time)
 		TArray<UObject*> ClassInstances;
 		GetObjectsOfClass(BlueprintClass, ClassInstances, /*bIncludeDerivedClasses =*/true);
-		checkSlow(bCircumventValidationChecks || ClassInstances.Num() == 0);
-#endif // TEST_CHECK_DEPENDENCY_LOAD_DEFERRING
+		DEFERRED_DEPENDENCY_CHECK(ClassInstances.Num() == 0);
+#endif // USE_DEFERRED_DEPENDENCY_CHECK_VERIFICATION_TESTS
 
 		// clear this so IsBlueprintFinalizationPending() doesn't report true
 		DeferredExportIndex = INDEX_NONE;
@@ -802,6 +778,9 @@ bool ULinkerLoad::IsBlueprintFinalizationPending() const
 #endif // USE_CIRCULAR_DEPENDENCY_LOAD_DEFERRING
 }
 
+// don't want other files ending up with this internal define
+#undef DEFERRED_DEPENDENCY_CHECK
+
 /*******************************************************************************
  * UObject
  ******************************************************************************/
diff --git a/Engine/Source/Runtime/CoreUObject/Private/UObject/LinkerLoad.cpp b/Engine/Source/Runtime/CoreUObject/Private/UObject/LinkerLoad.cpp
--- a/Engine/Source/Runtime/CoreUObject/Private/UObject/LinkerLoad.cpp
+++ b/Engine/Source/Runtime/CoreUObject/Private/UObject/LinkerLoad.cpp
@@ -2168,10 +2168,7 @@ bool ULinkerLoad::VerifyImportInner(const int32 ImportIndex, FString& WarningSuf
 			//      class's layout depends on the struct's size... in this case, 
 			//      we choke off circular loads by propagating this flag along 
 			//      to the struct linker (so it doesn't load any blueprints)
-			if (FBlueprintSupport::UseDeferredLoadingForSubsequentLoads())
-			{
-				InternalLoadFlags |= (LoadFlags & LOAD_DeferDependencyLoads);
-			}
+			InternalLoadFlags |= (LoadFlags & LOAD_DeferDependencyLoads);
 #endif // USE_CIRCULAR_DEPENDENCY_LOAD_DEFERRING
 
 			// we now fully load the package that we need a single export from - however, we still use CreatePackage below as it handles all cases when the package
@@ -2539,14 +2536,14 @@ void ULinkerLoad::LoadAllObjects( bool bForcePreload )
 	SlowTask.bVisibleOnUI = false;
 #endif
 
-#if TEST_CHECK_DEPENDENCY_LOAD_DEFERRING
+#if USE_DEFERRED_DEPENDENCY_CHECK_VERIFICATION_TESTS
 	// if we're re-entering a call to LoadAllObjects() while DeferDependencyLoads
 	// is set, then we're not doing our job (we're risking an export needing 
 	// another external asset)... if this is hit, then we're most likely already
 	// in this function (for this linker) further up the load chain; it should 
 	// finish the loads there
-	checkSlow( ((LoadFlags & LOAD_DeferDependencyLoads) == 0) || !FBlueprintSupport::UseDeferredDependencyVerificationChecks() );
-#endif // TEST_CHECK_DEPENDENCY_LOAD_DEFERRING
+	checkSlow((LoadFlags & LOAD_DeferDependencyLoads) == 0);
+#endif // USE_DEFERRED_DEPENDENCY_CHECK_VERIFICATION_TESTS
 
 	if ( (LoadFlags&LOAD_SeekFree) != 0 )
 	{
@@ -2774,13 +2771,13 @@ void ULinkerLoad::Preload( UObject* Object )
 	// classes we have)
 	bool const bIsBlueprintClass = (ObjectAsClass != nullptr) && !(ObjectAsClass->GetOutermost()->PackageFlags & PKG_CompiledIn);
 
-#if TEST_CHECK_DEPENDENCY_LOAD_DEFERRING
+#if USE_DEFERRED_DEPENDENCY_CHECK_VERIFICATION_TESTS
 	// we should NEVER be pre-loading another blueprint class when the 
 	// DeferDependencyLoads flag is set (some other blueprint class is already 
 	// being loaded further up the load chain, and this could introduce a 
 	// circular load)
-	checkSlow(!bIsBlueprintClass || !Object->HasAnyFlags(RF_NeedLoad) || !(LoadFlags & LOAD_DeferDependencyLoads) || !FBlueprintSupport::UseDeferredDependencyVerificationChecks());
-#endif // TEST_CHECK_DEPENDENCY_LOAD_DEFERRING
+	checkSlow(!bIsBlueprintClass || !Object->HasAnyFlags(RF_NeedLoad) || !(LoadFlags & LOAD_DeferDependencyLoads));
+#endif // USE_DEFERRED_DEPENDENCY_CHECK_VERIFICATION_TESTS
 #endif // USE_CIRCULAR_DEPENDENCY_LOAD_DEFERRING
 	
 	// Preload the object if necessary.
@@ -2835,34 +2832,37 @@ void ULinkerLoad::Preload( UObject* Object )
 
 				{
 					SCOPE_CYCLE_COUNTER(STAT_LinkerSerialize);
-					if ( Object->HasAnyFlags(RF_ClassDefaultObject) )
+					if (Object->HasAnyFlags(RF_ClassDefaultObject))
 					{
 #if USE_CIRCULAR_DEPENDENCY_LOAD_DEFERRING
-						if ( ((LoadFlags & LOAD_DeferDependencyLoads) != 0) && FBlueprintSupport::UseDeferredCDOSerialization() )
+						if ((LoadFlags & LOAD_DeferDependencyLoads) != 0)
 						{
-#if TEST_CHECK_DEPENDENCY_LOAD_DEFERRING
-							checkSlow(!FBlueprintSupport::UseDeferredDependencyVerificationChecks() || (DeferredExportIndex == INDEX_NONE) || (DeferredExportIndex == ExportIndex));
-#endif // TEST_CHECK_DEPENDENCY_LOAD_DEFERRING
-
-							// since serializing the CDO can introduce circular 
-							// dependencies, we want to stave that off until 
-							// we're ready to handle those 
-							DeferredExportIndex = ExportIndex;
-							// don't need to actually "consume" the data through
-							// serialization though (since we seek back to 
-							// SavedPos later on)
-
-							// reset the flag and return (don't worry, we make
-							// sure to force load this later)
-							Object->SetFlags(RF_NeedLoad);
-							return;
+#if USE_DEFERRED_DEPENDENCY_CHECK_VERIFICATION_TESTS
+							if (!FBlueprintSupport::IsDeferredCDOSerializationDisabled())
+							{
+								checkSlow((DeferredExportIndex == INDEX_NONE) || (DeferredExportIndex == ExportIndex));
+#else // USE_DEFERRED_DEPENDENCY_CHECK_VERIFICATION_TESTS
+							{
+#endif // USE_DEFERRED_DEPENDENCY_CHECK_VERIFICATION_TESTS
+
+								// since serializing the CDO can introduce circular 
+								// dependencies, we want to stave that off until 
+								// we're ready to handle those 
+								DeferredExportIndex = ExportIndex;
+								// don't need to actually "consume" the data through
+								// serialization though (since we seek back to 
+								// SavedPos later on)
+
+								// reset the flag and return (don't worry, we make
+								// sure to force load this later)
+								Object->SetFlags(RF_NeedLoad);
+								return;
+							}
 						}
-						else
 #endif // USE_CIRCULAR_DEPENDENCY_LOAD_DEFERRING
-						{
-							Object->GetClass()->SerializeDefaultObject(Object, *this);
-							Object->SetFlags(RF_LoadCompleted);
-						}
+
+						Object->GetClass()->SerializeDefaultObject(Object, *this);
+						Object->SetFlags(RF_LoadCompleted);
 					}
 					else
 					{
@@ -2880,16 +2880,15 @@ void ULinkerLoad::Preload( UObject* Object )
 					SCOPE_CYCLE_COUNTER(STAT_LinkerLoadDeferred);
 					if ((LoadFlags & LOAD_DeferDependencyLoads) != (*LoadFlagsGuard & LOAD_DeferDependencyLoads))
 					{
-#if TEST_CHECK_DEPENDENCY_LOAD_DEFERRING
-						bool const bCircumventValidationChecks = !FBlueprintSupport::UseDeferredDependencyVerificationChecks();
+#if USE_DEFERRED_DEPENDENCY_CHECK_VERIFICATION_TESTS
 						// since class serialization reads in the class's CDO, 
 						// then we can be certain that the CDO export object 
 						// exists (and DeferredExportIndex should reference it)
 						// ... FinalizeBlueprint() depends on this (and since 
 						// ResolveDeferredDependencies() can be recursive, we 
 						// check it out here, before it is called)
-						checkSlow(bCircumventValidationChecks || (DeferredExportIndex != INDEX_NONE));
-#endif // TEST_CHECK_DEPENDENCY_LOAD_DEFERRING
+						checkSlow((DeferredExportIndex != INDEX_NONE) || FBlueprintSupport::IsDeferredCDOSerializationDisabled());
+#endif // USE_DEFERRED_DEPENDENCY_CHECK_VERIFICATION_TESTS
 
 						ResolveDeferredDependencies(ObjectAsClass);
 						FinalizeBlueprint(ObjectAsClass);
@@ -3355,7 +3354,11 @@ UObject* ULinkerLoad::CreateExport( int32 Index )
 				&& ((Export.ObjectFlags&RF_ClassDefaultObject) != 0) )
 			{
 #if USE_CIRCULAR_DEPENDENCY_LOAD_DEFERRING
-				if ((LoadFlags & LOAD_DeferDependencyLoads) != 0)
+				if ( (LoadFlags & LOAD_DeferDependencyLoads) != 0
+#if USE_DEFERRED_DEPENDENCY_CHECK_VERIFICATION_TESTS
+					&& !FBlueprintSupport::IsDeferredCDOSerializationDisabled()
+#endif // USE_DEFERRED_DEPENDENCY_CHECK_VERIFICATION_TESTS
+				   )
 				{
 					// if LOAD_DeferDependencyLoads is set, then we're already
 					// serializing the blueprint's class somewhere up the chain,
diff --git a/Engine/Source/Runtime/CoreUObject/Private/UObject/LinkerPlaceholderClass.cpp b/Engine/Source/Runtime/CoreUObject/Private/UObject/LinkerPlaceholderClass.cpp
--- a/Engine/Source/Runtime/CoreUObject/Private/UObject/LinkerPlaceholderClass.cpp
+++ b/Engine/Source/Runtime/CoreUObject/Private/UObject/LinkerPlaceholderClass.cpp
@@ -7,21 +7,16 @@
 //------------------------------------------------------------------------------
 ULinkerPlaceholderClass::ULinkerPlaceholderClass(const FObjectInitializer& ObjectInitializer)
 	: Super(ObjectInitializer)
-#if TEST_CHECK_DEPENDENCY_LOAD_DEFERRING
 	, bResolvedReferences(false)
-#endif // TEST_CHECK_DEPENDENCY_LOAD_DEFERRING
 {
 }
 
 //------------------------------------------------------------------------------
 ULinkerPlaceholderClass::~ULinkerPlaceholderClass()
 {
-	// @TODO: should we assert when this happens, so we can figure out what 
-	//        placeholder references are still around?
-#if TEST_CHECK_DEPENDENCY_LOAD_DEFERRING
-	bool const bCircumventValidationChecks = !FBlueprintSupport::UseDeferredDependencyVerificationChecks();
-	checkSlow(bCircumventValidationChecks || !HasReferences());
-#endif // TEST_CHECK_DEPENDENCY_LOAD_DEFERRING
+#if USE_DEFERRED_DEPENDENCY_CHECK_VERIFICATION_TESTS
+	checkSlow(!HasReferences());
+#endif // USE_DEFERRED_DEPENDENCY_CHECK_VERIFICATION_TESTS
 
 	// by this point, we really shouldn't have any properties left (they should 
 	// have all got replaced), but just in case (so things don't blow up with a
@@ -71,29 +66,25 @@ void ULinkerPlaceholderClass::Bind()
 //------------------------------------------------------------------------------
 void ULinkerPlaceholderClass::AddTrackedReference(UProperty* ReferencingProperty)
 {
-#if TEST_CHECK_DEPENDENCY_LOAD_DEFERRING
-	bool const bCircumventValidationChecks = !FBlueprintSupport::UseDeferredDependencyVerificationChecks();
-	if (!bCircumventValidationChecks)
+#if USE_DEFERRED_DEPENDENCY_CHECK_VERIFICATION_TESTS
+	FObjectImport* PlaceholderImport = nullptr;
+	if (ULinkerLoad* PropertyLinker = ReferencingProperty->GetLinker())
 	{
-		FObjectImport* PlaceholderImport = nullptr;
-		if (ULinkerLoad* PropertyLinker = ReferencingProperty->GetLinker())
+		for (FObjectImport& Import : PropertyLinker->ImportMap)
 		{
-			for (FObjectImport& Import : PropertyLinker->ImportMap)
+			if (Import.XObject == this)
 			{
-				if (Import.XObject == this)
-				{
-					PlaceholderImport = &Import;
-					break;
-				}
+				PlaceholderImport = &Import;
+				break;
 			}
-			check(GetOuter() == PropertyLinker->LinkerRoot);
-			check(PropertyLinker->LoadFlags & LOAD_DeferDependencyLoads);
 		}
-		// if this check hits, then we're adding dependencies after we've 
-		// already resolved the placeholder (it won't be resolved again)
-		check(!bResolvedReferences); 
-	}	
-#endif // TEST_CHECK_DEPENDENCY_LOAD_DEFERRING
+		check(GetOuter() == PropertyLinker->LinkerRoot);
+		check((PropertyLinker->LoadFlags & LOAD_DeferDependencyLoads) || FBlueprintSupport::IsResolvingDeferredDependenciesDisabled());
+	}
+	// if this check hits, then we're adding dependencies after we've 
+	// already resolved the placeholder (it won't be resolved again)
+	check(!bResolvedReferences); 	
+#endif // USE_DEFERRED_DEPENDENCY_CHECK_VERIFICATION_TESTS
 
 	ReferencingProperties.Add(ReferencingProperty);
 }
@@ -119,10 +110,7 @@ bool ULinkerPlaceholderClass::HasBeenResolved() const
 //------------------------------------------------------------------------------
 void ULinkerPlaceholderClass::RemoveTrackedReference(UProperty* ReferencingProperty)
 {
-	//if (ReferencingProperty->IsValidLowLevel())
-	{
-		ReferencingProperties.Remove(ReferencingProperty);
-	}
+	ReferencingProperties.Remove(ReferencingProperty);
 }
 
 //------------------------------------------------------------------------------
diff --git a/Engine/Source/Runtime/CoreUObject/Public/Blueprint/BlueprintSupport.h b/Engine/Source/Runtime/CoreUObject/Public/Blueprint/BlueprintSupport.h
--- a/Engine/Source/Runtime/CoreUObject/Public/Blueprint/BlueprintSupport.h
+++ b/Engine/Source/Runtime/CoreUObject/Public/Blueprint/BlueprintSupport.h
@@ -23,9 +23,8 @@ struct FBlueprintSupport
 	 * aspects of the deferred loading (mostly for testing purposes). 
 	 */
 	static bool UseDeferredDependencyLoading();
-	static bool UseDeferredCDOSerialization();
-	static bool UseDeferredLoadingForSubsequentLoads();
-	static bool UseDeferredDependencyVerificationChecks();
+	static bool IsResolvingDeferredDependenciesDisabled();
+	static bool IsDeferredCDOSerializationDisabled();
 };
 
 #if WITH_EDITOR
