[UE-7296] Handling dependency placholder ref breaking in the property's BeginDestroy() method, instead of its destructor.

diff --git a/Engine/Source/Runtime/CoreUObject/Private/UObject/PropertyBaseObject.cpp b/Engine/Source/Runtime/CoreUObject/Private/UObject/PropertyBaseObject.cpp
--- a/Engine/Source/Runtime/CoreUObject/Private/UObject/PropertyBaseObject.cpp
+++ b/Engine/Source/Runtime/CoreUObject/Private/UObject/PropertyBaseObject.cpp
@@ -8,17 +8,16 @@
 	UObjectPropertyBase.
 -----------------------------------------------------------------------------*/
 
-UObjectPropertyBase::~UObjectPropertyBase()
+void UObjectPropertyBase::BeginDestroy()
 {
 #if USE_CIRCULAR_DEPENDENCY_LOAD_DEFERRING
-	if (PropertyClass->IsValidLowLevelFast(/*bRecursive =*/false))
+	if (ULinkerPlaceholderClass* PlaceholderClass = Cast<ULinkerPlaceholderClass>(PropertyClass))
 	{
-		if (ULinkerPlaceholderClass* PlaceholderClass = Cast<ULinkerPlaceholderClass>(PropertyClass))
-		{
-			PlaceholderClass->RemoveTrackedReference(this);
-		}
+		PlaceholderClass->RemoveTrackedReference(this);
 	}
 #endif // USE_CIRCULAR_DEPENDENCY_LOAD_DEFERRING
+
+	Super::BeginDestroy();
 }
 
 void UObjectPropertyBase::InstanceSubobjects(void* Data, void const* DefaultData, UObject* Owner, FObjectInstancingGraph* InstanceGraph )
diff --git a/Engine/Source/Runtime/CoreUObject/Private/UObject/PropertyClass.cpp b/Engine/Source/Runtime/CoreUObject/Private/UObject/PropertyClass.cpp
--- a/Engine/Source/Runtime/CoreUObject/Private/UObject/PropertyClass.cpp
+++ b/Engine/Source/Runtime/CoreUObject/Private/UObject/PropertyClass.cpp
@@ -8,17 +8,16 @@
 	UClassProperty.
 -----------------------------------------------------------------------------*/
 
-UClassProperty::~UClassProperty()
+void UClassProperty::BeginDestroy()
 {
 #if USE_CIRCULAR_DEPENDENCY_LOAD_DEFERRING
-	if (MetaClass->IsValidLowLevelFast(/*bRecursive =*/false))
+	if (ULinkerPlaceholderClass* PlaceholderClass = Cast<ULinkerPlaceholderClass>(MetaClass))
 	{
-		if (ULinkerPlaceholderClass* PlaceholderClass = Cast<ULinkerPlaceholderClass>(MetaClass))
-		{
-			PlaceholderClass->RemoveTrackedReference(this);
-		}
+		PlaceholderClass->RemoveTrackedReference(this);
 	}
 #endif // USE_CIRCULAR_DEPENDENCY_LOAD_DEFERRING
+
+	Super::BeginDestroy();
 }
 
 void UClassProperty::Serialize( FArchive& Ar )
diff --git a/Engine/Source/Runtime/CoreUObject/Private/UObject/PropertyInterface.cpp b/Engine/Source/Runtime/CoreUObject/Private/UObject/PropertyInterface.cpp
--- a/Engine/Source/Runtime/CoreUObject/Private/UObject/PropertyInterface.cpp
+++ b/Engine/Source/Runtime/CoreUObject/Private/UObject/PropertyInterface.cpp
@@ -7,17 +7,16 @@
 	UInterfaceProperty.
 -----------------------------------------------------------------------------*/
 
-UInterfaceProperty::~UInterfaceProperty()
+void UInterfaceProperty::BeginDestroy()
 {
 #if USE_CIRCULAR_DEPENDENCY_LOAD_DEFERRING
-	if (InterfaceClass->IsValidLowLevelFast(/*bRecursive =*/false))
+	if (ULinkerPlaceholderClass* PlaceholderClass = Cast<ULinkerPlaceholderClass>(InterfaceClass))
 	{
-		if (ULinkerPlaceholderClass* PlaceholderClass = Cast<ULinkerPlaceholderClass>(InterfaceClass))
-		{
-			PlaceholderClass->RemoveTrackedReference(this);
-		}
+		PlaceholderClass->RemoveTrackedReference(this);
 	}
 #endif // USE_CIRCULAR_DEPENDENCY_LOAD_DEFERRING
+
+	Super::BeginDestroy();
 }
 
 /**
diff --git a/Engine/Source/Runtime/CoreUObject/Public/UObject/UnrealType.h b/Engine/Source/Runtime/CoreUObject/Public/UObject/UnrealType.h
--- a/Engine/Source/Runtime/CoreUObject/Public/UObject/UnrealType.h
+++ b/Engine/Source/Runtime/CoreUObject/Public/UObject/UnrealType.h
@@ -1591,11 +1591,10 @@ class COREUOBJECT_API UObjectPropertyBase : public UProperty
 	,	PropertyClass( InClass )
 	{}
 
-	virtual ~UObjectPropertyBase();
-
 	// UObject interface
 	virtual void Serialize( FArchive& Ar ) override;
 	static void AddReferencedObjects(UObject* InThis, FReferenceCollector& Collector);
+	virtual void BeginDestroy() override;
 	// End of UObject interface
 
 	// UProperty interface
@@ -1934,11 +1933,10 @@ public:
 	{
 	}
 
-	virtual ~UClassProperty();
-
 	// UObject interface
 	virtual void Serialize( FArchive& Ar ) override;
 	static void AddReferencedObjects(UObject* InThis, FReferenceCollector& Collector);
+	virtual void BeginDestroy() override;
 	// End of UObject interface
 
 	// UHT interface
@@ -2033,8 +2031,6 @@ public:
 	{
 	}
 
-	virtual ~UInterfaceProperty();
-
 	// UHT interface
 	virtual FString GetCPPMacroType( FString& ExtendedTypeText ) const  override;
 	virtual FString GetCPPType( FString* ExtendedTypeText, uint32 CPPExportFlags ) const override;
@@ -2055,6 +2051,7 @@ public:
 	// UObject interface
 	virtual void Serialize( FArchive& Ar ) override;
 	virtual void EmitReferenceInfo(UClass& OwnerClass, int32 BaseOffset) override;
+	virtual void BeginDestroy() override;
 	// End of UObject interface
 
 	/**
