[INTEGRATE] Change 2412427 by [REDACTED_PERSON][REDACTED_HANDLE].7 on 2015/01/20 12:50:16

	Backing out change to make asset registry tag generation recurse into sub objects

	This caused Blueprint placement to break, which expects to find a parent class tag.

diff --git a/Engine/Source/Runtime/CoreUObject/Private/UObject/Obj.cpp b/Engine/Source/Runtime/CoreUObject/Private/UObject/Obj.cpp
--- a/Engine/Source/Runtime/CoreUObject/Private/UObject/Obj.cpp
+++ b/Engine/Source/Runtime/CoreUObject/Private/UObject/Obj.cpp
@@ -1123,42 +1123,13 @@ static FString GetConfigFilename( UObject* SourceObject )
 
 void UObject::FAssetRegistryTag::GetAssetRegistryTagsFromSearchableProperties(const UObject* Object, TArray<FAssetRegistryTag>& OutTags)
 {
-	static FName ReimportPathName(TEXT("ReimportPath"));
-
 	check(NULL != Object);
 	for( TFieldIterator<UProperty> FieldIt( Object->GetClass() ); FieldIt; ++FieldIt )
 	{
 		if ( FieldIt->HasAnyPropertyFlags(CPF_AssetRegistrySearchable) )
 		{
-			const uint8* PropertyAddr = FieldIt->ContainerPtrToValuePtr<uint8>(Object);
-
-			if (const auto* Property = Cast<UStructProperty>(*FieldIt))
-			{
-				if (Property->Struct)
-				{
-					GetAssetRegistryTagsFromSearchableProperties(Property->Struct, OutTags);
-				}
-				continue;
-			}
-			else if (const auto* Property = Cast<UObjectProperty>(*FieldIt))
-			{
-				if (const UObject* Obj = Property->GetObjectPropertyValue(PropertyAddr))
-				{
-					GetAssetRegistryTagsFromSearchableProperties(Obj, OutTags);
-				}
-				continue;
-			}
-
-
-#if WITH_EDITOR
-			if (!GIsEditor && FieldIt->HasMetaData(ReimportPathName))
-			{
-				// ignore reimport paths in non-editor
-				continue;
-			}
-#endif
-
 			FString PropertyStr;
+			const uint8* PropertyAddr = FieldIt->ContainerPtrToValuePtr<uint8>(Object);
 			FieldIt->ExportTextItem( PropertyStr, PropertyAddr, PropertyAddr, NULL, PPF_None );
 
 			FAssetRegistryTag::ETagType TagType;
@@ -1194,15 +1165,7 @@ void UObject::FAssetRegistryTag::GetAssetRegistryTagsFromSearchableProperties(co
 				TagType = FAssetRegistryTag::TT_Alphabetical;
 			}
 
-			FName Key = FieldIt->GetFName();
-
-#if WITH_EDITOR
-			if (GIsEditor && FieldIt->HasMetaData(ReimportPathName))
-			{
-				Key = ReimportPathName;
-			}
-#endif
-			OutTags.Add( FAssetRegistryTag(Key, PropertyStr, TagType) );
+			OutTags.Add( FAssetRegistryTag(FieldIt->GetFName(), PropertyStr, TagType) );
 		}
 	}
 }
