[INTEGRATE] Change 2415255 by [REDACTED_PERSON][REDACTED_HANDLE] on 2015/01/22 09:05:46

	Added useful asset registry tags to UClass

	UE-7184 - Show game c++ classes and engine c++ classes in the content browser

	These show up in the Content Browser tooltip.

diff --git a/Engine/Source/Runtime/CoreUObject/Private/UObject/Class.cpp b/Engine/Source/Runtime/CoreUObject/Private/UObject/Class.cpp
--- a/Engine/Source/Runtime/CoreUObject/Private/UObject/Class.cpp
+++ b/Engine/Source/Runtime/CoreUObject/Private/UObject/Class.cpp
@@ -2801,6 +2801,25 @@ FString UClass::GetDesc()
 	return GetName();
 }
 
+void UClass::GetAssetRegistryTags(TArray<FAssetRegistryTag>& OutTags) const
+{
+	Super::GetAssetRegistryTags(OutTags);
+
+#if WITH_EDITOR
+	static const FName ParentClassFName = "ParentClass";
+	const UClass* const ParentClass = GetSuperClass();
+	OutTags.Add( FAssetRegistryTag(ParentClassFName, ((ParentClass) ? ParentClass->GetFName() : NAME_None).ToString(), FAssetRegistryTag::TT_Alphabetical) );
+
+	static const FName ModuleNameFName = "ModuleName";
+	const UPackage* const ClassPackage = GetOuterUPackage();
+	OutTags.Add( FAssetRegistryTag(ModuleNameFName, ((ClassPackage) ? FPackageName::GetShortFName(ClassPackage->GetFName()) : NAME_None).ToString(), FAssetRegistryTag::TT_Alphabetical) );
+
+	static const FName ModuleRelativePathFName = "ModuleRelativePath";
+	const FString& ClassModuleRelativeIncludePath = GetMetaData(ModuleRelativePathFName);
+	OutTags.Add( FAssetRegistryTag(ModuleRelativePathFName, ClassModuleRelativeIncludePath, FAssetRegistryTag::TT_Alphabetical) );
+#endif
+}
+
 void UClass::Link(FArchive& Ar, bool bRelinkExistingProperties)
 {
 	check(!bRelinkExistingProperties || !(ClassFlags & CLASS_Intrinsic));
diff --git a/Engine/Source/Runtime/CoreUObject/Public/UObject/Class.h b/Engine/Source/Runtime/CoreUObject/Public/UObject/Class.h
--- a/Engine/Source/Runtime/CoreUObject/Public/UObject/Class.h
+++ b/Engine/Source/Runtime/CoreUObject/Public/UObject/Class.h
@@ -1917,6 +1917,7 @@ public:
 	static void AddReferencedObjects(UObject* InThis, FReferenceCollector& Collector);
 	virtual FRestoreForUObjectOverwrite* GetRestoreForUObjectOverwrite() override;
 	virtual FString GetDesc() override;
+	virtual void GetAssetRegistryTags(TArray<FAssetRegistryTag>& OutTags) const override;
 	virtual bool IsAsset() const override { return false; }
 	// End of UObject interface.
 
