[INTEGRATE] Change 2417819 by [REDACTED_PERSON][REDACTED_HANDLE] on 2015/01/24 11:54:16

	When duplicating an object try and keep the name the same if putting in a different outer and the name is not explicitly specified

diff --git a/Engine/Source/Runtime/CoreUObject/Private/UObject/UObjectGlobals.cpp b/Engine/Source/Runtime/CoreUObject/Private/UObject/UObjectGlobals.cpp
--- a/Engine/Source/Runtime/CoreUObject/Private/UObject/UObjectGlobals.cpp
+++ b/Engine/Source/Runtime/CoreUObject/Private/UObject/UObjectGlobals.cpp
@@ -1392,6 +1392,14 @@ UObject* StaticDuplicateObject(UObject const* SourceObject, UObject* DestOuter,
 	{
 		Parameters.DestName = FName(DestName, FNAME_Add, true);
 	}
+	else if (SourceObject->GetOuter() != DestOuter)
+	{
+		// try to keep the object name consistent if possible
+		if (FindObjectFast<UObject>(DestOuter, SourceObject->GetFName()) == nullptr)
+		{
+			Parameters.DestName = SourceObject->GetFName();
+		}
+	}
 
 	if ( DestClass == NULL )
 	{
