#UE4 Fixed a crash that happened when the source linker is null in ULinkerLoad::FinalizeBlueprint.

diff --git a/Engine/Source/Runtime/CoreUObject/Private/Blueprint/BlueprintSupport.cpp b/Engine/Source/Runtime/CoreUObject/Private/Blueprint/BlueprintSupport.cpp
--- a/Engine/Source/Runtime/CoreUObject/Private/Blueprint/BlueprintSupport.cpp
+++ b/Engine/Source/Runtime/CoreUObject/Private/Blueprint/BlueprintSupport.cpp
@@ -708,7 +708,10 @@ void ULinkerLoad::FinalizeBlueprint(UClass* LoadClass)
 		{
 			// because it is tracked by FUnresolvedStructTracker, it must be a struct
 			DEFERRED_DEPENDENCY_CHECK(Cast<UStruct>(Import.XObject) != nullptr);
-			Import.SourceLinker->ResolveDeferredDependencies((UStruct*)Import.XObject);
+			if (Import.SourceLinker)
+			{
+				Import.SourceLinker->ResolveDeferredDependencies((UStruct*)Import.XObject);
+			}
 		}
 	}
 
