A fix in StaticFindObject for a bug that was causing the infinite recursion in some cases.

diff --git a/Engine/Source/Runtime/CoreUObject/Private/UObject/UObjectGlobals.cpp b/Engine/Source/Runtime/CoreUObject/Private/UObject/UObjectGlobals.cpp
--- a/Engine/Source/Runtime/CoreUObject/Private/UObject/UObjectGlobals.cpp
+++ b/Engine/Source/Runtime/CoreUObject/Private/UObject/UObjectGlobals.cpp
@@ -173,21 +173,16 @@ UObject* StaticFindObject( UClass* ObjectClass, UObject* InObjectPackage, const
 	// and old packages can have invalid ones. The path now does not have a UCLASS as an outer. All mentioned
 	// types are just children of package of the file there were defined in.
 	if (!MatchingObject && ObjectPackage != nullptr &&
+		ObjectPackage->IsA<UClass>() && // Only if outer is a class.
 		!FPlatformProperties::RequiresCookedData() && // Cooked platforms will have all paths resolved.
 			(
 				ObjectClass == UEnum::StaticClass()	// Enums
-				|| ObjectClass == UScriptStruct::StaticClass() // Structs
+				|| ObjectClass == UScriptStruct::StaticClass() || ObjectClass == UStruct::StaticClass() // Structs
 				|| (ObjectClass == UFunction::StaticClass() && FString(OrigInName).EndsWith(HEADER_GENERATED_DELEGATE_SIGNATURE_SUFFIX)) // Delegates
 			)
 		)
 	{
-		UPackage* SearchPackage = ObjectPackage->GetOutermost();
-		const FString SearchName = ObjectName.ToString();
-		const bool bIdenticalSearch = (SearchPackage == InObjectPackage) && (SearchName == FString(OrigInName));
-		if (!bIdenticalSearch)
-		{
-			MatchingObject = StaticFindObject(ObjectClass, SearchPackage, *SearchName, ExactClass);
-		}
+		MatchingObject = StaticFindObject(ObjectClass, ObjectPackage->GetOutermost(), *ObjectName.ToString(), ExactClass);
 	}
 
 	return MatchingObject;
