Old localization property flags and related stuff removed:

PPF_LocalizedOnly
CPF_Localized
CLASS_Localized
UProperty::IsLocalized()
Localized specifiers in UPROPERTYs.

diff --git a/Engine/Source/Editor/KismetCompiler/Private/KismetCompilerCppBackend.cpp b/Engine/Source/Editor/KismetCompiler/Private/KismetCompilerCppBackend.cpp
--- a/Engine/Source/Editor/KismetCompiler/Private/KismetCompilerCppBackend.cpp
+++ b/Engine/Source/Editor/KismetCompiler/Private/KismetCompilerCppBackend.cpp
@@ -528,7 +528,6 @@ struct FEmitHelper
 		}
 
 		// OTHER
-		HANDLE_CPF_TAG(TEXT("Localized"), CPF_Localized | CPF_BlueprintReadOnly)
 		HANDLE_CPF_TAG(TEXT("Transient"), CPF_Transient)
 		HANDLE_CPF_TAG(TEXT("DuplicateTransient"), CPF_DuplicateTransient)
 		HANDLE_CPF_TAG(TEXT("TextExportTransient"), CPF_TextExportTransient)
diff --git a/Engine/Source/Programs/UnrealHeaderTool/Private/HeaderParser.cpp b/Engine/Source/Programs/UnrealHeaderTool/Private/HeaderParser.cpp
--- a/Engine/Source/Programs/UnrealHeaderTool/Private/HeaderParser.cpp
+++ b/Engine/Source/Programs/UnrealHeaderTool/Private/HeaderParser.cpp
@@ -2299,7 +2299,7 @@ bool FHeaderParser::GetVarType
 			}
 			else if (Specifier == TEXT("Localized"))
 			{
-				Flags |= CPF_Localized | CPF_BlueprintReadOnly;
+				FError::Throwf(TEXT("The Localized specifier is deprecated"));
 			}
 			else if (Specifier == TEXT("Transient"))
 			{
@@ -4050,11 +4050,6 @@ void PostParsingClassSetup(UClass* Class)
 	// Set all optimization ClassFlags based on property types
 	for (TFieldIterator<UProperty> It(Class, EFieldIteratorFlags::ExcludeSuper); It; ++It)
 	{
-		if (It->IsLocalized())
-		{
-			Class->ClassFlags |= CLASS_Localized;
-		}
-
 		if ((It->PropertyFlags & CPF_Config) != 0)
 		{
 			Class->ClassFlags |= CLASS_Config;
diff --git a/Engine/Source/Runtime/CoreUObject/Private/UObject/Class.cpp b/Engine/Source/Runtime/CoreUObject/Private/UObject/Class.cpp
--- a/Engine/Source/Runtime/CoreUObject/Private/UObject/Class.cpp
+++ b/Engine/Source/Runtime/CoreUObject/Private/UObject/Class.cpp
@@ -654,8 +654,7 @@ void UStruct::Link(FArchive& Ar, bool bRelinkExistingProperties)
 		}
 
 		// Link references to properties that require their values to be copied from CDO.
-		if ((Property->HasAnyPropertyFlags(CPF_Config) && Property->GetOwnerClass() && !Property->GetOwnerClass()->HasAnyClassFlags(CLASS_PerObjectConfig)) ||
-			 Property->HasAnyPropertyFlags(CPF_Localized))
+		if ((Property->HasAnyPropertyFlags(CPF_Config) && Property->GetOwnerClass() && !Property->GetOwnerClass()->HasAnyClassFlags(CLASS_PerObjectConfig)))
 		{
 			*PostConstructLinkPtr = Property;
 			PostConstructLinkPtr = &(*PostConstructLinkPtr)->PostConstructLinkNext;
diff --git a/Engine/Source/Runtime/CoreUObject/Private/UObject/Obj.cpp b/Engine/Source/Runtime/CoreUObject/Private/UObject/Obj.cpp
--- a/Engine/Source/Runtime/CoreUObject/Private/UObject/Obj.cpp
+++ b/Engine/Source/Runtime/CoreUObject/Private/UObject/Obj.cpp
@@ -2375,7 +2375,7 @@ TArray<const TCHAR*> ParsePropertyFlags(uint64 Flags)
 		TEXT("0x0000000000001000"),
 		TEXT("CPF_Transient"),
 		TEXT("CPF_Config"),
-		TEXT("CPF_Localized"),
+		TEXT("0x0000000000008000"),
 		TEXT("CPF_DisableEditOnInstance"),
 		TEXT("CPF_EditConst"),
 		TEXT("CPF_GlobalConfig"),
diff --git a/Engine/Source/Runtime/CoreUObject/Private/UObject/Property.cpp b/Engine/Source/Runtime/CoreUObject/Private/UObject/Property.cpp
--- a/Engine/Source/Runtime/CoreUObject/Private/UObject/Property.cpp
+++ b/Engine/Source/Runtime/CoreUObject/Private/UObject/Property.cpp
@@ -410,35 +410,27 @@ void UProperty::InitializeValueInternal( void* Dest ) const
  */
 bool UProperty::ValidateImportFlags( uint32 PortFlags, FOutputDevice* ErrorHandler ) const
 {
-	bool bResult = true;
-
 	// PPF_RestrictImportTypes is set when importing defaultproperties; it indicates that
 	// we should not allow config/localized properties to be imported here
-	if ( (PortFlags&PPF_RestrictImportTypes) != 0 &&
-		(PropertyFlags&(CPF_Localized|CPF_Config)) != 0 )
+	if ((PortFlags & PPF_RestrictImportTypes) && (PropertyFlags & CPF_Config))
 	{
-		FString PropertyType = (PropertyFlags&CPF_Config) != 0
-			? (PropertyFlags&CPF_Localized) != 0
-				? TEXT("config/localized")
-				: TEXT("config")
-			: TEXT("localized");
-
+		FString PropertyType = (PropertyFlags & CPF_Config) ? TEXT("config") : TEXT("localized");
 
 		FString ErrorMsg = FString::Printf(TEXT("Import failed for '%s': property is %s (Check to see if the property is listed in the DefaultProperties.  It should only be listed in the specific .ini/.int file)"), *GetName(), *PropertyType);
 
-		if( ErrorHandler != NULL )
+		if (ErrorHandler)
 		{
-			ErrorHandler->Logf( *ErrorMsg );
+			ErrorHandler->Logf(*ErrorMsg);
 		}
 		else
 		{
-			UE_LOG(LogProperty, Warning, TEXT("%s"), *ErrorMsg );
+			UE_LOG(LogProperty, Warning, TEXT("%s"), *ErrorMsg);
 		}
 
-		bResult = false;
+		return false;
 	}
 
-	return bResult;
+	return true;
 }
 
 FString UProperty::GetNameCPP() const
@@ -577,10 +569,6 @@ bool UProperty::ExportText_Direct
 	UObject*	ExportRootScope
 	) const
 {
-	if ( (PortFlags&PPF_LocalizedOnly) != 0 && !IsLocalized() )
-	{
-		return false;
-	}
 	if( Data==Delta || !Identical(Data,Delta,PortFlags) )
 	{
 		ExportTextItem
@@ -628,14 +616,6 @@ bool UProperty::NetSerializeItem( FArchive& Ar, UPackageMap* Map, void* Data, TA
 	return 1;
 }
 
-//
-// Return whether the property should be exported to localization files.
-//
-bool UProperty::IsLocalized() const
-{
-	return (PropertyFlags & CPF_Localized) != 0;
-}
-
 //
 // Return whether the property should be exported.
 //
@@ -781,7 +761,7 @@ static const int32 ReadArrayIndex(UStruct* ObjectStruct, const TCHAR*& Str, FOut
  * This normally only happens for empty strings or empty dynamic arrays, and the alternative
  * is for strings and dynamic arrays to always export blank delimiters, such as Array=() or String="", 
  * but this tends to cause problems with inherited property values being overwritten, especially in the localization 
- * import/export code, or when a property has both CPF_Localized & CPF_Config flags
+ * import/export code
 
  * The safest way is to interpret blank delimiters as an indication that the current value should be overwritten with an empty
  * value, while the lack of any value or delimiter as an indication to not import this property, thereby preventing any current
diff --git a/Engine/Source/Runtime/CoreUObject/Private/UObject/PropertyArray.cpp b/Engine/Source/Runtime/CoreUObject/Private/UObject/PropertyArray.cpp
--- a/Engine/Source/Runtime/CoreUObject/Private/UObject/PropertyArray.cpp
+++ b/Engine/Source/Runtime/CoreUObject/Private/UObject/PropertyArray.cpp
@@ -176,12 +176,6 @@ void UArrayProperty::ExportTextItem( FString& ValueStr, const void* PropertyValu
 		uint8* PropDefault = ( StructProperty != NULL ) ? StructDefaults :
 			( ( DefaultValue && DefaultArrayHelper.Num() > i ) ? DefaultArrayHelper.GetRawPtr(i) : NULL );
 
-		// Do not re-export duplicate data from superclass when exporting to .int file
-		if ( (PortFlags & PPF_LocalizedOnly) != 0 && Inner->Identical(PropData, PropDefault) )
-		{
-			continue;
-		}
-
 		Inner->ExportTextItem( ValueStr, PropData, PropDefault, Parent, PortFlags|PPF_Delimited, ExportRootScope );
 	}
 
@@ -215,11 +209,7 @@ const TCHAR* UArrayProperty::ImportText_Internal( const TCHAR* Buffer, void* Dat
 		return NULL;
 	}
 
-	// only clear the array if we're not importing localized text
-	if ( (PortFlags&PPF_LocalizedOnly) == 0 )
-	{
-		ArrayHelper.EmptyValues();
-	}
+	ArrayHelper.EmptyValues();
 
 	SkipWhitespace(Buffer);
 
@@ -324,10 +314,6 @@ void UArrayProperty::DestroyValueInternal( void* Dest ) const
 	//@todo UE4 potential double destroy later from this...would be ok for a script array, but still
 	((FScriptArray*)Dest)->~FScriptArray();
 }
-bool UArrayProperty::IsLocalized() const
-{
-	return Inner->IsLocalized() ? true : Super::IsLocalized();
-}
 bool UArrayProperty::PassCPPArgsByRef() const
 {
 	return true;
diff --git a/Engine/Source/Runtime/CoreUObject/Private/UObject/PropertyStruct.cpp b/Engine/Source/Runtime/CoreUObject/Private/UObject/PropertyStruct.cpp
--- a/Engine/Source/Runtime/CoreUObject/Private/UObject/PropertyStruct.cpp
+++ b/Engine/Source/Runtime/CoreUObject/Private/UObject/PropertyStruct.cpp
@@ -470,30 +470,6 @@ void UStructProperty::InstanceSubobjects( void* Data, void const* DefaultData, U
 	}
 }
 
-bool UStructProperty::IsLocalized() const
-{
-	// prevent recursion in the case of structs containing dynamic arrays of themselves
-	static TArray<const UStructProperty*> EncounteredStructProps;
-	if (EncounteredStructProps.Contains(this))
-	{
-		return Super::IsLocalized();
-	}
-	else
-	{
-		EncounteredStructProps.Add(this);
-		for ( TFieldIterator<UProperty> It(Struct); It; ++It )
-		{
-			if ( It->IsLocalized() )
-			{
-				EncounteredStructProps.RemoveSingleSwap(this);
-				return true;
-			}
-		}
-		EncounteredStructProps.RemoveSingleSwap(this);
-		return Super::IsLocalized();
-	}
-}
-
 bool UStructProperty::SameType(const UProperty* Other) const
 {
 	return Super::SameType(Other) && (Struct == ((UStructProperty*)Other)->Struct);
diff --git a/Engine/Source/Runtime/CoreUObject/Private/UObject/UObjectGlobals.cpp b/Engine/Source/Runtime/CoreUObject/Private/UObject/UObjectGlobals.cpp
--- a/Engine/Source/Runtime/CoreUObject/Private/UObject/UObjectGlobals.cpp
+++ b/Engine/Source/Runtime/CoreUObject/Private/UObject/UObjectGlobals.cpp
@@ -2156,7 +2156,7 @@ void FObjectInitializer::InitProperties(UObject* Obj, UClass* DefaultsClass, UOb
 			}
 			else
 			{
-				// Copy all properties that require additional initialization (CPF_Config, CPF_Localized).
+				// Copy all properties that require additional initialization (e.g. CPF_Config).
 				for (UProperty* P = Class->PostConstructLink; P; P = P->PostConstructLinkNext)
 				{
 					P->CopyCompleteValue_InContainer(Obj, DefaultData);
diff --git a/Engine/Source/Runtime/CoreUObject/Public/UObject/ObjectBase.h b/Engine/Source/Runtime/CoreUObject/Public/UObject/ObjectBase.h
--- a/Engine/Source/Runtime/CoreUObject/Public/UObject/ObjectBase.h
+++ b/Engine/Source/Runtime/CoreUObject/Public/UObject/ObjectBase.h
@@ -120,8 +120,8 @@ enum EClassFlags
 	CLASS_Transient			  = 0x00000008,
 	/** Successfully parsed. */
 	CLASS_Parsed              = 0x00000010,
-	/** Class contains localized text. */
-	CLASS_Localized           = 0x00000020,
+	/** */
+	//CLASS_                  = 0x00000020,
 	/** All the properties on the class are shown in the advanced section (which is hidden by default) unless SimpleDisplay is specified on the property */
 	CLASS_AdvancedDisplay	  = 0x00000040,
 	/** Class is a native class - native interfaces will have CLASS_Native set, but not RF_Native */
@@ -188,7 +188,7 @@ enum EClassFlags
 
 	/** @name Flags to inherit from base class */
 	//@{
-	CLASS_Inherit           = CLASS_Transient | CLASS_DefaultConfig | CLASS_Config | CLASS_Localized | CLASS_PerObjectConfig | CLASS_ConfigDoNotCheckDefaults | CLASS_NotPlaceable
+	CLASS_Inherit           = CLASS_Transient | CLASS_DefaultConfig | CLASS_Config | CLASS_PerObjectConfig | CLASS_ConfigDoNotCheckDefaults | CLASS_NotPlaceable
 							| CLASS_PointersDefaultToAutoWeak | CLASS_PointersDefaultToWeak | CLASS_Const
 							| CLASS_HasInstancedReference | CLASS_Deprecated | CLASS_DefaultToInstanced | CLASS_GlobalUserConfig,
 
@@ -209,7 +209,6 @@ enum EClassFlags
 		CLASS_GlobalUserConfig |
 		CLASS_Config |
 		CLASS_Transient |
-		CLASS_Localized |
 		CLASS_Native |
 		CLASS_NotPlaceable |
 		CLASS_PerObjectConfig |
@@ -309,7 +308,7 @@ typedef uint64 EClassCastFlags;
 //#define CPF_      						DECLARE_UINT64(0x0000000000001000)		// 
 #define CPF_Transient   					DECLARE_UINT64(0x0000000000002000)		// Property is transient: shouldn't be saved, zero-filled at load time.
 #define CPF_Config      					DECLARE_UINT64(0x0000000000004000)		// Property should be loaded/saved as permanent profile.
-#define CPF_Localized   					DECLARE_UINT64(0x0000000000008000)		// Property should be loaded as localizable text.
+//#define CPF_								DECLARE_UINT64(0x0000000000008000)		// 
 #define CPF_DisableEditOnInstance			DECLARE_UINT64(0x0000000000010000)		// Disable editing on an instance of this class
 #define CPF_EditConst   					DECLARE_UINT64(0x0000000000020000)		// Property is uneditable in the editor.
 #define CPF_GlobalConfig					DECLARE_UINT64(0x0000000000040000)		// Load config from base class, not subclass.
@@ -352,7 +351,7 @@ typedef uint64 EClassCastFlags;
 /** @name Combinations flags */
 //@{
 #define CPF_ParmFlags				(CPF_Parm | CPF_OutParm | CPF_ReturnParm | CPF_ReferenceParm | CPF_ConstParm)
-#define CPF_PropagateToArrayInner	(CPF_ExportObject | CPF_PersistentInstance | CPF_InstancedReference | CPF_ContainsInstancedReference | CPF_Localized | CPF_Config | CPF_EditConst | CPF_Deprecated | CPF_EditorOnly | CPF_AutoWeak )
+#define CPF_PropagateToArrayInner	(CPF_ExportObject | CPF_PersistentInstance | CPF_InstancedReference | CPF_ContainsInstancedReference | CPF_Config | CPF_EditConst | CPF_Deprecated | CPF_EditorOnly | CPF_AutoWeak )
 
 /** the flags that should never be set on interface properties */
 #define CPF_InterfaceClearMask		(CPF_ExportObject|CPF_InstancedReference|CPF_ContainsInstancedReference)
diff --git a/Engine/Source/Runtime/CoreUObject/Public/UObject/UnrealType.h b/Engine/Source/Runtime/CoreUObject/Public/UObject/UnrealType.h
--- a/Engine/Source/Runtime/CoreUObject/Public/UObject/UnrealType.h
+++ b/Engine/Source/Runtime/CoreUObject/Public/UObject/UnrealType.h
@@ -35,11 +35,9 @@ enum EPropertyPortFlags
 	
 	/** Indicates that importing values for config or localized properties is disallowed */
 	PPF_RestrictImportTypes			= 0x00000020,
-	
+
 	//								= 0x00000040,
-	
-	/** only include properties that have the CPF_Localized flag for ImportText/ExportText */
-	PPF_LocalizedOnly				= 0x00000080,
+	//								= 0x00000080,
 
 	/** only include properties which are marked CPF_InstancedReference */
 	PPF_SubobjectsOnly				= 0x00000100,
@@ -676,7 +674,6 @@ public:
 	bool ValidateImportFlags( uint32 PortFlags, FOutputDevice* ErrorText = NULL ) const;
 	bool ShouldPort( uint32 PortFlags=0 ) const;
 	virtual FName GetID() const;
-	virtual bool IsLocalized() const;
 
 	/**
 	 * Creates new copies of components
@@ -2185,7 +2182,6 @@ public:
 	virtual void CopyValuesInternal( void* Dest, void const* Src, int32 Count  ) const override;
 	virtual void ClearValueInternal( void* Data ) const override;
 	virtual void DestroyValueInternal( void* Dest ) const override;
-	virtual bool IsLocalized() const override;
 	virtual bool PassCPPArgsByRef() const override;
 	virtual void InstanceSubobjects( void* Data, void const* DefaultData, UObject* Owner, struct FObjectInstancingGraph* InstanceGraph ) override;
 	virtual bool ContainsObjectReference() const override;
@@ -2552,7 +2548,6 @@ public:
 	virtual void ClearValueInternal( void* Data ) const override;
 	virtual void DestroyValueInternal( void* Dest ) const override;
 	virtual void InitializeValueInternal( void* Dest ) const override;
-	virtual bool IsLocalized() const override;
 	virtual void InstanceSubobjects( void* Data, void const* DefaultData, UObject* Owner, struct FObjectInstancingGraph* InstanceGraph ) override;
 	virtual int32 GetMinAlignment() const override;
 	virtual bool ContainsObjectReference() const override;
diff --git a/Engine/Source/Runtime/Engine/Classes/Commandlets/Commandlet.h b/Engine/Source/Runtime/Engine/Classes/Commandlets/Commandlet.h
--- a/Engine/Source/Runtime/Engine/Classes/Commandlets/Commandlet.h
+++ b/Engine/Source/Runtime/Engine/Classes/Commandlets/Commandlet.h
@@ -35,23 +35,23 @@ class UCommandlet : public UObject
 	GENERATED_UCLASS_BODY()
 
 	/** Description of the commandlet's purpose */
-	UPROPERTY(localized)
+	UPROPERTY()
 	FString HelpDescription;
 
 	/** Usage template to show for "ucc help" */
-	UPROPERTY(localized)
+	UPROPERTY()
 	FString HelpUsage;
 
 	/** Hyperlink for more info */
-	UPROPERTY(localized)
+	UPROPERTY()
 	FString HelpWebLink;
 
 	/** The name of the parameter the commandlet takes */
-	UPROPERTY(localized)
+	UPROPERTY()
 	TArray<FString> HelpParamNames;
 
 	/** The description of the parameter */
-	UPROPERTY(localized)
+	UPROPERTY()
 	TArray<FString> HelpParamDescriptions;
 
 	/**
diff --git a/Engine/Source/Runtime/Engine/Classes/GameFramework/EngineMessage.h b/Engine/Source/Runtime/Engine/Classes/GameFramework/EngineMessage.h
--- a/Engine/Source/Runtime/Engine/Classes/GameFramework/EngineMessage.h
+++ b/Engine/Source/Runtime/Engine/Classes/GameFramework/EngineMessage.h
@@ -10,35 +10,35 @@ class  UEngineMessage : public ULocalMessage
 	GENERATED_UCLASS_BODY()
 
 	/** Message displayed in message dialog when player pawn fails to spawn because no playerstart was available. */
-	UPROPERTY(localized)
+	UPROPERTY()
 	FString FailedPlaceMessage;
 
 	/** Message when player join attempt is refused because the server is at capacity. */
-	UPROPERTY(localized)
+	UPROPERTY()
 	FString MaxedOutMessage;
 
 	/** Message when a new player enters the game. */
-	UPROPERTY(localized)
+	UPROPERTY()
 	FString EnteredMessage;
 
 	/** Message when a player leaves the game. */
-	UPROPERTY(localized)
+	UPROPERTY()
 	FString LeftMessage;
 
 	/** Message when a player changes his name. */
-	UPROPERTY(localized)
+	UPROPERTY()
 	FString GlobalNameChange;
 
 	/** Message when a new spectator enters the server (if spectator has a player name). */
-	UPROPERTY(localized)
+	UPROPERTY()
 	FString SpecEnteredMessage;
 
 	/** Message when a new player enters the server (if player is unnamed). */
-	UPROPERTY(localized)
+	UPROPERTY()
 	FString NewPlayerMessage;
 
 	/** Message when a new spectator enters the server (if spectator is unnamed). */
-	UPROPERTY(localized)
+	UPROPERTY()
 	FString NewSpecMessage;
 
 	// Begin ULocalMessage interface
diff --git a/Engine/Source/Runtime/Engine/Classes/GameFramework/GameMode.h b/Engine/Source/Runtime/Engine/Classes/GameFramework/GameMode.h
--- a/Engine/Source/Runtime/Engine/Classes/GameFramework/GameMode.h
+++ b/Engine/Source/Runtime/Engine/Classes/GameFramework/GameMode.h
@@ -199,7 +199,7 @@ public:
 	int32 CurrentID;    
 
 	/** The default player name assigned to players that join with no name specified. */
-	UPROPERTY(localized)
+	UPROPERTY()
 	FString DefaultPlayerName;
 
 	/** Array of available player starts. */
diff --git a/Engine/Source/Runtime/Engine/Classes/Sound/SoundWave.h b/Engine/Source/Runtime/Engine/Classes/Sound/SoundWave.h
--- a/Engine/Source/Runtime/Engine/Classes/Sound/SoundWave.h
+++ b/Engine/Source/Runtime/Engine/Classes/Sound/SoundWave.h
@@ -136,15 +136,15 @@ class ENGINE_API USoundWave : public USoundBase
 	uint32 bDynamicResource:1;
 
 	/** If set to true if this sound is considered to contain mature/adult content. */
-	UPROPERTY(EditAnywhere, localized, Category=Subtitles, AssetRegistrySearchable)
+	UPROPERTY(EditAnywhere, Category=Subtitles, AssetRegistrySearchable)
 	uint32 bMature:1;
 
 	/** If set to true will disable automatic generation of line breaks - use if the subtitles have been split manually. */
-	UPROPERTY(EditAnywhere, localized, Category=Subtitles )
+	UPROPERTY(EditAnywhere, Category=Subtitles )
 	uint32 bManualWordWrap:1;
 
 	/** If set to true the subtitles display as a sequence of single lines as opposed to multiline. */
-	UPROPERTY(EditAnywhere, localized, Category=Subtitles )
+	UPROPERTY(EditAnywhere, Category=Subtitles )
 	uint32 bSingleLine:1;
 
 	/** Whether this SoundWave was decompressed from OGG. */
@@ -193,12 +193,12 @@ class ENGINE_API USoundWave : public USoundBase
 	 * Subtitle cues.  If empty, use SpokenText as the subtitle.  Will often be empty,
 	 * as the contents of the subtitle is commonly identical to what is spoken.
 	 */
-	UPROPERTY(EditAnywhere, localized, Category=Subtitles)
+	UPROPERTY(EditAnywhere, Category=Subtitles)
 	TArray<struct FSubtitleCue> Subtitles;
 
 #if WITH_EDITORONLY_DATA
 	/** Provides contextual information for the sound to the translator. */
-	UPROPERTY(EditAnywhere, localized, Category=Subtitles )
+	UPROPERTY(EditAnywhere, Category=Subtitles )
 	FString Comment;
 
 #endif // WITH_EDITORONLY_DATA
