Exposed a pair of functions regarding configuration files associated with a UObject.

diff --git a/Engine/Source/Runtime/CoreUObject/Private/UObject/Obj.cpp b/Engine/Source/Runtime/CoreUObject/Private/UObject/Obj.cpp
--- a/Engine/Source/Runtime/CoreUObject/Private/UObject/Obj.cpp
+++ b/Engine/Source/Runtime/CoreUObject/Private/UObject/Obj.cpp
@@ -1093,7 +1093,7 @@ bool UObject::CheckDefaultSubobjectsInternal()
 /**
  * Determines whether the specified object should load values using PerObjectConfig rules
  */
-static bool UsesPerObjectConfig( UObject* SourceObject )
+bool UsesPerObjectConfig( UObject* SourceObject )
 {
 	checkSlow(SourceObject);
 	return (SourceObject->GetClass()->HasAnyClassFlags(CLASS_PerObjectConfig) && !SourceObject->HasAnyFlags(RF_ClassDefaultObject));
@@ -1102,7 +1102,7 @@ static bool UsesPerObjectConfig( UObject* SourceObject )
 /**
  * Returns the file to load ini values from for the specified object, taking into account PerObjectConfig-ness
  */
-static FString GetConfigFilename( UObject* SourceObject )
+FString GetConfigFilename( UObject* SourceObject )
 {
 	checkSlow(SourceObject);
 
diff --git a/Engine/Source/Runtime/CoreUObject/Public/UObject/UObjectGlobals.h b/Engine/Source/Runtime/CoreUObject/Public/UObject/UObjectGlobals.h
--- a/Engine/Source/Runtime/CoreUObject/Public/UObject/UObjectGlobals.h
+++ b/Engine/Source/Runtime/CoreUObject/Public/UObject/UObjectGlobals.h
@@ -1077,6 +1077,16 @@ T* DuplicateObject(T const* SourceObject,UObject* Outer,const TCHAR* Name = TEXT
 	return NULL;
 }
 
+/**
+ * Determines whether the specified object should load values using PerObjectConfig rules
+ */
+COREUOBJECT_API bool UsesPerObjectConfig( UObject* SourceObject );
+
+/**
+ * Returns the file to load ini values from for the specified object, taking into account PerObjectConfig-ness
+ */
+COREUOBJECT_API FString GetConfigFilename( UObject* SourceObject );
+
 /*----------------------------------------------------------------------------
 	Core templates.
 ----------------------------------------------------------------------------*/
