Making sure LoadProgressScope in ULinkerLoad gets deleted in case there's been an error reading the package.

Fixes a crash when loading a package with invalid package version.

UE-7659

diff --git a/Engine/Source/Runtime/CoreUObject/Private/UObject/LinkerLoad.cpp b/Engine/Source/Runtime/CoreUObject/Private/UObject/LinkerLoad.cpp
--- a/Engine/Source/Runtime/CoreUObject/Private/UObject/LinkerLoad.cpp
+++ b/Engine/Source/Runtime/CoreUObject/Private/UObject/LinkerLoad.cpp
@@ -678,18 +678,43 @@ bool ULinkerLoad::IsTimeLimitExceeded( const TCHAR* CurrentTask, int32 Granulari
 	return bTimeLimitExceeded;
 }
 
+#if WITH_EDITOR
+/** Makes sure LoadProgressScope gets deleted in case of an error */
+class FLoadProgressAutoDelete
+{
+	FScopedSlowTask*& LoadProgressScope;
+	bool bDoNotDelete;
+public:
+	FLoadProgressAutoDelete(FScopedSlowTask*& InLoadProgressScope)
+		: LoadProgressScope(InLoadProgressScope)
+		, bDoNotDelete(false)
+	{}
+	~FLoadProgressAutoDelete()
+	{
+		if (!bDoNotDelete)
+		{
+			delete LoadProgressScope;
+			LoadProgressScope = nullptr;
+		}
+	}
+	void DoNotDeleteLoadProgress()
+	{
+		bDoNotDelete = true;
+	}
+};
+#endif
+
 /**
  * Creates loader used to serialize content.
  */
 ULinkerLoad::ELinkerStatus ULinkerLoad::CreateLoader()
 {
 #if WITH_EDITOR
-
 	if (!LoadProgressScope)
 	{
 		LoadProgressScope = new FScopedSlowTask(ULinkerDefs::TotalProgressSteps, NSLOCTEXT("Core", "GenericLoading", "Loading..."), ShouldReportProgress());
 	}
-
+	FLoadProgressAutoDelete LoadProgressAutoDelete(LoadProgressScope);
 #endif
 
 	CreateActiveRedirectsMap( GEngineIni );
@@ -816,6 +841,9 @@ ULinkerLoad::ELinkerStatus ULinkerLoad::CreateLoader()
 		bExecuteNextStep = Loader->Precache( 0, PrecacheSize);
 	}
 
+#if WITH_EDITOR
+	LoadProgressAutoDelete.DoNotDeleteLoadProgress();
+#endif
 	return (bExecuteNextStep && !IsTimeLimitExceeded( TEXT("creating loader") )) ? LINKER_Loaded : LINKER_TimedOut;
 }
 
@@ -824,6 +852,10 @@ ULinkerLoad::ELinkerStatus ULinkerLoad::CreateLoader()
  */
 ULinkerLoad::ELinkerStatus ULinkerLoad::SerializePackageFileSummary()
 {
+#if WITH_EDITOR
+	FLoadProgressAutoDelete LoadProgressAutoDelete(LoadProgressScope);
+#endif
+
 	if( bHasSerializedPackageFileSummary == false )
 	{
 #if WITH_EDITOR
@@ -1018,6 +1050,9 @@ ULinkerLoad::ELinkerStatus ULinkerLoad::SerializePackageFileSummary()
 		bHasSerializedPackageFileSummary = true;
 	}
 
+#if WITH_EDITOR
+	LoadProgressAutoDelete.DoNotDeleteLoadProgress();
+#endif
 	return !IsTimeLimitExceeded( TEXT("serializing package file summary") ) ? LINKER_Loaded : LINKER_TimedOut;
 }
 
