Nicer fix for UE-7659

diff --git a/Engine/Source/Runtime/CoreUObject/Private/UObject/LinkerLoad.cpp b/Engine/Source/Runtime/CoreUObject/Private/UObject/LinkerLoad.cpp
--- a/Engine/Source/Runtime/CoreUObject/Private/UObject/LinkerLoad.cpp
+++ b/Engine/Source/Runtime/CoreUObject/Private/UObject/LinkerLoad.cpp
@@ -610,6 +610,14 @@ ULinkerLoad::ELinkerStatus ULinkerLoad::Tick( float InTimeLimit, bool bInUseTime
 			);
 	}
 
+#if WITH_EDITOR
+	if (Status == LINKER_Failed)
+	{
+		delete LoadProgressScope;
+		LoadProgressScope = nullptr;
+	}
+#endif
+
 	// Return whether we completed or not.
 	return Status;
 }
@@ -678,43 +686,18 @@ bool ULinkerLoad::IsTimeLimitExceeded( const TCHAR* CurrentTask, int32 Granulari
 	return bTimeLimitExceeded;
 }
 
-#if WITH_EDITOR
-/** Makes sure LoadProgressScope gets deleted in case of an error */
-class FLoadProgressAutoDelete
-{
-	FScopedSlowTask*& LoadProgressScope;
-	bool bDoNotDelete;
-public:
-	FLoadProgressAutoDelete(FScopedSlowTask*& InLoadProgressScope)
-		: LoadProgressScope(InLoadProgressScope)
-		, bDoNotDelete(false)
-	{}
-	~FLoadProgressAutoDelete()
-	{
-		if (!bDoNotDelete)
-		{
-			delete LoadProgressScope;
-			LoadProgressScope = nullptr;
-		}
-	}
-	void DoNotDeleteLoadProgress()
-	{
-		bDoNotDelete = true;
-	}
-};
-#endif
-
 /**
  * Creates loader used to serialize content.
  */
 ULinkerLoad::ELinkerStatus ULinkerLoad::CreateLoader()
 {
 #if WITH_EDITOR
+
 	if (!LoadProgressScope)
 	{
 		LoadProgressScope = new FScopedSlowTask(ULinkerDefs::TotalProgressSteps, NSLOCTEXT("Core", "GenericLoading", "Loading..."), ShouldReportProgress());
 	}
-	FLoadProgressAutoDelete LoadProgressAutoDelete(LoadProgressScope);
+
 #endif
 
 	CreateActiveRedirectsMap( GEngineIni );
@@ -841,9 +824,6 @@ ULinkerLoad::ELinkerStatus ULinkerLoad::CreateLoader()
 		bExecuteNextStep = Loader->Precache( 0, PrecacheSize);
 	}
 
-#if WITH_EDITOR
-	LoadProgressAutoDelete.DoNotDeleteLoadProgress();
-#endif
 	return (bExecuteNextStep && !IsTimeLimitExceeded( TEXT("creating loader") )) ? LINKER_Loaded : LINKER_TimedOut;
 }
 
@@ -852,10 +832,6 @@ ULinkerLoad::ELinkerStatus ULinkerLoad::CreateLoader()
  */
 ULinkerLoad::ELinkerStatus ULinkerLoad::SerializePackageFileSummary()
 {
-#if WITH_EDITOR
-	FLoadProgressAutoDelete LoadProgressAutoDelete(LoadProgressScope);
-#endif
-
 	if( bHasSerializedPackageFileSummary == false )
 	{
 #if WITH_EDITOR
@@ -1050,9 +1026,6 @@ ULinkerLoad::ELinkerStatus ULinkerLoad::SerializePackageFileSummary()
 		bHasSerializedPackageFileSummary = true;
 	}
 
-#if WITH_EDITOR
-	LoadProgressAutoDelete.DoNotDeleteLoadProgress();
-#endif
 	return !IsTimeLimitExceeded( TEXT("serializing package file summary") ) ? LINKER_Loaded : LINKER_TimedOut;
 }
 
