[INTEGRATE] Change 2422752 by [REDACTED_PERSON][REDACTED_HANDLE] on 2015/01/28 15:50:25

	Write old UE3 Version out again so that older versions of the engine can identify them as incompatible without crashing.

diff --git a/Engine/Source/Runtime/CoreUObject/Private/UObject/PackageFileSummary.cpp b/Engine/Source/Runtime/CoreUObject/Private/UObject/PackageFileSummary.cpp
--- a/Engine/Source/Runtime/CoreUObject/Private/UObject/PackageFileSummary.cpp
+++ b/Engine/Source/Runtime/CoreUObject/Private/UObject/PackageFileSummary.cpp
@@ -35,19 +35,20 @@ FArchive& operator<<( FArchive& Ar, FPackageFileSummary& Sum )
 		 *
 		 * Lower 16 bits stores the UE3 engine version
 		 * Upper 16 bits stores the UE4/licensee version
-		 * For newer packages this is -4
+		 * For newer packages this is -5
 		 *		-2 indicates presence of enum-based custom versions
 		 *		-3 indicates guid-based custom versions
-		 *		-4 indicates removal of the UE3 version
+		 *		-4 indicates removal of the UE3 version. Packages saved with this ID cannot be loaded in older engine versions 
+		 *      -5 indicates the replacement of writing out the "UE3 version" so older versions of engine can gracefully fail to open newer packages
 		 */
-		int32 LegacyFileVersion = -4;
+		int32 LegacyFileVersion = -5;
 		Ar << LegacyFileVersion;
 
 		if (Ar.IsLoading())
 		{
 			if (LegacyFileVersion < 0) // means we have modern version numbers
 			{
-				if (LegacyFileVersion > -4)
+				if (LegacyFileVersion != -4)
 				{
 					int32 LegacyUE3Version = 0;
 					Ar << LegacyUE3Version;
@@ -82,14 +83,18 @@ FArchive& operator<<( FArchive& Ar, FPackageFileSummary& Sum )
 			if (Sum.bUnversioned)
 			{
 				int32 Zero = 0;
-				Ar << Zero;
-				Ar << Zero;
+				Ar << Zero; // LegacyUE3version
+				Ar << Zero; // VersionUE4
+				Ar << Zero; // VersionLicenseeUE4
 
 				FCustomVersionContainer NoCustomVersions;
 				NoCustomVersions.Serialize(Ar);
 			}
 			else
 			{
+				// Must write out the last UE3 engine version, so that older versions identify it as new
+				int32 LegacyUE3Version = 864;
+				Ar << LegacyUE3Version;
 				Ar << Sum.FileVersionUE4;
 				Ar << Sum.FileVersionLicenseeUE4;
 
