#ue4 - LinkerLoad circular dependency fixes: ensure that any parent default subobjects have been loaded before trying to create a CDO.

- Fixes components in derived blueprints not inheriting default values from their parent blueprint.

Jira UE-7987, UE-4964

diff --git a/Engine/Source/Runtime/CoreUObject/Private/UObject/LinkerLoad.cpp b/Engine/Source/Runtime/CoreUObject/Private/UObject/LinkerLoad.cpp
--- a/Engine/Source/Runtime/CoreUObject/Private/UObject/LinkerLoad.cpp
+++ b/Engine/Source/Runtime/CoreUObject/Private/UObject/LinkerLoad.cpp
@@ -3346,6 +3346,35 @@ UObject* ULinkerLoad::CreateExport( int32 Index )
 
 		FName NewName = Export.ObjectName;
 
+
+		// If we are about to create a CDO, we need to ensure that all parent sub-objects are loaded
+		// to get default value initialization to work.
+		if ((ObjectLoadFlags & RF_ClassDefaultObject) != 0)
+		{
+			UClass* SuperClass = LoadClass->GetSuperClass();
+			if (SuperClass && !SuperClass->HasAnyFlags(RF_Native))
+			{
+				UObject* SuperCDO = SuperClass->GetDefaultObject();
+				TArray<UObject*> SuperSubObjects;
+				GetObjectsWithOuter(SuperCDO, SuperSubObjects, /*bIncludeNestedObjects=*/ false, /*ExclusionFlags=*/ RF_Native);
+
+				for (UObject* SubObject : SuperSubObjects)
+				{
+					if (SubObject->HasAnyFlags(RF_NeedLoad))
+					{
+						Preload(SubObject);
+					}
+				}
+
+				// Preload may have already created this object.
+				if (Export.Object)
+				{
+					return Export.Object;
+				}
+			}
+		}
+
+
 		LoadClass->GetDefaultObject();
 
 
