Slightly more useful loading percentage reporting for async packages.

diff --git a/Engine/Source/Runtime/CoreUObject/Private/Serialization/AsyncLoading.cpp b/Engine/Source/Runtime/CoreUObject/Private/Serialization/AsyncLoading.cpp
--- a/Engine/Source/Runtime/CoreUObject/Private/Serialization/AsyncLoading.cpp
+++ b/Engine/Source/Runtime/CoreUObject/Private/Serialization/AsyncLoading.cpp
@@ -165,7 +165,19 @@ FAsyncPackage::FAsyncPackage(const FName& InPackageName, const FGuid* InPackageG
  */
 float FAsyncPackage::GetLoadPercentage() const
 {
-	return LoadPercentage;
+	float Accum = LoadPercentage;
+
+	for (auto Import : PendingImportedPackages)
+	{
+		Accum += Import->GetLoadPercentage();
+	}
+
+	for (auto Import : ReferencedImports)
+	{
+		Accum += Import->GetLoadPercentage();
+	}
+
+	return Accum / (float)(1 + PendingImportedPackages.Num() + ReferencedImports.Num());
 }
 
 /**
