Changing deprecation version of operator new for UObjects from 4.8 to 4.7

diff --git a/Engine/Source/Runtime/CoreUObject/Public/UObject/ObjectBase.h b/Engine/Source/Runtime/CoreUObject/Public/UObject/ObjectBase.h
--- a/Engine/Source/Runtime/CoreUObject/Public/UObject/ObjectBase.h
+++ b/Engine/Source/Runtime/CoreUObject/Public/UObject/ObjectBase.h
@@ -1178,7 +1178,7 @@ public: \
 	{ \
 		return TStaticCastFlags; \
 	} \
-	DEPRECATED(4.8, "operator new has been deprecated for UObjects - please use NewObject or NewNamedObject instead") \
+	DEPRECATED(4.7, "operator new has been deprecated for UObjects - please use NewObject or NewNamedObject instead") \
 	inline void* operator new( const size_t InSize, UObject* InOuter=(UObject*)GetTransientPackage(), FName InName=NAME_None, EObjectFlags InSetFlags=RF_NoFlags ) \
 	{ \
 		return StaticAllocateObject( StaticClass(), InOuter, InName, InSetFlags ); \
