GitHub 777 - More informative error message.

Integrated.

diff --git a/Engine/Source/Runtime/CoreUObject/Private/UObject/LinkerLoad.cpp b/Engine/Source/Runtime/CoreUObject/Private/UObject/LinkerLoad.cpp
--- a/Engine/Source/Runtime/CoreUObject/Private/UObject/LinkerLoad.cpp
+++ b/Engine/Source/Runtime/CoreUObject/Private/UObject/LinkerLoad.cpp
@@ -2138,7 +2138,7 @@ bool ULinkerLoad::VerifyImportInner(const int32 ImportIndex, FString& WarningSuf
 	// Find or load the linker load that contains the FObjectExport for this import
 	if (Import.OuterIndex.IsNull() && Import.ClassName!=NAME_Package )
 	{
-		UE_LOG(LogLinker, Warning, TEXT("%s has an inappropriate outermost, it was probably saved with a deprecated outer."), *Import.ObjectName.ToString());
+		UE_LOG(LogLinker, Warning, TEXT("%s has an inappropriate outermost, it was probably saved with a deprecated outer (file: %s)"), *Import.ObjectName.ToString(), *Filename);
 		Import.SourceLinker = NULL;
 		return false;
 	}
