Simplify determination of whether to call SerializeScriptProperties

diff --git a/Engine/Source/Runtime/CoreUObject/Private/UObject/Obj.cpp b/Engine/Source/Runtime/CoreUObject/Private/UObject/Obj.cpp
--- a/Engine/Source/Runtime/CoreUObject/Private/UObject/Obj.cpp
+++ b/Engine/Source/Runtime/CoreUObject/Private/UObject/Obj.cpp
@@ -828,18 +828,11 @@ void UObject::Serialize( FArchive& Ar )
 	}
 
 	// Serialize object properties which are defined in the class.
-	if( !Class->IsChildOf(UClass::StaticClass()) )
+	// Handle derived UClass objects (exact UClass objects are native only and shouldn't be touched)
+	if (Class != UClass::StaticClass())
 	{
 		SerializeScriptProperties(Ar);
 	}
-	else
-	{
-		// Handle derived UClass objects (exact UClass objects are native only and shouldn't be touched)
-		if (Class != UClass::StaticClass())
-		{
-			SerializeScriptProperties(Ar);
-		}
-	}
 
 	// Keep track of pending kill
 	if( Ar.IsTransacting() )
