Fix Linux and Mac debug builds.

diff --git a/Engine/Source/Runtime/CoreUObject/Public/UObject/UObjectGlobals.h b/Engine/Source/Runtime/CoreUObject/Public/UObject/UObjectGlobals.h
--- a/Engine/Source/Runtime/CoreUObject/Public/UObject/UObjectGlobals.h
+++ b/Engine/Source/Runtime/CoreUObject/Public/UObject/UObjectGlobals.h
@@ -1045,7 +1045,7 @@ T* NewObject(UObject* Outer = (UObject*)GetTransientPackage(), UClass* Class = T
 		FObjectInitializer::AssertIfInConstructor(Outer, TEXT("NewObject with empty name can't be used to create default subobjects (inside of UObject derived class constructor) as it produces inconsistent object names. Use ObjectInitializer.CreateDefaultSuobject<> instead."));
 	}
 	checkf(Class, TEXT("NewObject called with a nullptr class object"));
-	checkSlow(Class->IsChildOf(T::StaticClass()));
+	// this breaks Linux and Mac debug builds - checkSlow(Class->IsChildOf(T::StaticClass()));
 	return static_cast<T*>(StaticConstructObject(Class, Outer, Name, Flags, Template, bCopyTransientsFromClassDefaults, InInstanceGraph));
 }
 
