Linux/Mac debug build fix.

diff --git a/Engine/Source/Runtime/CoreUObject/Private/UObject/UObjectGlobals.cpp b/Engine/Source/Runtime/CoreUObject/Private/UObject/UObjectGlobals.cpp
--- a/Engine/Source/Runtime/CoreUObject/Private/UObject/UObjectGlobals.cpp
+++ b/Engine/Source/Runtime/CoreUObject/Private/UObject/UObjectGlobals.cpp
@@ -2215,6 +2215,11 @@ void FObjectInitializer::AssertIfSubobjectSetupIsNotAllowed(const TCHAR* Subobje
 		TEXT("%s.%s: Subobject class setup is only allowed in base class constructor call (in the initialization list)"), Obj ? *Obj->GetFullName() : TEXT("NULL"), SubobjectName);
 }
 
+bool DebugIsClassChildOf_Internal(UClass* Parent, UClass* Child)
+{
+	return Child->IsChildOf(Parent);
+}
+
 UObject* StaticConstructObject
 (
 	UClass*			InClass,
diff --git a/Engine/Source/Runtime/CoreUObject/Public/UObject/UObjectGlobals.h b/Engine/Source/Runtime/CoreUObject/Public/UObject/UObjectGlobals.h
--- a/Engine/Source/Runtime/CoreUObject/Public/UObject/UObjectGlobals.h
+++ b/Engine/Source/Runtime/CoreUObject/Public/UObject/UObjectGlobals.h
@@ -1031,6 +1031,17 @@ template< class T >
 DEPRECATED(4.8, "ConstructObject is deprecated. Use NewObject instead")
 T* ConstructObject(UClass* Class, UObject* Outer = (UObject*)GetTransientPackage(), FName Name=NAME_None, EObjectFlags SetFlags=RF_NoFlags, UObject* Template=NULL, bool bCopyTransientsFromClassDefaults=false, struct FObjectInstancingGraph* InstanceGraph=NULL );
 
+/*
+ * Check if class is child of another class.
+ * 
+ * @param Parent	Parent class
+ * @param Child		Child class
+ * 
+ * @return True if Child is a child of Parent, false otherwise.
+ *
+ **/
+bool DebugIsClassChildOf_Internal(UClass* Parent, UClass* Child);
+
 /**
  * Convenience template for constructing a gameplay object
  *
@@ -1045,7 +1056,7 @@ T* NewObject(UObject* Outer = (UObject*)GetTransientPackage(), UClass* Class = T
 		FObjectInitializer::AssertIfInConstructor(Outer, TEXT("NewObject with empty name can't be used to create default subobjects (inside of UObject derived class constructor) as it produces inconsistent object names. Use ObjectInitializer.CreateDefaultSuobject<> instead."));
 	}
 	checkf(Class, TEXT("NewObject called with a nullptr class object"));
-	// this breaks Linux and Mac debug builds - checkSlow(Class->IsChildOf(T::StaticClass()));
+	checkSlow(DebugIsClassChildOf_Internal(T::StaticClass(), Class));
 	return static_cast<T*>(StaticConstructObject(Class, Outer, Name, Flags, Template, bCopyTransientsFromClassDefaults, InInstanceGraph));
 }
 
