Core: Fixed linker errors in OnlineSubsystem
Merging using UE4-To-UE4-LauncherDev

DebugIsClassChildOf must be exported by CoreUObject, because it is used in the new NewObject template function.

diff --git a/Engine/Source/Runtime/CoreUObject/Public/UObject/UObjectGlobals.h b/Engine/Source/Runtime/CoreUObject/Public/UObject/UObjectGlobals.h
--- a/Engine/Source/Runtime/CoreUObject/Public/UObject/UObjectGlobals.h
+++ b/Engine/Source/Runtime/CoreUObject/Public/UObject/UObjectGlobals.h
@@ -1040,7 +1040,7 @@ T* ConstructObject(UClass* Class, UObject* Outer = (UObject*)GetTransientPackage
  * @return True if Child is a child of Parent, false otherwise.
  *
  **/
-bool DebugIsClassChildOf_Internal(UClass* Parent, UClass* Child);
+COREUOBJECT_API bool DebugIsClassChildOf_Internal(UClass* Parent, UClass* Child);
 
 /**
  * Convenience template for constructing a gameplay object
