Fix spelling error

diff --git a/Engine/Source/Runtime/CoreUObject/Private/UObject/UObjectGlobals.cpp b/Engine/Source/Runtime/CoreUObject/Private/UObject/UObjectGlobals.cpp
--- a/Engine/Source/Runtime/CoreUObject/Private/UObject/UObjectGlobals.cpp
+++ b/Engine/Source/Runtime/CoreUObject/Private/UObject/UObjectGlobals.cpp
@@ -910,7 +910,7 @@ UPackage* LoadPackageInternal(UPackage* InOuter, const TCHAR* InLongPackageName,
 		{
 			// Make sure we pass the property that's currently being serialized by the linker that owns the import 
 			// that triggered this LoadPackage call
-			auto OldSerialziedProperty = Linker->GetSerializedProperty();
+			UProperty* OldSerializedProperty = Linker->GetSerializedProperty();
 			if (ImportLinker)
 			{
 				Linker->SetSerializedProperty(ImportLinker->GetSerializedProperty());
@@ -918,7 +918,7 @@ UPackage* LoadPackageInternal(UPackage* InOuter, const TCHAR* InLongPackageName,
 
 			Linker->LoadAllObjects();
 
-			Linker->SetSerializedProperty(OldSerialziedProperty);
+			Linker->SetSerializedProperty(OldSerializedProperty);
 		}
 
 		SlowTask.EnterProgressFrame(30);
