UE-7688: Adding a UFUNCTION macro to an existing function can crash the Editor.

It was actually a bug in hot-reloaded exposed when someone was trying to get static class of hot-reloaded type, which was happening for vtable ptr retrieval as it needs to create an object. Now first all classes are registered and then vtable ptrs are updated.

diff --git a/Engine/Source/Runtime/CoreUObject/Private/UObject/Class.cpp b/Engine/Source/Runtime/CoreUObject/Private/UObject/Class.cpp
--- a/Engine/Source/Runtime/CoreUObject/Private/UObject/Class.cpp
+++ b/Engine/Source/Runtime/CoreUObject/Private/UObject/Class.cpp
@@ -3371,10 +3371,6 @@ bool UClass::HotReloadPrivateStaticClass(
 		UClass::GetDefaultPropertiesFeedbackContext().Logf(ELogVerbosity::Warning, TEXT("Property size mismatch. Will not update class %s (was %d, new %d)."), *GetName(), PropertiesSize, InSize);
 		return false;
 	}
-	//We could do this later, but might as well get it before we start corrupting the object
-	UObject* CDO = GetDefaultObject();
-	void* OldVTable = *(void**)CDO;
-
 
 	//@todo safe? ClassFlags = InClassFlags | CLASS_Native;
 	//@todo safe? ClassCastFlags = InClassCastFlags;
@@ -3395,51 +3391,23 @@ bool UClass::HotReloadPrivateStaticClass(
 	ClassWithin = TClass_WithinClass_StaticClass;
 	*/
 
-	UE_LOG(LogClass, Verbose, TEXT("Attempting to change VTable for class %s."),*GetName());
-	ClassWithin = UPackage::StaticClass();  // We are just avoiding error checks with this...we don't care about this temp object other than to get the vtable.
-	UObject* TempObjectForVTable = StaticConstructObject(this, GetTransientPackage(), NAME_None, RF_NeedLoad | RF_ClassDefaultObject);
-
-	if( !TempObjectForVTable->IsRooted() )
-	{
-		TempObjectForVTable->MarkPendingKill();
-	}
-	else
-	{
-		UE_LOG(LogClass, Warning, TEXT("Hot Reload:  Was not expecting temporary object '%s' for class '%s' to become rooted during construction.  This object cannot be marked pending kill." ), *TempObjectForVTable->GetFName().ToString(), *this->GetName() );
-	}
-
-	ClassWithin = TClass_WithinClass_StaticClass;
-
-	void* NewVTable = *(void**)TempObjectForVTable;
-	if (NewVTable != OldVTable)
+	int32 CountClass = 0;
+	for (FRawObjectIterator It; It; ++It)
 	{
-		int32 Count = 0;
-		int32 CountClass = 0;
-		for ( FRawObjectIterator It; It; ++It )
+		UObject* Target = *It;
+		if (dynamic_cast<UClass*>(Target))
 		{
-			UObject* Target = *It;
-			if (OldVTable == *(void**)Target)
-			{
-				*(void**)Target = NewVTable;
-				Count++;
-			}
-			else if (dynamic_cast<UClass*>(Target))
+			UClass *Class = CastChecked<UClass>(Target);
+			if (Class->ClassConstructor == OldClassConstructor)
 			{
-				UClass *Class = CastChecked<UClass>(Target);
-				if (Class->ClassConstructor == OldClassConstructor)
-				{
-					Class->ClassConstructor = ClassConstructor;
-					Class->ClassAddReferencedObjects = ClassAddReferencedObjects;
-					CountClass++;
-				}
+				Class->ClassConstructor = ClassConstructor;
+				Class->ClassAddReferencedObjects = ClassAddReferencedObjects;
+				CountClass++;
 			}
 		}
-		UE_LOG(LogClass, Verbose, TEXT("Updated the vtable for %d live objects and %d blueprint classes.  %016llx -> %016llx"), Count, CountClass, PTRINT(OldVTable), PTRINT(NewVTable));
-	}
-	else
-	{
-		UE_LOG(LogClass, Error, TEXT("VTable for class %s did not change?"),*GetName());
 	}
+	UE_LOG(LogClass, Verbose, TEXT("Updated the internal methods %d blueprint classes."), CountClass);
+
 	return true;
 }
 
diff --git a/Engine/Source/Runtime/CoreUObject/Private/UObject/UObjectBase.cpp b/Engine/Source/Runtime/CoreUObject/Private/UObject/UObjectBase.cpp
--- a/Engine/Source/Runtime/CoreUObject/Private/UObject/UObjectBase.cpp
+++ b/Engine/Source/Runtime/CoreUObject/Private/UObject/UObjectBase.cpp
@@ -621,6 +621,114 @@ FString RemoveClassPrefix(const TCHAR* ClassName)
 	return NameWithoutPrefix;
 }
 
+/**
+ * Helper class for updating VTables of hot reloaded classes.
+ */
+class FUClassVTableHelper
+{
+public:
+	/**
+	 * Adds UClasses VTable address if it's not cached yet.
+	 *
+	 * @param Class A UClass for which to add the VTable address.
+	 */
+	void TouchUClassVTable(UClass* Class)
+	{
+		auto* Found = VTableMap.Find(Class);
+
+		if (Found != nullptr)
+		{
+			return;
+		}
+
+		VTableMap.Add(Class, GetUClassVTable(Class));
+	}
+
+	/**
+	 * Changes old VTable address to the new one after hot reload.
+	 *
+	 * @param OldClass Old class.
+	 * @param NewClass New class.
+	 */
+	void UpdateVTables(UClass* OldClass, UClass* NewClass)
+	{
+		UE_LOG(LogClass, Verbose, TEXT("Attempting to change VTable for class %s."), *OldClass->GetName());
+
+		// We could do this later, but might as well get it before we start corrupting the object
+		void* OldVTable = VTableMap[OldClass];
+		void* NewVTable = GetUClassVTable(NewClass);
+
+		if (NewVTable != OldVTable)
+		{
+			int32 Count = 0;
+			for (FRawObjectIterator It; It; ++It)
+			{
+				UObject* Target = *It;
+				if (OldVTable == *(void**)Target)
+				{
+					*(void**)Target = NewVTable;
+					Count++;
+				}
+			}
+			UE_LOG(LogClass, Verbose, TEXT("Updated the vtable for %d live objects. %016llx -> %016llx"), Count, PTRINT(OldVTable), PTRINT(NewVTable));
+
+			VTableMap[OldClass] = NewVTable;
+			VTableMap.Add(NewClass, NewVTable);
+		}
+		else
+		{
+			UE_LOG(LogClass, Error, TEXT("VTable for class %s did not change?"), *NewClass->GetName());
+		}
+	}
+
+	/**
+	 * Gets the singleton.
+	 */
+	static FUClassVTableHelper& Get()
+	{
+		static FUClassVTableHelper Helper;
+		return Helper;
+	}
+
+private:
+	/**
+	 * Forbid public construction.
+	 */
+	FUClassVTableHelper() {}
+
+	/**
+	 * Gets UClasses VTable address.
+	 *
+	 * Note that this function creates an object of UClass type.
+	 *
+	 * @param Class A UClass for which to get VTable address.
+	 *
+	 * @returns VTable address.
+	 */
+	static void* GetUClassVTable(UClass* Class)
+	{
+		auto* ClassWithin = Class->ClassWithin;
+		Class->ClassWithin = UPackage::StaticClass(); // We are just avoiding error checks with this... we don't care about this temp object other than to get the vtable.
+		UObject* TempObjectForVTable = StaticConstructObject(Class, GetTransientPackage(), NAME_None, RF_NeedLoad | RF_ClassDefaultObject);
+
+		if (!TempObjectForVTable->IsRooted())
+		{
+			TempObjectForVTable->MarkPendingKill();
+		}
+		else
+		{
+			UE_LOG(LogClass, Warning, TEXT("Hot Reload: Was not expecting temporary object '%s' for class '%s' to become rooted during construction. This object cannot be marked pending kill."), *TempObjectForVTable->GetFName().ToString(), *Class->GetName());
+		}
+
+		Class->ClassWithin = ClassWithin;
+
+		return *(void**)TempObjectForVTable;
+	}
+
+	// A map of classes VTables.
+	TMap<UClass*, void*> VTableMap;
+};
+
 void UClassCompiledInDefer(FFieldCompiledInInfo* ClassInfo, const TCHAR* Name, SIZE_T ClassSize, uint32 Crc)
 {
 	const FName CPPClassName(Name);
@@ -639,6 +747,9 @@ void UClassCompiledInDefer(FFieldCompiledInInfo* ClassInfo, const TCHAR* Name, S
 
 		if (ClassInfo->bHasChanged)
 		{
+			// Check in cache and create if doesn't exist UCLASSes VTable address.
+			FUClassVTableHelper::Get().TouchUClassVTable(ExistingClass);
+
 			// Rename the old class and move it to transient package
 			ExistingClass->ClearFlags(RF_RootSet | RF_Standalone | RF_Public);
 			ExistingClass->GetDefaultObject()->ClearFlags(RF_RootSet | RF_Standalone | RF_Public);
@@ -698,21 +809,38 @@ void UClassRegisterAllCompiledInClasses()
 /** Re-instance all existing classes that have changed during hot-reload */
 void UClassReplaceHotReloadClasses()
 {
-	if (FCoreUObjectDelegates::ReplaceHotReloadClassDelegate.IsBound())
+	struct FClassTuple
 	{
-		for (auto Class : GetHotReloadClasses())
-		{
-			check(Class->OldClass);
+		UClass* OldClass;
+		UClass* NewClass;
 
-			UClass* RegisteredClass = nullptr;
-			if (Class->bHasChanged)
-			{
-				RegisteredClass = Class->Register();
-			}
+		FClassTuple(UClass* OldClass, UClass* NewClass)
+			: OldClass(OldClass), NewClass(NewClass) { }
+	};
+	TArray<FClassTuple> VTablesToUpdate;
+
+	for (auto Class : GetHotReloadClasses())
+	{
+		check(Class->OldClass);
 
+		UClass* RegisteredClass = nullptr;
+		if (Class->bHasChanged)
+		{
+			RegisteredClass = Class->Register();
+			VTablesToUpdate.Add(FClassTuple(Class->OldClass, RegisteredClass));
+		}
+
+		if (FCoreUObjectDelegates::ReplaceHotReloadClassDelegate.IsBound())
+		{
 			FCoreUObjectDelegates::ReplaceHotReloadClassDelegate.Execute(Class->OldClass, RegisteredClass);
 		}
 	}
+
+	for (auto Tuple : VTablesToUpdate)
+	{
+		FUClassVTableHelper::Get().UpdateVTables(Tuple.OldClass, Tuple.NewClass);
+	}
+
 	GetHotReloadClasses().Empty();
 }
 #endif
