[INTEGRATE] Change 2425992 by [REDACTED_PERSON][REDACTED_HANDLE] on 2015/01/30 12:01:19

	Fix broken logic in UObject::Modify that forces packages to be dirtied even when requesting not to.

diff --git a/Engine/Source/Runtime/CoreUObject/Private/UObject/Obj.cpp b/Engine/Source/Runtime/CoreUObject/Private/UObject/Obj.cpp
--- a/Engine/Source/Runtime/CoreUObject/Private/UObject/Obj.cpp
+++ b/Engine/Source/Runtime/CoreUObject/Private/UObject/Obj.cpp
@@ -737,7 +737,7 @@ bool UObject::Modify( bool bAlwaysMarkDirty/*=true*/ )
 			// Attempt to mark the package dirty and save a copy of the object to the transaction
 			// buffer. The save will fail if there isn't a valid transactor, the object isn't
 			// transactional, etc.
-			bSavedToTransactionBuffer = SaveToTransactionBuffer(this, true);
+			bSavedToTransactionBuffer = SaveToTransactionBuffer(this, bAlwaysMarkDirty);
 
 			// If we failed to save to the transaction buffer, but the user requested the package
 			// marked dirty anyway, do so
