[INTEGRATE] Change 2426004 by [REDACTED_PERSON][REDACTED_HANDLE] on 2015/01/30 12:16:19

	Forcing blueprint components earlier in Export loading (so they don't get created during class serialization).

diff --git a/Engine/Source/Runtime/CoreUObject/Private/Blueprint/BlueprintSupport.cpp b/Engine/Source/Runtime/CoreUObject/Private/Blueprint/BlueprintSupport.cpp
--- a/Engine/Source/Runtime/CoreUObject/Private/Blueprint/BlueprintSupport.cpp
+++ b/Engine/Source/Runtime/CoreUObject/Private/Blueprint/BlueprintSupport.cpp
@@ -313,6 +313,10 @@ bool ULinkerLoad::DeferPotentialCircularImport(const int32 Index)
 			if (const UClass* ImportClass = FindObject<UClass>(ClassPackage, *Import.ClassName.ToString()))
 			{
 				bool const bIsBlueprintClass = ImportClass->IsChildOf<UClass>();
+				// @TODO: if we could see if the related package is created 
+				//        (without loading it further), AND it already has this 
+				//        import loaded, then we can probably return that rather 
+				//        than allocating a placeholder
 				if (bIsBlueprintClass)
 				{
 					UPackage* PlaceholderOuter = LinkerRoot;
diff --git a/Engine/Source/Runtime/CoreUObject/Private/UObject/LinkerLoad.cpp b/Engine/Source/Runtime/CoreUObject/Private/UObject/LinkerLoad.cpp
--- a/Engine/Source/Runtime/CoreUObject/Private/UObject/LinkerLoad.cpp
+++ b/Engine/Source/Runtime/CoreUObject/Private/UObject/LinkerLoad.cpp
@@ -3097,6 +3097,16 @@ UObject* ULinkerLoad::CreateExport( int32 Index )
 			return NULL;
 		}
 
+#if USE_DEFERRED_DEPENDENCY_CHECK_VERIFICATION_TESTS
+		// we're going to have troubles if we're attempting to create an export 
+		// for a placeholder class... this means that class serialization lead
+		// to this (non-CDO) export being created, and because of the way we 
+		// guard against cyclic dependencies (with these placeholder classes), 
+		// we need to make sure this export is created BEFORE we start serializing
+		// the class (maybe through ordering in this linker's ExportMap?)
+		check(Cast<ULinkerPlaceholderClass>(LoadClass) == nullptr);
+#endif // USE_DEFERRED_DEPENDENCY_CHECK_VERIFICATION_TESTS
+
 		check(LoadClass);
 		check(dynamic_cast<UClass*>(LoadClass) != NULL);
 
diff --git a/Engine/Source/Runtime/CoreUObject/Private/UObject/SavePackage.cpp b/Engine/Source/Runtime/CoreUObject/Private/UObject/SavePackage.cpp
--- a/Engine/Source/Runtime/CoreUObject/Private/UObject/SavePackage.cpp
+++ b/Engine/Source/Runtime/CoreUObject/Private/UObject/SavePackage.cpp
@@ -1701,7 +1701,7 @@ public:
 								UByteProperty* ByteProp = dynamic_cast<UByteProperty*>(Object);
 								if (ByteProp != NULL && ByteProp->Enum != NULL)
 								{
-									HandleDependency(ByteProp->Enum, true);
+									HandleDependency(ByteProp->Enum, /*bProcessObject =*/true);
 								}
 								// a normal field - property, enum, const; just insert it into the list and keep going
 								ProcessedObjects.Add(Object);
@@ -1718,10 +1718,19 @@ public:
 				}
 				else
 				{
-					// since normal references to objects aren't force-loaded, we do not need to pass true for bProcessObject
-					// (which would indicate that Object must be inserted into the sorted export list before the object that contains
-					// this object reference - i.e. the object we're currently serializing)
-					HandleDependency(Object);
+					// since normal references to objects aren't force-loaded, 
+					// we do not need to pass true for bProcessObject by default
+					// (true would indicate that Object must be inserted into 
+					// the sorted export list before the object that contains 
+					// this object reference - i.e. the object we're currently
+					// serializing)
+					// 
+					// sometimes (rarely) this is the case though, so we use 
+					// ForceLoadObjects to determine if the object we're 
+					// serializing would force load Object (if so, it'll come 
+					// first in the ExportMap)
+					bool const bProcessObject = ForceLoadObjects.Contains(Object);
+					HandleDependency(Object, bProcessObject);
 				}
 			}
 		}
@@ -1820,9 +1829,35 @@ public:
 					// so we turn off field serialization so that we don't have to worry about handling this struct's fields just yet
 					bIgnoreFieldReferences = true;
 
+					// most often, we don't want/need object references getting  
+					// recorded as dependencies, but some structs (classes) 
+					// require certain non-field objects be prioritized in the 
+					// ExportMap earlier (see UClass::GetRequiredPreloadDependencies() 
+					// for more details)... this array records/holds those 
+					// required sub-objects
+					TArray<UObject*> StructForceLoadObjects;
+
+					bool const bIsClassObject = (dynamic_cast<UClass*>(StructObject) != nullptr);
+					if (bIsClassObject)
+					{
+						UClass* AsClass = (UClass*)StructObject;
+						AsClass->GetRequiredPreloadDependencies(StructForceLoadObjects);
+					}
+					// append rather than replace (in case we're nested in a 
+					// recursive call)... adding these to ForceLoadObjects 
+					// ensures that any reference to a StructForceLoadObjects
+					// object gets recorded in the ExportMap before StructObject
+					// (see operator<<, where we utilize ForceLoadObjects)
+					ForceLoadObjects.Append(StructForceLoadObjects);
+					int32 const ForceLoadCount = StructForceLoadObjects.Num();
+
 					SerializedObjects.Add(StructObject);
 					StructObject->Serialize(*this);
 
+					// remove (pop) rather than empty, in case ClassForceLoadObjects
+					// had entries from a previous call to this function up that chain
+					ForceLoadObjects.RemoveAt(ForceLoadObjects.Num() - ForceLoadCount, ForceLoadCount);
+
 					// at this point, any objects which were referenced through this struct's script or defaults will be in the list of exports, and 
 					// the CurrentInsertIndex will have been advanced so that the object processed will be inserted just before this struct in the array
 					// (i.e. just after class/archetypes for any objects which were referenced by this struct's script)
@@ -1830,15 +1865,16 @@ public:
 					// now re-enable field serialization and process the struct's properties, functions, enums, structs, etc.  They will be inserted into the list
 					// just ahead of the struct itself, so that those objects are created first during seek-free loading.
 					bIgnoreFieldReferences = false;
-
+					
 					// invoke the serialize operator rather than calling Serialize directly so that the object is handled correctly (i.e. if it is a struct, then we should
 					// call ProcessStruct, etc. and all this logic is already contained in the serialization operator)
-					if (!dynamic_cast<UClass*>(StructObject))
+					if (!bIsClassObject)
 					{
 						// before processing the Children reference, set the CurrentClass to the class which contains this StructObject so that we
 						// don't inadvertently serialize other fields of the owning class too early.
 						CurrentClass = StructObject->GetOwnerClass();
-					}
+					}					
+
 					(*this) << (UObject*&)StructObject->Children;
 					CurrentClass = NULL;
 
@@ -1928,6 +1964,13 @@ private:
 	 * hasn't been created yet.
 	 */
 	UClass* CurrentClass;
+
+	/** 
+	 * This is a list of objects that would be force loaded by a struct/class 
+	 * currently being handled by ProcessStruct() (meaning that they should be
+	 * prioritized in the target ExportMap, before the struct).
+	 */
+	TArray<UObject*> ForceLoadObjects;
 };
 
 #define EXPORT_SORTING_DETAILED_LOGGING 0
diff --git a/Engine/Source/Runtime/CoreUObject/Public/UObject/Class.h b/Engine/Source/Runtime/CoreUObject/Public/UObject/Class.h
--- a/Engine/Source/Runtime/CoreUObject/Public/UObject/Class.h
+++ b/Engine/Source/Runtime/CoreUObject/Public/UObject/Class.h
@@ -2186,6 +2186,22 @@ public:
 	 */
 	virtual bool HasProperty(UProperty* InProperty) const;
 
+	/**
+	 * On save, we order a package's exports in class dependency order (so that
+	 * on load, we create the class dependencies before we create the class).  
+	 * More often than not, the class doesn't require any non-struct objects 
+	 * before it is created/serialized (only super-classes, and its UField 
+	 * members, see FExportReferenceSorter::operator<<() for reference). 
+	 * However, in some special occasions, there might be an export that we 
+	 * would like force loaded prior to the class's serialization (like  
+	 * component templates for blueprint classes). This function returns a list
+	 * of those non-struct dependencies, so that FExportReferenceSorter knows to 
+	 * prioritize them earlier in the ExportMap.
+	 * 
+	 * @param  DependenciesOut	Will be filled with a list of dependencies that need to be created before this class is recreated (on load).
+	 */
+	virtual void GetRequiredPreloadDependencies(TArray<UObject*>& DependenciesOut) {}
+
 private:
 	// This signature intentionally hides the method declared in UObjectBaseUtility to make it private.
 	// Call IsChildOf instead; Hidden because calling IsA on a class almost always indicates an error where the caller should use IsChildOf.
diff --git a/Engine/Source/Runtime/Engine/Classes/Engine/BlueprintGeneratedClass.h b/Engine/Source/Runtime/Engine/Classes/Engine/BlueprintGeneratedClass.h
--- a/Engine/Source/Runtime/Engine/Classes/Engine/BlueprintGeneratedClass.h
+++ b/Engine/Source/Runtime/Engine/Classes/Engine/BlueprintGeneratedClass.h
@@ -486,6 +486,7 @@ public:
 	virtual void Link(FArchive& Ar, bool bRelinkExistingProperties) override;
 	virtual void PurgeClass(bool bRecompilingOnLoad) override;
 	virtual void Bind() override;
+	virtual void GetRequiredPreloadDependencies(TArray<UObject*>& DependenciesOut) override;
 	// End UClass interface
 
 	static void AddReferencedObjectsInUbergraphFrame(UObject* InThis, FReferenceCollector& Collector);
diff --git a/Engine/Source/Runtime/Engine/Private/BlueprintGeneratedClass.cpp b/Engine/Source/Runtime/Engine/Private/BlueprintGeneratedClass.cpp
--- a/Engine/Source/Runtime/Engine/Private/BlueprintGeneratedClass.cpp
+++ b/Engine/Source/Runtime/Engine/Private/BlueprintGeneratedClass.cpp
@@ -89,6 +89,27 @@ UClass* UBlueprintGeneratedClass::GetAuthoritativeClass()
 	return (GeneratingBP->GeneratedClass != NULL) ? GeneratingBP->GeneratedClass : this;
 }
 
+void UBlueprintGeneratedClass::GetRequiredPreloadDependencies(TArray<UObject*>& DependenciesOut)
+{
+	Super::GetRequiredPreloadDependencies(DependenciesOut);
+	for (UActorComponent* Component : ComponentTemplates)
+	{
+		// because of the linker's way of handling circular dependencies (with 
+		// placeholder blueprint classes), we need to ensure that class owned 
+		// blueprint components are created before the class's  is 
+		// ComponentTemplates member serialized in (otherwise, the component 
+		// would be created as a ULinkerPlaceholderClass instance)
+		//
+		// by returning these in the DependenciesOut array, we're making it 
+		// known that they should be prioritized in the package's ExportMap 
+		// before this class
+		if (Component->GetClass()->HasAnyClassFlags(CLASS_CompiledFromBlueprint))
+		{
+			DependenciesOut.Add(Component);
+		}
+	}
+}
+
 bool FStructUtils::ArePropertiesTheSame(const UProperty* A, const UProperty* B, bool bCheckPropertiesNames)
 {
 	if (A == B)
