[INTEGRATE] Change 2426408 by [REDACTED_PERSON][REDACTED_HANDLE] on 2015/01/30 16:51:34

	Resaving some BP component test cases with CL 2426004 (to reorder ExportMap priorities).

diff --git a/Engine/Source/Runtime/CoreUObject/Private/UObject/LinkerLoad.cpp b/Engine/Source/Runtime/CoreUObject/Private/UObject/LinkerLoad.cpp
--- a/Engine/Source/Runtime/CoreUObject/Private/UObject/LinkerLoad.cpp
+++ b/Engine/Source/Runtime/CoreUObject/Private/UObject/LinkerLoad.cpp
@@ -3104,7 +3104,7 @@ UObject* ULinkerLoad::CreateExport( int32 Index )
 		// guard against cyclic dependencies (with these placeholder classes), 
 		// we need to make sure this export is created BEFORE we start serializing
 		// the class (maybe through ordering in this linker's ExportMap?)
-		check(Cast<ULinkerPlaceholderClass>(LoadClass) == nullptr);
+		check(!LoadClass->IsChildOf<ULinkerPlaceholderClass>());
 #endif // USE_DEFERRED_DEPENDENCY_CHECK_VERIFICATION_TESTS
 
 		check(LoadClass);
