Added metadata documentation to ObjectBase.h

diff --git a/Engine/Source/Runtime/CoreUObject/Public/UObject/ObjectBase.h b/Engine/Source/Runtime/CoreUObject/Public/UObject/ObjectBase.h
--- a/Engine/Source/Runtime/CoreUObject/Public/UObject/ObjectBase.h
+++ b/Engine/Source/Runtime/CoreUObject/Public/UObject/ObjectBase.h
@@ -894,7 +894,7 @@ namespace UM
 		/// [ClassMetadata] [PropertyMetadata] The name to display for this class or property instead of auto-generating it from the name.
 		DisplayName,
 
-		/// [ClassMetadata]
+		/// [ClassMetadata] Specifies that this class is an acceptable base class for creating blueprints.
 		IsBlueprintBase,
 
 		/// [ClassMetadata] Comma delimited list of blueprint events that are not be allowed to be overriden in classes of this type
@@ -949,19 +949,19 @@ namespace UM
 		/// [PropertyMetadata] Indicates that the property is an asset type and it should display the thumbnail of the selected asset.
 		DisplayThumbnail,	
 	
-		/// [PropertyMetadata] Species a boolean property that is used to indicate whether editing of this property is disabled.
+		/// [PropertyMetadata] Specifies a boolean property that is used to indicate whether editing of this property is disabled.
 		EditCondition,
 		
 		/// [PropertyMetadata] Used for FStringAssetReference properties in conjunction with AllowedClasses. Indicates whether only the exact classes specified in AllowedClasses can be used or whether subclasses are valid.
 		ExactClass,
 
-		/// [PropertyMetadata]
+		/// [PropertyMetadata] Specifies a list of categories whose functions should be exposed when building a function list in the Blueprint Editor.
 		ExposeFunctionCategories,
 
 		/// [PropertyMetadata] Specifies whether the property should be exposed on a Spawn Actor for the class type.
 		ExposeOnSpawn,
 
-		/// [PropertyMetadata]
+		/// [PropertyMetadata] Deprecated.
 		FixedIncrement,
 
 		/// [PropertyMetadata] Used for FColor and FLinearColor properties. Indicates that the Alpha property should be hidden when displaying the property widget in the details.
@@ -973,7 +973,7 @@ namespace UM
 		/// [PropertyMetadata] Used for Subclass properties. Indicates whether only placeable classes should be shown in the class picker.
 		OnlyPlaceable,
 
-		/// [PropertyMetadata]
+		/// [PropertyMetadata] Used for Vector/Rotator properties. Indicates that the property should be exposed in the viewport as a movable widget.
 		MakeEditWidget,
 
 		/// [PropertyMetadata] For properties in a structure indicates the default value of the property in a blueprint make structure node.
@@ -982,7 +982,7 @@ namespace UM
 		/// [PropertyMetadata] Used FStringClassReference properties. Indicates the parent class that the class picker will use when filtering which classes to display.
 		MetaClass,
 
-		/// [PropertyMetadata]
+		/// [PropertyMetadata] Used for numeric properties. Stipulates that the value must be a multiple of the metadata value.
 		Multiple,
 
 		/// [PropertyMetadata] Used for array properties. Indicates that the duplicate icon should not be shown for entries of this array in the property panel.
@@ -994,19 +994,19 @@ namespace UM
 		/// [PropertyMetadata] Used by FFilePath properties. Indicates the path filter to display in the file picker.
 		FilePathFilter,
 
-		/// [PropertyMetadta] Used by FDirectoryPath properties. 
+		/// [PropertyMetadta] Used by FDirectoryPath properties. Indicates that the directory dialog will output a relative path when setting the property.
 		RelativePath,
 
-		/// [PropertyMetadta] Used by FDirectoryPath properties. 
+		/// [PropertyMetadta] Used by FDirectoryPath properties. Indicates that the directory dialog will output a path relative to the game content directory when setting the property.
 		RelativeToGameContentDir,
 
-		// [PropertyMetadata]
+		// [PropertyMetadata] Used by struct properties. Indicates that the inner properties will not be shown inside an expandable struct, but promoted up a level.
 		ShowOnlyInnerProperties,
 
-		// [PropertyMetadata]
+		// [PropertyMetadata] Used by numeric properties. Indicates how rapidly the value will grow when moving an unbounded slider.
 		SliderExponent,
 
-		// [PropertyMetadata] Overrides the automatically generated tooltip from the property comment
+		// [PropertyMetadata] Overrides the automatically generated tooltip from the property comment.
 		ToolTip,
 
 		/// [PropertyMetadata] Used for float and integer properties.  Specifies the lowest that the value slider should represent.
