Fix to the struct/class mismatch for FMatrix in generated code.

diff --git a/Engine/Source/Runtime/CoreUObject/Private/UObject/PropertyStruct.cpp b/Engine/Source/Runtime/CoreUObject/Private/UObject/PropertyStruct.cpp
--- a/Engine/Source/Runtime/CoreUObject/Private/UObject/PropertyStruct.cpp
+++ b/Engine/Source/Runtime/CoreUObject/Private/UObject/PropertyStruct.cpp
@@ -254,14 +254,15 @@ FString UStructProperty::GetCPPTypeForwardDeclaration() const
 		FName(TEXT("Rotator")),
 		FName(TEXT("Timespan")),
 		NAME_Transform,
-		NAME_Vector
+		NAME_Vector,
+		NAME_Matrix
 	};
 
 	auto Name = Struct->GetFName();
 
-	for (int32 i = 0; i < ARRAY_COUNT(ClassNames); ++i)
+	for (auto ClassName : ClassNames)
 	{
-		if (ClassNames[i] == Name)
+		if (ClassName == Name)
 		{
 			return FString::Printf(TEXT("class F%s;"), *Struct->GetName());
 		}
