[INTEGRATE] Change 2429227 by [REDACTED_PERSON][REDACTED_HANDLE] on 2015/02/02 22:24:33

	Core: Document MultiLine property metadata (for FString and FText, indicates that editors should allow multiple lines of text to be entered)

diff --git a/Engine/Source/Runtime/CoreUObject/Public/UObject/ObjectBase.h b/Engine/Source/Runtime/CoreUObject/Public/UObject/ObjectBase.h
--- a/Engine/Source/Runtime/CoreUObject/Public/UObject/ObjectBase.h
+++ b/Engine/Source/Runtime/CoreUObject/Public/UObject/ObjectBase.h
@@ -985,6 +985,9 @@ namespace UM
 		/// [PropertyMetadata] Used for numeric properties. Stipulates that the value must be a multiple of the metadata value.
 		Multiple,
 
+		/// [PropertyMetadata] Used for FString and FText properties.  Indicates that the edit field should be multi-line, allowing entry of newlines.
+		MultiLine,
+
 		/// [PropertyMetadata] Used for array properties. Indicates that the duplicate icon should not be shown for entries of this array in the property panel.
 		NoElementDuplicate,
 
