[INTEGRATE] Change 2429911 by [REDACTED_PERSON][REDACTED_HANDLE] on 2015/02/03 10:10:56

	Mark BlueprintClass owned ActorComponents public in a way that resets their ExportTable

diff --git a/Engine/Source/Runtime/CoreUObject/Private/UObject/LinkerLoad.cpp b/Engine/Source/Runtime/CoreUObject/Private/UObject/LinkerLoad.cpp
--- a/Engine/Source/Runtime/CoreUObject/Private/UObject/LinkerLoad.cpp
+++ b/Engine/Source/Runtime/CoreUObject/Private/UObject/LinkerLoad.cpp
@@ -438,7 +438,6 @@ void ULinkerLoad::PRIVATE_PatchNewObjectIntoExport(UObject* OldObject, UObject*
 
 void ULinkerLoad::InvalidateExport(UObject* OldObject)
 {
-	// Cache off the old object's linker and export index.  We'll slide the new object in here.
 	ULinkerLoad* OldObjectLinker = OldObject->GetLinker();
 	const int32 CachedLinkerIndex = OldObject->GetLinkerIndex();
 
@@ -449,6 +448,19 @@ void ULinkerLoad::InvalidateExport(UObject* OldObject)
 	}
 }
 
+void ULinkerLoad::RefreshExportFlags(UObject* Object)
+{
+	ULinkerLoad* ObjectLinker = Object->GetLinker();
+	const int32 CachedLinkerIndex = Object->GetLinkerIndex();
+
+	if (ObjectLinker && ObjectLinker->ExportMap.IsValidIndex(CachedLinkerIndex))
+	{
+		FObjectExport& ObjExport = ObjectLinker->ExportMap[CachedLinkerIndex];
+		ObjExport.ObjectFlags = Object->GetMaskedFlags();
+	}
+}
+
+
 FName ULinkerLoad::FindSubobjectRedirectName(const FName& Name)
 {
 	FSubobjectRedirect *Redirect = SubobjectNameRedirects.Find(Name);
diff --git a/Engine/Source/Runtime/CoreUObject/Public/UObject/Linker.h b/Engine/Source/Runtime/CoreUObject/Public/UObject/Linker.h
--- a/Engine/Source/Runtime/CoreUObject/Public/UObject/Linker.h
+++ b/Engine/Source/Runtime/CoreUObject/Public/UObject/Linker.h
@@ -1595,6 +1595,11 @@ public:
 	 */
 	COREUOBJECT_API static void InvalidateExport(UObject* OldObject);
 
+	/**
+	 * Updates the export table with an Object's object flags
+	 */
+	COREUOBJECT_API static void RefreshExportFlags(UObject* Object);
+
 	/** Used by Matinee to fixup component renaming */
 	COREUOBJECT_API static FName FindSubobjectRedirectName(const FName& Name);
 
diff --git a/Engine/Source/Runtime/Engine/Private/ActorComponent.cpp b/Engine/Source/Runtime/Engine/Private/ActorComponent.cpp
--- a/Engine/Source/Runtime/Engine/Private/ActorComponent.cpp
+++ b/Engine/Source/Runtime/Engine/Private/ActorComponent.cpp
@@ -134,6 +134,12 @@ void UActorComponent::PostLoad()
 	{
 		CreationMethod = EComponentCreationMethod::Instance;
 	}
+
+	if (!HasAllFlags(RF_Public) && GetOuter()->IsA<UBlueprintGeneratedClass>())
+	{
+		SetFlags(RF_Public);
+		ULinkerLoad::RefreshExportFlags(this);
+	}
 }
 
 void UActorComponent::PostRename(UObject* OldOuter, const FName OldName)
diff --git a/Engine/Source/Runtime/Engine/Private/SimpleConstructionScript.cpp b/Engine/Source/Runtime/Engine/Private/SimpleConstructionScript.cpp
--- a/Engine/Source/Runtime/Engine/Private/SimpleConstructionScript.cpp
+++ b/Engine/Source/Runtime/Engine/Private/SimpleConstructionScript.cpp
@@ -174,22 +174,6 @@ void USimpleConstructionScript::PostLoad()
 		}
 	}
 
-	// Templates are used as archetypes, so they should be Public
-	auto OwnerClass = GetOwnerClass();
-	if (OwnerClass)
-	{
-		for (auto Node : Nodes)
-		{
-			auto ComponentTemplate = Node ? Node->ComponentTemplate : nullptr;
-			if (ComponentTemplate && !ComponentTemplate->HasAllFlags(RF_Public) && (ComponentTemplate->GetOuter() == OwnerClass))
-			{
-				ComponentTemplate->SetFlags(RF_Public);
-
-				// Package should be marked as dirty here, but it's not user friendly
-			}
-		}
-	}
-
 #endif // WITH_EDITOR
 	// Fix up native/inherited parent attachments, in case anything has changed
 	FixupRootNodeParentReferences();
