[INTEGRATE] Change 2431559 by [REDACTED_PERSON][REDACTED_HANDLE] on 2015/02/04 08:15:14

	Temporarily stop using archetypes from blueprint generated class for blueprint created components as it is causing level load issues during cook.

diff --git a/Engine/Source/Runtime/CoreUObject/Private/UObject/UObjectArchetype.cpp b/Engine/Source/Runtime/CoreUObject/Private/UObject/UObjectArchetype.cpp
--- a/Engine/Source/Runtime/CoreUObject/Private/UObject/UObjectArchetype.cpp
+++ b/Engine/Source/Runtime/CoreUObject/Private/UObject/UObjectArchetype.cpp
@@ -46,7 +46,7 @@ UObject* UObject::GetArchetypeFromRequiredInfo(UClass* Class, UObject* Outer, FN
 			}
 			else
 			{
-				Result = ArchetypeToSearch->GetClass()->FindArchetype(Class, Name);
+				//Result = ArchetypeToSearch->GetClass()->FindArchetype(Class, Name);
 			}
 		}
 
