[INTEGRATE] Change 2431758 by [REDACTED_PERSON][REDACTED_HANDLE] on 2015/02/04 10:53:54

	UE-8676 Crash when changing a Blueprint's vector variable to a transform variable

	It's a temporal workaround. In fact UseBinaryOrNativeSerialization() should not be caled/checked in UStruct::SerializeTaggedProperties.

diff --git a/Engine/Source/Runtime/CoreUObject/Private/UObject/Class.cpp b/Engine/Source/Runtime/CoreUObject/Private/UObject/Class.cpp
--- a/Engine/Source/Runtime/CoreUObject/Private/UObject/Class.cpp
+++ b/Engine/Source/Runtime/CoreUObject/Private/UObject/Class.cpp
@@ -813,6 +813,11 @@ void UStruct::SerializeTaggedProperties(FArchive& Ar, uint8* Data, UStruct* Defa
 			{
 				break;
 			}
+			if (!Tag.Name.IsValid())
+			{
+				UE_LOG(LogClass, Warning, TEXT("Invalid tag name: struct '%s', archive '%s'"), *GetName(), *Ar.GetArchiveName());
+				break;
+			}
 
 			// Move to the next property to be serialized
 			if( AdvanceProperty && --RemainingArrayDim <= 0 )
diff --git a/Engine/Source/Runtime/CoreUObject/Public/UObject/PropertyTag.h b/Engine/Source/Runtime/CoreUObject/Public/UObject/PropertyTag.h
--- a/Engine/Source/Runtime/CoreUObject/Public/UObject/PropertyTag.h
+++ b/Engine/Source/Runtime/CoreUObject/Public/UObject/PropertyTag.h
@@ -25,7 +25,17 @@ struct FPropertyTag
 
 	// Constructors.
 	FPropertyTag()
+		: Type      (NAME_None)
+		, BoolVal   (0)
+		, Name      (NAME_None)
+		, StructName(NAME_None)
+		, EnumName  (NAME_None)
+		, InnerType (NAME_None)
+		, Size      (0)
+		, ArrayIndex(INDEX_NONE)
+		, SizeOffset(INDEX_NONE)
 	{}
+
 	FPropertyTag( FArchive& InSaveAr, UProperty* Property, int32 InIndex, uint8* Value, uint8* Defaults )
 		: Type      (Property->GetID())
 		, BoolVal   (0)
@@ -67,7 +77,7 @@ struct FPropertyTag
 	{
 		// Name.
 		Ar << Tag.Name;
-		if( Tag.Name == NAME_None )
+		if ((Tag.Name == NAME_None) || !Tag.Name.IsValid())
 		{
 			return Ar;
 		}
