[INTEGRATE] Change 2432091 by [REDACTED_PERSON][REDACTED_HANDLE] on 2015/02/04 14:09:06

	[UE-9004] Making sure blueprint component classes are fully regenerated before we instantiate templates (so we don't have to re-instance on load, in cyclic component scenarios).

diff --git a/Engine/Source/Runtime/CoreUObject/Private/Blueprint/BlueprintSupport.cpp b/Engine/Source/Runtime/CoreUObject/Private/Blueprint/BlueprintSupport.cpp
--- a/Engine/Source/Runtime/CoreUObject/Private/Blueprint/BlueprintSupport.cpp
+++ b/Engine/Source/Runtime/CoreUObject/Private/Blueprint/BlueprintSupport.cpp
@@ -1,6 +1,8 @@
 // Copyright 1998-2015 Epic Games, Inc. All Rights Reserved.
+
 #include "CoreUObjectPrivate.h"
 #include "UObject/LinkerPlaceholderClass.h"
+#include "UObject/LinkerPlaceholderExportObject.h"
 #include "Linker.h"
 #include "PropertyTag.h"
 
@@ -71,6 +73,16 @@ bool FBlueprintSupport::IsDeferredCDOSerializationDisabled()
 #endif
 }
 
+bool FBlueprintSupport::IsDeferredExportCreationDisabled()
+{
+#if USE_CIRCULAR_DEPENDENCY_LOAD_DEFERRING
+	static const FBoolConfigValueHelper NoDeferredExports(TEXT("Kismet"), TEXT("bForceDisableDeferredExportCreation"), GEngineIni);
+	return NoDeferredExports;
+#else
+	return false;
+#endif
+}
+
 /*******************************************************************************
  * FScopedClassDependencyGather
  ******************************************************************************/
@@ -347,9 +359,54 @@ bool ULinkerLoad::DeferPotentialCircularImport(const int32 Index)
 #endif // USE_CIRCULAR_DEPENDENCY_LOAD_DEFERRING
 }
 
+UObject* ULinkerLoad::DeferExportCreation(const int32 Index)
+{
+	FObjectExport& Export = ExportMap[Index];
+
+#if USE_CIRCULAR_DEPENDENCY_LOAD_DEFERRING
+	if (!FBlueprintSupport::UseDeferredDependencyLoading() || FBlueprintSupport::IsDeferredExportCreationDisabled())
+	{
+		return Export.Object;
+	}
+
+	if (Export.Object != nullptr)
+	{
+		return Export.Object;
+	}
+
+	UClass* LoadClass = GetExportLoadClass(Index);
+	if (LoadClass == nullptr)
+	{
+		return Export.Object;
+	}
+
+	ULinkerPlaceholderClass* AsPlaceholderClass = Cast<ULinkerPlaceholderClass>(LoadClass);
+	bool const bIsPlaceholderClass = (AsPlaceholderClass != nullptr);
+
+	ULinkerLoad* ClassLinker = LoadClass->GetLinker();
+	if ( !bIsPlaceholderClass && ((ClassLinker == nullptr) || !ClassLinker->IsBlueprintFinalizationPending()) )
+	{
+		return Export.Object;
+	}
+
+	UPackage* PlaceholderOuter = LinkerRoot;
+	UClass*   PlaceholderType  = ULinkerPlaceholderExportObject::StaticClass();
+
+	FString ClassName = LoadClass->GetName();
+	//ClassName.RemoveFromEnd("_C");
+	FName PlaceholderName(*FString::Printf(TEXT("PLACEHOLDER-INST_of_%s"), *ClassName));
+	PlaceholderName = MakeUniqueObjectName(PlaceholderOuter, PlaceholderType, PlaceholderName);
+
+	ULinkerPlaceholderExportObject* Placeholder = ConstructObject<ULinkerPlaceholderExportObject>(PlaceholderType, PlaceholderOuter, PlaceholderName, RF_Public | RF_Transient);
+	Export.Object = Placeholder;
+#endif // USE_CIRCULAR_DEPENDENCY_LOAD_DEFERRING
+
+	return Export.Object;
+}
+
 /**
  * This utility struct helps track blueprint structs/linkers that are currently 
- * in the  middle of a call to ResolveDeferredDependencies(). This can be used  
+ * in the middle of a call to ResolveDeferredDependencies(). This can be used  
  * to know if a dependency's resolve needs to be finished (to avoid unwanted 
  * placeholder references ending up in script-code).
  */
@@ -441,6 +498,7 @@ void ULinkerLoad::ResolveDeferredDependencies(UStruct* LoadStruct)
 	TGuardValue<int32> RecursiveDepthGuard(RecursiveDepth, RecursiveDepth + 1);
 
 	DEFERRED_DEPENDENCY_CHECK( (LoadStruct != nullptr) && (LoadStruct->GetLinker() == this) );
+	DEFERRED_DEPENDENCY_CHECK( LoadStruct->HasAnyFlags(RF_LoadCompleted) );
 #endif // USE_DEFERRED_DEPENDENCY_CHECK_VERIFICATION_TESTS
 
 	// scoped block to manage the lifetime of ScopedResolveTracker, so that 
@@ -573,16 +631,26 @@ int32 ULinkerLoad::ResolveDependencyPlaceholder(UClass* PlaceholderIn, UClass* R
 	int32 const ImportIndex = PlaceholderClass->ImportIndex;
 	FObjectImport& Import = ImportMap[ImportIndex];
 
-	DEFERRED_DEPENDENCY_CHECK( (Import.XObject == PlaceholderClass) || (Import.XObject == nullptr) );
 	
-	// clear the placeholder from the import, so that a call to CreateImport()
-	// properly fills it in
-	Import.XObject = nullptr;
-	// NOTE: this is a possible point of recursion... CreateImport() could 
-	//       continue to load a package already started up the stack and you 
-	//       could end up in another ResolveDependencyPlaceholder() for some  
-	//       other placeholder before this one has completely finished resolving
-	UClass* RealClassObj = CastChecked<UClass>(CreateImport(ImportIndex), ECastCheckedType::NullAllowed);
+	UClass* RealClassObj = nullptr;
+	if ((Import.XObject != nullptr) && (Import.XObject != PlaceholderClass))
+	{
+		DEFERRED_DEPENDENCY_CHECK(ResolvingDeferredPlaceholder == PlaceholderIn);
+		RealClassObj = CastChecked<UClass>(Import.XObject);
+	}
+	else 
+	{
+		// clear the placeholder from the import, so that a call to CreateImport()
+		// properly fills it in
+		Import.XObject = nullptr;
+		// NOTE: this is a possible point of recursion... CreateImport() could 
+		//       continue to load a package already started up the stack and you 
+		//       could end up in another ResolveDependencyPlaceholder() for some  
+		//       other placeholder before this one has completely finished resolving
+		RealClassObj = CastChecked<UClass>(CreateImport(ImportIndex), ECastCheckedType::NullAllowed);
+	}
+
+	DEFERRED_DEPENDENCY_CHECK(RealClassObj == nullptr || RealClassObj->HasAnyFlags(RF_LoadCompleted));
 
 	int32 ReplacementCount = 0;
 	if (ReferencingClass != nullptr)
@@ -645,6 +713,7 @@ void ULinkerLoad::FinalizeBlueprint(UClass* LoadClass)
 	{
 		return;
 	}
+	DEFERRED_DEPENDENCY_CHECK(LoadClass->HasAnyFlags(RF_LoadCompleted));
 
 	//--------------------------------------
 	// Phase 3: Finalize (serialize CDO & regenerate class)
@@ -725,80 +794,175 @@ void ULinkerLoad::FinalizeBlueprint(UClass* LoadClass)
 	// get fully executed in some nested call
 	if (IsBlueprintFinalizationPending())
 	{
-		DEFERRED_DEPENDENCY_CHECK(DeferredExportIndex != INDEX_NONE);
-		FObjectExport& CDOExport = ExportMap[DeferredExportIndex];
+		ResolveDeferredExports(LoadClass);
 
-		UObject* CDO = CDOExport.Object;
-		DEFERRED_DEPENDENCY_CHECK(CDO != nullptr);
+		FObjectExport& CDOExport = ExportMap[DeferredCDOIndex];
+		// clear this so IsBlueprintFinalizationPending() doesn't report true
+		DeferredCDOIndex = INDEX_NONE;
 
 #if USE_DEFERRED_DEPENDENCY_CHECK_VERIFICATION_TESTS
-		if (!FBlueprintSupport::IsDeferredCDOSerializationDisabled())
-#endif // USE_DEFERRED_DEPENDENCY_CHECK_VERIFICATION_TESTS
-		{
-			// have to prematurely set the CDO's linker so we can force a Preload()/
-			// Serialization of the CDO before we regenerate the Blueprint class
-			{
-				const EObjectFlags OldFlags = CDO->GetFlags();
-				CDO->ClearFlags(RF_NeedLoad | RF_NeedPostLoad);
-				CDO->SetLinker(this, DeferredExportIndex, /*bShouldDetatchExisting =*/false);
-				CDO->SetFlags(OldFlags);
-			}
-			// should load the CDO (ensuring that it has been serialized in by the 
-			// time we get to class regeneration)
-			//
-			// NOTE: this is point where circular dependencies could reveal 
-			//       themselves, as the CDO could depend on a class not listed in 
-			//       the package's imports
-			//
-			// NOTE: we don't guard against reentrant behavior... if the CDO has
-			//       already been "finalized", then its RF_NeedLoad flag would 
-			//       have been cleared
-			Preload(CDO);
+		// at this point there should not be any instances of the Blueprint 
+		// (else, we'd have to re-instance and that is too expensive an 
+		// operation to have at load time)
+		TArray<UObject*> ClassInstances;
+		GetObjectsOfClass(LoadClass, ClassInstances, /*bIncludeDerivedClasses =*/true);
 
-			DEFERRED_DEPENDENCY_CHECK(CDO->HasAnyFlags(RF_LoadCompleted));
+		for (UObject* ClassInst : ClassInstances)
+		{
+			// in the case that we do end up with instances, use this to find 
+			// where they are referenced (to help sleuth out when/where they 
+			// were created)
+			FReferencerInformationList InstanceReferences;
+			bool const bIsReferenced = IsReferenced(ClassInst, RF_NoFlags, /*bCheckSubObjects =*/false, &InstanceReferences);
+			DEFERRED_DEPENDENCY_CHECK(!bIsReferenced);
 		}
+		DEFERRED_DEPENDENCY_CHECK(ClassInstances.Num() == 0);
 
 		UClass* BlueprintClass = Cast<UClass>(IndexToObject(CDOExport.ClassIndex));
 		DEFERRED_DEPENDENCY_CHECK(BlueprintClass == LoadClass);
 		DEFERRED_DEPENDENCY_CHECK(BlueprintClass->HasAnyClassFlags(CLASS_CompiledFromBlueprint));
-
-#if USE_DEFERRED_DEPENDENCY_CHECK_VERIFICATION_TESTS
-		// at this point there should not be any instances of the Blueprint (else,
-		// we'd have to reinstance and that is too expensive an operation to 
-		// have at load time)
-		TArray<UObject*> ClassInstances;
-		GetObjectsOfClass(BlueprintClass, ClassInstances, /*bIncludeDerivedClasses =*/true);
-		DEFERRED_DEPENDENCY_CHECK(ClassInstances.Num() == 0);
 #endif // USE_DEFERRED_DEPENDENCY_CHECK_VERIFICATION_TESTS
 
-		// clear this so IsBlueprintFinalizationPending() doesn't report true
-		DeferredExportIndex = INDEX_NONE;
-
 		// just in case we choose to enable the deferred dependency loading for 
 		// cooked builds... we want to keep from regenerating in that scenario
 		if (!LoadClass->bCooked)
 		{
-			UObject* OldCDO = BlueprintClass->ClassDefaultObject;
-			if (RegenerateBlueprintClass(BlueprintClass, CDO))
+			UObject* OldCDO = LoadClass->ClassDefaultObject;
+			if (RegenerateBlueprintClass(LoadClass, CDOExport.Object))
 			{
 				// emulate class CDO serialization (RegenerateBlueprintClass() could 
 				// have a side-effect where it overwrites the class's CDO; so we 
 				// want to make sure that we don't overwrite that new CDO with a 
 				// stale one)
-				if (OldCDO == BlueprintClass->ClassDefaultObject)
+				if (OldCDO == LoadClass->ClassDefaultObject)
+				{
+					LoadClass->ClassDefaultObject = CDOExport.Object;
+				}
+			}
+		}
+	}
+#endif // USE_CIRCULAR_DEPENDENCY_LOAD_DEFERRING
+}
+
+void ULinkerLoad::ResolveDeferredExports(UClass* LoadClass)
+{
+#if USE_CIRCULAR_DEPENDENCY_LOAD_DEFERRING
+
+	DEFERRED_DEPENDENCY_CHECK(DeferredCDOIndex != INDEX_NONE);
+	FObjectExport& CDOExport = ExportMap[DeferredCDOIndex];
+
+	UObject* BlueprintCDO = CDOExport.Object;
+	DEFERRED_DEPENDENCY_CHECK(BlueprintCDO != nullptr);
+
+#if USE_DEFERRED_DEPENDENCY_CHECK_VERIFICATION_TESTS
+	static TSet<ULinkerLoad*> ResolvingExportsTracker;
+	// the first half of this function CANNOT be re-entrant (specifically for 
+	// this linker)
+	//
+	// if this ends up being a problem, then we have two blueprints that 
+	// rely on an instances of each other... this is an unexpected problem, and 
+	// we'll have to re-instance on load (an operation we'd rather avoid,  
+	// because of how slow it is)
+	DEFERRED_DEPENDENCY_CHECK(!ResolvingExportsTracker.Contains(this));
+	ResolvingExportsTracker.Add(this);
+#endif // USE_DEFERRED_DEPENDENCY_CHECK_VERIFICATION_TESTS
+
+	if (!FBlueprintSupport::IsDeferredExportCreationDisabled())
+	{
+		// we may have circumvented an export creation or two to avoid instantiating
+		// an BP object before its class has been finalized (to avoid costly re-
+		// instancing operations when the class ultimately finalizes)... so here, we
+		// find those skipped exports and properly create them (before we finalize 
+		// our own class)
+		for (int32 ExportIndex = 0; ExportIndex < ExportMap.Num(); ++ExportIndex)
+		{
+			FObjectExport& Export = ExportMap[ExportIndex];
+			if (ULinkerPlaceholderExportObject* PlaceholderExport = Cast<ULinkerPlaceholderExportObject>(Export.Object))
+			{
+				DEFERRED_DEPENDENCY_CHECK(Export.ClassIndex.IsImport());
+				UClass* ImportClass = GetExportLoadClass(ExportIndex);
+				DEFERRED_DEPENDENCY_CHECK((ImportClass != nullptr) && !ImportClass->HasAnyClassFlags(CLASS_Intrinsic) && ImportClass->HasAnyClassFlags(CLASS_CompiledFromBlueprint));
+				ULinkerLoad* ClassLinker = ImportClass->GetLinker();
+				DEFERRED_DEPENDENCY_CHECK((ClassLinker != nullptr) && (ClassLinker != this));
+
+				// make sure this export's class is fully regenerated before we 
+				// instantiate it 
+				if (ClassLinker->HasUnresolvedDependencies())
 				{
-					BlueprintClass->ClassDefaultObject = CDOExport.Object;
+					ClassLinker->ResolveDeferredDependencies(ImportClass);
 				}
+				if (ClassLinker->IsBlueprintFinalizationPending())
+				{
+					ClassLinker->FinalizeBlueprint(ImportClass);
+				}
+
+				// replace the placeholder with the proper object instance
+				Export.Object = nullptr;
+				UObject* ExportObj = CreateExport(ExportIndex);
+
+				// NOTE: we don't count how many references were resolved (and 
+				//       assert on it), because this could have only been created as 
+				//       part of the LoadAllObjects() pass (not for any specific 
+				//       container object).
+				PlaceholderExport->ReplaceReferencingObjectValues(LoadClass, ExportObj);
+				PlaceholderExport->MarkPendingKill();
+
+				// if we hadn't used a ULinkerPlaceholderExportObject in place of 
+				// the expected export, then someone may have wanted it preloaded
+				Preload(ExportObj);
+
+#if USE_DEFERRED_DEPENDENCY_CHECK_VERIFICATION_TESTS
+				FReferencerInformationList UnresolvedReferences;
+				UObject* PlaceholderObj = PlaceholderExport;
+				bool const bIsReferenced = IsReferenced(PlaceholderObj, RF_NoFlags, /*bCheckSubObjects =*/false, &UnresolvedReferences);
+#endif // USE_DEFERRED_DEPENDENCY_CHECK_VERIFICATION_TESTS
 			}
 		}
 	}
+
+#if USE_DEFERRED_DEPENDENCY_CHECK_VERIFICATION_TESTS
+	// we can be re-entrant from here on (in-fact, we expect to with CDO 
+	// serialization)
+	ResolvingExportsTracker.Remove(this);
+
+	for (TObjectIterator<ULinkerPlaceholderExportObject> PlaceholderIt; PlaceholderIt; ++PlaceholderIt)
+	{
+		ULinkerPlaceholderExportObject* PlaceholderObj = *PlaceholderIt;
+		DEFERRED_DEPENDENCY_CHECK((PlaceholderObj->GetOuter() != LinkerRoot) || PlaceholderObj->IsPendingKill());
+	}
+
+	if (!FBlueprintSupport::IsDeferredCDOSerializationDisabled())
+#endif // USE_DEFERRED_DEPENDENCY_CHECK_VERIFICATION_TESTS
+	{
+		// have to prematurely set the CDO's linker so we can force a Preload()/
+		// Serialization of the CDO before we regenerate the Blueprint class
+		{
+			const EObjectFlags OldFlags = BlueprintCDO->GetFlags();
+			BlueprintCDO->ClearFlags(RF_NeedLoad | RF_NeedPostLoad);
+			BlueprintCDO->SetLinker(this, DeferredCDOIndex, /*bShouldDetatchExisting =*/false);
+			BlueprintCDO->SetFlags(OldFlags);
+		}
+		// should load the CDO (ensuring that it has been serialized in by the 
+		// time we get to class regeneration)
+		//
+		// NOTE: this is point where circular dependencies could reveal 
+		//       themselves, as the CDO could depend on a class not listed in 
+		//       the package's imports
+		//
+		// NOTE: how we don't guard against re-entrant behavior... if the CDO 
+		//       has already been "finalized", then its RF_NeedLoad flag would 
+		//       be cleared (and this will do nothing the 2nd time around)
+		Preload(BlueprintCDO);
+
+		DEFERRED_DEPENDENCY_CHECK(BlueprintCDO->HasAnyFlags(RF_LoadCompleted));
+	}
 #endif // USE_CIRCULAR_DEPENDENCY_LOAD_DEFERRING
 }
 
 bool ULinkerLoad::IsBlueprintFinalizationPending() const
 {
 #if USE_CIRCULAR_DEPENDENCY_LOAD_DEFERRING
-	return (DeferredExportIndex != INDEX_NONE);
+	return (DeferredCDOIndex != INDEX_NONE);
 #else  // USE_CIRCULAR_DEPENDENCY_LOAD_DEFERRING
 	return false;
 #endif // USE_CIRCULAR_DEPENDENCY_LOAD_DEFERRING
diff --git a/Engine/Source/Runtime/CoreUObject/Private/UObject/LinkerLoad.cpp b/Engine/Source/Runtime/CoreUObject/Private/UObject/LinkerLoad.cpp
--- a/Engine/Source/Runtime/CoreUObject/Private/UObject/LinkerLoad.cpp
+++ b/Engine/Source/Runtime/CoreUObject/Private/UObject/LinkerLoad.cpp
@@ -7,6 +7,7 @@
 #include "UObjectToken.h"
 #include "EngineVersion.h"
 #include "LinkerPlaceholderClass.h"
+#include "LinkerPlaceholderExportObject.h"
 
 #define LOCTEXT_NAMESPACE "LinkerLoad"
 
@@ -396,6 +397,37 @@ ULinkerLoad* ULinkerLoad::FindExistingLinkerForPackage(UPackage* Package)
 	return GObjLoaders.FindRef(Package);
 }
 
+ULinkerLoad* ULinkerLoad::FindExistingLinkerForImport(int32 Index) const
+{
+	const FObjectImport& Import = ImportMap[Index];
+	if (Import.SourceLinker != nullptr)
+	{
+		return Import.SourceLinker;
+	}
+	else if (Import.XObject != nullptr)
+	{
+		if (ULinkerLoad* ObjLinker = Import.XObject->GetLinker())
+		{
+			return ObjLinker;
+		}
+	}
+
+	ULinkerLoad* FoundLinker = nullptr;
+	if (Import.OuterIndex.IsNull() && (Import.ClassName == NAME_Package))
+	{
+		FString PackageName = Import.ObjectName.ToString();
+		if (UPackage* FoundPackage = FindObject<UPackage>(/*Outer =*/nullptr, *PackageName))
+		{
+			FoundLinker = ULinkerLoad::FindExistingLinkerForPackage(FoundPackage);
+		}
+	}
+	else if (Import.OuterIndex.IsImport())
+	{
+		FoundLinker = FindExistingLinkerForImport(Import.OuterIndex.ToImport());
+	}
+	return FoundLinker;
+}
+
 /**
  * !!!!!!!!!!!!!!!!!!!!!!!!!!!!!!!!!!!!!!!!!!!!!!!!!!!!!!!!!!!!!!!!!!!!!!!!!!!!!!!!!!!!!!!!!!!!!!!!!!!!!!!!!!!!!!!!!!!!!!!!!!!!!!!!!!
  * CAUTION:  This function is potentially DANGEROUS.  Should only be used when you're really, really sure you know what you're doing.
@@ -649,7 +681,7 @@ ULinkerLoad::ULinkerLoad( const FObjectInitializer& ObjectInitializer, UPackage*
 ,	LoadProgressScope( nullptr )
 #endif // WITH_EDITOR
 #if USE_CIRCULAR_DEPENDENCY_LOAD_DEFERRING
-,	DeferredExportIndex(INDEX_NONE)
+,	DeferredCDOIndex(INDEX_NONE)
 ,	ResolvingDeferredPlaceholder(nullptr)
 #endif // USE_CIRCULAR_DEPENDENCY_LOAD_DEFERRING
 {
@@ -2509,6 +2541,34 @@ UObject* ULinkerLoad::CreateExportAndPreload(int32 ExportIndex, bool bForcePrelo
 	return Object;
 }
 
+UClass* ULinkerLoad::GetExportLoadClass(int32 Index)
+{
+	FObjectExport& Export = ExportMap[Index];
+
+#if USE_CIRCULAR_DEPENDENCY_LOAD_DEFERRING
+	// VerifyImport() runs the risk of loading up another package, and we can't 
+	// have that when we're explicitly trying to block dependency loads...
+	// if this needs a class from another package, IndexToObject() should return 
+	// a ULinkerPlaceholderClass instead
+	if (Export.ClassIndex.IsImport() && !(LoadFlags & LOAD_DeferDependencyLoads))
+#else  // USE_CIRCULAR_DEPENDENCY_LOAD_DEFERRING
+	if (Export.ClassIndex.IsImport())
+#endif // USE_CIRCULAR_DEPENDENCY_LOAD_DEFERRING
+	{
+		// @TODO: I believe IndexToObject() -> CreateImport() will verify this  
+		//        for us, if it has to; so is this necessary?
+		VerifyImport(Export.ClassIndex.ToImport());
+	}
+
+	UClass* ExportClass = (UClass*)IndexToObject(Export.ClassIndex);
+#if USE_DEFERRED_DEPENDENCY_CHECK_VERIFICATION_TESTS
+	check(!Export.ClassIndex.IsImport() || !(LoadFlags & LOAD_DeferDependencyLoads) || 
+		(ExportClass && ExportClass->HasAnyClassFlags(CLASS_Intrinsic)) || (Cast<ULinkerPlaceholderClass>(ExportClass) != nullptr));
+#endif // USE_DEFERRED_DEPENDENCY_CHECK_VERIFICATION_TESTS
+
+	return ExportClass;
+}
+
 int32 ULinkerLoad::LoadMetaDataFromExportMap(bool bForcePreload/* = false */)
 {
 	int32 MetaDataIndex = INDEX_NONE;
@@ -2797,6 +2857,9 @@ void ULinkerLoad::Preload( UObject* Object )
 	// being loaded further up the load chain, and this could introduce a 
 	// circular load)
 	check(!bIsBlueprintClass || !Object->HasAnyFlags(RF_NeedLoad) || !(LoadFlags & LOAD_DeferDependencyLoads));
+	// right now there are no known scenarios where someone requests a Preloa() 
+	// on a temporary ULinkerPlaceholderExportObject
+	check(!Object->IsA<ULinkerPlaceholderExportObject>());
 #endif // USE_DEFERRED_DEPENDENCY_CHECK_VERIFICATION_TESTS
 #endif // USE_CIRCULAR_DEPENDENCY_LOAD_DEFERRING
 	
@@ -2860,7 +2923,7 @@ void ULinkerLoad::Preload( UObject* Object )
 #if USE_DEFERRED_DEPENDENCY_CHECK_VERIFICATION_TESTS
 							if (!FBlueprintSupport::IsDeferredCDOSerializationDisabled())
 							{
-								check((DeferredExportIndex == INDEX_NONE) || (DeferredExportIndex == ExportIndex));
+								check((DeferredCDOIndex == INDEX_NONE) || (DeferredCDOIndex == ExportIndex));
 #else // USE_DEFERRED_DEPENDENCY_CHECK_VERIFICATION_TESTS
 							{
 #endif // USE_DEFERRED_DEPENDENCY_CHECK_VERIFICATION_TESTS
@@ -2868,7 +2931,7 @@ void ULinkerLoad::Preload( UObject* Object )
 								// since serializing the CDO can introduce circular 
 								// dependencies, we want to stave that off until 
 								// we're ready to handle those 
-								DeferredExportIndex = ExportIndex;
+								DeferredCDOIndex = ExportIndex;
 								// don't need to actually "consume" the data through
 								// serialization though (since we seek back to 
 								// SavedPos later on)
@@ -2920,7 +2983,7 @@ void ULinkerLoad::Preload( UObject* Object )
 						// the class wouldn't be regenerated there either)... in
 						// this scenario FinalizeBlueprint() essentially does nothing
 						// @TODO: maybe only allow a null ClassGeneratedBy when PIE'ing/running a game?
-						check((DeferredExportIndex != INDEX_NONE) || (ObjectAsClass->ClassGeneratedBy == nullptr) || FBlueprintSupport::IsDeferredCDOSerializationDisabled());
+						check((DeferredCDOIndex != INDEX_NONE) || (ObjectAsClass->ClassGeneratedBy == nullptr) || FBlueprintSupport::IsDeferredCDOSerializationDisabled());
 #endif // USE_DEFERRED_DEPENDENCY_CHECK_VERIFICATION_TESTS
 
 						ResolveDeferredDependencies(ObjectAsClass);
@@ -3074,16 +3137,12 @@ UObject* ULinkerLoad::CreateExport( int32 Index )
 		check(Export.ObjectName!=NAME_None || !(Export.ObjectFlags&RF_Public));
 		check(IsLoading());
 
-		// Get the object's class.
-		if(Export.ClassIndex.IsImport() ) 
-		{
-			VerifyImport(Export.ClassIndex.ToImport()); 
-		}
-		UClass* LoadClass = (UClass*)IndexToObject( Export.ClassIndex );
+		UClass* LoadClass = GetExportLoadClass(Index);
 		if( !LoadClass && !Export.ClassIndex.IsNull() ) // Hack to load packages with classes which do not exist.
 		{
 			return NULL;
 		}
+
 #if WITH_EDITOR
 		// NULL (None) active class redirect.
 		if( !LoadClass && Export.ObjectName.IsNone() && Export.ClassIndex.IsNull() && !Export.OldClassName.IsNone() )
@@ -3109,16 +3168,6 @@ UObject* ULinkerLoad::CreateExport( int32 Index )
 			return NULL;
 		}
 
-#if USE_DEFERRED_DEPENDENCY_CHECK_VERIFICATION_TESTS
-		// we're going to have troubles if we're attempting to create an export 
-		// for a placeholder class... this means that class serialization lead
-		// to this (non-CDO) export being created, and because of the way we 
-		// guard against cyclic dependencies (with these placeholder classes), 
-		// we need to make sure this export is created BEFORE we start serializing
-		// the class (maybe through ordering in this linker's ExportMap?)
-		check(!LoadClass->IsChildOf<ULinkerPlaceholderClass>());
-#endif // USE_DEFERRED_DEPENDENCY_CHECK_VERIFICATION_TESTS
-
 		check(LoadClass);
 		check(dynamic_cast<UClass*>(LoadClass) != NULL);
 
@@ -3181,7 +3230,21 @@ UObject* ULinkerLoad::CreateExport( int32 Index )
 
 		if ( !LoadClass->HasAnyClassFlags(CLASS_Intrinsic) )
 		{
-			Preload( LoadClass );
+#if USE_CIRCULAR_DEPENDENCY_LOAD_DEFERRING
+			if (LoadClass->HasAnyFlags(RF_NeedLoad))
+			{
+				Preload(LoadClass);
+			}
+			else if ((Export.Object == nullptr) && !(Export.ObjectFlags & RF_ClassDefaultObject))
+			{
+				if (UObject* PlaceholderObj = DeferExportCreation(Index))
+				{
+					return PlaceholderObj; 
+				}				
+			}
+#else  // USE_CIRCULAR_DEPENDENCY_LOAD_DEFERRING
+			Preload(LoadClass);
+#endif // USE_CIRCULAR_DEPENDENCY_LOAD_DEFERRING
 
 			// Check if the Preload() above caused the class to be regenerated (LoadClass will be out of date), and refresh the LoadClass pointer if that is the case
 			if( LoadClass->HasAnyClassFlags(CLASS_NewerVersionExists) )
@@ -3206,6 +3269,32 @@ UObject* ULinkerLoad::CreateExport( int32 Index )
 				}
 			}
 		}
+#if USE_CIRCULAR_DEPENDENCY_LOAD_DEFERRING
+		else if (ULinkerPlaceholderClass* PlaceholderClass = Cast<ULinkerPlaceholderClass>(LoadClass))
+		{
+#if USE_DEFERRED_DEPENDENCY_CHECK_VERIFICATION_TESTS
+			// this blueprint package's CDO should NOT be using a placeholder 
+			// class (its class should be internal to this package)
+			check(!(Export.ObjectFlags & RF_ClassDefaultObject));
+#endif // USE_DEFERRED_DEPENDENCY_CHECK_VERIFICATION_TESTS
+			if (Export.Object == nullptr)
+			{
+				if (UObject* PlaceholderObj = DeferExportCreation(Index))
+				{
+					return PlaceholderObj;
+				}
+			}
+		}
+#endif // USE_CIRCULAR_DEPENDENCY_LOAD_DEFERRING
+
+#if USE_DEFERRED_DEPENDENCY_CHECK_VERIFICATION_TESTS
+		// we're going to have troubles if we're attempting to create an export 
+		// for a placeholder class past this point... placeholder-classes should
+		// have generated an export-placeholder in the above 
+		// !LoadClass->HasAnyClassFlags(CLASS_Intrinsic) block (with the call to
+		// DeferExportCreation)
+		check(Cast<ULinkerPlaceholderClass>(LoadClass) == nullptr);
+#endif // USE_DEFERRED_DEPENDENCY_CHECK_VERIFICATION_TESTS
 
 		// detect cases where a class has been made transient when there are existing instances of this class in content packages,
 		// and this isn't the class default object; when this happens, it can cause issues which are difficult to debug since they'll
@@ -3435,7 +3524,7 @@ UObject* ULinkerLoad::CreateExport( int32 Index )
 				{
 					// if LOAD_DeferDependencyLoads is set, then we're already
 					// serializing the blueprint's class somewhere up the chain,
-					DeferredExportIndex = Index;
+					DeferredCDOIndex = Index;
 					return Export.Object;
 				}
 #endif // USE_CIRCULAR_DEPENDENCY_LOAD_DEFERRING
diff --git a/Engine/Source/Runtime/CoreUObject/Private/UObject/PropertyBaseObject.cpp b/Engine/Source/Runtime/CoreUObject/Private/UObject/PropertyBaseObject.cpp
--- a/Engine/Source/Runtime/CoreUObject/Private/UObject/PropertyBaseObject.cpp
+++ b/Engine/Source/Runtime/CoreUObject/Private/UObject/PropertyBaseObject.cpp
@@ -3,6 +3,7 @@
 #include "CoreUObjectPrivate.h"
 #include "PropertyHelper.h"
 #include "LinkerPlaceholderClass.h"
+#include "LinkerPlaceholderExportObject.h"
 
 /*-----------------------------------------------------------------------------
 	UObjectPropertyBase.
@@ -421,7 +422,20 @@ void UObjectPropertyBase::CheckValidObject(void* Value) const
 		// in the middle of an FArchiveReplaceObjectRef pass)
 		bool bIsReplacingClassRefs = PropertyClass->HasAnyClassFlags(CLASS_NewerVersionExists) != ObjectClass->HasAnyClassFlags(CLASS_NewerVersionExists);
 		
-		if ((PropertyClass != nullptr) && !ObjectClass->IsChildOf(PropertyClass) && !bIsReplacingClassRefs)
+#if USE_CIRCULAR_DEPENDENCY_LOAD_DEFERRING
+		ULinkerLoad* PropertyLinker = GetLinker();
+		bool const bIsDeferringValueLoad = ((PropertyLinker == nullptr) || (PropertyLinker->LoadFlags & LOAD_DeferDependencyLoads)) &&
+			Object->IsA<ULinkerPlaceholderExportObject>();
+
+#if USE_DEFERRED_DEPENDENCY_CHECK_VERIFICATION_TESTS
+		check(bIsDeferringValueLoad || !Object->IsA<ULinkerPlaceholderExportObject>());
+#endif // USE_DEFERRED_DEPENDENCY_CHECK_VERIFICATION_TESTS
+
+#else  // USE_CIRCULAR_DEPENDENCY_LOAD_DEFERRING 
+		bool const bIsDeferringValueLoad = false;
+#endif // USE_CIRCULAR_DEPENDENCY_LOAD_DEFERRING
+
+		if ((PropertyClass != nullptr) && !ObjectClass->IsChildOf(PropertyClass) && !bIsReplacingClassRefs && !bIsDeferringValueLoad)
 		{
 			UE_LOG(LogProperty, Warning,
 				TEXT("Serialized %s for a property of %s. Reference will be NULLed.\n    Property = %s\n    Item = %s"),
diff --git a/Engine/Source/Runtime/CoreUObject/Private/UObject/PropertyObject.cpp b/Engine/Source/Runtime/CoreUObject/Private/UObject/PropertyObject.cpp
--- a/Engine/Source/Runtime/CoreUObject/Private/UObject/PropertyObject.cpp
+++ b/Engine/Source/Runtime/CoreUObject/Private/UObject/PropertyObject.cpp
@@ -1,6 +1,7 @@
 // Copyright 1998-2015 Epic Games, Inc. All Rights Reserved.
 
 #include "CoreUObjectPrivate.h"
+#include "LinkerPlaceholderExportObject.h"
 
 /*-----------------------------------------------------------------------------
 	UObjectProperty.
@@ -26,13 +27,45 @@ void UObjectProperty::SerializeItem( FArchive& Ar, void* Value, void const* Defa
 {
 	UObject* ObjectValue = GetObjectPropertyValue(Value);
 	Ar << ObjectValue;
-	if (ObjectValue != GetObjectPropertyValue(Value))
+
+	UObject* CurrentValue = GetObjectPropertyValue(Value);
+	if (ObjectValue != CurrentValue)
 	{
+#if USE_CIRCULAR_DEPENDENCY_LOAD_DEFERRING
+		if (ULinkerPlaceholderExportObject* PlaceholderVal = Cast<ULinkerPlaceholderExportObject>(ObjectValue))
+		{
+			PlaceholderVal->AddReferencingProperty(this);
+		}
+		// NOTE: we don't remove this from CurrentValue if it is a 
+		//       ULinkerPlaceholderExportObject; this is because this property 
+		//       could be an array inner, and another member of that array (also 
+		//       referenced through this property)... if this becomes a problem,
+		//       then we could inc/decrement a ref count per referencing property 
+		//
+		// @TODO: if this becomes problematic (because ObjectValue doesn't match 
+		//        this property's PropertyClass), then we could spawn another
+		//        placeholder object (of PropertyClass's type), or use null; but
+		//        we'd have to modify ULinkerPlaceholderExportObject::ReplaceReferencingObjectValues()
+		//        to accommodate this (as it depends on finding itself as the set value)
+#endif // USE_CIRCULAR_DEPENDENCY_LOAD_DEFERRING
+
 		SetObjectPropertyValue(Value, ObjectValue);
 		CheckValidObject(Value);
 	}
 }
 
+void UObjectProperty::SetObjectPropertyValue(void* PropertyValueAddress, UObject* Value) const
+{
+#if USE_DEFERRED_DEPENDENCY_CHECK_VERIFICATION_TESTS
+	if (ULinkerPlaceholderExportObject* PlaceholderVal = Cast<ULinkerPlaceholderExportObject>(Value))
+	{
+		check(PlaceholderVal->IsReferencedBy(this));
+	}
+#endif // USE_DEFERRED_DEPENDENCY_CHECK_VERIFICATION_TESTS
+
+	SetPropertyValue(PropertyValueAddress, Value);
+}
+
 IMPLEMENT_CORE_INTRINSIC_CLASS(UObjectProperty, UObjectPropertyBase,
 	{
 	}
diff --git a/Engine/Source/Runtime/CoreUObject/Public/Blueprint/BlueprintSupport.h b/Engine/Source/Runtime/CoreUObject/Public/Blueprint/BlueprintSupport.h
--- a/Engine/Source/Runtime/CoreUObject/Public/Blueprint/BlueprintSupport.h
+++ b/Engine/Source/Runtime/CoreUObject/Public/Blueprint/BlueprintSupport.h
@@ -25,6 +25,7 @@ struct FBlueprintSupport
 	static bool UseDeferredDependencyLoading();
 	static bool IsResolvingDeferredDependenciesDisabled();
 	static bool IsDeferredCDOSerializationDisabled();
+	static bool IsDeferredExportCreationDisabled();
 };
 
 #if WITH_EDITOR
diff --git a/Engine/Source/Runtime/CoreUObject/Public/UObject/Linker.h b/Engine/Source/Runtime/CoreUObject/Public/UObject/Linker.h
--- a/Engine/Source/Runtime/CoreUObject/Public/UObject/Linker.h
+++ b/Engine/Source/Runtime/CoreUObject/Public/UObject/Linker.h
@@ -1618,6 +1618,14 @@ private:
 	 */
 	UObject* CreateExportAndPreload(int32 ExportIndex, bool bForcePreload = false);
 
+	/**
+	 * Utility function for easily retrieving the specified export's UClass.
+	 * 
+	 * @param  ExportIndex    Index of the export you want a class for.
+	 * @return The class that the specified export's ClassIndex references.
+	 */
+	UClass* GetExportLoadClass(int32 ExportIndex);
+
 	/** 
 	 * Looks for and loads meta data object from export map.
 	 *
@@ -1641,6 +1649,17 @@ private:
 	 */
 	bool IsImportNative(const int32 ImportIndex) const;
 
+	/**
+	 * Attempts to lookup and return the corresponding ULinkerLoad object for 
+	 * the specified import WITHOUT invoking  a load, or continuing to load 
+	 * the import package (will only return one if it has already been 
+	 * created... could still be in the process of loading).
+	 * 
+	 * @param  ImportIndex    Specifies the import that you would like a linker for.
+	 * @return The imports associated linker (null if it hasn't been created yet).
+	 */
+	ULinkerLoad* FindExistingLinkerForImport(int32 ImportIndex) const;
+
 	UObject* IndexToObject( FPackageIndex Index );
 
 	void DetachExport( int32 i );
@@ -1838,6 +1857,17 @@ private:
 	 * @return True if the specified import was deferred, other wise false (it is ok to load it).
 	 */
 	bool DeferPotentialCircularImport(const int32 ImportIndex);
+	
+	/**
+	 * Stubs in a ULinkerPlaceholderExportObject for the specified export (if 
+	 * one is required, meaning: the export's LoadClass is not fully formed). 
+	 * This should rarely happen, but has been seen in cyclic Blueprint 
+	 * scenarios involving Blueprinted components.
+	 * 
+	 * @param  ExportIndex    Identifies the export you want deferred.
+	 * @return The specified FObjectExport's Object member (unchanged if the deferral failed, or wasn't required).
+	 */
+	UObject* DeferExportCreation(const int32 ExportIndex);
 
 	/**
 	 * Combs the ImportMap for any imports that were deferred, and then creates 
@@ -1876,6 +1906,18 @@ private:
 	 */
 	void FinalizeBlueprint(UClass* LoadClass);
 
+	/**
+	 * Combs the ExportMap for any stubbed in ULinkerPlaceholderExportObjects,
+	 * and finalizes the real export's class before actually creating it
+	 * (exports are deferred when their class isn't fully formed at the time
+	 * CreateExport() is called). Also, this function ensures that deferred CDO  
+	 * serialization is executed (expects its class to be fully resolved at this
+	 * point).
+	 *
+	 * @param  LoadClass    A fully loaded/serialized class that may have property references to placeholder export objects (in need of fix-up).
+	 */
+	void ResolveDeferredExports(UClass* LoadClass);
+
 	/**
 	 * Query method to help handle recursive behavior. When this returns true, 
 	 * this linker is in the middle of, or is about to call FinalizeBlueprint()
@@ -1894,7 +1936,7 @@ private:
 	 * happening, we instead defer it and record the export's index here (so we 
 	 * can return to it later).
 	 */
-	int32 DeferredExportIndex;
+	int32 DeferredCDOIndex;
 
 	/** 
 	 * Used to track dependency placeholders currently being resolved inside of 
diff --git a/Engine/Source/Runtime/CoreUObject/Public/UObject/UnrealType.h b/Engine/Source/Runtime/CoreUObject/Public/UObject/UnrealType.h
--- a/Engine/Source/Runtime/CoreUObject/Public/UObject/UnrealType.h
+++ b/Engine/Source/Runtime/CoreUObject/Public/UObject/UnrealType.h
@@ -1790,10 +1790,7 @@ class COREUOBJECT_API UObjectProperty : public TUObjectPropertyBase<UObject*>
 	{
 		return GetPropertyValue(PropertyValueAddress);
 	}
-	virtual void SetObjectPropertyValue(void* PropertyValueAddress, UObject* Value) const override
-	{
-		SetPropertyValue(PropertyValueAddress, Value);
-	}
+	virtual void SetObjectPropertyValue(void* PropertyValueAddress, UObject* Value) const override;
 	// End of UObjectPropertyBase interface
 };
 
diff --git a/Engine/Source/Runtime/Engine/Private/BlueprintGeneratedClass.cpp b/Engine/Source/Runtime/Engine/Private/BlueprintGeneratedClass.cpp
--- a/Engine/Source/Runtime/Engine/Private/BlueprintGeneratedClass.cpp
+++ b/Engine/Source/Runtime/Engine/Private/BlueprintGeneratedClass.cpp
@@ -83,22 +83,29 @@ void UBlueprintGeneratedClass::PostLoad()
 void UBlueprintGeneratedClass::GetRequiredPreloadDependencies(TArray<UObject*>& DependenciesOut)
 {
 	Super::GetRequiredPreloadDependencies(DependenciesOut);
-	for (UActorComponent* Component : ComponentTemplates)
-	{
-		// because of the linker's way of handling circular dependencies (with 
-		// placeholder blueprint classes), we need to ensure that class owned 
-		// blueprint components are created before the class's  is 
-		// ComponentTemplates member serialized in (otherwise, the component 
-		// would be created as a ULinkerPlaceholderClass instance)
-		//
-		// by returning these in the DependenciesOut array, we're making it 
-		// known that they should be prioritized in the package's ExportMap 
-		// before this class
-		if (Component->GetClass()->HasAnyClassFlags(CLASS_CompiledFromBlueprint))
-		{
-			DependenciesOut.Add(Component);
-		}
-	}
+
+	// the component templates are no longer needed as Preload() dependencies 
+	// (ULinkerLoad now handles these with placeholder export objects instead)...
+	// this change was prompted by a cyclic case, where creating the first
+	// component-template tripped the serialization of its class outer, before 
+	// another second component-template could be created (even though the 
+	// second component was listed in the ExportMap before the class)
+// 	for (UActorComponent* Component : ComponentTemplates)
+// 	{
+// 		// because of the linker's way of handling circular dependencies (with 
+// 		// placeholder blueprint classes), we need to ensure that class owned 
+// 		// blueprint components are created before the class's  is 
+// 		// ComponentTemplates member serialized in (otherwise, the component 
+// 		// would be created as a ULinkerPlaceholderClass instance)
+// 		//
+// 		// by returning these in the DependenciesOut array, we're making it 
+// 		// known that they should be prioritized in the package's ExportMap 
+// 		// before this class
+// 		if (Component->GetClass()->HasAnyClassFlags(CLASS_CompiledFromBlueprint))
+// 		{
+// 			DependenciesOut.Add(Component);
+// 		}
+// 	}
 }
 
 #if WITH_EDITOR
