[INTEGRATE] Change 2432194 by [REDACTED_PERSON][REDACTED_HANDLE] on 2015/02/04 14:44:15

	[UE-9004] Missing files to go along with CL 2432091.

diff --git a/Engine/Source/Runtime/CoreUObject/Private/UObject/LinkerPlaceholderExportObject.cpp b/Engine/Source/Runtime/CoreUObject/Private/UObject/LinkerPlaceholderExportObject.cpp
new file mode 100644
--- /dev/null
+++ b/Engine/Source/Runtime/CoreUObject/Private/UObject/LinkerPlaceholderExportObject.cpp
@@ -0,0 +1,160 @@
+// Copyright 1998-2015 Epic Games, Inc. All Rights Reserved.
+
+#include "CoreUObjectPrivate.h"
+#include "LinkerPlaceholderExportObject.h"
+#include "Blueprint/BlueprintSupport.h"
+
+//------------------------------------------------------------------------------
+ULinkerPlaceholderExportObject::ULinkerPlaceholderExportObject(const FObjectInitializer& ObjectInitializer)
+	: Super(ObjectInitializer)
+#if USE_DEFERRED_DEPENDENCY_CHECK_VERIFICATION_TESTS
+	, bWasResolved(false)
+#endif // USE_DEFERRED_DEPENDENCY_CHECK_VERIFICATION_TESTS
+{
+}
+
+//------------------------------------------------------------------------------
+IMPLEMENT_CORE_INTRINSIC_CLASS(ULinkerPlaceholderExportObject, UObject,
+	{
+	}
+);
+
+//------------------------------------------------------------------------------
+void ULinkerPlaceholderExportObject::BeginDestroy()
+{
+#if USE_DEFERRED_DEPENDENCY_CHECK_VERIFICATION_TESTS
+	check(bWasResolved || HasAnyFlags(RF_ClassDefaultObject));
+	check(ReferencingProperties.Num() == 0);
+#endif // USE_DEFERRED_DEPENDENCY_CHECK_VERIFICATION_TESTS
+
+	Super::BeginDestroy();
+}
+
+//------------------------------------------------------------------------------
+void ULinkerPlaceholderExportObject::AddReferencingProperty(const UObjectProperty* ReferencingProperty)
+{
+#if USE_DEFERRED_DEPENDENCY_CHECK_VERIFICATION_TESTS
+	check(bWasResolved != true);
+	check(!ReferencingProperty->GetLinker() || ReferencingProperty->GetLinker()->LinkerRoot == GetOuter());
+#endif // USE_DEFERRED_DEPENDENCY_CHECK_VERIFICATION_TESTS
+
+	ReferencingProperties.Add(ReferencingProperty);
+}
+
+//------------------------------------------------------------------------------
+bool ULinkerPlaceholderExportObject::IsReferencedBy(const UObjectProperty* Property)
+{
+	return ReferencingProperties.Contains(Property);
+}
+
+//------------------------------------------------------------------------------
+void ULinkerPlaceholderExportObject::ReplaceReferencingObjectValues(UClass* Container, UObject* NewObjectValue)
+{
+	UClass* ClassOwner = Container->GetClass();
+	for (TWeakObjectPtr<UObjectProperty> PropertyPtr : ReferencingProperties)
+	{
+		if (!PropertyPtr.IsValid())
+		{
+			continue;
+		}
+		const UObjectProperty* Property = PropertyPtr.Get();
+
+#if USE_DEFERRED_DEPENDENCY_CHECK_VERIFICATION_TESTS
+		check(Property->GetOwnerClass() == ClassOwner);
+#endif // USE_DEFERRED_DEPENDENCY_CHECK_VERIFICATION_TESTS
+
+		TArray<const UProperty*> PropertyChain;
+		PropertyChain.Add(Property);
+		const UProperty* OutermostProperty = Property;
+
+		UObject* PropertyOuter = Property->GetOuter();
+		while ((PropertyOuter != nullptr) && (PropertyOuter != ClassOwner))
+		{
+			if (const UProperty* PropertyOwner = Cast<const UProperty>(PropertyOuter))
+			{
+				OutermostProperty = PropertyOwner;
+				PropertyChain.Add(PropertyOwner);
+			}
+			PropertyOuter = PropertyOuter->GetOuter();
+		}
+
+		uint8* OutermostAddress = OutermostProperty->ContainerPtrToValuePtr<uint8>((uint8*)Container, /*ArrayIndex =*/0);
+		int32 ResolvedCount = ResolvePlaceholderValues(PropertyChain, PropertyChain.Num() - 1, OutermostAddress, NewObjectValue);
+
+#if USE_DEFERRED_DEPENDENCY_CHECK_VERIFICATION_TESTS
+		// we expect that (because we have had ReferencingProperties added) 
+		// there should be at least one reference that is resolved... if this 
+		// there were none, then a property could have changed its value after
+		// it was set to this
+		// 
+		// NOTE: this may seem it can be resolved by properties removing 
+		//       themselves from ReferencingProperties, but certain properties
+		//       may be the inner of a UArrayProperty (meaning there could be
+		//       multiple references per property)... we'd have to inc/decrement
+		//       a property ref-count to resolve that scenario
+		check(ResolvedCount > 0);
+#endif // USE_DEFERRED_DEPENDENCY_CHECK_VERIFICATION_TESTS
+	}
+
+	ReferencingProperties.Empty();
+#if USE_DEFERRED_DEPENDENCY_CHECK_VERIFICATION_TESTS
+	bWasResolved = true;
+#endif // USE_DEFERRED_DEPENDENCY_CHECK_VERIFICATION_TESTS
+}
+
+//------------------------------------------------------------------------------
+int32 ULinkerPlaceholderExportObject::ResolvePlaceholderValues(const TArray<const UProperty*>& PropertyChain, int32 ChainIndex, uint8* ValueAddress, UObject* ReplacementValue)
+{
+	int32 ReplacementCount = 0;
+
+	for (int32 PropertyIndex = ChainIndex; PropertyIndex >= 0; --PropertyIndex)
+	{
+		const UProperty* Property = PropertyChain[PropertyIndex];
+		if (PropertyIndex == 0)
+		{
+#if USE_DEFERRED_DEPENDENCY_CHECK_VERIFICATION_TESTS
+			check(Property->IsA<UObjectProperty>());
+#endif // USE_DEFERRED_DEPENDENCY_CHECK_VERIFICATION_TESTS
+
+			const UObjectProperty* ReferencingProperty = (const UObjectProperty*)Property;
+			if (ReferencingProperty->GetObjectPropertyValue(ValueAddress) == this)
+			{
+				ReferencingProperty->SetObjectPropertyValue(ValueAddress, ReplacementValue);
+				// @TODO: unfortunately, this is currently protected
+				//ReferencingProperty->CheckValidObject(ValueAddress);
+
+				++ReplacementCount;
+			}
+		}
+		else if (const UArrayProperty* ArrayProperty = Cast<const UArrayProperty>(Property))
+		{
+#if USE_DEFERRED_DEPENDENCY_CHECK_VERIFICATION_TESTS
+			const UProperty* NextProperty = PropertyChain[PropertyIndex - 1];
+			check(NextProperty == ArrayProperty->Inner);
+#endif // USE_DEFERRED_DEPENDENCY_CHECK_VERIFICATION_TESTS
+
+			// because we can't know which array entry was set with a reference 
+			// to this object, we have to comb through them all
+			FScriptArrayHelper ArrayHelper(ArrayProperty, ValueAddress);
+			for (int32 ArrayIndex = 0; ArrayIndex < ArrayHelper.Num(); ++ArrayIndex)
+			{
+				uint8* MemberAddress = ArrayHelper.GetRawPtr(ArrayIndex);
+				ReplacementCount += ResolvePlaceholderValues(PropertyChain, PropertyIndex - 1, MemberAddress, ReplacementValue);
+			}
+
+			// the above recursive call chewed through the rest of the
+			// PropertyChain, no need to keep on here
+			break;
+		}
+		else
+		{
+			const UProperty* NextProperty = PropertyChain[PropertyIndex - 1];
+#if USE_DEFERRED_DEPENDENCY_CHECK_VERIFICATION_TESTS
+			check(NextProperty->GetOuter() == Property);
+#endif // USE_DEFERRED_DEPENDENCY_CHECK_VERIFICATION_TESTS
+			ValueAddress = Property->ContainerPtrToValuePtr<uint8>(ValueAddress, /*ArrayIndex =*/0);
+		}
+	}
+
+	return ReplacementCount;
+}
diff --git a/Engine/Source/Runtime/CoreUObject/Private/UObject/LinkerPlaceholderExportObject.h b/Engine/Source/Runtime/CoreUObject/Private/UObject/LinkerPlaceholderExportObject.h
new file mode 100644
--- /dev/null
+++ b/Engine/Source/Runtime/CoreUObject/Private/UObject/LinkerPlaceholderExportObject.h
@@ -0,0 +1,92 @@
+// Copyright 1998-2014 Epic Games, Inc. All Rights Reserved.
+
+#pragma once
+#include "Object.h"
+#include "Set.h"
+
+// Forward declarations
+class FObjectInitializer;
+class UObjectProperty;
+
+/**  
+ * A utility class for the deferred dependency loader, used to stub in temporary
+ * export objects so we don't spawn any Blueprint class instances before the 
+ * class is fully regenerated.
+ */ 
+class ULinkerPlaceholderExportObject : public UObject
+{
+public:
+	DECLARE_CASTED_CLASS_INTRINSIC_NO_CTOR(ULinkerPlaceholderExportObject, UObject, /*TStaticFlags =*/0, CoreUObject, /*TStaticCastFlags =*/0, NO_API)
+
+	ULinkerPlaceholderExportObject(const FObjectInitializer& ObjectInitializer);
+
+	// UObject interface
+	virtual void BeginDestroy() override;
+	// End of UObject interface
+
+	/**
+	 * Logs the supplied property as referencing this class (meaning an instance  
+	 * is referencing this object through that property's value).
+	 * 
+	 * @param  ReferencingProperty    The property whose object-value is pointing to this class.
+	 */
+	void AddReferencingProperty(const UObjectProperty* ReferencingProperty);
+
+	/**
+	 * Checks to see if the specified property is being tracked by this as a
+	 * referencer.
+	 * 
+	 * @param  Property    The property you want to check for.
+	 * @return True if the property was added as a referencer through AddReferencingProperty(), otherwise false.
+	 */
+	bool IsReferencedBy(const UObjectProperty* Property);
+
+	/**
+	 * Resolves all referencing object values (assuming they were all made as 
+	 * part of the specified Container object).
+	 *
+	 * NOTE: there is a very strict assumption here, that any property values 
+	 *       that were set to this (and recorded through AddReferencingProperty)
+	 *       were properties belonging to the specified Container object (if not
+	 *       there could be dangling, unresolved references).
+	 * 
+	 * @param  Container		The object whose properties may be referencing this temporary object.
+	 * @param  NewObjectValue   The object which should replace all references to this. 
+	 */
+	void ReplaceReferencingObjectValues(UClass* Container, UObject* NewObjectValue);
+
+private:
+	/**
+	 * A recursive method that replaces all leaf references to this object with 
+	 * the supplied ReplacementValue.
+	 *
+	 * This function recurses the property chain (from class owner down) because 
+	 * at the time of AddReferencingProperty() we cannot know/record the 
+	 * address/index of array properties (as they may change during array 
+	 * re-allocation or compaction). So we must follow the property chain and 
+	 * check every UArrayProperty member for references to this (hence, the need
+	 * for this recursive function).
+	 * 
+	 * @param  PropertyChain    An ascending outer chain, where the property at index zero is the leaf (referencer) property.
+	 * @param  ChainIndex		An index into the PropertyChain that this call should start at and iterate DOWN to zero.
+	 * @param  ValueAddress     The memory address of the value corresponding to the property at ChainIndex.
+	 * @param  ReplacementValue	The new object to replace all references to this with.
+	 * @return The number of references that were replaced.
+	 */
+	int32 ResolvePlaceholderValues(const TArray<const UProperty*>& PropertyChain, int32 ChainIndex, uint8* ValueAddress, UObject* ReplacementValue);
+
+	/**
+	 * A list of KNOWN referencing properties (meaning an object instance 
+	 * references this object through these properties).
+	 *
+	 * @TODO: consider making this a map with a ref-counter, so we can remove 
+	 *        property references (as they may be replaced prior to resolve)
+	 */
+	TSet< TWeakObjectPtr<UObjectProperty> > ReferencingProperties;
+
+#if USE_DEFERRED_DEPENDENCY_CHECK_VERIFICATION_TESTS
+public:
+	/** Used to catch any references that were added after ULinkerLoad has already done a resolve pass. */
+	bool bWasResolved;
+#endif // USE_DEFERRED_DEPENDENCY_CHECK_VERIFICATION_TESTS
+}; 
