[INTEGRATE] Change 2432849 by [REDACTED_PERSON][REDACTED_HANDLE] on 2015/02/04 21:02:45

	Core: Add support for optional localizable short tooltip metadata
	- UField::GetToolTipText now takes an optional parameter to indicate if the short tooltip is desired, but will fall back to the long tooltip if none was specified

diff --git a/Engine/Config/Localization/ToolTips.ini b/Engine/Config/Localization/ToolTips.ini
--- a/Engine/Config/Localization/ToolTips.ini
+++ b/Engine/Config/Localization/ToolTips.ini
@@ -23,6 +23,9 @@ ExcludePaths=*/Source/Developer/NoRedist/CommunityPortalServices/*
 InputKeys=ToolTip
 OutputNamespaces=UObjectToolTips
 OutputKeys="{FieldPath}"
+InputKeys=ShortTooltip
+OutputNamespaces=UObjectShortTooltips
+OutputKeys="{FieldPath}"
 
 ;Write Manifest
 [GatherTextStep1]
diff --git a/Engine/Source/Runtime/CoreUObject/Private/UObject/Class.cpp b/Engine/Source/Runtime/CoreUObject/Private/UObject/Class.cpp
--- a/Engine/Source/Runtime/CoreUObject/Private/UObject/Class.cpp
+++ b/Engine/Source/Runtime/CoreUObject/Private/UObject/Class.cpp
@@ -251,12 +251,32 @@ FText UField::GetDisplayNameText() const
  *
  * @return The tooltip for this object.
  */
-FText UField::GetToolTipText() const
+FText UField::GetToolTipText(bool bShortTooltip) const
 {
+	bool bFoundShortTooltip = false;
+	static const FName NAME_Tooltip(TEXT("Tooltip"));
+	static const FName NAME_ShortTooltip(TEXT("ShortTooltip"));
 	FText LocalizedToolTip;
-	FString NativeToolTip = GetMetaData( TEXT("Tooltip") );
+	FString NativeToolTip;
+	
+	if (bShortTooltip)
+	{
+		NativeToolTip = GetMetaData(NAME_ShortTooltip);
+		if (NativeToolTip.IsEmpty())
+		{
+			NativeToolTip = GetMetaData(NAME_Tooltip);
+		}
+		else
+		{
+			bFoundShortTooltip = true;
+		}
+	}
+	else
+	{
+		NativeToolTip = GetMetaData(NAME_Tooltip);
+	}
 
-	static const FString Namespace = TEXT("UObjectToolTips");
+	const FString Namespace = bFoundShortTooltip ? TEXT("UObjectShortTooltips") : TEXT("UObjectToolTips");
 	const FString Key = GetFullGroupName(true) + TEXT(".") + GetName();
 	if ( !(FText::FindText( Namespace, Key, /*OUT*/LocalizedToolTip )) || *FTextInspector::GetSourceString(LocalizedToolTip) != NativeToolTip)
 	{
diff --git a/Engine/Source/Runtime/CoreUObject/Public/UObject/Class.h b/Engine/Source/Runtime/CoreUObject/Public/UObject/Class.h
--- a/Engine/Source/Runtime/CoreUObject/Public/UObject/Class.h
+++ b/Engine/Source/Runtime/CoreUObject/Public/UObject/Class.h
@@ -82,9 +82,11 @@ class COREUOBJECT_API UField : public UObject
 	/**
 	 * Finds the localized tooltip or native tooltip as a fallback.
 	 *
+	 * @param bShortTooltip Look for a shorter version of the tooltip (falls back to the long tooltip if none was specified)
+	 *
 	 * @return The tooltip for this object.
 	 */
-	FText GetToolTipText() const;
+	FText GetToolTipText(bool bShortTooltip = false) const;
 
 	/**
 	 * Determines if the property has any metadata associated with the key
diff --git a/Engine/Source/Runtime/CoreUObject/Public/UObject/ObjectBase.h b/Engine/Source/Runtime/CoreUObject/Public/UObject/ObjectBase.h
--- a/Engine/Source/Runtime/CoreUObject/Public/UObject/ObjectBase.h
+++ b/Engine/Source/Runtime/CoreUObject/Public/UObject/ObjectBase.h
@@ -876,6 +876,16 @@ namespace US
 // Metadata specifiers
 namespace UM
 {
+	// Metadata usable in any UField (UCLASS(), USTRUCT(), UPROPERTY(), UFUNCTION(), etc...)
+	enum
+	{
+		/// Overrides the automatically generated tooltip from the class comment
+		ToolTip,
+
+		/// A short tooltip that is used in some contexts where the full tooltip might be overwhelming (such as the parent class picker dialog)
+		ShortTooltip,
+	};
+
 	// Metadata usable in UCLASS
 	enum
 	{
@@ -909,7 +919,6 @@ namespace UM
 		/// [ClassMetadata] Indicates that when placing blueprint nodes in graphs owned by this class that the hidden world context pin should be visible because the self context of the class cannot
 		///                 provide the world context and it must be wired in manually
 		ShowWorldContextPin,
-
 	};
 
 	// Metadata usable in USTRUCT
@@ -1009,9 +1018,6 @@ namespace UM
 		// [PropertyMetadata] Used by numeric properties. Indicates how rapidly the value will grow when moving an unbounded slider.
 		SliderExponent,
 
-		// [PropertyMetadata] Overrides the automatically generated tooltip from the property comment.
-		ToolTip,
-
 		/// [PropertyMetadata] Used for float and integer properties.  Specifies the lowest that the value slider should represent.
 		UIMin,
 
