Making FPackageName::FilenameToLongPackageName use TryConvertFilenameToLongPackageName to make error checking and reporting use one code path.

diff --git a/Engine/Source/Runtime/CoreUObject/Private/Misc/PackageName.cpp b/Engine/Source/Runtime/CoreUObject/Private/Misc/PackageName.cpp
--- a/Engine/Source/Runtime/CoreUObject/Private/Misc/PackageName.cpp
+++ b/Engine/Source/Runtime/CoreUObject/Private/Misc/PackageName.cpp
@@ -276,30 +276,21 @@ FString FPackageName::InternalFilenameToLongPackageName(const FString& InFilenam
 	return Result;
 }
 
-bool FPackageName::TryConvertFilenameToLongPackageName(const FString& InFilename, FString& OutPackageName)
+bool FPackageName::TryConvertFilenameToLongPackageName(const FString& InFilename, FString& OutPackageName, FString* OutFailureReason)
 {
-	FString Result = InternalFilenameToLongPackageName(InFilename);
+	FString LongPackageName = InternalFilenameToLongPackageName(InFilename);
 
 	// we don't support loading packages from outside of well defined places
-	if (Result.Contains(TEXT("."), ESearchCase::CaseSensitive) || Result.Contains(TEXT("\\"), ESearchCase::CaseSensitive) || Result.Contains(TEXT(":"), ESearchCase::CaseSensitive))
+	const bool bContainsDot = LongPackageName.Contains(TEXT("."), ESearchCase::CaseSensitive);
+	const bool bContainsBackslash = LongPackageName.Contains(TEXT("\\"), ESearchCase::CaseSensitive);
+	const bool bContainsColon = LongPackageName.Contains(TEXT(":"), ESearchCase::CaseSensitive);
+	const bool bResult = !(bContainsDot || bContainsBackslash || bContainsColon);
+
+	if (bResult)
 	{
-		return false;
+		OutPackageName = MoveTemp(LongPackageName);
 	}
-
-	OutPackageName = MoveTemp(Result);
-	return true;
-}
-
-FString FPackageName::FilenameToLongPackageName(const FString& InFilename)
-{
-	FString Result = InternalFilenameToLongPackageName(InFilename);
-
-	// we don't support loading packages from outside of well defined places
-	const bool bContainsDot = Result.Contains(TEXT("."), ESearchCase::CaseSensitive);
-	const bool bContainsBackslash = Result.Contains(TEXT("\\"), ESearchCase::CaseSensitive);
-	const bool bContainsColon = Result.Contains(TEXT(":"), ESearchCase::CaseSensitive);
-
-	if (bContainsDot || bContainsBackslash || bContainsColon)
+	else if (OutFailureReason != nullptr)
 	{
 		FString InvalidChars;
 		if (bContainsDot)
@@ -314,10 +305,19 @@ FString FPackageName::FilenameToLongPackageName(const FString& InFilename)
 		{
 			InvalidChars += TEXT(":");
 		}
-
-		UE_LOG(LogPackageName, Fatal,TEXT("FilenameToLongPackageName failed to convert '%s'. Attempt result was '%s', but the path contains illegal characters '%s'"), *InFilename, *Result, *InvalidChars);
+		*OutFailureReason = FString::Printf(TEXT("FilenameToLongPackageName failed to convert '%s'. Attempt result was '%s', but the path contains illegal characters '%s'"), *InFilename, *LongPackageName, *InvalidChars);
 	}
+	return bResult;
+}
 
+FString FPackageName::FilenameToLongPackageName(const FString& InFilename)
+{
+	FString FailureReason;
+	FString Result;
+	if (!TryConvertFilenameToLongPackageName(InFilename, Result, &FailureReason))
+	{
+		UE_LOG(LogPackageName, Fatal, TEXT("%s"), *FailureReason);
+	}
 	return Result;
 }
 
diff --git a/Engine/Source/Runtime/CoreUObject/Public/Misc/PackageName.h b/Engine/Source/Runtime/CoreUObject/Public/Misc/PackageName.h
--- a/Engine/Source/Runtime/CoreUObject/Public/Misc/PackageName.h
+++ b/Engine/Source/Runtime/CoreUObject/Public/Misc/PackageName.h
@@ -37,9 +37,10 @@ public:
 	 * 
 	 * @param InFilename Filename to convert.
 	 * @param OutPackageName The resulting long package name if the conversion was successful.
+	 * @param OutFailureReason Description of an error if the conversion failed.
 	 * @return Returns true if the supplied filename properly maps to one of the long package roots.
 	 */
-	static bool TryConvertFilenameToLongPackageName(const FString& InFilename, FString& OutPackageName);
+	static bool TryConvertFilenameToLongPackageName(const FString& InFilename, FString& OutPackageName, FString* OutFailureReason = nullptr);
 	/** 
 	 * Converts the supplied filename to long package name.
 	 *  Throws a fatal error if the conversion is not successfull.
