[UE-8639] PR #779: Fix invalid default value for structures embedded in parent classes (Contributed by [REDACTED_PERSON])
https://github.com/EpicGames/UnrealEngine/pull/779

diff --git a/Engine/Source/Runtime/CoreUObject/Private/UObject/Class.cpp b/Engine/Source/Runtime/CoreUObject/Private/UObject/Class.cpp
--- a/Engine/Source/Runtime/CoreUObject/Private/UObject/Class.cpp
+++ b/Engine/Source/Runtime/CoreUObject/Private/UObject/Class.cpp
@@ -1231,9 +1231,10 @@ void UStruct::SerializeTaggedProperties(FArchive& Ar, uint8* Data, UStruct* Defa
 			else
 			{
 				uint8* DestAddress = Property->ContainerPtrToValuePtr<uint8>(Data, Tag.ArrayIndex);  
+				uint8* DefaultsFromParent = Defaults ? Property->ContainerPtrToValuePtr<uint8>(Defaults, Tag.ArrayIndex) : nullptr;
 
 				// This property is ok.			
-				Tag.SerializeTaggedProperty( Ar, Property, DestAddress, NULL );
+				Tag.SerializeTaggedProperty(Ar, Property, DestAddress, DefaultsFromParent);
 
 				AdvanceProperty = true;
 				continue;
diff --git a/Engine/Source/Runtime/Engine/Private/UserDefinedStruct.cpp b/Engine/Source/Runtime/Engine/Private/UserDefinedStruct.cpp
--- a/Engine/Source/Runtime/Engine/Private/UserDefinedStruct.cpp
+++ b/Engine/Source/Runtime/Engine/Private/UserDefinedStruct.cpp
@@ -121,7 +121,7 @@ void UUserDefinedStruct::SerializeTaggedProperties(FArchive& Ar, uint8* Data, US
 	}
 	if (bLoadDefaultFirst)
 	{
-		UDDefaultsStruct->CopyScriptStruct(Data, StructDefaultMem.GetStructMemory());
+		UDDefaultsStruct->CopyScriptStruct(Data, Defaults ? Defaults : StructDefaultMem.GetStructMemory());
 	}
 #endif // WITH_EDITOR
 	Super::SerializeTaggedProperties(Ar, Data, DefaultsStruct, Defaults);
