[INTEGRATE] Change 2432988 by [REDACTED_PERSON][REDACTED_HANDLE] on 2015/02/05 00:20:27

	Core: Fix error message stripping off the _# from a FName when objects have the same fully qualified name but different paths

diff --git a/Engine/Source/Runtime/CoreUObject/Private/UObject/UObjectGlobals.cpp b/Engine/Source/Runtime/CoreUObject/Private/UObject/UObjectGlobals.cpp
--- a/Engine/Source/Runtime/CoreUObject/Private/UObject/UObjectGlobals.cpp
+++ b/Engine/Source/Runtime/CoreUObject/Private/UObject/UObjectGlobals.cpp
@@ -1750,7 +1750,7 @@ UObject* StaticAllocateObject
 				TEXT("Objects have the same fully qualified name but different paths.\n")
 				TEXT("\tNew Object: %s %s.%s\n")
 				TEXT("\tExisting Object: %s"),
-				*InClass->GetName(), InOuter ? *InOuter->GetPathName() : TEXT(""), *InName.GetPlainNameString(),
+				*InClass->GetName(), InOuter ? *InOuter->GetPathName() : TEXT(""), *InName.ToString(),
 				*Obj->GetFullName());
 		}
 	}
