Fixed leaking of file handles in cooked builds.

ResetLoaders removes the relevant Linker from GDelayedLinkerClosePackages. This meant that GDelayedLinkerClosePackages was being modified as it was being iterated over, causing many items to be skipped. ResetLoaders also deletes the loader, so I've replaced the delete here with a check.

diff --git a/Engine/Source/Runtime/CoreUObject/Private/UObject/UObjectGlobals.cpp b/Engine/Source/Runtime/CoreUObject/Private/UObject/UObjectGlobals.cpp
--- a/Engine/Source/Runtime/CoreUObject/Private/UObject/UObjectGlobals.cpp
+++ b/Engine/Source/Runtime/CoreUObject/Private/UObject/UObjectGlobals.cpp
@@ -1205,20 +1205,18 @@ void EndLoad()
 		DissociateImportsAndForcedExports();
 
 		// close any linkers' loaders that were requested to be closed once GObjBeginLoadCount goes to 0
-		for (int32 Index = 0; Index < GDelayedLinkerClosePackages.Num(); Index++)
+		auto PackagesToClose = MoveTemp(GDelayedLinkerClosePackages);
+		for (auto Linker: PackagesToClose)
 		{
-			ULinkerLoad* Linker = GDelayedLinkerClosePackages[Index];
 			if (Linker)
 			{
 				if (Linker->Loader && Linker->LinkerRoot)
 				{
 					ResetLoaders(Linker->LinkerRoot);
 				}
-				delete Linker->Loader;
-				Linker->Loader = NULL;
+				check(Linker->Loader == nullptr);
 			}
 		}
-		GDelayedLinkerClosePackages.Empty();
 	}
 }
 
