UE-9617 Crash occurs in RegenerateBlueprintClass opening the UMG content examples map.

This blocks automation testing from finishing on the content examples project.

diff --git a/Engine/Source/Runtime/CoreUObject/Private/UObject/Class.cpp b/Engine/Source/Runtime/CoreUObject/Private/UObject/Class.cpp
--- a/Engine/Source/Runtime/CoreUObject/Private/UObject/Class.cpp
+++ b/Engine/Source/Runtime/CoreUObject/Private/UObject/Class.cpp
@@ -1231,7 +1231,7 @@ void UStruct::SerializeTaggedProperties(FArchive& Ar, uint8* Data, UStruct* Defa
 			else
 			{
 				uint8* DestAddress = Property->ContainerPtrToValuePtr<uint8>(Data, Tag.ArrayIndex);  
-				uint8* DefaultsFromParent = Defaults ? Property->ContainerPtrToValuePtr<uint8>(Defaults, Tag.ArrayIndex) : nullptr;
+				uint8* DefaultsFromParent = Property->ContainerPtrToValuePtrForDefaults<uint8>(DefaultsStruct, Defaults, Tag.ArrayIndex);
 
 				// This property is ok.			
 				Tag.SerializeTaggedProperty(Ar, Property, DestAddress, DefaultsFromParent);
