[INTEGRATE] Change 2442128 by [REDACTED_PERSON][REDACTED_HANDLE] on 2015/02/11 14:52

	Ensure that all ActorComponents that belong to a BlueprintGeneratedClass are correctly flagged as RF_Public in the ExportMap
	Remove now unnecessary code previously trying to get flag set correctly

diff --git a/Engine/Source/Runtime/CoreUObject/Private/UObject/LinkerLoad.cpp b/Engine/Source/Runtime/CoreUObject/Private/UObject/LinkerLoad.cpp
--- a/Engine/Source/Runtime/CoreUObject/Private/UObject/LinkerLoad.cpp
+++ b/Engine/Source/Runtime/CoreUObject/Private/UObject/LinkerLoad.cpp
@@ -480,19 +480,6 @@ void ULinkerLoad::InvalidateExport(UObject* OldObject)
 	}
 }
 
-void ULinkerLoad::RefreshExportFlags(UObject* Object)
-{
-	ULinkerLoad* ObjectLinker = Object->GetLinker();
-	const int32 CachedLinkerIndex = Object->GetLinkerIndex();
-
-	if (ObjectLinker && ObjectLinker->ExportMap.IsValidIndex(CachedLinkerIndex))
-	{
-		FObjectExport& ObjExport = ObjectLinker->ExportMap[CachedLinkerIndex];
-		ObjExport.ObjectFlags = Object->GetMaskedFlags();
-	}
-}
-
-
 FName ULinkerLoad::FindSubobjectRedirectName(const FName& Name)
 {
 	FSubobjectRedirect *Redirect = SubobjectNameRedirects.Find(Name);
@@ -4244,6 +4231,25 @@ ULinkerLoad::ELinkerStatus ULinkerLoad::FixupExportMap()
 			FName NameClass = GetExportClassName(ExportMapIdx);
 			FName NamePackage = GetExportClassPackage(ExportMapIdx);
 
+			// ActorComponents outered to a BlueprintGeneratedClass need to be marked RF_Public, but older content was not created
+			// as such.  This updates the ExportTable such that they are correctly flagged when created and when other packages
+			// validate their imports
+			// TODO: Wrap this in a version check if possible. Mark up for future removal when minimum version is incremented
+			if ((Export.ObjectFlags & RF_Public) == 0)
+			{
+				static const FName NAME_BlueprintGeneratedClass("BlueprintGeneratedClass");
+				const FName OuterClassName = GetExportClassName(Export.OuterIndex);
+				if (OuterClassName == NAME_BlueprintGeneratedClass)
+				{
+					static const UClass* ActorComponentClass = FindObjectChecked<UClass>(ANY_PACKAGE, TEXT("ActorComponent"), true);
+					UClass* Class = FindObject<UClass>(ANY_PACKAGE, *NameClass.ToString());
+					if (Class && Class->IsChildOf(ActorComponentClass))
+					{
+						Export.ObjectFlags |= RF_Public;
+					}
+				}
+			}
+
 			{
 				FSubobjectRedirect *Redirect = SubobjectNameRedirects.Find(Export.ObjectName);
 				if (Redirect)
diff --git a/Engine/Source/Runtime/CoreUObject/Public/UObject/Linker.h b/Engine/Source/Runtime/CoreUObject/Public/UObject/Linker.h
--- a/Engine/Source/Runtime/CoreUObject/Public/UObject/Linker.h
+++ b/Engine/Source/Runtime/CoreUObject/Public/UObject/Linker.h
@@ -1595,11 +1595,6 @@ public:
 	 */
 	COREUOBJECT_API static void InvalidateExport(UObject* OldObject);
 
-	/**
-	 * Updates the export table with an Object's object flags
-	 */
-	COREUOBJECT_API static void RefreshExportFlags(UObject* Object);
-
 	/** Used by Matinee to fixup component renaming */
 	COREUOBJECT_API static FName FindSubobjectRedirectName(const FName& Name);
 
diff --git a/Engine/Source/Runtime/Engine/Classes/Components/ActorComponent.h b/Engine/Source/Runtime/Engine/Classes/Components/ActorComponent.h
--- a/Engine/Source/Runtime/Engine/Classes/Components/ActorComponent.h
+++ b/Engine/Source/Runtime/Engine/Classes/Components/ActorComponent.h
@@ -565,7 +565,6 @@ public:
 	 */
 	virtual void ApplyWorldOffset(const FVector& InOffset, bool bWorldShift) {};
 
-	virtual void Serialize(FArchive& Ar) override;
 #if !(UE_BUILD_SHIPPING || UE_BUILD_TEST)
 
 private:
diff --git a/Engine/Source/Runtime/Engine/Private/ActorComponent.cpp b/Engine/Source/Runtime/Engine/Private/ActorComponent.cpp
--- a/Engine/Source/Runtime/Engine/Private/ActorComponent.cpp
+++ b/Engine/Source/Runtime/Engine/Private/ActorComponent.cpp
@@ -164,18 +164,6 @@ void UActorComponent::PostLoad()
 // 	}
 }
 
-void UActorComponent::Serialize(FArchive& Ar)
-{
-	Super::Serialize(Ar);
-
-	// Fixup older blueprints that were private but now need to be public
-	if (!HasAllFlags(RF_Public) && GetOuter()->IsA<UBlueprintGeneratedClass>())
-	{
-		SetFlags(RF_Public);
-		ULinkerLoad::RefreshExportFlags(this);
-	}
-}
-
 void UActorComponent::PostRename(UObject* OldOuter, const FName OldName)
 {
 	if (OldOuter != GetOuter())
