Back out changelist 2433051 as it was causing some other troubles in hot-reload.

diff --git a/Engine/Source/Runtime/CoreUObject/Private/UObject/Class.cpp b/Engine/Source/Runtime/CoreUObject/Private/UObject/Class.cpp
--- a/Engine/Source/Runtime/CoreUObject/Private/UObject/Class.cpp
+++ b/Engine/Source/Runtime/CoreUObject/Private/UObject/Class.cpp
@@ -3397,6 +3397,10 @@ bool UClass::HotReloadPrivateStaticClass(
 		UClass::GetDefaultPropertiesFeedbackContext().Logf(ELogVerbosity::Warning, TEXT("Property size mismatch. Will not update class %s (was %d, new %d)."), *GetName(), PropertiesSize, InSize);
 		return false;
 	}
+	//We could do this later, but might as well get it before we start corrupting the object
+	UObject* CDO = GetDefaultObject();
+	void* OldVTable = *(void**)CDO;
+
 
 	//@todo safe? ClassFlags = InClassFlags | CLASS_Native;
 	//@todo safe? ClassCastFlags = InClassCastFlags;
@@ -3417,23 +3421,51 @@ bool UClass::HotReloadPrivateStaticClass(
 	ClassWithin = TClass_WithinClass_StaticClass;
 	*/
 
-	int32 CountClass = 0;
-	for (FRawObjectIterator It; It; ++It)
+	UE_LOG(LogClass, Verbose, TEXT("Attempting to change VTable for class %s."),*GetName());
+	ClassWithin = UPackage::StaticClass();  // We are just avoiding error checks with this...we don't care about this temp object other than to get the vtable.
+	UObject* TempObjectForVTable = StaticConstructObject(this, GetTransientPackage(), NAME_None, RF_NeedLoad | RF_ClassDefaultObject);
+
+	if( !TempObjectForVTable->IsRooted() )
+	{
+		TempObjectForVTable->MarkPendingKill();
+	}
+	else
+	{
+		UE_LOG(LogClass, Warning, TEXT("Hot Reload:  Was not expecting temporary object '%s' for class '%s' to become rooted during construction.  This object cannot be marked pending kill." ), *TempObjectForVTable->GetFName().ToString(), *this->GetName() );
+	}
+
+	ClassWithin = TClass_WithinClass_StaticClass;
+
+	void* NewVTable = *(void**)TempObjectForVTable;
+	if (NewVTable != OldVTable)
 	{
-		UObject* Target = *It;
-		if (dynamic_cast<UClass*>(Target))
+		int32 Count = 0;
+		int32 CountClass = 0;
+		for ( FRawObjectIterator It; It; ++It )
 		{
-			UClass *Class = CastChecked<UClass>(Target);
-			if (Class->ClassConstructor == OldClassConstructor)
+			UObject* Target = *It;
+			if (OldVTable == *(void**)Target)
+			{
+				*(void**)Target = NewVTable;
+				Count++;
+			}
+			else if (dynamic_cast<UClass*>(Target))
 			{
-				Class->ClassConstructor = ClassConstructor;
-				Class->ClassAddReferencedObjects = ClassAddReferencedObjects;
-				CountClass++;
+				UClass *Class = CastChecked<UClass>(Target);
+				if (Class->ClassConstructor == OldClassConstructor)
+				{
+					Class->ClassConstructor = ClassConstructor;
+					Class->ClassAddReferencedObjects = ClassAddReferencedObjects;
+					CountClass++;
+				}
 			}
 		}
+		UE_LOG(LogClass, Verbose, TEXT("Updated the vtable for %d live objects and %d blueprint classes.  %016llx -> %016llx"), Count, CountClass, PTRINT(OldVTable), PTRINT(NewVTable));
+	}
+	else
+	{
+		UE_LOG(LogClass, Error, TEXT("VTable for class %s did not change?"),*GetName());
 	}
-	UE_LOG(LogClass, Verbose, TEXT("Updated the internal methods %d blueprint classes."), CountClass);
-
 	return true;
 }
 
diff --git a/Engine/Source/Runtime/CoreUObject/Private/UObject/UObjectBase.cpp b/Engine/Source/Runtime/CoreUObject/Private/UObject/UObjectBase.cpp
--- a/Engine/Source/Runtime/CoreUObject/Private/UObject/UObjectBase.cpp
+++ b/Engine/Source/Runtime/CoreUObject/Private/UObject/UObjectBase.cpp
@@ -621,114 +621,6 @@ FString RemoveClassPrefix(const TCHAR* ClassName)
 	return NameWithoutPrefix;
 }
 
-/**
- * Helper class for updating VTables of hot reloaded classes.
- */
-class FUClassVTableHelper
-{
-public:
-	/**
-	 * Adds UClasses VTable address if it's not cached yet.
-	 *
-	 * @param Class A UClass for which to add the VTable address.
-	 */
-	void TouchUClassVTable(UClass* Class)
-	{
-		auto* Found = VTableMap.Find(Class);
-
-		if (Found != nullptr)
-		{
-			return;
-		}
-
-		VTableMap.Add(Class, GetUClassVTable(Class));
-	}
-
-	/**
-	 * Changes old VTable address to the new one after hot reload.
-	 *
-	 * @param OldClass Old class.
-	 * @param NewClass New class.
-	 */
-	void UpdateVTables(UClass* OldClass, UClass* NewClass)
-	{
-		UE_LOG(LogClass, Verbose, TEXT("Attempting to change VTable for class %s."), *OldClass->GetName());
-
-		// We could do this later, but might as well get it before we start corrupting the object
-		void* OldVTable = VTableMap[OldClass];
-		void* NewVTable = GetUClassVTable(NewClass);
-
-		if (NewVTable != OldVTable)
-		{
-			int32 Count = 0;
-			for (FRawObjectIterator It; It; ++It)
-			{
-				UObject* Target = *It;
-				if (OldVTable == *(void**)Target)
-				{
-					*(void**)Target = NewVTable;
-					Count++;
-				}
-			}
-			UE_LOG(LogClass, Verbose, TEXT("Updated the vtable for %d live objects. %016llx -> %016llx"), Count, PTRINT(OldVTable), PTRINT(NewVTable));
-
-			VTableMap[OldClass] = NewVTable;
-			VTableMap.Add(NewClass, NewVTable);
-		}
-		else
-		{
-			UE_LOG(LogClass, Error, TEXT("VTable for class %s did not change?"), *NewClass->GetName());
-		}
-	}
-
-	/**
-	 * Gets the singleton.
-	 */
-	static FUClassVTableHelper& Get()
-	{
-		static FUClassVTableHelper Helper;
-		return Helper;
-	}
-
-private:
-	/**
-	 * Forbid public construction.
-	 */
-	FUClassVTableHelper() {}
-
-	/**
-	 * Gets UClasses VTable address.
-	 *
-	 * Note that this function creates an object of UClass type.
-	 *
-	 * @param Class A UClass for which to get VTable address.
-	 *
-	 * @returns VTable address.
-	 */
-	static void* GetUClassVTable(UClass* Class)
-	{
-		auto* ClassWithin = Class->ClassWithin;
-		Class->ClassWithin = UPackage::StaticClass(); // We are just avoiding error checks with this... we don't care about this temp object other than to get the vtable.
-		UObject* TempObjectForVTable = StaticConstructObject_Internal(Class, GetTransientPackage(), NAME_None, RF_NeedLoad | RF_ClassDefaultObject);
-
-		if (!TempObjectForVTable->IsRooted())
-		{
-			TempObjectForVTable->MarkPendingKill();
-		}
-		else
-		{
-			UE_LOG(LogClass, Warning, TEXT("Hot Reload: Was not expecting temporary object '%s' for class '%s' to become rooted during construction. This object cannot be marked pending kill."), *TempObjectForVTable->GetFName().ToString(), *Class->GetName());
-		}
-
-		Class->ClassWithin = ClassWithin;
-
-		return *(void**)TempObjectForVTable;
-	}
-
-	// A map of classes VTables.
-	TMap<UClass*, void*> VTableMap;
-};
-
 void UClassCompiledInDefer(FFieldCompiledInInfo* ClassInfo, const TCHAR* Name, SIZE_T ClassSize, uint32 Crc)
 {
 	const FName CPPClassName(Name);
@@ -747,9 +639,6 @@ void UClassCompiledInDefer(FFieldCompiledInInfo* ClassInfo, const TCHAR* Name, S
 
 		if (ClassInfo->bHasChanged)
 		{
-			// Check in cache and create if doesn't exist UCLASSes VTable address.
-			FUClassVTableHelper::Get().TouchUClassVTable(ExistingClass);
-
 			// Rename the old class and move it to transient package
 			ExistingClass->ClearFlags(RF_RootSet | RF_Standalone | RF_Public);
 			ExistingClass->GetDefaultObject()->ClearFlags(RF_RootSet | RF_Standalone | RF_Public);
@@ -809,38 +698,21 @@ void UClassRegisterAllCompiledInClasses()
 /** Re-instance all existing classes that have changed during hot-reload */
 void UClassReplaceHotReloadClasses()
 {
-	struct FClassTuple
-	{
-		UClass* OldClass;
-		UClass* NewClass;
-
-		FClassTuple(UClass* OldClass, UClass* NewClass)
-			: OldClass(OldClass), NewClass(NewClass) { }
-	};
-	TArray<FClassTuple> VTablesToUpdate;
-
-	for (auto Class : GetHotReloadClasses())
+	if (FCoreUObjectDelegates::ReplaceHotReloadClassDelegate.IsBound())
 	{
-		check(Class->OldClass);
-
-		UClass* RegisteredClass = nullptr;
-		if (Class->bHasChanged)
+		for (auto Class : GetHotReloadClasses())
 		{
-			RegisteredClass = Class->Register();
-			VTablesToUpdate.Add(FClassTuple(Class->OldClass, RegisteredClass));
-		}
+			check(Class->OldClass);
+
+			UClass* RegisteredClass = nullptr;
+			if (Class->bHasChanged)
+			{
+				RegisteredClass = Class->Register();
+			}
 
-		if (FCoreUObjectDelegates::ReplaceHotReloadClassDelegate.IsBound())
-		{
 			FCoreUObjectDelegates::ReplaceHotReloadClassDelegate.Execute(Class->OldClass, RegisteredClass);
 		}
 	}
-
-	for (auto Tuple : VTablesToUpdate)
-	{
-		FUClassVTableHelper::Get().UpdateVTables(Tuple.OldClass, Tuple.NewClass);
-	}
-
 	GetHotReloadClasses().Empty();
 }
 #endif
