Mark temp hot-reload vtable objects as RF_TagGarbageTemp

diff --git a/Engine/Source/Runtime/CoreUObject/Private/UObject/Class.cpp b/Engine/Source/Runtime/CoreUObject/Private/UObject/Class.cpp
--- a/Engine/Source/Runtime/CoreUObject/Private/UObject/Class.cpp
+++ b/Engine/Source/Runtime/CoreUObject/Private/UObject/Class.cpp
@@ -3423,7 +3423,7 @@ bool UClass::HotReloadPrivateStaticClass(
 
 	UE_LOG(LogClass, Verbose, TEXT("Attempting to change VTable for class %s."),*GetName());
 	ClassWithin = UPackage::StaticClass();  // We are just avoiding error checks with this...we don't care about this temp object other than to get the vtable.
-	UObject* TempObjectForVTable = StaticConstructObject(this, GetTransientPackage(), NAME_None, RF_NeedLoad | RF_ClassDefaultObject);
+	UObject* TempObjectForVTable = StaticConstructObject(this, GetTransientPackage(), NAME_None, RF_NeedLoad | RF_ClassDefaultObject | RF_TagGarbageTemp);
 
 	if( !TempObjectForVTable->IsRooted() )
 	{
