UE-9653 Some blueprints fail to compile only when playing with -game

Merging

//depot/UE4-UT/Engine/Source/Runtime/CoreUObject/Private/Blueprint/BlueprintSupport.cpp

to //depot/UE4/Engine/Source/Runtime/CoreUObject/Private/Blueprint/BlueprintSupport.cpp

diff --git a/Engine/Source/Runtime/CoreUObject/Private/Blueprint/BlueprintSupport.cpp b/Engine/Source/Runtime/CoreUObject/Private/Blueprint/BlueprintSupport.cpp
--- a/Engine/Source/Runtime/CoreUObject/Private/Blueprint/BlueprintSupport.cpp
+++ b/Engine/Source/Runtime/CoreUObject/Private/Blueprint/BlueprintSupport.cpp
@@ -781,9 +781,10 @@ void ULinkerLoad::FinalizeBlueprint(UClass* LoadClass)
 		{
 			// because it is tracked by FUnresolvedStructTracker, it must be a struct
 			DEFERRED_DEPENDENCY_CHECK(Cast<UStruct>(Import.XObject) != nullptr);
-			if (Import.SourceLinker)
+			auto SourceLinker = FindExistingLinkerForImport(ImportIndex);
+			if (SourceLinker)
 			{
-				Import.SourceLinker->ResolveDeferredDependencies((UStruct*)Import.XObject);
+				SourceLinker->ResolveDeferredDependencies((UStruct*)Import.XObject);
 			}
 		}
 	}
