Fixing ConstructObject deprecation warnings

diff --git a/Engine/Source/Runtime/CoreUObject/Private/Blueprint/BlueprintSupport.cpp b/Engine/Source/Runtime/CoreUObject/Private/Blueprint/BlueprintSupport.cpp
--- a/Engine/Source/Runtime/CoreUObject/Private/Blueprint/BlueprintSupport.cpp
+++ b/Engine/Source/Runtime/CoreUObject/Private/Blueprint/BlueprintSupport.cpp
@@ -397,7 +397,7 @@ UObject* ULinkerLoad::DeferExportCreation(const int32 Index)
 	FName PlaceholderName(*FString::Printf(TEXT("PLACEHOLDER-INST_of_%s"), *ClassName));
 	PlaceholderName = MakeUniqueObjectName(PlaceholderOuter, PlaceholderType, PlaceholderName);
 
-	ULinkerPlaceholderExportObject* Placeholder = ConstructObject<ULinkerPlaceholderExportObject>(PlaceholderType, PlaceholderOuter, PlaceholderName, RF_Public | RF_Transient);
+	ULinkerPlaceholderExportObject* Placeholder = NewObject<ULinkerPlaceholderExportObject>(PlaceholderOuter, PlaceholderType, PlaceholderName, RF_Public | RF_Transient);
 	Export.Object = Placeholder;
 #endif // USE_CIRCULAR_DEPENDENCY_LOAD_DEFERRING
 
diff --git a/Engine/Source/Runtime/CoreUObject/Private/UObject/Class.cpp b/Engine/Source/Runtime/CoreUObject/Private/UObject/Class.cpp
--- a/Engine/Source/Runtime/CoreUObject/Private/UObject/Class.cpp
+++ b/Engine/Source/Runtime/CoreUObject/Private/UObject/Class.cpp
@@ -3423,7 +3423,7 @@ bool UClass::HotReloadPrivateStaticClass(
 
 	UE_LOG(LogClass, Verbose, TEXT("Attempting to change VTable for class %s."),*GetName());
 	ClassWithin = UPackage::StaticClass();  // We are just avoiding error checks with this...we don't care about this temp object other than to get the vtable.
-	UObject* TempObjectForVTable = StaticConstructObject(this, GetTransientPackage(), NAME_None, RF_NeedLoad | RF_ClassDefaultObject | RF_TagGarbageTemp);
+	UObject* TempObjectForVTable = StaticConstructObject_Internal(this, GetTransientPackage(), NAME_None, RF_NeedLoad | RF_ClassDefaultObject | RF_TagGarbageTemp);
 
 	if( !TempObjectForVTable->IsRooted() )
 	{
