[INTEGRATE] Change 2436241 by [REDACTED_PERSON][REDACTED_HANDLE] on 2015/02/06 17:56:24

	[UE-9128] References to private objects are now null'd when copied to an unrelated object in a different package, this fixes a GLEO that would occur when converting an actor to a class blueprint.

diff --git a/Engine/Source/Runtime/CoreUObject/Public/Serialization/ArchiveUObject.h b/Engine/Source/Runtime/CoreUObject/Public/Serialization/ArchiveUObject.h
--- a/Engine/Source/Runtime/CoreUObject/Public/Serialization/ArchiveUObject.h
+++ b/Engine/Source/Runtime/CoreUObject/Public/Serialization/ArchiveUObject.h
@@ -1204,6 +1204,64 @@ protected:
 	bool bNullPrivateReferences;
 };
 
+/*----------------------------------------------------------------------------
+FArchiveReplaceOrClearExternalReferences.
+----------------------------------------------------------------------------*/
+/**
+* Identical to FArchiveReplaceObjectRef, but for references to private objects
+* in other packages we clear the reference instead of preserving it (unless it
+* makes it into the replacement map)
+*/
+template< class T >
+class FArchiveReplaceOrClearExternalReferences : public FArchiveReplaceObjectRef<T>
+{
+	typedef FArchiveReplaceObjectRef<T> TSuper;
+public:
+	FArchiveReplaceOrClearExternalReferences
+		( UObject* InSearchObject
+		, const TMap<T*, T*>& InReplacementMap
+		, UPackage* InDestPackage
+		, bool bDelayStart = false )
+		: TSuper(InSearchObject, InReplacementMap, false, false, false, true)
+		, DestPackage(InDestPackage)
+	{
+		if (!bDelayStart)
+		{
+			SerializeSearchObject();
+		}
+	}
+
+	FArchive& operator<<(UObject*& Obj)
+	{
+		UObject* Resolved = Obj;
+		TSuper::operator<<(Resolved);
+		// if Resolved is a private object in another package just clear the reference:
+		if (Resolved)
+		{
+			UObject* Outermost = Resolved->GetOutermost();
+			if (Outermost)
+			{
+				UPackage* ObjPackage = dynamic_cast<UPackage*>(Outermost);
+				if (ObjPackage)
+				{
+					if (ObjPackage != Obj && 
+						DestPackage != ObjPackage && 
+						!Obj->HasAnyFlags(RF_Public))
+					{
+						Resolved = nullptr;
+					}
+				}
+			}
+		}
+		Obj = Resolved;
+		return *this;
+	}
+
+protected:
+	/** Package that we are loading into, references to private objects in other packages will be cleared */
+	UPackage* DestPackage;
+};
+
 /*----------------------------------------------------------------------------
 	FArchiveObjectPropertyMapper.
 ----------------------------------------------------------------------------*/
diff --git a/Engine/Source/Runtime/Engine/Private/UnrealEngine.cpp b/Engine/Source/Runtime/Engine/Private/UnrealEngine.cpp
--- a/Engine/Source/Runtime/Engine/Private/UnrealEngine.cpp
+++ b/Engine/Source/Runtime/Engine/Private/UnrealEngine.cpp
@@ -10449,13 +10449,14 @@ void UEngine::CopyPropertiesForUnrelatedObjects(UObject* OldObject, UObject* New
 	}
 
 	// Replace references to old classes and instances on this object with the corresponding new ones
-	FArchiveReplaceObjectRef<UObject> ReplaceInCDOAr(NewObject, ReferenceReplacementMap, /*bNullPrivateRefs=*/ false, /*bIgnoreOuterRef=*/ false, /*bIgnoreArchetypeRef=*/ false);
+	UPackage* NewPackage = Cast<UPackage>(NewObject->GetOutermost());
+	FArchiveReplaceOrClearExternalReferences<UObject> ReplaceInCDOAr(NewObject, ReferenceReplacementMap, NewPackage);
 
 	// Replace references inside each individual component. This is always required because if something is in ReferenceReplacementMap, the above replace code will skip fixing child properties
 	for (int32 ComponentIndex = 0; ComponentIndex < ComponentsOnNewObject.Num(); ++ComponentIndex)
 	{
 		UObject* NewComponent = ComponentsOnNewObject[ComponentIndex];
-		FArchiveReplaceObjectRef<UObject> ReplaceInComponentAr(NewComponent, ReferenceReplacementMap, /*bNullPrivateRefs=*/ false, /*bIgnoreOuterRef=*/ false, /*bIgnoreArchetypeRef=*/ false);
+		FArchiveReplaceOrClearExternalReferences<UObject> ReplaceInComponentAr(NewComponent, ReferenceReplacementMap, NewPackage);
 	}
 
 	// Restore the root component reference
