[INTEGRATE] Change 2436313 by [REDACTED_PERSON][REDACTED_HANDLE] on 2015/02/06 18:30:37

	Probably fix compile error on Mac (CL 2436241)

diff --git a/Engine/Source/Runtime/CoreUObject/Public/Serialization/ArchiveUObject.h b/Engine/Source/Runtime/CoreUObject/Public/Serialization/ArchiveUObject.h
--- a/Engine/Source/Runtime/CoreUObject/Public/Serialization/ArchiveUObject.h
+++ b/Engine/Source/Runtime/CoreUObject/Public/Serialization/ArchiveUObject.h
@@ -1227,7 +1227,7 @@ public:
 	{
 		if (!bDelayStart)
 		{
-			SerializeSearchObject();
+			this->SerializeSearchObject();
 		}
 	}
 
