Made sure GC referencer object is checked for references to other objects even if it resides in permanent object pool

diff --git a/Engine/Source/Runtime/CoreUObject/Private/UObject/GarbageCollection.cpp b/Engine/Source/Runtime/CoreUObject/Private/UObject/GarbageCollection.cpp
--- a/Engine/Source/Runtime/CoreUObject/Private/UObject/GarbageCollection.cpp
+++ b/Engine/Source/Runtime/CoreUObject/Private/UObject/GarbageCollection.cpp
@@ -436,7 +436,12 @@ public:
 		GObjectCountDuringLastMarkPhase = 0;
 
 		// Presize array and add a bit of extra slack for prefetching.
-		ObjectsToSerialize.Empty( GUObjectArray.GetObjectArrayNumMinusPermanent() + 2 );
+		ObjectsToSerialize.Empty( GUObjectArray.GetObjectArrayNumMinusPermanent() + 3 );
+		// Make sure GC referencer object is checked for references to other objects even if it resides in permanent object pool
+		if (FPlatformProperties::RequiresCookedData() && FGCObject::GGCObjectReferencer && GUObjectArray.IsDisregardForGC(FGCObject::GGCObjectReferencer))
+		{
+			ObjectsToSerialize.Add(FGCObject::GGCObjectReferencer);
+		}
 
 		for ( FRawObjectIterator It(true); It; ++It )
 		{
