[INTEGRATE] Change 2438760 by [REDACTED_PERSON][REDACTED_HANDLE] on 2015/02/09 16:50:55

	Now tracking script bytecode  references (to placeholder UClass objects), and fixing up those references once the real class is loaded.

diff --git a/Engine/Source/Runtime/CoreUObject/Private/UObject/Class.cpp b/Engine/Source/Runtime/CoreUObject/Private/UObject/Class.cpp
--- a/Engine/Source/Runtime/CoreUObject/Private/UObject/Class.cpp
+++ b/Engine/Source/Runtime/CoreUObject/Private/UObject/Class.cpp
@@ -7,6 +7,7 @@
 #include "CoreUObjectPrivate.h"
 #include "PropertyTag.h"
 #include "HotReloadInterface.h"
+#include "LinkerPlaceholderClass.h"
 
 DECLARE_LOG_CATEGORY_EXTERN(LogScriptSerialization, Log, All);
 DEFINE_LOG_CATEGORY(LogScriptSerialization);
@@ -1554,6 +1555,37 @@ bool UStruct::GetStringMetaDataHierarchical(const FName& Key, FString* OutValue)
 }
 
 #endif
+
+#if USE_CIRCULAR_DEPENDENCY_LOAD_DEFERRING
+	/**
+	 * If we're loading, then the value of the script's UObject* expression 
+	 * could be pointing at a ULinkerPlaceholderClass (used by the linker to 
+	 * fight cyclic dependency issues on load). So here, if that's the case, we
+	 * have the placeholder track this ref (so it'll replace it once the real 
+	 * class is loaded).
+	 * 
+	 * @param  ScriptPtr    Reference to the point in the bytecode buffer, where a UObject* has been stored (for us to check).
+	 */
+	static void HandlePlaceholderScriptRef(ScriptPointerType& ScriptPtr)
+	{
+		UObject*& ExprPtrRef = (UObject*&)ScriptPtr;
+		if (ULinkerPlaceholderClass* PlaceholderObj = Cast<ULinkerPlaceholderClass>(ExprPtrRef)) \
+		{
+			PlaceholderObj->AddReferencingScriptExpr((ULinkerPlaceholderClass**)(&ExprPtrRef));
+		}
+	}
+
+	#define FIXUP_EXPR_OBJECT_POINTER(Type) \
+	{ \
+		if (!Ar.IsSaving()) \
+		{ \
+			int32 const ExprIndex = iCode - sizeof(ScriptPointerType); \
+			ScriptPointerType& ScriptPtr = (ScriptPointerType&)Script[ExprIndex]; \
+			HandlePlaceholderScriptRef(ScriptPtr); \
+		} \
+	}
+#endif // #if USE_CIRCULAR_DEPENDENCY_LOAD_DEFERRING
+
 //
 // Serialize an expression to an archive.
 // Returns expression token.
@@ -1571,6 +1603,7 @@ EExprToken UStruct::SerializeExpr( int32& iCode, FArchive& Ar )
 #undef XFER_FUNC_POINTER
 #undef XFER_FUNC_NAME
 #undef XFER_PROP_POINTER
+#undef FIXUP_EXPR_OBJECT_POINTER
 }
 
 void UStruct::InstanceSubobjectTemplates( void* Data, void const* DefaultData, UStruct* DefaultStruct, UObject* Owner, FObjectInstancingGraph* InstanceGraph )
diff --git a/Engine/Source/Runtime/CoreUObject/Private/UObject/LinkerPlaceholderClass.cpp b/Engine/Source/Runtime/CoreUObject/Private/UObject/LinkerPlaceholderClass.cpp
--- a/Engine/Source/Runtime/CoreUObject/Private/UObject/LinkerPlaceholderClass.cpp
+++ b/Engine/Source/Runtime/CoreUObject/Private/UObject/LinkerPlaceholderClass.cpp
@@ -64,7 +64,7 @@ void ULinkerPlaceholderClass::Bind()
 }
 
 //------------------------------------------------------------------------------
-void ULinkerPlaceholderClass::AddTrackedReference(UProperty* ReferencingProperty)
+void ULinkerPlaceholderClass::AddReferencingProperty(UProperty* ReferencingProperty)
 {
 #if USE_DEFERRED_DEPENDENCY_CHECK_VERIFICATION_TESTS
 	FObjectImport* PlaceholderImport = nullptr;
@@ -89,16 +89,26 @@ void ULinkerPlaceholderClass::AddTrackedReference(UProperty* ReferencingProperty
 	ReferencingProperties.Add(ReferencingProperty);
 }
 
+//------------------------------------------------------------------------------
+void ULinkerPlaceholderClass::AddReferencingScriptExpr(ULinkerPlaceholderClass** ExpressionPtr)
+{
+#if USE_DEFERRED_DEPENDENCY_CHECK_VERIFICATION_TESTS
+	check(*ExpressionPtr == this);
+#endif // USE_DEFERRED_DEPENDENCY_CHECK_VERIFICATION_TESTS
+
+	ReferencingScriptExpressions.Add((UClass**)ExpressionPtr);
+}
+
 //------------------------------------------------------------------------------
 bool ULinkerPlaceholderClass::HasReferences() const
 {
-	return (ReferencingProperties.Num() > 0);
+	return (GetRefCount() > 0);
 }
 
 //------------------------------------------------------------------------------
 int32 ULinkerPlaceholderClass::GetRefCount() const
 {
-	return ReferencingProperties.Num();
+	return ReferencingProperties.Num() + ReferencingScriptExpressions.Num();
 }
 
 //------------------------------------------------------------------------------
@@ -108,7 +118,7 @@ bool ULinkerPlaceholderClass::HasBeenResolved() const
 }
 
 //------------------------------------------------------------------------------
-void ULinkerPlaceholderClass::RemoveTrackedReference(UProperty* ReferencingProperty)
+void ULinkerPlaceholderClass::RemovePropertyReference(UProperty* ReferencingProperty)
 {
 	ReferencingProperties.Remove(ReferencingProperty);
 }
@@ -145,7 +155,18 @@ int32 ULinkerPlaceholderClass::ReplaceTrackedReferences(UClass* ReplacementClass
 		}
 	}
 	ReferencingProperties.Empty();
-	bResolvedReferences = true;
 
+	for (UClass** ScriptRefPtr : ReferencingScriptExpressions)
+	{
+		UClass*& ScriptRef = *ScriptRefPtr;
+		if (ScriptRef == this)
+		{
+			ScriptRef = ReplacementClass;
+			++ReplacementCount;
+		}
+	}
+	ReferencingScriptExpressions.Empty();
+
+	bResolvedReferences = true;
 	return ReplacementCount;
 }
diff --git a/Engine/Source/Runtime/CoreUObject/Private/UObject/LinkerPlaceholderClass.h b/Engine/Source/Runtime/CoreUObject/Private/UObject/LinkerPlaceholderClass.h
--- a/Engine/Source/Runtime/CoreUObject/Private/UObject/LinkerPlaceholderClass.h
+++ b/Engine/Source/Runtime/CoreUObject/Private/UObject/LinkerPlaceholderClass.h
@@ -38,7 +38,21 @@ public:
 	 * 
 	 * @param  ReferencingProperty    A property that uses and stores this class.
 	 */
-	void AddTrackedReference(UProperty* ReferencingProperty);
+	void AddReferencingProperty(UProperty* ReferencingProperty);
+
+	/**
+	 * Records a raw pointer, directly to the UClass* script expression (so that
+	 * we can switch-out its value in ReplaceTrackedReferences). 
+	 *
+	 * NOTE: We don't worry about creating some kind of weak ref to the script 
+	 *       pointer (or facilitate a way for this tracked reference to be 
+	 *       removed). We're not worried about the script ref being deleted 
+	 *       before we call ReplaceTrackedReferences (because we expect that we
+	 *       do this all within the same frame; before GC can be ran).
+	 * 
+	 * @param  ExpressionPtr    A direct pointer to the UClass* that is now referencing this placeholder.
+	 */
+	void AddReferencingScriptExpr(ULinkerPlaceholderClass** ExpressionPtr);
 
 	/**
 	 * A query method that let's us check to see if this class is currently 
@@ -73,7 +87,7 @@ public:
 	 * 
 	 * @param  ReferencingProperty    A property that used to use this class, and now no longer does.
 	 */
-	void RemoveTrackedReference(UProperty* ReferencingProperty);
+	void RemovePropertyReference(UProperty* ReferencingProperty);
 
 	/**
 	 * Iterates over all referencing properties and attempts to replace their 
@@ -94,4 +108,7 @@ private:
 
 	/** Used to catch references that are added after we've already resolved all references */
 	bool bResolvedReferences;
+
+	/** Points directly at UClass* refs that we're serialized in as part of script bytecode */
+	TSet<UClass**> ReferencingScriptExpressions;
 }; 
diff --git a/Engine/Source/Runtime/CoreUObject/Private/UObject/PropertyBaseObject.cpp b/Engine/Source/Runtime/CoreUObject/Private/UObject/PropertyBaseObject.cpp
--- a/Engine/Source/Runtime/CoreUObject/Private/UObject/PropertyBaseObject.cpp
+++ b/Engine/Source/Runtime/CoreUObject/Private/UObject/PropertyBaseObject.cpp
@@ -14,7 +14,7 @@ void UObjectPropertyBase::BeginDestroy()
 #if USE_CIRCULAR_DEPENDENCY_LOAD_DEFERRING
 	if (ULinkerPlaceholderClass* PlaceholderClass = Cast<ULinkerPlaceholderClass>(PropertyClass))
 	{
-		PlaceholderClass->RemoveTrackedReference(this);
+		PlaceholderClass->RemovePropertyReference(this);
 	}
 #endif // USE_CIRCULAR_DEPENDENCY_LOAD_DEFERRING
 
@@ -94,7 +94,7 @@ void UObjectPropertyBase::Serialize( FArchive& Ar )
 	{
 		if (ULinkerPlaceholderClass* PlaceholderClass = Cast<ULinkerPlaceholderClass>(PropertyClass))
 		{
-			PlaceholderClass->AddTrackedReference(this);
+			PlaceholderClass->AddReferencingProperty(this);
 		}
 	}
 #endif // USE_CIRCULAR_DEPENDENCY_LOAD_DEFERRING
@@ -105,12 +105,12 @@ void UObjectPropertyBase::SetPropertyClass(UClass* NewPropertyClass)
 {
 	if (ULinkerPlaceholderClass* NewPlaceholderClass = Cast<ULinkerPlaceholderClass>(NewPropertyClass))
 	{
-		NewPlaceholderClass->AddTrackedReference(this);
+		NewPlaceholderClass->AddReferencingProperty(this);
 	}
 	
 	if (ULinkerPlaceholderClass* OldPlaceholderClass = Cast<ULinkerPlaceholderClass>(PropertyClass))
 	{
-		OldPlaceholderClass->RemoveTrackedReference(this);
+		OldPlaceholderClass->RemovePropertyReference(this);
 	}
 	PropertyClass = NewPropertyClass;
 }
diff --git a/Engine/Source/Runtime/CoreUObject/Private/UObject/PropertyClass.cpp b/Engine/Source/Runtime/CoreUObject/Private/UObject/PropertyClass.cpp
--- a/Engine/Source/Runtime/CoreUObject/Private/UObject/PropertyClass.cpp
+++ b/Engine/Source/Runtime/CoreUObject/Private/UObject/PropertyClass.cpp
@@ -13,7 +13,7 @@ void UClassProperty::BeginDestroy()
 #if USE_CIRCULAR_DEPENDENCY_LOAD_DEFERRING
 	if (ULinkerPlaceholderClass* PlaceholderClass = Cast<ULinkerPlaceholderClass>(MetaClass))
 	{
-		PlaceholderClass->RemoveTrackedReference(this);
+		PlaceholderClass->RemovePropertyReference(this);
 	}
 #endif // USE_CIRCULAR_DEPENDENCY_LOAD_DEFERRING
 
@@ -30,7 +30,7 @@ void UClassProperty::Serialize( FArchive& Ar )
 	{
 		if (ULinkerPlaceholderClass* PlaceholderClass = Cast<ULinkerPlaceholderClass>(MetaClass))
 		{
-			PlaceholderClass->AddTrackedReference(this);
+			PlaceholderClass->AddReferencingProperty(this);
 		}
 	}
 #endif // USE_CIRCULAR_DEPENDENCY_LOAD_DEFERRING
@@ -53,12 +53,12 @@ void UClassProperty::SetMetaClass(UClass* NewMetaClass)
 {
 	if (ULinkerPlaceholderClass* NewPlaceholderClass = Cast<ULinkerPlaceholderClass>(NewMetaClass))
 	{
-		NewPlaceholderClass->AddTrackedReference(this);
+		NewPlaceholderClass->AddReferencingProperty(this);
 	}
 
 	if (ULinkerPlaceholderClass* OldPlaceholderClass = Cast<ULinkerPlaceholderClass>(NewMetaClass))
 	{
-		OldPlaceholderClass->RemoveTrackedReference(this);
+		OldPlaceholderClass->RemovePropertyReference(this);
 	}
 	MetaClass = NewMetaClass;
 }
diff --git a/Engine/Source/Runtime/CoreUObject/Private/UObject/PropertyInterface.cpp b/Engine/Source/Runtime/CoreUObject/Private/UObject/PropertyInterface.cpp
--- a/Engine/Source/Runtime/CoreUObject/Private/UObject/PropertyInterface.cpp
+++ b/Engine/Source/Runtime/CoreUObject/Private/UObject/PropertyInterface.cpp
@@ -12,7 +12,7 @@ void UInterfaceProperty::BeginDestroy()
 #if USE_CIRCULAR_DEPENDENCY_LOAD_DEFERRING
 	if (ULinkerPlaceholderClass* PlaceholderClass = Cast<ULinkerPlaceholderClass>(InterfaceClass))
 	{
-		PlaceholderClass->RemoveTrackedReference(this);
+		PlaceholderClass->RemovePropertyReference(this);
 	}
 #endif // USE_CIRCULAR_DEPENDENCY_LOAD_DEFERRING
 
@@ -217,7 +217,7 @@ void UInterfaceProperty::Serialize( FArchive& Ar )
 	{
 		if (ULinkerPlaceholderClass* PlaceholderClass = Cast<ULinkerPlaceholderClass>(InterfaceClass))
 		{
-			PlaceholderClass->AddTrackedReference(this);
+			PlaceholderClass->AddReferencingProperty(this);
 		}
 	}
 #endif // USE_CIRCULAR_DEPENDENCY_LOAD_DEFERRING
@@ -244,12 +244,12 @@ void UInterfaceProperty::SetInterfaceClass(UClass* NewInterfaceClass)
 {
 	if (ULinkerPlaceholderClass* NewPlaceholderClass = Cast<ULinkerPlaceholderClass>(NewInterfaceClass))
 	{
-		NewPlaceholderClass->AddTrackedReference(this);
+		NewPlaceholderClass->AddReferencingProperty(this);
 	}
 
 	if (ULinkerPlaceholderClass* OldPlaceholderClass = Cast<ULinkerPlaceholderClass>(InterfaceClass))
 	{
-		OldPlaceholderClass->RemoveTrackedReference(this);
+		OldPlaceholderClass->RemovePropertyReference(this);
 	}
 	InterfaceClass = NewInterfaceClass;
 }
diff --git a/Engine/Source/Runtime/CoreUObject/Public/UObject/ScriptSerialization.h b/Engine/Source/Runtime/CoreUObject/Public/UObject/ScriptSerialization.h
--- a/Engine/Source/Runtime/CoreUObject/Public/UObject/ScriptSerialization.h
+++ b/Engine/Source/Runtime/CoreUObject/Public/UObject/ScriptSerialization.h
@@ -91,6 +91,12 @@
 	#define XFER_OBJECT_POINTER(Type)	XFERPTR(Type)
 #endif
 
+#ifndef FIXUP_EXPR_OBJECT_POINTER
+	// sometimes after a UOBject* expression is loaded it may require some post-
+	// processing (see: the overridden FIXUP_EXPR_OBJECT_POINTER(), defined in Class.cpp)
+	#define FIXUP_EXPR_OBJECT_POINTER(Type) 
+#endif
+
 /** UStruct::SerializeExpr() */
 #ifdef SERIALIZEEXPR_INC
 	EExprToken Expr=(EExprToken)0;
@@ -115,7 +121,9 @@
 			// A conversion from an object or interface variable to a native interface variable.  
 			// We use a different bytecode to avoid the branching each time we process a cast token.
 			
-			XFERPTR(UClass*);	// the interface class to convert to
+			XFER_OBJECT_POINTER(UClass*); // the interface class to convert to
+			FIXUP_EXPR_OBJECT_POINTER(UClass*);
+
 			SerializeExpr( iCode, Ar );
 			break;
 		}
@@ -281,6 +289,8 @@
 		case EX_ObjectConst:
 		{
 			XFER_OBJECT_POINTER(UObject*);
+			FIXUP_EXPR_OBJECT_POINTER(UObject*);
+
 			break;
 		}
 		case EX_NameConst:
@@ -341,12 +351,16 @@
 		case EX_MetaCast:
 		{
 			XFER_OBJECT_POINTER(UClass*);
+			FIXUP_EXPR_OBJECT_POINTER(UClass*);
+
 			SerializeExpr( iCode, Ar );
 			break;
 		}
 		case EX_DynamicCast:
 		{
 			XFER_OBJECT_POINTER(UClass*);
+			FIXUP_EXPR_OBJECT_POINTER(UClass*);
+
 			SerializeExpr( iCode, Ar );
 			break;
 		}
