[INTEGRATE] Change 2439348 by [REDACTED_PERSON][REDACTED_HANDLE] on 2015/02/09 23:38:36

	Engine - Adding a method to StringClassReference and StringAssetReference to try and load the class/asset referenced.

diff --git a/Engine/Source/Runtime/CoreUObject/Private/Misc/StringAssetReference.cpp b/Engine/Source/Runtime/CoreUObject/Private/Misc/StringAssetReference.cpp
--- a/Engine/Source/Runtime/CoreUObject/Private/Misc/StringAssetReference.cpp
+++ b/Engine/Source/Runtime/CoreUObject/Private/Misc/StringAssetReference.cpp
@@ -111,7 +111,17 @@ bool FStringAssetReference::SerializeFromMismatchedTag(struct FPropertyTag const
 	return SerializeFromMismatchedTagTemplate<UObjectTypePolicy>(AssetLongPathname, Tag, Ar);
 }
 
-UObject *FStringAssetReference::ResolveObject() const
+UObject* FStringAssetReference::TryLoad() const
+{
+	if ( IsValid() )
+	{
+		return LoadObject<UObject>(nullptr, *ToString());
+	}
+
+	return nullptr;
+}
+
+UObject* FStringAssetReference::ResolveObject() const
 {
 	// Don't try to resolve if we're saving a package because StaticFindObject can't be used here
 	// and we usually don't want to force references to weak pointers while saving.
diff --git a/Engine/Source/Runtime/CoreUObject/Public/Misc/StringAssetReference.h b/Engine/Source/Runtime/CoreUObject/Public/Misc/StringAssetReference.h
--- a/Engine/Source/Runtime/CoreUObject/Public/Misc/StringAssetReference.h
+++ b/Engine/Source/Runtime/CoreUObject/Public/Misc/StringAssetReference.h
@@ -47,11 +47,17 @@ struct COREUOBJECT_API FStringAssetReference
 		return AssetLongPathname;
 	}
 
+	/**
+	 * Attempts to load the asset.
+	 * @return Loaded UObject, or null if the asset fails to load, or if the reference is not valid.
+	 */
+	UObject* TryLoad() const;
+
 	/**
 	 * Attempts to find a currently loaded object that matches this object ID
 	 * @return Found UObject, or NULL if not currently loaded
 	 */
-	UObject *ResolveObject() const;
+	UObject* ResolveObject() const;
 
 	/**
 	 * Resets reference to point to NULL
diff --git a/Engine/Source/Runtime/CoreUObject/Public/Misc/StringClassReference.h b/Engine/Source/Runtime/CoreUObject/Public/Misc/StringClassReference.h
--- a/Engine/Source/Runtime/CoreUObject/Public/Misc/StringClassReference.h
+++ b/Engine/Source/Runtime/CoreUObject/Public/Misc/StringClassReference.h
@@ -30,11 +30,26 @@ struct COREUOBJECT_API FStringClassReference : public FStringAssetReference
 		: FStringAssetReference(InClass)
 	{ }
 
+	/**
+	* Attempts to load the class.
+	* @return Loaded UObject, or null if the class fails to load, or if the reference is not valid.
+	*/
+	template<typename T>
+	UClass* TryLoadClass() const
+	{
+		if ( IsValid() )
+		{
+			return LoadClass<T>(nullptr, *ToString(), nullptr, LOAD_None, nullptr);
+		}
+
+		return nullptr;
+	}
+
 	/**
 	 * Attempts to find a currently loaded object that matches this object ID
 	 * @return Found UClass, or NULL if not currently loaded
 	 */
-	UClass *ResolveClass() const;
+	UClass* ResolveClass() const;
 
 	bool SerializeFromMismatchedTag(const FPropertyTag& Tag, FArchive& Ar);
 
