Integrate CL 2445618 by [REDACTED_PERSON] from from 4.7:
Handle the case where the class of the object is not yet loaded, but it is almost certainly a component if its class ends in _C

diff --git a/Engine/Source/Runtime/CoreUObject/Private/UObject/LinkerLoad.cpp b/Engine/Source/Runtime/CoreUObject/Private/UObject/LinkerLoad.cpp
--- a/Engine/Source/Runtime/CoreUObject/Private/UObject/LinkerLoad.cpp
+++ b/Engine/Source/Runtime/CoreUObject/Private/UObject/LinkerLoad.cpp
@@ -4245,8 +4245,15 @@ ULinkerLoad::ELinkerStatus ULinkerLoad::FixupExportMap()
 				if (OuterClassName == NAME_BlueprintGeneratedClass)
 				{
 					static const UClass* ActorComponentClass = FindObjectChecked<UClass>(ANY_PACKAGE, TEXT("ActorComponent"), true);
-					UClass* Class = FindObject<UClass>(ANY_PACKAGE, *NameClass.ToString());
-					if (Class && Class->IsChildOf(ActorComponentClass))
+					static const FString BPGeneratedClassPostfix(TEXT("_C"));
+					const FString NameClassString = NameClass.ToString();
+					UClass* Class = FindObject<UClass>(ANY_PACKAGE, *NameClassString);
+
+					// It is (obviously) a component if the class is a child of actor component
+					// and (almost certainly) a component if the class cannot be loaded but it ends in _C meaning it was generated from a blueprint
+					// However, it (probably) isn't safe to load the blueprint class, so we just check the _C and it is (probably) good enough
+					if (    ((Class != nullptr) && Class->IsChildOf(ActorComponentClass))
+						 || ((Class == nullptr) && NameClassString.EndsWith(BPGeneratedClassPostfix)))
 					{
 						Export.ObjectFlags |= RF_Public;
 					}
