Integrate CL 2446359 by [REDACTED_PERSON] from 4.7:
Add Blueprint as an outer class we consider for fixing up the RF_Public flag in FixupExportTable for older content that hasn't moved ownership in to BlueprintGeneratedClass yet.

diff --git a/Engine/Source/Runtime/CoreUObject/Private/UObject/LinkerLoad.cpp b/Engine/Source/Runtime/CoreUObject/Private/UObject/LinkerLoad.cpp
--- a/Engine/Source/Runtime/CoreUObject/Private/UObject/LinkerLoad.cpp
+++ b/Engine/Source/Runtime/CoreUObject/Private/UObject/LinkerLoad.cpp
@@ -4234,15 +4234,15 @@ ULinkerLoad::ELinkerStatus ULinkerLoad::FixupExportMap()
 			FName NameClass = GetExportClassName(ExportMapIdx);
 			FName NamePackage = GetExportClassPackage(ExportMapIdx);
 
-			// ActorComponents outered to a BlueprintGeneratedClass need to be marked RF_Public, but older content was not created
-			// as such.  This updates the ExportTable such that they are correctly flagged when created and when other packages
-			// validate their imports
+			// ActorComponents outered to a BlueprintGeneratedClass (or even older ones that are outered to Blueprint) need to be marked RF_Public, but older content was 
+			// not created as such.  This updates the ExportTable such that they are correctly flagged when created and when other packages validate their imports.
 			// TODO: Wrap this in a version check if possible. Mark up for future removal when minimum version is incremented
 			if ((Export.ObjectFlags & RF_Public) == 0)
 			{
 				static const FName NAME_BlueprintGeneratedClass("BlueprintGeneratedClass");
+				static const FName NAME_Blueprint("Blueprint");
 				const FName OuterClassName = GetExportClassName(Export.OuterIndex);
-				if (OuterClassName == NAME_BlueprintGeneratedClass)
+				if (OuterClassName == NAME_BlueprintGeneratedClass || OuterClassName == NAME_Blueprint)
 				{
 					static const UClass* ActorComponentClass = FindObjectChecked<UClass>(ANY_PACKAGE, TEXT("ActorComponent"), true);
 					static const FString BPGeneratedClassPostfix(TEXT("_C"));
