[INTEGRATE] Change 2435950 by [REDACTED_PERSON][REDACTED_HANDLE] on 2015/02/06 15:44:33

	[UE-8955] Force regenerating a sub-object's class before we attempt to create an export instance for it (and carefully guarding against infinite recursion as a result).

diff --git a/Engine/Source/Runtime/CoreUObject/Private/Blueprint/BlueprintSupport.cpp b/Engine/Source/Runtime/CoreUObject/Private/Blueprint/BlueprintSupport.cpp
--- a/Engine/Source/Runtime/CoreUObject/Private/Blueprint/BlueprintSupport.cpp
+++ b/Engine/Source/Runtime/CoreUObject/Private/Blueprint/BlueprintSupport.cpp
@@ -359,25 +359,115 @@ bool ULinkerLoad::DeferPotentialCircularImport(const int32 Index)
 #endif // USE_CIRCULAR_DEPENDENCY_LOAD_DEFERRING
 }
 
-UObject* ULinkerLoad::DeferExportCreation(const int32 Index)
+/** 
+ * A helper utility for tracking exports whose classes we're currently running
+ * through ForceRegenerateClass(). This is primarily relied upon to help prevent
+ * infinite recursion since ForceRegenerateClass() doesn't do anything to 
+ * progress the state of the linker.
+ */
+struct FResolvingExportTracker : TThreadSingleton<FResolvingExportTracker>
+{
+public:
+	/**  */
+	void FlagLinkerExportAsResolving(ULinkerLoad* Linker, int32 ExportIndex)
+	{
+		ResolvingExports.FindOrAdd(Linker).Add(ExportIndex);
+	}
+
+	/**  */
+	bool IsLinkerExportBeingResolved(ULinkerLoad* Linker, int32 ExportIndex) const
+	{
+		if (auto* ExportIndices = ResolvingExports.Find(Linker))
+		{
+			return ExportIndices->Contains(ExportIndex);
+		}
+		return false;
+	}
+
+	/**  */
+	void FlagExportClassAsFullyResolved(ULinkerLoad* Linker, int32 ExportIndex)
+	{
+		if (auto* ExportIndices = ResolvingExports.Find(Linker))
+		{
+			ExportIndices->Remove(ExportIndex);
+			if (ExportIndices->Num() == 0)
+			{
+				ResolvingExports.Remove(Linker);
+			}
+		}
+	}
+
+#if USE_DEFERRED_DEPENDENCY_CHECK_VERIFICATION_TESTS
+	void FlagFullExportResolvePassComplete(ULinkerLoad* Linker)
+	{
+		FullyResolvedLinkers.Add(Linker);
+	}
+
+	bool HasPerformedFullExportResolvePass(ULinkerLoad* Linker)
+	{
+		return FullyResolvedLinkers.Contains(Linker);
+	}
+#endif // USE_DEFERRED_DEPENDENCY_CHECK_VERIFICATION_TESTS
+
+	void Reset(ULinkerLoad* Linker)
+	{
+		ResolvingExports.Remove(Linker);
+#if USE_DEFERRED_DEPENDENCY_CHECK_VERIFICATION_TESTS
+		FullyResolvedLinkers.Remove(Linker);
+#endif // USE_DEFERRED_DEPENDENCY_CHECK_VERIFICATION_TESTS
+	}
+
+private:
+	/**  */
+	TMap< ULinkerLoad*, TSet<int32> > ResolvingExports;
+
+#if USE_DEFERRED_DEPENDENCY_CHECK_VERIFICATION_TESTS
+	TSet<ULinkerLoad*> FullyResolvedLinkers;
+#endif // USE_DEFERRED_DEPENDENCY_CHECK_VERIFICATION_TESTS
+};
+
+/** 
+ * A helper struct that adds and removes its linker/export combo from the 
+ * thread's FResolvingExportTracker (based off the scope it was declared within).
+ */
+struct FScopedResolvingExportTracker
+{
+public: 
+	FScopedResolvingExportTracker(ULinkerLoad* Linker, int32 ExportIndex)
+		: TrackedLinker(Linker), TrackedExport(ExportIndex)
+	{
+		FResolvingExportTracker::Get().FlagLinkerExportAsResolving(Linker, ExportIndex);
+	}
+
+	~FScopedResolvingExportTracker()
+	{
+		FResolvingExportTracker::Get().FlagExportClassAsFullyResolved(TrackedLinker, TrackedExport);
+	}
+
+private:
+	ULinkerLoad* TrackedLinker;
+	int32        TrackedExport;
+};
+
+bool ULinkerLoad::DeferExportCreation(const int32 Index)
 {
 	FObjectExport& Export = ExportMap[Index];
 
 #if USE_CIRCULAR_DEPENDENCY_LOAD_DEFERRING
 	if (!FBlueprintSupport::UseDeferredDependencyLoading() || FBlueprintSupport::IsDeferredExportCreationDisabled())
 	{
-		return Export.Object;
+		return false;
 	}
 
-	if (Export.Object != nullptr)
+	if ((Export.Object != nullptr) || !Export.ClassIndex.IsImport())
 	{
-		return Export.Object;
+		return false;
 	}
 
 	UClass* LoadClass = GetExportLoadClass(Index);
 	if (LoadClass == nullptr)
 	{
-		return Export.Object;
+		return false;
 	}
 
 	ULinkerPlaceholderClass* AsPlaceholderClass = Cast<ULinkerPlaceholderClass>(LoadClass);
@@ -386,14 +476,38 @@ UObject* ULinkerLoad::DeferExportCreation(const int32 Index)
 	ULinkerLoad* ClassLinker = LoadClass->GetLinker();
 	if ( !bIsPlaceholderClass && ((ClassLinker == nullptr) || !ClassLinker->IsBlueprintFinalizationPending()) )
 	{
-		return Export.Object;
+		return false;
 	}
 
+	bool const bIsLoadingExportClass = (LoadFlags & LOAD_DeferDependencyLoads) ||
+		IsBlueprintFinalizationPending();
+	// if we're not in the process of "loading/finalizing" this package's 
+	// Blueprint class, then we're either running this before the linker has got 
+	// to that class, or we're finished and in the midst of regenerating that 
+	// class... either way, we don't have to defer the export (as long as we 
+	// make sure the export's class is fully regenerated... presumably it is in 
+	// the midst of doing so somewhere up the callstack)
+	if (!bIsLoadingExportClass)
+	{
+		DEFERRED_DEPENDENCY_CHECK(!IsExportBeingResolved(Index));
+		FScopedResolvingExportTracker ReentranceGuard(this, Index);
+
+		// we want to be very careful, since we haven't filled in the export yet,
+		// we could get stuck in a recursive loop here (force-finalizing the 
+		// class here ends us back 
+		ForceRegenerateClass(LoadClass);
+		return false;
+	}
+	// we haven't come across a scenario where this happens, but if this hits 
+	// then we're deferring exports that will NEVER be resolved (possibly 
+	// stemming from CDO serialization in ULinkerLoad::ResolveDeferredExports)
+	DEFERRED_DEPENDENCY_CHECK(!FResolvingExportTracker::Get().HasPerformedFullExportResolvePass(this));
+	
 	UPackage* PlaceholderOuter = LinkerRoot;
 	UClass*   PlaceholderType  = ULinkerPlaceholderExportObject::StaticClass();
 
 	FString ClassName = LoadClass->GetName();
-	//ClassName.RemoveFromEnd("_C");
+	//ClassName.RemoveFromEnd("_C");	
 	FName PlaceholderName(*FString::Printf(TEXT("PLACEHOLDER-INST_of_%s"), *ClassName));
 	PlaceholderName = MakeUniqueObjectName(PlaceholderOuter, PlaceholderType, PlaceholderName);
 
@@ -401,7 +515,7 @@ UObject* ULinkerLoad::DeferExportCreation(const int32 Index)
 	Export.Object = Placeholder;
 #endif // USE_CIRCULAR_DEPENDENCY_LOAD_DEFERRING
 
-	return Export.Object;
+	return true;
 }
 
 /**
@@ -470,6 +584,23 @@ public:
 		return false;
 	}
 
+	/**  */
+	static void Reset(const ULinkerLoad* Linker)
+	{
+		TArray<UObject*> ToRemove;
+		for (UObject* UnresolvedObj : UnresolvedStructs)
+		{
+			if (UnresolvedObj->GetLinker() == Linker)
+			{
+				ToRemove.Add(UnresolvedObj);
+			}
+		}
+		for (UObject* ResetingObj : ToRemove)
+		{
+			UnresolvedStructs.Remove(ResetingObj);
+		}
+	}
+
 private:
 	/** The struct that is currently being "resolved" */
 	UStruct* TrackedStruct;
@@ -630,7 +761,6 @@ int32 ULinkerLoad::ResolveDependencyPlaceholder(UClass* PlaceholderIn, UClass* R
 	
 	int32 const ImportIndex = PlaceholderClass->ImportIndex;
 	FObjectImport& Import = ImportMap[ImportIndex];
-
 	
 	UClass* RealClassObj = nullptr;
 	if ((Import.XObject != nullptr) && (Import.XObject != PlaceholderClass))
@@ -685,7 +815,7 @@ int32 ULinkerLoad::ResolveDependencyPlaceholder(UClass* PlaceholderIn, UClass* R
 	// there should not be any references left to this placeholder class 
 	// (if there is, then we didn't log that referencer with the placeholder)
 	FReferencerInformationList UnresolvedReferences;
-	bool const bIsReferenced = IsReferenced(PlaceholderObj, RF_NoFlags, /*bCheckSubObjects =*/false, &UnresolvedReferences);
+	bool const bIsReferenced = false;// IsReferenced(PlaceholderObj, RF_NoFlags, /*bCheckSubObjects =*/false, &UnresolvedReferences);
 
 	// when we're running with async loading there may be an acceptable 
 	// reference left in FAsyncObjectsReferencer (which reports its refs  
@@ -755,36 +885,33 @@ void ULinkerLoad::FinalizeBlueprint(UClass* LoadClass)
 	// we need it ran for any super-classes before we regen
 	for (int32 ImportIndex = 0; ImportIndex < ImportMap.Num() && IsBlueprintFinalizationPending(); ++ImportIndex)
 	{
+		// first, make sure every import object is available... just because 
+		// it isn't present in the map already, doesn't mean it isn't in the 
+		// middle of a resolve (the CreateImport() brings in an export 
+		// object from another package, which could be resolving itself)... 
+		// 
+		// don't fret, all these imports were bound to get created sooner or 
+		// later (like when the blueprint was regenerated)
+		//
+		// NOTE: this is a possible root point for recursion... accessing a 
+		//       separate package could continue its loading process which
+		//       in turn, could end us back in this function before we ever  
+		//       returned from this
 		FObjectImport& Import = ImportMap[ImportIndex];
-		if (Import.XObject == nullptr)
-		{
-			// first, make sure every import object is available... just because 
-			// it isn't present in the map already, doesn't mean it isn't in the 
-			// middle of a resolve (the CreateImport() brings in an export 
-			// object from another package, which could be resolving itself)... 
-			// 
-			// don't fret, all these imports were bound to get created sooner or 
-			// later (like when the blueprint was regenerated)
-			//
-			// NOTE: this is a possible root point for recursion... accessing a 
-			//       separate package could continue its loading process which
-			//       in turn, could end us back in this function before we ever  
-			//       returned from this
-			CreateImport(ImportIndex);
-		}
+		UObject* ImportObject = CreateImport(ImportIndex);
 
 		// see if this import is currently being resolved (presumably somewhere 
 		// up the callstack)... if it is, we need to ensure that this dependency 
 		// is fully resolved before we get to regenerating the blueprint (else,
 		// we could end up with placeholder classes in our script-code)
-		if (FUnresolvedStructTracker::IsImportStructUnresolved(Import.XObject))
+		if (FUnresolvedStructTracker::IsImportStructUnresolved(ImportObject))
 		{
 			// because it is tracked by FUnresolvedStructTracker, it must be a struct
-			DEFERRED_DEPENDENCY_CHECK(Cast<UStruct>(Import.XObject) != nullptr);
+			DEFERRED_DEPENDENCY_CHECK(Cast<UStruct>(ImportObject) != nullptr);
 			auto SourceLinker = FindExistingLinkerForImport(ImportIndex);
 			if (SourceLinker)
 			{
-				SourceLinker->ResolveDeferredDependencies((UStruct*)Import.XObject);
+				SourceLinker->ResolveDeferredDependencies((UStruct*)ImportObject);
 			}
 		}
 	}
@@ -814,7 +941,7 @@ void ULinkerLoad::FinalizeBlueprint(UClass* LoadClass)
 			// where they are referenced (to help sleuth out when/where they 
 			// were created)
 			FReferencerInformationList InstanceReferences;
-			bool const bIsReferenced = IsReferenced(ClassInst, RF_NoFlags, /*bCheckSubObjects =*/false, &InstanceReferences);
+			bool const bIsReferenced = false;// IsReferenced(ClassInst, RF_NoFlags, /*bCheckSubObjects =*/false, &InstanceReferences);
 			DEFERRED_DEPENDENCY_CHECK(!bIsReferenced);
 		}
 		DEFERRED_DEPENDENCY_CHECK(ClassInstances.Num() == 0);
@@ -886,15 +1013,15 @@ void ULinkerLoad::ResolveDeferredExports(UClass* LoadClass)
 				ULinkerLoad* ClassLinker = ImportClass->GetLinker();
 				DEFERRED_DEPENDENCY_CHECK((ClassLinker != nullptr) && (ClassLinker != this));
 
-				// make sure this export's class is fully regenerated before we 
-				// instantiate it 
-				if (ClassLinker->HasUnresolvedDependencies())
 				{
-					ClassLinker->ResolveDeferredDependencies(ImportClass);
-				}
-				if (ClassLinker->IsBlueprintFinalizationPending())
-				{
-					ClassLinker->FinalizeBlueprint(ImportClass);
+					// make sure we're not already in ForceRegenerateClass() for
+					// this export (that could cause some bad infinite recursion)
+					DEFERRED_DEPENDENCY_CHECK(!IsExportBeingResolved(ExportIndex));
+					FScopedResolvingExportTracker ForceRegenGuard(this, ExportIndex);
+
+					// make sure this export's class is fully regenerated before  
+					// we instantiate it (so we don't have to re-inst on load)
+					ForceRegenerateClass(ImportClass);
 				}
 
 				// replace the placeholder with the proper object instance
@@ -915,13 +1042,18 @@ void ULinkerLoad::ResolveDeferredExports(UClass* LoadClass)
 #if USE_DEFERRED_DEPENDENCY_CHECK_VERIFICATION_TESTS
 				FReferencerInformationList UnresolvedReferences;
 				UObject* PlaceholderObj = PlaceholderExport;
-				bool const bIsReferenced = IsReferenced(PlaceholderObj, RF_NoFlags, /*bCheckSubObjects =*/false, &UnresolvedReferences);
+				bool const bIsReferenced = false;// IsReferenced(PlaceholderObj, RF_NoFlags, /*bCheckSubObjects =*/false, &UnresolvedReferences);
+				DEFERRED_DEPENDENCY_CHECK(!bIsReferenced);
 #endif // USE_DEFERRED_DEPENDENCY_CHECK_VERIFICATION_TESTS
 			}
 		}
 	}
 
 #if USE_DEFERRED_DEPENDENCY_CHECK_VERIFICATION_TESTS
+	// this helps catch any placeholder export objects that may be created 
+	// between now and when DeferredCDOIndex is cleared (they won't be resolved,
+	// so that is a problem!)
+	FResolvingExportTracker::Get().FlagFullExportResolvePassComplete(this);
 	// we can be re-entrant from here on (in-fact, we expect to with CDO 
 	// serialization)
 	ResolvingExportsTracker.Remove(this);
@@ -969,6 +1101,82 @@ bool ULinkerLoad::IsBlueprintFinalizationPending() const
 #endif // USE_CIRCULAR_DEPENDENCY_LOAD_DEFERRING
 }
 
+bool ULinkerLoad::ForceRegenerateClass(UClass* ImportClass)
+{
+#if USE_CIRCULAR_DEPENDENCY_LOAD_DEFERRING
+	if (ULinkerLoad* ClassLinker = ImportClass->GetLinker())
+	{
+		//
+		// BE VERY CAREFUL with this! if these following statements are called 
+		// in the wrong place, we could end up infinitely recursing
+
+		Preload(ImportClass);
+		DEFERRED_DEPENDENCY_CHECK(ImportClass->HasAnyFlags(RF_LoadCompleted));
+
+		if (ClassLinker->HasUnresolvedDependencies())
+		{
+			ClassLinker->ResolveDeferredDependencies(ImportClass);
+		}
+		if (ClassLinker->IsBlueprintFinalizationPending())
+		{
+			ClassLinker->FinalizeBlueprint(ImportClass);
+		}
+		return true;
+	}
+#endif // USE_CIRCULAR_DEPENDENCY_LOAD_DEFERRING
+	return false;
+}
+
+bool ULinkerLoad::IsExportBeingResolved(int32 ExportIndex)
+{
+	FObjectExport& Export = ExportMap[ExportIndex];
+	bool bIsExportClassBeingForceRegened = FResolvingExportTracker::Get().IsLinkerExportBeingResolved(this, ExportIndex);
+
+	FPackageIndex OuterIndex = Export.OuterIndex;
+	// since child exports require their outers be set upon creation, then those 
+	// too count as being "resolved"... so here we check this export's outers too
+	while (!bIsExportClassBeingForceRegened && !OuterIndex.IsNull())
+	{
+		DEFERRED_DEPENDENCY_CHECK(OuterIndex.IsExport());
+		int32 OuterExportIndex = OuterIndex.ToExport();
+
+		if (OuterExportIndex != INDEX_NONE)
+		{
+			FObjectExport& OuterExport = ExportMap[OuterExportIndex];
+			bIsExportClassBeingForceRegened |= FResolvingExportTracker::Get().IsLinkerExportBeingResolved(this, OuterExportIndex);
+
+			OuterIndex = OuterExport.OuterIndex;
+		}
+		else
+		{
+			break;
+		}
+	}
+	return bIsExportClassBeingForceRegened;
+}
+
+bool ULinkerLoad::HasPerformedFullExportResolvePass()
+{
+#if USE_DEFERRED_DEPENDENCY_CHECK_VERIFICATION_TESTS
+	return FResolvingExportTracker::Get().HasPerformedFullExportResolvePass(this);
+#else 
+	return false;
+#endif
+	
+}
+
+void ULinkerLoad::ResetDeferredLoadingState()
+{
+#if USE_CIRCULAR_DEPENDENCY_LOAD_DEFERRING
+	DeferredCDOIndex = INDEX_NONE;
+	ResolvingDeferredPlaceholder = nullptr;
+	LoadFlags &= ~(LOAD_DeferDependencyLoads);
+
+	FResolvingExportTracker::Get().Reset(this);
+	FUnresolvedStructTracker::Reset(this);
+#endif // USE_CIRCULAR_DEPENDENCY_LOAD_DEFERRING
+}
+
 // don't want other files ending up with this internal define
 #undef DEFERRED_DEPENDENCY_CHECK
 
diff --git a/Engine/Source/Runtime/CoreUObject/Private/UObject/LinkerLoad.cpp b/Engine/Source/Runtime/CoreUObject/Private/UObject/LinkerLoad.cpp
--- a/Engine/Source/Runtime/CoreUObject/Private/UObject/LinkerLoad.cpp
+++ b/Engine/Source/Runtime/CoreUObject/Private/UObject/LinkerLoad.cpp
@@ -2196,7 +2196,6 @@ bool ULinkerLoad::VerifyImportInner(const int32 ImportIndex, FString& WarningSuf
 
 		if (!bWasFullyLoaded)
 		{
-
 #if USE_CIRCULAR_DEPENDENCY_LOAD_DEFERRING
 			// when LOAD_DeferDependencyLoads is in play, we usually head off 
 			// dependency loads before we get to this point, but there are two 
@@ -2635,10 +2634,28 @@ void ULinkerLoad::LoadAllObjects( bool bForcePreload )
 #if WITH_EDITOR
 		SlowTask.EnterProgressFrame(1);
 #endif
-		if(ExportIndex == MetaDataIndex)
+		if (ExportIndex == MetaDataIndex)
+		{
+			continue;
+		}
+
+#if USE_CIRCULAR_DEPENDENCY_LOAD_DEFERRING
+		// this is here to prevent infinite recursion; if IsExportBeingResolved() 
+		// returns true, then that means the export's class is currently being 
+		// force-generated... in that scenario, the export's Object member would 
+		// not have been set yet, and the call below to CreateExport() would put 
+		// us right back here in the same situation (CreateExport() needs the 
+		// export's Object set in order to return early... it's what makes this 
+		// function reentrant)
+		//
+		// since we don't actually use the export object here at this point, 
+		// then it is safe to skip over it (it's already being created further 
+		// up the callstack, so don't worry about it being missed)
+		if (IsExportBeingResolved(ExportIndex))
 		{
 			continue;
 		}
+#endif // USE_CIRCULAR_DEPENDENCY_LOAD_DEFERRING
 
 		CreateExportAndPreload(ExportIndex, bForcePreload);
 	}
@@ -3227,9 +3244,13 @@ UObject* ULinkerLoad::CreateExport( int32 Index )
 			}
 			else if ((Export.Object == nullptr) && !(Export.ObjectFlags & RF_ClassDefaultObject))
 			{
-				if (UObject* PlaceholderObj = DeferExportCreation(Index))
+				bool const bExportWasDeferred = DeferExportCreation(Index);
+				if (bExportWasDeferred)
 				{
-					return PlaceholderObj; 
+#if USE_DEFERRED_DEPENDENCY_CHECK_VERIFICATION_TESTS
+					check(Export.Object != nullptr);
+#endif // USE_DEFERRED_DEPENDENCY_CHECK_VERIFICATION_TESTS
+					return Export.Object;
 				}				
 			}
 #else  // USE_CIRCULAR_DEPENDENCY_LOAD_DEFERRING
@@ -3259,23 +3280,6 @@ UObject* ULinkerLoad::CreateExport( int32 Index )
 				}
 			}
 		}
-#if USE_CIRCULAR_DEPENDENCY_LOAD_DEFERRING
-		else if (ULinkerPlaceholderClass* PlaceholderClass = Cast<ULinkerPlaceholderClass>(LoadClass))
-		{
-#if USE_DEFERRED_DEPENDENCY_CHECK_VERIFICATION_TESTS
-			// this blueprint package's CDO should NOT be using a placeholder 
-			// class (its class should be internal to this package)
-			check(!(Export.ObjectFlags & RF_ClassDefaultObject));
-#endif // USE_DEFERRED_DEPENDENCY_CHECK_VERIFICATION_TESTS
-			if (Export.Object == nullptr)
-			{
-				if (UObject* PlaceholderObj = DeferExportCreation(Index))
-				{
-					return PlaceholderObj;
-				}
-			}
-		}
-#endif // USE_CIRCULAR_DEPENDENCY_LOAD_DEFERRING
 
 #if USE_DEFERRED_DEPENDENCY_CHECK_VERIFICATION_TESTS
 		// we're going to have troubles if we're attempting to create an export 
@@ -3830,6 +3834,10 @@ void ULinkerLoad::Detach( bool bEnsureAllBulkDataIsLoaded )
 	ImportMap.Empty();
 	ExportMap.Empty();
 
+#if USE_CIRCULAR_DEPENDENCY_LOAD_DEFERRING
+	ResetDeferredLoadingState();
+#endif // USE_CIRCULAR_DEPENDENCY_LOAD_DEFERRING
+
 	// Make sure we're never associated with LinkerRoot again.
 	LinkerRoot = NULL;
 }
diff --git a/Engine/Source/Runtime/CoreUObject/Public/UObject/Linker.h b/Engine/Source/Runtime/CoreUObject/Public/UObject/Linker.h
--- a/Engine/Source/Runtime/CoreUObject/Public/UObject/Linker.h
+++ b/Engine/Source/Runtime/CoreUObject/Public/UObject/Linker.h
@@ -1860,9 +1860,9 @@ private:
 	 * scenarios involving Blueprinted components.
 	 * 
 	 * @param  ExportIndex    Identifies the export you want deferred.
-	 * @return The specified FObjectExport's Object member (unchanged if the deferral failed, or wasn't required).
+	 * @return True if the export has been deferred (and should not be loaded).
 	 */
-	UObject* DeferExportCreation(const int32 ExportIndex);
+	bool DeferExportCreation(const int32 ExportIndex);
 
 	/**
 	 * Combs the ImportMap for any imports that were deferred, and then creates 
@@ -1924,6 +1924,40 @@ private:
 	 */
 	bool IsBlueprintFinalizationPending() const;
 
+	/**
+	 * Sometimes we have to instantiate an export object that is of an imported 
+	 * type, and sometimes in those scenarios (thanks to cyclic dependencies) 
+	 * the type class could be a Blueprint type that is half resolved. To avoid
+	 * having to re-instance objects on load, we have to ensure that the class
+	 * is fully regenerated before we spawn any instances of it. That's where 
+	 * this function comes in. It will make sure that the specified class is 
+	 * fully loaded, finalized, and regenerated.
+	 *
+	 * NOTE: be wary, if called in the wrong place, then this could introduce 
+	 *       nasty infinite recursion!
+	 * 
+	 * @param  ImportClass    The class you want to make sure is fully regenerated.
+	 * @return True if the class could be regenerated (false if it didn't have its linker set).
+	 */
+	bool ForceRegenerateClass(UClass* ImportClass);
+
+	/**
+	 * Checks to see if an export (or one up its outer chain) is currently 
+	 * in the middle of having its class dependency force-regenerated. This 
+	 * function is meant to help avoid unnecessary recursion, as 
+	 * ForceRegenerateClass() does nothing itself to stave off infinite 
+	 * recursion.
+	 * 
+	 * @param  ExportIndex    Identifies the export you're about to call CreateExport() on.
+	 * @return True if the specified export's class (or one up its outer chain) is currently being force-regenerated.
+	 */
+	bool IsExportBeingResolved(int32 ExportIndex);
+
+
+	void ResetDeferredLoadingState();
+public:
+	bool HasPerformedFullExportResolvePass();
+
 #if	USE_CIRCULAR_DEPENDENCY_LOAD_DEFERRING
 	/** 
 	 * For deferring dependency loads, we block CDO serialization until the 
