[INTEGRATE] Change 2439906 by [REDACTED_PERSON][REDACTED_HANDLE] on 2015/02/10 11:51:34

	Making it so we can easily enable deferred dependency loading in cooked builds (but leaving it disabled for now?).

diff --git a/Engine/Config/BaseEngine.ini b/Engine/Config/BaseEngine.ini
--- a/Engine/Config/BaseEngine.ini
+++ b/Engine/Config/BaseEngine.ini
@@ -1507,6 +1507,7 @@ bOptimizeExecutionFlowStack=true
 bOptimizeAdjacentStates=true
 bEnableInheritableComponents=true
 bDeferDependencyLoads=true
+bForceDisableCookedDependencyDeferring=true
 
 [/Script/Engine.Blueprint]
 bRecompileOnLoad=true
diff --git a/Engine/Source/Runtime/Core/Public/Misc/Build.h b/Engine/Source/Runtime/Core/Public/Misc/Build.h
--- a/Engine/Source/Runtime/Core/Public/Misc/Build.h
+++ b/Engine/Source/Runtime/Core/Public/Misc/Build.h
@@ -241,5 +241,5 @@
 /** Enable UberGraphPersistentFrame feature. It can speed up BP compilation (re-instancing) in editor, but introduce an unnecessary overhead in runtime. */
 #define USE_UBER_GRAPH_PERSISTENT_FRAME 1
 
-#define USE_CIRCULAR_DEPENDENCY_LOAD_DEFERRING (WITH_EDITOR && 1)
+#define USE_CIRCULAR_DEPENDENCY_LOAD_DEFERRING 1
 #define USE_DEFERRED_DEPENDENCY_CHECK_VERIFICATION_TESTS (USE_CIRCULAR_DEPENDENCY_LOAD_DEFERRING && 0)
diff --git a/Engine/Source/Runtime/CoreUObject/Private/Blueprint/BlueprintSupport.cpp b/Engine/Source/Runtime/CoreUObject/Private/Blueprint/BlueprintSupport.cpp
--- a/Engine/Source/Runtime/CoreUObject/Private/Blueprint/BlueprintSupport.cpp
+++ b/Engine/Source/Runtime/CoreUObject/Private/Blueprint/BlueprintSupport.cpp
@@ -46,7 +46,14 @@ bool FBlueprintSupport::UseDeferredDependencyLoading()
 {
 #if USE_CIRCULAR_DEPENDENCY_LOAD_DEFERRING
 	static const FBoolConfigValueHelper DeferDependencyLoads(TEXT("Kismet"), TEXT("bDeferDependencyLoads"), GEngineIni);
-	return DeferDependencyLoads;
+	bool bUseDeferredDependencyLoading = DeferDependencyLoads;
+
+	if (FPlatformProperties::RequiresCookedData())
+	{
+		static const FBoolConfigValueHelper DisableCookedBuildDefering(TEXT("Kismet"), TEXT("bForceDisableCookedDependencyDeferring"), GEngineIni);
+		bUseDeferredDependencyLoading &= !((bool)DisableCookedBuildDefering);
+	}
+	return bUseDeferredDependencyLoading;
 #else
 	return false;
 #endif
