[INTEGRATE] Change 2441881 by [REDACTED_PERSON][REDACTED_HANDLE] on 2015/02/11 12:28:57

	Fixing up deferred dependency loading for cooked builds (making sure the CDO export is properly flagged in cooked builds, so we don't miss serializing it in).

diff --git a/Engine/Config/BaseEngine.ini b/Engine/Config/BaseEngine.ini
--- a/Engine/Config/BaseEngine.ini
+++ b/Engine/Config/BaseEngine.ini
@@ -1509,7 +1509,7 @@ bOptimizeExecutionFlowStack=true
 bOptimizeAdjacentStates=true
 bEnableInheritableComponents=true
 bDeferDependencyLoads=true
-bForceDisableCookedDependencyDeferring=true
+bForceDisableCookedDependencyDeferring=false
 
 [/Script/Engine.Blueprint]
 bRecompileOnLoad=true
diff --git a/Engine/Source/Runtime/CoreUObject/Private/Blueprint/BlueprintSupport.cpp b/Engine/Source/Runtime/CoreUObject/Private/Blueprint/BlueprintSupport.cpp
--- a/Engine/Source/Runtime/CoreUObject/Private/Blueprint/BlueprintSupport.cpp
+++ b/Engine/Source/Runtime/CoreUObject/Private/Blueprint/BlueprintSupport.cpp
@@ -84,7 +84,7 @@ bool FBlueprintSupport::IsDeferredExportCreationDisabled()
 {
 #if USE_CIRCULAR_DEPENDENCY_LOAD_DEFERRING
 	static const FBoolConfigValueHelper NoDeferredExports(TEXT("Kismet"), TEXT("bForceDisableDeferredExportCreation"), GEngineIni);
-	return NoDeferredExports;
+	return !UseDeferredDependencyLoading() || NoDeferredExports;
 #else
 	return false;
 #endif
@@ -958,6 +958,7 @@ void ULinkerLoad::FinalizeBlueprint(UClass* LoadClass)
 		DEFERRED_DEPENDENCY_CHECK(BlueprintClass->HasAnyClassFlags(CLASS_CompiledFromBlueprint));
 #endif // USE_DEFERRED_DEPENDENCY_CHECK_VERIFICATION_TESTS
 
+		DEFERRED_DEPENDENCY_CHECK(LoadClass->GetOutermost() != GetTransientPackage());
 		// just in case we choose to enable the deferred dependency loading for 
 		// cooked builds... we want to keep from regenerating in that scenario
 		if (!LoadClass->bCooked)
diff --git a/Engine/Source/Runtime/CoreUObject/Private/UObject/LinkerLoad.cpp b/Engine/Source/Runtime/CoreUObject/Private/UObject/LinkerLoad.cpp
--- a/Engine/Source/Runtime/CoreUObject/Private/UObject/LinkerLoad.cpp
+++ b/Engine/Source/Runtime/CoreUObject/Private/UObject/LinkerLoad.cpp
@@ -3503,26 +3503,30 @@ UObject* ULinkerLoad::CreateExport( int32 Index )
 		
 		if( Export.Object )
 		{
-			// Check to see if LoadClass is a blueprint, which potentially needs to be refreshed and regenerated.  If so, regenerate and patch it back into the export table
-			if( !LoadClass->bCooked
-				&& LoadClass->ClassGeneratedBy 
-				&& (LoadClass->GetOutermost() != GetTransientPackage()) 
-				&& ((Export.ObjectFlags&RF_ClassDefaultObject) != 0) )
-			{
+			bool const bIsBlueprintCDO = ((Export.ObjectFlags & RF_ClassDefaultObject) != 0) && (LoadClass->ClassGeneratedBy != nullptr);
+
 #if USE_CIRCULAR_DEPENDENCY_LOAD_DEFERRING
-				if ( (LoadFlags & LOAD_DeferDependencyLoads) != 0
+			bool bDeferCDOSerialization = bIsBlueprintCDO && ((LoadFlags & LOAD_DeferDependencyLoads) != 0);
 #if USE_DEFERRED_DEPENDENCY_CHECK_VERIFICATION_TESTS
-					&& !FBlueprintSupport::IsDeferredCDOSerializationDisabled()
-#endif // USE_DEFERRED_DEPENDENCY_CHECK_VERIFICATION_TESTS
-				   )
-				{
-					// if LOAD_DeferDependencyLoads is set, then we're already
-					// serializing the blueprint's class somewhere up the chain,
-					DeferredCDOIndex = Index;
-					return Export.Object;
-				}
-#endif // USE_CIRCULAR_DEPENDENCY_LOAD_DEFERRING
+			bDeferCDOSerialization &= !FBlueprintSupport::IsDeferredCDOSerializationDisabled();
+#endif // USE_DEFERRED_DEPENDENCY_CHECK_VERIFICATION_TESTS	
 
+			if (bDeferCDOSerialization)			
+			{
+				// if LOAD_DeferDependencyLoads is set, then we're already
+				// serializing the blueprint's class somewhere up the chain... 
+				// we don't want the class regenerated while it in the middle of
+				// serializing
+				DeferredCDOIndex = Index;
+				return Export.Object;
+			}
+			else 
+#endif // USE_CIRCULAR_DEPENDENCY_LOAD_DEFERRING
+			// Check to see if LoadClass is a blueprint, which potentially needs 
+			// to be refreshed and regenerated.  If so, regenerate and patch it 
+			// back into the export table
+			if( !LoadClass->bCooked && bIsBlueprintCDO && (LoadClass->GetOutermost() != GetTransientPackage()) )
+			{			
 				{
 					// For classes that are about to be regenerated, make sure we register them with the linker, so future references to this linker index will be valid
 					const EObjectFlags OldFlags = Export.Object->GetFlags();
