[INTEGRATE] Change 2442291 by [REDACTED_PERSON][REDACTED_HANDLE] on 2015/02/11 16:26:33

	UE-9064 Crash on shutdown after adding a child actor component to a Blueprint and PIE-ing a few times

diff --git a/Engine/Source/Runtime/CoreUObject/Private/UObject/Obj.cpp b/Engine/Source/Runtime/CoreUObject/Private/UObject/Obj.cpp
--- a/Engine/Source/Runtime/CoreUObject/Private/UObject/Obj.cpp
+++ b/Engine/Source/Runtime/CoreUObject/Private/UObject/Obj.cpp
@@ -3611,8 +3611,12 @@ void StaticExit()
 		// Valid object.
 		GObjectCountDuringLastMarkPhase++;
 
-		// Mark as unreachable so purge phase will kill it.
-		It->SetFlags( RF_Unreachable );
+		UObject* Obj = *It;
+		if (Obj && !Obj->IsA<UField>()) // Skip Structures, properties, etc.. They could be still necessary while GC.
+		{
+			// Mark as unreachable so purge phase will kill it.
+			It->SetFlags(RF_Unreachable);
+		}
 	}
 
 	// Fully purge all objects, not using time limit.
@@ -3634,6 +3638,27 @@ void StaticExit()
 
 	IncrementalPurgeGarbage( false );
 
+	{
+		//Repeat GC for every object, including structures and properties.
+		for (FRawObjectIterator It; It; ++It)
+		{
+			// Mark as unreachable so purge phase will kill it.
+			It->SetFlags(RF_Unreachable);
+		}
+
+		for (FRawObjectIterator It; It; ++It)
+		{
+			UObject* Object = *It;
+			if (Object->HasAnyFlags(RF_Unreachable))
+			{
+				// Begin the object's asynchronous destruction.
+				Object->ConditionalBeginDestroy();
+			}
+		}
+
+		IncrementalPurgeGarbage(false);
+	}
+
 	UObjectBaseShutdown();
 	// Empty arrays to prevent falsely-reported memory leaks.
 	GObjLoaded			.Empty();
