[INTEGRATE] Change 2444981 by [REDACTED_PERSON][REDACTED_HANDLE] on 2015/02/13 13:14:13

	[UE-9737] Making sure we defer dependency loads during user-defined-struct serialization as well.

diff --git a/Engine/Source/Runtime/CoreUObject/Private/UObject/LinkerLoad.cpp b/Engine/Source/Runtime/CoreUObject/Private/UObject/LinkerLoad.cpp
--- a/Engine/Source/Runtime/CoreUObject/Private/UObject/LinkerLoad.cpp
+++ b/Engine/Source/Runtime/CoreUObject/Private/UObject/LinkerLoad.cpp
@@ -1,5 +1,4 @@
 // Copyright 1998-2015 Epic Games, Inc. All Rights Reserved.
-
 #include "CoreUObjectPrivate.h"
 #include "SecureHash.h"
 #include "DebuggingDefines.h"
@@ -2852,17 +2851,26 @@ void ULinkerLoad::Preload( UObject* Object )
 	check(Object);
 
 #if USE_CIRCULAR_DEPENDENCY_LOAD_DEFERRING
-	UClass* ObjectAsClass = Cast<UClass>(Object);
-	// we can determine that this is a blueprint class, if it is indeed a class
-	// AND it is not native (blueprint classes are the only asset package 
-	// classes we have)
-	bool const bIsBlueprintClass = (ObjectAsClass != nullptr) && !(ObjectAsClass->GetOutermost()->PackageFlags & PKG_CompiledIn);
+	bool const bIsNonNativeObject = !(Object->GetOutermost()->PackageFlags & PKG_CompiledIn);
+	// we can determine that this is a blueprint class/struct by checking if it 
+	// is a class/struct object AND if it is not native (blueprint 
+	// structs/classes are the only asset package structs/classes we have)
+	bool const bIsBlueprintClass  = (Cast<UClass>(Object) != nullptr) && bIsNonNativeObject;
+	bool const bIsBlueprintStruct = (Cast<UScriptStruct>(Object) != nullptr) && bIsNonNativeObject;
+	// to avoid cyclic dependency issues, we want to defer all external loads 
+	// that MAY rely on this class/struct (meaning all other blueprint packages)  
+	bool const bDeferDependencyLoads = (bIsBlueprintClass || bIsBlueprintStruct) && FBlueprintSupport::UseDeferredDependencyLoading();
 
 #if USE_DEFERRED_DEPENDENCY_CHECK_VERIFICATION_TESTS
 	// we should NEVER be pre-loading another blueprint class when the 
-	// DeferDependencyLoads flag is set (some other blueprint class is already 
-	// being loaded further up the load chain, and this could introduce a 
-	// circular load)
+	// DeferDependencyLoads flag is set (some other blueprint class/struct is  
+	// already being loaded further up the load chain, and this could introduce  
+	// a circular load)
+	//
+	// NOTE: we do allow Preload() calls for structs (because we need a struct 
+	//       loaded to determine its size), but structs will be prevented from 
+	//       further loading any of its BP class dependencies (we pass along the 
+	//       LOAD_DeferDependencyLoads flag)
 	check(!bIsBlueprintClass || !Object->HasAnyFlags(RF_NeedLoad) || !(LoadFlags & LOAD_DeferDependencyLoads));
 	// right now there are no known scenarios where someone requests a Preloa() 
 	// on a temporary ULinkerPlaceholderExportObject
@@ -2891,7 +2899,7 @@ void ULinkerLoad::Preload( UObject* Object )
 
 #if USE_CIRCULAR_DEPENDENCY_LOAD_DEFERRING
 			TGuardValue<uint32> LoadFlagsGuard(LoadFlags, LoadFlags);			
-			if (FBlueprintSupport::UseDeferredDependencyLoading() && bIsBlueprintClass)
+			if (bDeferDependencyLoads)
 			{
 				LoadFlags |= LOAD_DeferDependencyLoads;
 			}
@@ -2934,7 +2942,6 @@ void ULinkerLoad::Preload( UObject* Object )
 #else // USE_DEFERRED_DEPENDENCY_CHECK_VERIFICATION_TESTS
 							{
 #endif // USE_DEFERRED_DEPENDENCY_CHECK_VERIFICATION_TESTS
-
 								// since serializing the CDO can introduce circular 
 								// dependencies, we want to stave that off until 
 								// we're ready to handle those 
@@ -2970,31 +2977,43 @@ void ULinkerLoad::Preload( UObject* Object )
 					SCOPE_CYCLE_COUNTER(STAT_LinkerLoadDeferred);
 					if ((LoadFlags & LOAD_DeferDependencyLoads) != (*LoadFlagsGuard & LOAD_DeferDependencyLoads))
 					{
+						if (bIsBlueprintStruct)
+						{
+							ResolveDeferredDependencies((UScriptStruct*)Object); 
+							// user-defined-structs don't have classes/CDOs, so 
+							// we don't have to call FinalizeBlueprint() (to 
+							// serialize/regenerate them)
+						}
+						else
+						{
+							UClass* ObjectAsClass = (UClass*)Object;
 #if USE_DEFERRED_DEPENDENCY_CHECK_VERIFICATION_TESTS
-						// since class serialization reads in the class's CDO, 
-						// then we can be certain that the CDO export object 
-						// exists (and DeferredExportIndex should reference it);
-						// FinalizeBlueprint() depends on this (and since 
-						// ResolveDeferredDependencies() can recurse into 
-						// FinalizeBlueprint(), we check it here, before the
-						// resolve is handled)
-						//
-						// however, sometimes DeferredExportIndex doesn't get
-						// set at all, and that happens when the class's 
-						// ClassGeneratedBy is serialized in null... this would 
-						// normally be a problem in the editor (we don't end up
-						// regenerating the class), but if we're running the 
-						// game (in PIE or elsewhere) it shouldn't be a concern 
-						// (if it makes you feel better: with the old way of 
-						// doing things, in CreateExport(), you can see that 
-						// the class wouldn't be regenerated there either)... in
-						// this scenario FinalizeBlueprint() essentially does nothing
-						// @TODO: maybe only allow a null ClassGeneratedBy when PIE'ing/running a game?
-						check((DeferredCDOIndex != INDEX_NONE) || (ObjectAsClass->ClassGeneratedBy == nullptr) || FBlueprintSupport::IsDeferredCDOSerializationDisabled());
+							check(bIsBlueprintClass);
+							// since class serialization reads in the class's CDO, 
+							// then we can be certain that the CDO export object 
+							// exists (and DeferredExportIndex should reference it);
+							// FinalizeBlueprint() depends on this (and since 
+							// ResolveDeferredDependencies() can recurse into 
+							// FinalizeBlueprint(), we check it here, before the
+							// resolve is handled)
+							//
+							// however, sometimes DeferredExportIndex doesn't get
+							// set at all, and that happens when the class's 
+							// ClassGeneratedBy is serialized in null... this would 
+							// normally be a problem in the editor (we don't end up
+							// regenerating the class), but if we're running the 
+							// game (in PIE or elsewhere) it shouldn't be a concern 
+							// (if it makes you feel better: with the old way of 
+							// doing things, in CreateExport(), you can see that 
+							// the class wouldn't be regenerated there either)... in
+							// this scenario FinalizeBlueprint() essentially does nothing
+							// @TODO: maybe only allow a null ClassGeneratedBy when PIE'ing/running a game?
+							check((DeferredCDOIndex != INDEX_NONE) || (ObjectAsClass->ClassGeneratedBy == nullptr) || FBlueprintSupport::IsDeferredCDOSerializationDisabled());
 #endif // USE_DEFERRED_DEPENDENCY_CHECK_VERIFICATION_TESTS
 
-						ResolveDeferredDependencies(ObjectAsClass);
-						FinalizeBlueprint(ObjectAsClass);
+							ResolveDeferredDependencies(ObjectAsClass);
+							FinalizeBlueprint(ObjectAsClass);
+						}
 					}
 #endif // USE_CIRCULAR_DEPENDENCY_LOAD_DEFERRING
 				}
@@ -3479,10 +3498,8 @@ UObject* ULinkerLoad::CreateExport( int32 Index )
 			}
 		}
 
-
 		LoadClass->GetDefaultObject();
 
-
 		Export.Object = StaticConstructObject_Internal
 		(
 			LoadClass,
@@ -3637,7 +3654,7 @@ UObject* ULinkerLoad::CreateImport( int32 Index )
 	// actively trying to avoid that at this point in the load process), then 
 	// this will stub in the Import with a placeholder object, to be replace 
 	// later on (this will return true if the import was actually deferred)
-	DeferPotentialCircularImport(Index);
+	DeferPotentialCircularImport(Index); 
 #endif // USE_CIRCULAR_DEPENDENCY_LOAD_DEFERRING
 
 	if( Import.XObject == NULL )
