[INTEGRATE] Change 2445052 by [REDACTED_PERSON][REDACTED_HANDLE] on 2015/02/13 13:44:07

	[UE-9736] Making sure we have the CDO's export index before we attempt a resolve/finalize pass for loading blueprints.

diff --git a/Engine/Source/Runtime/CoreUObject/Private/Blueprint/BlueprintSupport.cpp b/Engine/Source/Runtime/CoreUObject/Private/Blueprint/BlueprintSupport.cpp
--- a/Engine/Source/Runtime/CoreUObject/Private/Blueprint/BlueprintSupport.cpp
+++ b/Engine/Source/Runtime/CoreUObject/Private/Blueprint/BlueprintSupport.cpp
@@ -6,7 +6,6 @@
 #include "Linker.h"
 #include "PropertyTag.h"
 
-
 /** 
  * Defined in BlueprintSupport.cpp
  * Duplicates all fields of a class in depth-first order. It makes sure that everything contained
@@ -525,6 +524,24 @@ bool ULinkerLoad::DeferExportCreation(const int32 Index)
 	return true;
 }
 
+int32 ULinkerLoad::FindCDOExportIndex(UClass* LoadClass)
+{
+	DEFERRED_DEPENDENCY_CHECK(LoadClass->GetLinker() == this);
+	int32 const ClassExportIndex = LoadClass->GetLinkerIndex();
+
+	// @TODO: the cdo SHOULD be listed after the class in the ExportMap, so we 
+	//        could start with ClassExportIndex to save on some cycles
+	for (int32 ExportIndex = 0; ExportIndex < ExportMap.Num(); ++ExportIndex)
+	{
+		FObjectExport& Export = ExportMap[ExportIndex];
+		if ((Export.ObjectFlags & RF_ClassDefaultObject) && Export.ClassIndex.IsExport() && (Export.ClassIndex.ToExport() == ClassExportIndex))
+		{
+			return ExportIndex;
+		}
+	}
+	return INDEX_NONE;
+}
+
 /**
  * This utility struct helps track blueprint structs/linkers that are currently 
  * in the middle of a call to ResolveDeferredDependencies(). This can be used  
diff --git a/Engine/Source/Runtime/CoreUObject/Private/UObject/LinkerLoad.cpp b/Engine/Source/Runtime/CoreUObject/Private/UObject/LinkerLoad.cpp
--- a/Engine/Source/Runtime/CoreUObject/Private/UObject/LinkerLoad.cpp
+++ b/Engine/Source/Runtime/CoreUObject/Private/UObject/LinkerLoad.cpp
@@ -2989,26 +2989,28 @@ void ULinkerLoad::Preload( UObject* Object )
 							UClass* ObjectAsClass = (UClass*)Object;
 #if USE_DEFERRED_DEPENDENCY_CHECK_VERIFICATION_TESTS
 							check(bIsBlueprintClass);
-							// since class serialization reads in the class's CDO, 
-							// then we can be certain that the CDO export object 
-							// exists (and DeferredExportIndex should reference it);
-							// FinalizeBlueprint() depends on this (and since 
-							// ResolveDeferredDependencies() can recurse into 
-							// FinalizeBlueprint(), we check it here, before the
+							// since class serialization reads in the class's CDO, then we can be certain that the CDO export object exists 
+							// (and DeferredExportIndex should reference it); FinalizeBlueprint() depends on DeferredExportIndex being set 
+							// (and since ResolveDeferredDependencies() can recurse into FinalizeBlueprint(), we check it here, before the 
 							// resolve is handled)
 							//
-							// however, sometimes DeferredExportIndex doesn't get
-							// set at all, and that happens when the class's 
-							// ClassGeneratedBy is serialized in null... this would 
-							// normally be a problem in the editor (we don't end up
-							// regenerating the class), but if we're running the 
-							// game (in PIE or elsewhere) it shouldn't be a concern 
-							// (if it makes you feel better: with the old way of 
-							// doing things, in CreateExport(), you can see that 
-							// the class wouldn't be regenerated there either)... in
-							// this scenario FinalizeBlueprint() essentially does nothing
-							// @TODO: maybe only allow a null ClassGeneratedBy when PIE'ing/running a game?
-							check((DeferredCDOIndex != INDEX_NONE) || (ObjectAsClass->ClassGeneratedBy == nullptr) || FBlueprintSupport::IsDeferredCDOSerializationDisabled());
+							// however, sometimes DeferredExportIndex doesn't get set at all (we have to utilize FindCDOExportIndex() to set
+							// it), and that happens when the class's ClassGeneratedBy is serialized in null... this will happen for cooked 
+							// builds (because Blueprints are editor-only objects)
+							check((DeferredCDOIndex != INDEX_NONE) || FPlatformProperties::RequiresCookedData() || FBlueprintSupport::IsDeferredCDOSerializationDisabled());
+
+							if ((DeferredCDOIndex == INDEX_NONE) && !FBlueprintSupport::IsDeferredCDOSerializationDisabled())
+							{
+								DeferredCDOIndex = FindCDOExportIndex(ObjectAsClass);
+								check(DeferredCDOIndex != INDEX_NONE);
+							}
+#else  // USE_DEFERRED_DEPENDENCY_CHECK_VERIFICATION_TESTS
+							// just because DeferredCDOIndex wasn't set (in cooked/PIE scenarios) doesn't mean that we don't need it 
+							// (FinalizeBlueprint() relies on it being set), so here we make sure we flag the CDO so it gets resolved
+							if ((DeferredCDOIndex == INDEX_NONE))
+							{
+								DeferredCDOIndex = FindCDOExportIndex(ObjectAsClass);
+							}
 #endif // USE_DEFERRED_DEPENDENCY_CHECK_VERIFICATION_TESTS
 
 							ResolveDeferredDependencies(ObjectAsClass);
diff --git a/Engine/Source/Runtime/CoreUObject/Public/UObject/Linker.h b/Engine/Source/Runtime/CoreUObject/Public/UObject/Linker.h
--- a/Engine/Source/Runtime/CoreUObject/Public/UObject/Linker.h
+++ b/Engine/Source/Runtime/CoreUObject/Public/UObject/Linker.h
@@ -1864,6 +1864,16 @@ private:
 	 */
 	bool DeferExportCreation(const int32 ExportIndex);
 
+	/**
+	 * Iterates through this linker's ExportMap, looking for the corresponding
+	 * class-default-object for the specified class (assumes that the supplied 
+	 * class is an export itself, making this a Blueprint package).
+	 * 
+	 * @param  LoadClass    The Blueprint class that this linker is in charge of loading (also belonging to its ExportMap).
+	 * @return An index into this linker's ExportMap array (INDEX_NONE if the CDO wasn't found).
+	 */
+	int32 FindCDOExportIndex(UClass* LoadClass);
+
 	/**
 	 * Combs the ImportMap for any imports that were deferred, and then creates 
 	 * them (via CreateImport).
