Added ini option to flush streaming before GC

diff --git a/Engine/Config/BaseEngine.ini b/Engine/Config/BaseEngine.ini
--- a/Engine/Config/BaseEngine.ini
+++ b/Engine/Config/BaseEngine.ini
@@ -1272,6 +1272,7 @@ AsyncIOBandwidthLimit=0
 CutdownPaths=%GAMEDIR%CutdownPackages
 ZeroEngineVersionWarning=True
 UseStrictEngineVersioning=True
+FlushStreamingOnGC=False
 
 [Internationalization]
 +LocalizationPaths=../../../Engine/Content/Localization/Engine
diff --git a/Engine/Source/Runtime/CoreUObject/Private/UObject/GarbageCollection.cpp b/Engine/Source/Runtime/CoreUObject/Private/UObject/GarbageCollection.cpp
--- a/Engine/Source/Runtime/CoreUObject/Private/UObject/GarbageCollection.cpp
+++ b/Engine/Source/Runtime/CoreUObject/Private/UObject/GarbageCollection.cpp
@@ -1055,6 +1055,27 @@ static const auto CVarAllowParallelGC =
 
 void CollectGarbage( EObjectFlags KeepFlags, bool bPerformFullPurge )
 {
+	// Helper class to register FlushAsyncLoadingCallback on first GC run.
+	struct FAddFlushAsyncLoadingCallback
+	{
+		FAddFlushAsyncLoadingCallback()
+		{
+			bool bFlushStreaming = false;
+			GConfig->GetBool(TEXT("Core.System"), TEXT("FlushStreamingOnGC"), bFlushStreaming, GEngineIni);
+			if (bFlushStreaming)
+			{
+				FCoreUObjectDelegates::PreGarbageCollect.AddStatic(FlushAsyncLoadingCallback);
+			}
+		}
+		/** Wrapper function to handle default parameter when used as function pointer */
+		static void FlushAsyncLoadingCallback()
+		{
+			FlushAsyncLoading();
+		}
+	};
+	// Add FlushAsyncLoadingCallback the first time CollectGarbage is called if requested by ini settings
+	static FAddFlushAsyncLoadingCallback MaybeAddFlushAsyncLoadingCallback;
+
 	// We can't collect garbage while there's a load in progress. E.g. one potential issue is Import.XObject
 	check( !IsLoading() );
 
