[INTEGRATE] Change 2445257 by [REDACTED_PERSON][REDACTED_HANDLE] on 2015/02/13 15:22:19

	Fix Mac build failure.

diff --git a/Engine/Source/Runtime/CoreUObject/Private/UObject/LinkerLoad.cpp b/Engine/Source/Runtime/CoreUObject/Private/UObject/LinkerLoad.cpp
--- a/Engine/Source/Runtime/CoreUObject/Private/UObject/LinkerLoad.cpp
+++ b/Engine/Source/Runtime/CoreUObject/Private/UObject/LinkerLoad.cpp
@@ -3007,7 +3007,7 @@ void ULinkerLoad::Preload( UObject* Object )
 #else  // USE_DEFERRED_DEPENDENCY_CHECK_VERIFICATION_TESTS
 							// just because DeferredCDOIndex wasn't set (in cooked/PIE scenarios) doesn't mean that we don't need it 
 							// (FinalizeBlueprint() relies on it being set), so here we make sure we flag the CDO so it gets resolved
-							if ((DeferredCDOIndex == INDEX_NONE))
+							if (DeferredCDOIndex == INDEX_NONE)
 							{
 								DeferredCDOIndex = FindCDOExportIndex(ObjectAsClass);
 							}
