[INTEGRATE] Change 2445561 by [REDACTED_PERSON][REDACTED_HANDLE] on 2015/02/13 17:46:17

	[UE-9741] Deferring script serialization for game builds (so we don't introduce possible cyclic loads during class serialization).

diff --git a/Engine/Source/Runtime/CoreUObject/Private/Blueprint/BlueprintSupport.cpp b/Engine/Source/Runtime/CoreUObject/Private/Blueprint/BlueprintSupport.cpp
--- a/Engine/Source/Runtime/CoreUObject/Private/Blueprint/BlueprintSupport.cpp
+++ b/Engine/Source/Runtime/CoreUObject/Private/Blueprint/BlueprintSupport.cpp
@@ -5,6 +5,7 @@
 #include "UObject/LinkerPlaceholderExportObject.h"
 #include "Linker.h"
 #include "PropertyTag.h"
+#include "UObject/StructScriptLoader.h"
 
 /** 
  * Defined in BlueprintSupport.cpp
@@ -358,6 +359,12 @@ bool ULinkerLoad::DeferPotentialCircularImport(const int32 Index)
 				}
 			}
 		}
+
+		// not the best way to check this (but we don't have ObjectFlags on an 
+		// import), but we don't want non-native (blueprint) CDO refs slipping 
+		// through... we've only seen these needed when serializing a class's 
+		// bytecode, and we resolved that by deferring script serialization
+		DEFERRED_DEPENDENCY_CHECK(!Import.ObjectName.ToString().StartsWith("Default__"));
 	}
 	return (Import.XObject != nullptr);
 #else // !USE_CIRCULAR_DEPENDENCY_LOAD_DEFERRING
@@ -975,6 +982,14 @@ void ULinkerLoad::FinalizeBlueprint(UClass* LoadClass)
 		DEFERRED_DEPENDENCY_CHECK(BlueprintClass->HasAnyClassFlags(CLASS_CompiledFromBlueprint));
 #endif // USE_DEFERRED_DEPENDENCY_CHECK_VERIFICATION_TESTS
 
+		// for cooked builds (we skip script serialization for editor builds), 
+		// certain scripts can contain references to external dependencies; and 
+		// since the script is serialized in with the class (functions) we want 
+		// those dependencies deferred until now (we expect this to be the right 
+		// spot, because in editor builds, with RegenerateBlueprintClass(), this 
+		// is where script code is regenerated)
+		FStructScriptLoader::ResolveDeferredScriptLoads(this);
+
 		DEFERRED_DEPENDENCY_CHECK(LoadClass->GetOutermost() != GetTransientPackage());
 		// just in case we choose to enable the deferred dependency loading for 
 		// cooked builds... we want to keep from regenerating in that scenario
@@ -1033,9 +1048,9 @@ void ULinkerLoad::ResolveDeferredExports(UClass* LoadClass)
 			if (ULinkerPlaceholderExportObject* PlaceholderExport = Cast<ULinkerPlaceholderExportObject>(Export.Object))
 			{
 				DEFERRED_DEPENDENCY_CHECK(Export.ClassIndex.IsImport());
-				UClass* ImportClass = GetExportLoadClass(ExportIndex);
-				DEFERRED_DEPENDENCY_CHECK((ImportClass != nullptr) && !ImportClass->HasAnyClassFlags(CLASS_Intrinsic) && ImportClass->HasAnyClassFlags(CLASS_CompiledFromBlueprint));
-				ULinkerLoad* ClassLinker = ImportClass->GetLinker();
+				UClass* ExportClass = GetExportLoadClass(ExportIndex);
+				DEFERRED_DEPENDENCY_CHECK((ExportClass != nullptr) && !ExportClass->HasAnyClassFlags(CLASS_Intrinsic) && ExportClass->HasAnyClassFlags(CLASS_CompiledFromBlueprint));
+				ULinkerLoad* ClassLinker = ExportClass->GetLinker();
 				DEFERRED_DEPENDENCY_CHECK((ClassLinker != nullptr) && (ClassLinker != this));
 
 				{
@@ -1046,7 +1061,7 @@ void ULinkerLoad::ResolveDeferredExports(UClass* LoadClass)
 
 					// make sure this export's class is fully regenerated before  
 					// we instantiate it (so we don't have to re-inst on load)
-					ForceRegenerateClass(ImportClass);
+					ForceRegenerateClass(ExportClass);
 				}
 
 				// replace the placeholder with the proper object instance
diff --git a/Engine/Source/Runtime/CoreUObject/Private/UObject/Class.cpp b/Engine/Source/Runtime/CoreUObject/Private/UObject/Class.cpp
--- a/Engine/Source/Runtime/CoreUObject/Private/UObject/Class.cpp
+++ b/Engine/Source/Runtime/CoreUObject/Private/UObject/Class.cpp
@@ -8,53 +8,12 @@
 #include "PropertyTag.h"
 #include "HotReloadInterface.h"
 #include "LinkerPlaceholderClass.h"
+#include "StructScriptLoader.h"
 
 DECLARE_LOG_CATEGORY_EXTERN(LogScriptSerialization, Log, All);
 DEFINE_LOG_CATEGORY(LogScriptSerialization);
 DEFINE_LOG_CATEGORY(LogClass);
 
-
-/*----------------------------------------------------------------------------
-	FArchiveScriptReferenceCollector.
-----------------------------------------------------------------------------*/
-
-class FArchiveScriptReferenceCollector : public FArchiveUObject
-{
-public:
-	/**
-	 * Constructor
-	 *
-	 * @param	InObjectArray			Array to add object references to
-	 */
-	FArchiveScriptReferenceCollector( TArray<UObject*>& InObjectArray )
-	:	ObjectArray( InObjectArray )
-	{
-		ArIsObjectReferenceCollector = true;
-		ArIsPersistent = false;
-		ArIgnoreArchetypeRef = false;
-	}
-protected:
-	/** 
-	 * UObject serialize operator implementation
-	 *
-	 * @param Object	reference to Object reference
-	 * @return reference to instance of this class
-	 */
-	FArchive& operator<<( UObject*& Object )
-	{
-		// Avoid duplicate entries.
-		if ( Object != NULL && !ObjectArray.Contains(Object) )
-		{
-			check( Object->IsValidLowLevel() );
-			ObjectArray.Add( Object );
-		}
-		return *this;
-	}
-
-	/** Stored reference to array of objects we add object references to */
-	TArray<UObject*>&		ObjectArray;
-};
-
 //////////////////////////////////////////////////////////////////////////
 // FPropertySpecifier
 
@@ -1322,82 +1281,54 @@ void UStruct::Serialize( FArchive& Ar )
 	Ar << SuperStruct;
 	Ar << Children;
 
-	// Script code.
-	// Skip serialization if we're duplicating classes for reinstancing, since we only need the memory layout
-	int32 ScriptBytecodeSize = !GIsDuplicatingClassForReinstancing ? Script.Num() : 0;
-	int32 ScriptStorageSize = 0;
-	int32 ScriptStorageSizeOffset = 0;
-	if ( Ar.IsLoading() )
+	if (Ar.IsLoading())
 	{
-		Ar << ScriptBytecodeSize;
-		Ar << ScriptStorageSize;
+		FStructScriptLoader ScriptLoadHelper(/*TargetScriptContainer =*/this, Ar);
+#if USE_CIRCULAR_DEPENDENCY_LOAD_DEFERRING
+		bool const bAllowDeferredScriptSerialization = true;
+#else  // USE_CIRCULAR_DEPENDENCY_LOAD_DEFERRING
+		bool const bAllowDeferredScriptSerialization = false;
+#endif // USE_CIRCULAR_DEPENDENCY_LOAD_DEFERRING
 
-		Script.Empty( ScriptBytecodeSize );
-		Script.AddUninitialized( ScriptBytecodeSize );
-	}
+		// NOTE: if bAllowDeferredScriptSerialization is set to true, then this
+		//       could temporarily skip script serialization (as it could 
+		//       introduce unwanted dependency loads at this time)
+		ScriptLoadHelper.LoadStructWithScript(this, Ar, bAllowDeferredScriptSerialization);
 
-	// Ensure that last byte in script code is EX_EndOfScript to work around script debugger implementation.
-	else if( Ar.IsSaving() )
+		if (!dynamic_cast<UClass*>(this) && !(Ar.GetPortFlags() & PPF_Duplicate)) // classes are linked in the UClass serializer, which just called me
+		{
+			// Link the properties.
+			Link(Ar, true);
+		}
+	}
+	else
 	{
-		Ar << ScriptBytecodeSize;
+		int32 ScriptBytecodeSize = Script.Num();
+		int32 ScriptStorageSizeOffset = INDEX_NONE;
 
-		// drop a zero here.  will seek back later and re-write it when we know it
-		ScriptStorageSizeOffset = Ar.Tell();
-		Ar << ScriptStorageSize;
-	}
+		if (Ar.IsSaving())
+		{
+			Ar << ScriptBytecodeSize;
 
-	// If we're duplicating for reinstancing, we only need memory layout, and cyclic dependencies within object literals can potentially cause problems, so do not serialize bytecode
-	if( !GIsDuplicatingClassForReinstancing )
-	{
-		// no bytecode patch for this struct - serialize normally [i.e. from disk]
-		int32 iCode = 0;
-		int32 const BytecodeStartOffset = Ar.Tell();
+			int32 ScriptStorageSize = 0;
+			// drop a zero here.  will seek back later and re-write it when we know it
+			ScriptStorageSizeOffset = Ar.Tell();
+			Ar << ScriptStorageSize;
+		}
 
-		if (Ar.IsPersistent() && Ar.GetLinker())
+		// Skip serialization if we're duplicating classes for reinstancing, since we only need the memory layout
+		if (!GIsDuplicatingClassForReinstancing)
 		{
-			if (Ar.IsLoading())
-			{
-				// make sure this is a ULinkerLoad
-				ULinkerLoad* LinkerLoad = CastChecked<ULinkerLoad>(Ar.GetLinker());
 
-				// preload the bytecode
-				TArray<uint8> TempScript;
-				TempScript.AddUninitialized(ScriptStorageSize);
-				int32 ScriptStart = Ar.Tell();
-				Ar.Serialize(TempScript.GetData(), ScriptStorageSize);
-				const int32 ScriptEnd = Ar.Tell();
+			// no bytecode patch for this struct - serialize normally [i.e. from disk]
+			int32 iCode = 0;
+			int32 const BytecodeStartOffset = Ar.Tell();
 
-				bool bSkipByteCodeSerialization = false;
-#if WITH_EDITOR
-				static const FBoolConfigValueHelper SkipByteCodeHelper(TEXT("StructSerialization"), TEXT("SkipByteCodeSerialization"));
-				bSkipByteCodeSerialization = SkipByteCodeHelper;
-#endif // WITH_EDITOR
-				if (bSkipByteCodeSerialization || (Ar.UE4Ver() < VER_MIN_SCRIPTVM_UE4) || (Ar.LicenseeUE4Ver() < VER_MIN_SCRIPTVM_LICENSEEUE4))
-				{
-					// Discard the bytecode as it's too old and might cause serialization errors
-					ScriptStorageSize = 0;
-					ScriptBytecodeSize = 0;
-					TempScript.Empty();
-					Script.Empty();
-				}
-				else
-				{
-					Ar.Seek(ScriptStart); // seek back and load it again
-					// now, use the linker to load the byte code, but reading from memory
-					while( iCode < ScriptBytecodeSize )
-					{	
-						SerializeExpr( iCode, Ar );
-					}
-					ensure(Ar.Tell() == ScriptEnd);
-				}
-				// and update the SHA (does nothing if not currently calculating SHA)
-				LinkerLoad->UpdateScriptSHAKey(TempScript);
-			}
-			else
+			if (Ar.IsPersistent() && Ar.GetLinker())
 			{
 				// make sure this is a ULinkerSave
 				ULinkerSave* LinkerSave = CastChecked<ULinkerSave>(Ar.GetLinker());
-			
+
 				// remember how we were saving
 				FArchive* SavedSaver = LinkerSave->Saver;
 
@@ -1407,9 +1338,9 @@ void UStruct::Serialize( FArchive& Ar )
 				LinkerSave->Saver = &MemWriter;
 
 				// now, use the linker to save the byte code, but writing to memory
-				while( iCode < ScriptBytecodeSize )
-				{	
-					SerializeExpr( iCode, Ar );
+				while (iCode < ScriptBytecodeSize)
+				{
+					SerializeExpr(iCode, Ar);
 				}
 
 				// restore the saver
@@ -1421,54 +1352,32 @@ void UStruct::Serialize( FArchive& Ar )
 				// and update the SHA (does nothing if not currently calculating SHA)
 				LinkerSave->UpdateScriptSHAKey(TempScript);
 			}
-		}
-		else
-		{
-			while( iCode < ScriptBytecodeSize )
-			{	
-				SerializeExpr( iCode, Ar );
+			else
+			{
+				while (iCode < ScriptBytecodeSize)
+				{
+					SerializeExpr(iCode, Ar);
+				}
 			}
-		}
 
-		if( iCode != ScriptBytecodeSize )
-		{	
-			UE_LOG(LogClass, Fatal, TEXT("Script serialization mismatch: Got %i, expected %i"), iCode, ScriptBytecodeSize );
-		}
-
-		if (Ar.IsSaving())
-		{
-			int32 const BytecodeEndOffset = Ar.Tell();
-
-			// go back and write on-disk size
-			Ar.Seek(ScriptStorageSizeOffset);
-			ScriptStorageSize = BytecodeEndOffset - BytecodeStartOffset;
-			Ar << ScriptStorageSize;
+			if (iCode != ScriptBytecodeSize)
+			{
+				UE_LOG(LogClass, Fatal, TEXT("Script serialization mismatch: Got %i, expected %i"), iCode, ScriptBytecodeSize);
+			}
 
-			// back to where we were
-			Ar.Seek(BytecodeEndOffset);
-		}
-		if( Ar.IsLoading() )
-		{
-			// Collect references to objects embedded in script and store them in easily accessible array. This is skipped if
-			// the struct is disregarded for GC as the references won't be of any use.
-			ScriptObjectReferences.Empty();
-			if( !GUObjectArray.IsDisregardForGC(this) )
+			if (Ar.IsSaving())
 			{
-				FArchiveScriptReferenceCollector ObjectReferenceCollector( ScriptObjectReferences );
+				int32 const BytecodeEndOffset = Ar.Tell();
 
-				int32 iCode2 = 0;
-				while( iCode2 < Script.Num() )
-				{	
-					SerializeExpr( iCode2, ObjectReferenceCollector );
-				}
-			}
-		}	
-	}
+				// go back and write on-disk size
+				Ar.Seek(ScriptStorageSizeOffset);
+				int32 ScriptStorageSize = BytecodeEndOffset - BytecodeStartOffset;
+				Ar << ScriptStorageSize;
 
-	if (Ar.IsLoading() && !dynamic_cast<UClass*>(this) && !(Ar.GetPortFlags() & PPF_Duplicate)) // classes are linked in the UClass serializer, which just called me
-	{
-		// Link the properties.
-		Link( Ar, true );
+				// back to where we were
+				Ar.Seek(BytecodeEndOffset);
+			}
+		} // if !GIsDuplicatingClassForReinstancing
 	}
 }
 
diff --git a/Engine/Source/Runtime/CoreUObject/Private/UObject/StructScriptLoader.cpp b/Engine/Source/Runtime/CoreUObject/Private/UObject/StructScriptLoader.cpp
new file mode 100644
--- /dev/null
+++ b/Engine/Source/Runtime/CoreUObject/Private/UObject/StructScriptLoader.cpp
@@ -0,0 +1,290 @@
+// Copyright 1998-2015 Epic Games, Inc. All Rights Reserved.
+
+#include "CoreUObjectPrivate.h"
+#include "StructScriptLoader.h"
+
+/*******************************************************************************
+ * FDeferredScriptTracker
+ ******************************************************************************/
+
+/** 
+ * Pairs a FStructScriptLoader with the target script container (so that the 
+ * script can be properly serialized in at a later time).
+ */
+struct FDeferredScriptLoader : private FStructScriptLoader
+{
+public:
+	FDeferredScriptLoader(const FStructScriptLoader& RHS, UStruct* ScriptContainer);
+
+	/**
+	 * If the target script container is still valid, then this will load it 
+	 * with script bytecode from the supplied archiver (expects that the 
+	 * archiver is the same one that originally attempted to load the script).  
+	 * 
+	 * @param  Ar    The archiver that this script was originally supposed to load with.
+	 * @return True if the target was loaded with new bytecode, otherwise false.
+	 */
+	bool Resolve(FArchive& Ar);
+
+private:
+	/** Kept as a weak pointer in case the target has since been destroyed (since the initial deferral) */
+	TWeakObjectPtr<UStruct> TargetScriptContainerPtr;
+};
+
+/** 
+ * Tracks all deferred script loads (so that they can be resolved at a later 
+ * time, via FStructScriptLoader::ResolveDeferredScriptLoads). Utilized to avoid
+ * having to load possible cyclic dependencies during class serialization.
+ */
+struct FDeferredScriptTracker : TThreadSingleton<FDeferredScriptTracker>
+{
+public:
+	FDeferredScriptTracker();
+
+	/**
+	 * Stores the target struct along with the serialization offset, script 
+	 * size, etc. (so the script can be resolved at a later time).
+	 * 
+	 * @param  Linker					The loader responsible for serializing in the target struct's script.
+	 * @param  TargetScriptContainer    The struct that the script should ultimately be serialized into
+	 * @param  ScriptLoader				The script serialization helper that contains info on the script's serializtion offset (buffer size, etc.)
+	 */
+	void AddDeferredScriptObject(ULinkerLoad* Linker, UStruct* TargetScriptContainer, const FStructScriptLoader& ScriptLoader);
+
+	/**
+	 * Goes through every deferred script load associated with the specified 
+	 * linker and attempts to resolve each one (will fail to resolve any if the 
+	 * linker is still flagged with LOAD_DeferDependencyLoads).
+	 * 
+	 * @param  Linker    The linker that may have deferred script serialization (possibly for many functions).
+	 * @return The number of script loads that were successfully resolved.
+	 */
+	int32 ResolveDeferredScripts(ULinkerLoad* Linker);
+
+private:
+#if USE_DEFERRED_DEPENDENCY_CHECK_VERIFICATION_TESTS
+	/** Used to catch any deferred script loads that are added during a call to ResolveDeferredScripts() */
+	ULinkerLoad* ResolvingLinker;
+#endif // USE_DEFERRED_DEPENDENCY_CHECK_VERIFICATION_TESTS
+
+	/** Keeps track of scripts (and their target containers) that need to be serialized in later */
+	TMultiMap<ULinkerLoad*, FDeferredScriptLoader> DeferredScriptLoads;
+};
+
+//------------------------------------------------------------------------------
+FDeferredScriptTracker::FDeferredScriptTracker()
+#if USE_DEFERRED_DEPENDENCY_CHECK_VERIFICATION_TESTS
+	: ResolvingLinker(nullptr)
+#endif // USE_DEFERRED_DEPENDENCY_CHECK_VERIFICATION_TESTS
+{
+}
+
+//------------------------------------------------------------------------------
+void FDeferredScriptTracker::AddDeferredScriptObject(ULinkerLoad* Linker, UStruct* TargetScriptContainer, const FStructScriptLoader& ScriptLoader)
+{
+#if USE_DEFERRED_DEPENDENCY_CHECK_VERIFICATION_TESTS
+	check(ResolvingLinker == nullptr);
+#endif // USE_DEFERRED_DEPENDENCY_CHECK_VERIFICATION_TESTS
+
+	DeferredScriptLoads.Add(Linker, FDeferredScriptLoader(ScriptLoader, TargetScriptContainer));
+}
+
+//------------------------------------------------------------------------------
+int32 FDeferredScriptTracker::ResolveDeferredScripts(ULinkerLoad* Linker)
+{
+	FArchive& Ar = *Linker;
+	if (FStructScriptLoader::ShouldDeferScriptSerialization(Ar))
+	{
+		return 0;
+	}
+
+#if USE_DEFERRED_DEPENDENCY_CHECK_VERIFICATION_TESTS
+	TGuardValue<ULinkerLoad*> ScopedResolvingLinker(ResolvingLinker, Linker);
+#endif // USE_DEFERRED_DEPENDENCY_CHECK_VERIFICATION_TESTS
+
+	TArray<FDeferredScriptLoader> DefferedLinkerScripts;
+	DeferredScriptLoads.MultiFind(Linker, DefferedLinkerScripts);
+	// remove before we resolve, because a failed call to 
+	// FDeferredScriptLoader::Resolve() could insert back into this list 
+	DeferredScriptLoads.Remove(Linker);
+
+	int32 const SerializationPosToRestore = Ar.Tell();
+
+	int32 ResolveCount = 0;
+	for (FDeferredScriptLoader& DeferredScript : DefferedLinkerScripts)
+	{
+		if (DeferredScript.Resolve(Ar))
+		{
+			++ResolveCount;
+		}
+	}
+
+	Ar.Seek(SerializationPosToRestore);
+	return ResolveCount;
+}
+
+//------------------------------------------------------------------------------
+FDeferredScriptLoader::FDeferredScriptLoader(const FStructScriptLoader& RHS, UStruct* TargetScriptContainer)
+	: FStructScriptLoader(RHS)
+	, TargetScriptContainerPtr(TargetScriptContainer)
+{
+}
+
+//------------------------------------------------------------------------------
+bool FDeferredScriptLoader::Resolve(FArchive& Ar)
+{
+	if (UStruct* Target = TargetScriptContainerPtr.Get())
+	{
+		return LoadStructWithScript(Target, Ar);
+	}
+	return false;
+}
+
+/*******************************************************************************
+ * FStructScriptLoader
+ ******************************************************************************/
+
+//------------------------------------------------------------------------------
+FStructScriptLoader::FStructScriptLoader(UStruct* TargetScriptContainer, FArchive& Ar)
+	: BytecodeBufferSize(0)
+	, SerializedScriptSize(0)
+	, ScriptSerializationOffset(INDEX_NONE)
+{
+	if (!Ar.IsLoading())
+	{
+		return;
+	}
+
+	Ar << BytecodeBufferSize;
+	Ar << SerializedScriptSize;
+	
+	if (SerializedScriptSize > 0)
+	{
+		ScriptSerializationOffset = Ar.Tell();
+	}
+
+	ClearScriptCode(TargetScriptContainer);
+}
+
+//------------------------------------------------------------------------------
+bool FStructScriptLoader::IsPrimed()
+{
+	return (SerializedScriptSize > 0) && (ScriptSerializationOffset != INDEX_NONE);
+}
+
+//------------------------------------------------------------------------------
+bool FStructScriptLoader::ShouldDeferScriptSerialization(FArchive& Ar)
+{
+#if USE_CIRCULAR_DEPENDENCY_LOAD_DEFERRING
+	if (Ar.IsLoading() && Ar.IsPersistent())
+	{
+		if (ULinkerLoad* Linker = Cast<ULinkerLoad>(Ar.GetLinker()))
+		{
+			return ((Linker->LoadFlags & LOAD_DeferDependencyLoads) != 0);
+		}
+	}
+#endif // USE_CIRCULAR_DEPENDENCY_LOAD_DEFERRING
+	return false;
+}
+
+//------------------------------------------------------------------------------
+bool FStructScriptLoader::LoadStructWithScript(UStruct* DestScriptContainer, FArchive& Ar, bool bAllowDeferredSerialization)
+{
+	if (!Ar.IsLoading() || !IsPrimed() || GIsDuplicatingClassForReinstancing)
+	{
+		return false;
+	}
+
+	bool const bIsLinkerLoader = Ar.IsPersistent() && (Ar.GetLinker() != nullptr);
+	int32 const ScriptEndOffset = ScriptSerializationOffset + SerializedScriptSize;
+
+	// to help us move development forward (and not have to support ancient 
+	// script code), we define a minimum script version
+	bool bSkipScriptSerialization = (Ar.UE4Ver() < VER_MIN_SCRIPTVM_UE4) || (Ar.LicenseeUE4Ver() < VER_MIN_SCRIPTVM_LICENSEEUE4);
+#if WITH_EDITOR
+	static const FBoolConfigValueHelper SkipByteCodeHelper(TEXT("StructSerialization"), TEXT("SkipByteCodeSerialization"));
+	// in editor builds, we're going to regenerate the bytecode anyways, so it
+	// is a waste of cycles to try and serialize it in
+	bSkipScriptSerialization |= (bool)SkipByteCodeHelper;
+#endif // WITH_EDITOR
+	bSkipScriptSerialization &= bIsLinkerLoader; // to keep consistent with old UStruct::Serialize() functionality
+
+	if (bSkipScriptSerialization)
+	{
+		int32 TrackedBufferSize = BytecodeBufferSize;
+		BytecodeBufferSize = 0; // temporarily clear so that ClearScriptCode() doesn't leave Class->Script with anything allocated
+		ClearScriptCode(DestScriptContainer);
+		BytecodeBufferSize = TrackedBufferSize;
+
+		// we have to at least move the archiver forward, so it is positioned 
+		// where it expects to be (as if we read in the script)
+		Ar.Seek(ScriptEndOffset);
+		return false;
+	}
+
+	bAllowDeferredSerialization &= bIsLinkerLoader;
+	if (bAllowDeferredSerialization && ShouldDeferScriptSerialization(Ar))
+	{
+		ULinkerLoad* Linker = CastChecked<ULinkerLoad>(Ar.GetLinker());
+		FDeferredScriptTracker::Get().AddDeferredScriptObject(Linker, DestScriptContainer, *this);
+
+		// we have to at least move the archiver forward, so it is positioned 
+		// where it expects to be (as if we read in the script)
+		Ar.Seek(ScriptEndOffset);
+		return false;
+	}
+
+	Ar.Seek(ScriptSerializationOffset);
+	if (bIsLinkerLoader)
+	{
+		ULinkerLoad* LinkerLoad = CastChecked<ULinkerLoad>(Ar.GetLinker());
+
+		TArray<uint8> ShaScriptBuffer;
+		ShaScriptBuffer.AddUninitialized(SerializedScriptSize);
+
+		Ar.Serialize(ShaScriptBuffer.GetData(), SerializedScriptSize);
+		ensure(ScriptEndOffset == Ar.Tell());
+		LinkerLoad->UpdateScriptSHAKey(ShaScriptBuffer);
+
+		Ar.Seek(ScriptSerializationOffset);
+	}
+
+	DestScriptContainer->Script.Empty(BytecodeBufferSize);
+	DestScriptContainer->Script.AddUninitialized(BytecodeBufferSize);
+
+	int32 BytecodeIndex = 0;
+	while (BytecodeIndex < BytecodeBufferSize)
+	{
+		DestScriptContainer->SerializeExpr(BytecodeIndex, Ar);
+	}
+	ensure(ScriptEndOffset == Ar.Tell());
+	checkf(BytecodeIndex == BytecodeBufferSize, TEXT("'%s' script expression-count mismatch; Expected: %i, Got: %i"), *DestScriptContainer->GetName(), BytecodeBufferSize, BytecodeIndex);
+
+	if (!GUObjectArray.IsDisregardForGC(DestScriptContainer))
+	{
+		DestScriptContainer->ScriptObjectReferences.Empty();
+		FArchiveScriptReferenceCollector ObjRefCollector(DestScriptContainer->ScriptObjectReferences);
+
+		BytecodeIndex = 0;
+		while (BytecodeIndex > BytecodeBufferSize)
+		{
+			DestScriptContainer->SerializeExpr(BytecodeIndex, ObjRefCollector);
+		}
+	}
+
+	// success! (we filled the target with serialized script code)
+	return true;
+}
+
+//------------------------------------------------------------------------------
+int32 FStructScriptLoader::ResolveDeferredScriptLoads(ULinkerLoad* Linker)
+{
+	return FDeferredScriptTracker::Get().ResolveDeferredScripts(Linker);
+}
+
+//------------------------------------------------------------------------------
+void FStructScriptLoader::ClearScriptCode(UStruct* ScriptContainer)
+{
+	ScriptContainer->Script.Empty(BytecodeBufferSize);
+	ScriptContainer->ScriptObjectReferences.Empty();
+}
diff --git a/Engine/Source/Runtime/CoreUObject/Private/UObject/StructScriptLoader.h b/Engine/Source/Runtime/CoreUObject/Private/UObject/StructScriptLoader.h
new file mode 100644
--- /dev/null
+++ b/Engine/Source/Runtime/CoreUObject/Private/UObject/StructScriptLoader.h
@@ -0,0 +1,130 @@
+// Copyright 1998-2014 Epic Games, Inc. All Rights Reserved.
+
+#pragma once
+
+/*******************************************************************************
+ * FStructScriptLoader
+ ******************************************************************************/
+
+/**  
+ * Utility class to aid in struct (UFunction) script serialization. Will defer 
+ * or skip script loading (if deemed necessary).
+ */ 
+class FStructScriptLoader
+{
+public:	
+	/**
+	 * Caches data regarding the script's serialized form (archiver offset, 
+	 * serialized size, etc.), so that given the same archiver later (in 
+	 * LoadStructWithScript) it can seek and serialize the target's bytecode.
+	 *
+	 * NOTE: This expects that the supplied archiver is already positioned to 
+	 *       read the start of the script's "header" (the script's BytecodeBufferSize).
+	 * 
+	 * @param  TargetScriptContainer    The target struct (will have its current script cleared).
+	 * @param  Loader					The archiver that holds serialized script data (expected to be seeked to the proper place). 
+	 */
+	FStructScriptLoader(UStruct* TargetScriptContainer, FArchive& Loader);
+
+	/**
+	 * Checks to see if this was created from a valid archiver with script code
+	 * to serialize in.
+	 * 
+	 * @return True if the cached serialization info is valid, otherwise false.
+	 */
+	bool IsPrimed();
+
+	/**
+	 * Can be used to determine if the specified archiver wants possible 
+	 * dependency load points (such as bytecode) deferred (until after class 
+	 * serialization).
+	 * 
+	 * @param  Ar    The archiver to check.
+	 * @return True if the specified archiver has LOAD_DeferDependencyLoads set, otherwise false.
+	 */
+	static bool ShouldDeferScriptSerialization(FArchive& Ar);
+
+	/**
+	 * Attempts to load the specified target struct with bytecode contained in 
+	 * the supplied archiver (expects the archiver to be the same one that this
+	 * was created from).
+	 *
+	 * NOTE: Serialization could be skipped if: 1) this isn't properly "primed",
+	 *       2) we've opted to skip bytecode serialization (for editor builds), 
+	 *       or 3) the Loader wishes to have dependency link points deferred 
+	 *       (unless bAllowDeferredSerialization is set false)
+	 * 
+	 * @param  DestScriptContainer			The target struct that you want script code loaded in to.	
+	 * @param  Loader						The archiver that holds the serialized script code meant for the target.
+	 * @param  bAllowDeferredSerialization	If false, then serialization will be forced through regardless of ShouldDeferScriptSerialization().
+	 * @return 
+	 */
+	bool LoadStructWithScript(UStruct* DestScriptContainer, FArchive& Loader, bool bAllowDeferredSerialization = true);
+
+	/**
+	 * Looks for any struct scripts that were deferred as part of a 
+	 * LoadStructWithScript() call, and attempts to serialize the original 
+	 * targets with the deferred load.
+	 * 
+	 * @param  Linker    The linker/archiver that the structs were originally meant to be loaded by.
+	 * @return The number of structs that were loaded with new bytecode.
+	 */
+	static int32 ResolveDeferredScriptLoads(ULinkerLoad* Linker);
+
+private:
+	/**
+	 * Empties the specified struct's Script array, as well as its 
+	 * ScriptObjectReferences container.
+	 * 
+	 * @param  ScriptContainer    The struct you want emptied.
+	 */
+	void ClearScriptCode(UStruct* ScriptContainer);
+
+	/** Determines the size of the target struct's Script array */
+	int32 BytecodeBufferSize;
+	/** Denotes how many bytes the serialized script code is, stored in the archiver (NOTE: could be more compact than BytecodeBufferSize) */
+	int32 SerializedScriptSize;
+	/** Denotes where to offset the archiver to start reading in bytecode */
+	int32 ScriptSerializationOffset;
+}; 
+
+/*******************************************************************************
+ * FArchiveScriptReferenceCollector
+ ******************************************************************************/
+
+class FArchiveScriptReferenceCollector : public FArchiveUObject
+{
+public:
+	/**
+	* Constructor
+	*
+	* @param	InObjectArray			Array to add object references to
+	*/
+	FArchiveScriptReferenceCollector(TArray<UObject*>& InObjectArray)
+		: ObjectArray(InObjectArray)
+	{
+		ArIsObjectReferenceCollector = true;
+		ArIsPersistent = false;
+		ArIgnoreArchetypeRef = false;
+	}
+protected:
+	/**
+	* UObject serialize operator implementation
+	*
+	* @param Object	reference to Object reference
+	* @return reference to instance of this class
+	*/
+	FArchive& operator<<(UObject*& Object)
+	{
+		// Avoid duplicate entries.
+		if (Object != NULL && !ObjectArray.Contains(Object))
+		{
+			check(Object->IsValidLowLevel());
+			ObjectArray.Add(Object);
+		}
+		return *this;
+	}
+
+	/** Stored reference to array of objects we add object references to */
+	TArray<UObject*>&		ObjectArray;
+};
