[INTEGRATE] Change 2448017 by [REDACTED_PERSON][REDACTED_HANDLE] on 2015/02/16 19:59:10

	Fix incorrect sign resulting in object references not being collected when loading in a struct

diff --git a/Engine/Source/Runtime/CoreUObject/Private/UObject/StructScriptLoader.cpp b/Engine/Source/Runtime/CoreUObject/Private/UObject/StructScriptLoader.cpp
--- a/Engine/Source/Runtime/CoreUObject/Private/UObject/StructScriptLoader.cpp
+++ b/Engine/Source/Runtime/CoreUObject/Private/UObject/StructScriptLoader.cpp
@@ -266,7 +266,7 @@ bool FStructScriptLoader::LoadStructWithScript(UStruct* DestScriptContainer, FAr
 		FArchiveScriptReferenceCollector ObjRefCollector(DestScriptContainer->ScriptObjectReferences);
 
 		BytecodeIndex = 0;
-		while (BytecodeIndex > BytecodeBufferSize)
+		while (BytecodeIndex < BytecodeBufferSize)
 		{
 			DestScriptContainer->SerializeExpr(BytecodeIndex, ObjRefCollector);
 		}
