[INTEGRATE] Change 2449696 by [REDACTED_PERSON][REDACTED_HANDLE] on 2015/02/17 23:31:40

	UE-9989 Copying and Pasting BP with Comms into New Level Hangs Editor

diff --git a/Engine/Source/Runtime/CoreUObject/Private/UObject/PropertyBaseObject.cpp b/Engine/Source/Runtime/CoreUObject/Private/UObject/PropertyBaseObject.cpp
--- a/Engine/Source/Runtime/CoreUObject/Private/UObject/PropertyBaseObject.cpp
+++ b/Engine/Source/Runtime/CoreUObject/Private/UObject/PropertyBaseObject.cpp
@@ -171,7 +171,14 @@ void UObjectPropertyBase::ExportTextItem( FString& ValueStr, const void* Propert
 		}
 		else if (Parent && !Parent->HasAnyFlags(RF_ClassDefaultObject) && Temp->IsDefaultSubobject())
 		{
-			ValueStr += Temp->GetName();
+			if ((PortFlags & PPF_Delimited) && (!Temp->GetFName().IsValidXName(INVALID_OBJECTNAME_CHARACTERS)))
+			{
+				ValueStr += FString::Printf(TEXT("\"%s\""), *Temp->GetName().ReplaceQuotesWithEscapedQuotes());
+			}
+			else
+			{
+				ValueStr += Temp->GetName();
+			}
 		}
 		else
 		{
