Fixing a bug in GC where IncrementalPurge would not call destructors on objects residing in permanent pool while performing exit purge

diff --git a/Engine/Source/Runtime/CoreUObject/Private/UObject/GarbageCollection.cpp b/Engine/Source/Runtime/CoreUObject/Private/UObject/GarbageCollection.cpp
--- a/Engine/Source/Runtime/CoreUObject/Private/UObject/GarbageCollection.cpp
+++ b/Engine/Source/Runtime/CoreUObject/Private/UObject/GarbageCollection.cpp
@@ -955,10 +955,9 @@ void IncrementalPurgeGarbage( bool bUseTimeLimit, float TimeLimit )
 				GGCObjectsPendingDestruction.Empty( 256 );
 
 				// Destroy has been routed to all objects so it's safe to delete objects now.
-				GObjFinishDestroyHasBeenRoutedToAllObjects		= true;
-				GObjCurrentPurgeObjectIndexNeedsReset			= true;
-				GObjCurrentPurgeObjectIndexResetPastPermanent	= true;
-
+				GObjFinishDestroyHasBeenRoutedToAllObjects = true;
+				GObjCurrentPurgeObjectIndexNeedsReset = true;
+				GObjCurrentPurgeObjectIndexResetPastPermanent = !GExitPurge;
 			}
 		}
 	}		
