Fix for crash on exit after async loading of a package failed

diff --git a/Engine/Source/Runtime/CoreUObject/Private/Serialization/AsyncLoading.cpp b/Engine/Source/Runtime/CoreUObject/Private/Serialization/AsyncLoading.cpp
--- a/Engine/Source/Runtime/CoreUObject/Private/Serialization/AsyncLoading.cpp
+++ b/Engine/Source/Runtime/CoreUObject/Private/Serialization/AsyncLoading.cpp
@@ -900,15 +900,34 @@ EAsyncPackageState::Type FAsyncPackage::FinishObjects()
 		Object->ClearFlags( RF_AsyncLoading );
 	}
 	GObjConstructedDuringAsyncLoading.Empty();
-			
-	// Simulate what EndLoad does.
-	GObjLoaded.Empty();
+		
+	EAsyncLoadingResult::Type LoadingResult;
+	if (!bLoadHasFailed)
+	{
+		GObjLoaded.Empty();
+		LoadingResult = EAsyncLoadingResult::Succeeded;		
+	}
+	else
+	{
+		// Cleanup objects from this package only
+		for (int ObjectIndex = GObjLoaded.Num() - 1; ObjectIndex >= 0; --ObjectIndex)
+		{
+			UObject* Object = GObjLoaded[ObjectIndex];
+			if (Object->GetOutermost()->GetFName() == PackageName)
+			{
+				Object->ClearFlags(RF_NeedPostLoad | RF_NeedLoad | RF_NeedPostLoadSubobjects);
+				Object->MarkPendingKill();
+				GObjLoaded.RemoveAt(ObjectIndex);
+			}			
+		}
+		LoadingResult = EAsyncLoadingResult::Failed;
+	}
+
+	// Simulate what EndLoad does.	
 	DissociateImportsAndForcedExports(); //@todo: this should be avoidable
 	PreLoadIndex = 0;
 	PostLoadIndex = 0;
 
-	EAsyncLoadingResult::Type LoadingResult = bLoadHasFailed ? EAsyncLoadingResult::Failed : EAsyncLoadingResult::Succeeded;
-
 	// If we successfully loaded
 	if (!bLoadHasFailed)
 	{
