GC performace optimization.

diff --git a/Engine/Source/Runtime/CoreUObject/Private/UObject/GarbageCollection.cpp b/Engine/Source/Runtime/CoreUObject/Private/UObject/GarbageCollection.cpp
--- a/Engine/Source/Runtime/CoreUObject/Private/UObject/GarbageCollection.cpp
+++ b/Engine/Source/Runtime/CoreUObject/Private/UObject/GarbageCollection.cpp
@@ -1238,17 +1238,6 @@ void UObject::AddReferencedObjects(UObject* This, FReferenceCollector& Collector
 		Collector.AllowEliminatingReferences(true);
 		Collector.AddReferencedObject( Class, This );
 
-		// Serialize object properties which are defined in the class.
-		// Note: This check is intentionally excluding UClass objects but including subclasses like UBlueprintGeneratedClass
-		// @TODO: is it right to also exclude ULinkerPlaceholderClass here (it was causing a crash in here).
-		if (Class != UClass::StaticClass() && Class != ULinkerPlaceholderClass::StaticClass())
-		{
-			// Script properties
-			FSimpleObjectReferenceCollectorArchive ObjectReferenceCollector( This, Collector );
-			ObjectReferenceCollector.SetSerializedProperty(Collector.GetSerializedProperty());
-			This->SerializeScriptProperties( ObjectReferenceCollector );
-		}
-
 #if TEST_ARO_FINDS_ALL_OBJECTS
 		TArray<UObject*> SerializedObjects;
 		AddReferencedObjectsViaSerialization( This, SerializedObjects );
diff --git a/Engine/Source/Runtime/Foliage/Private/ProceduralFoliage.cpp b/Engine/Source/Runtime/Foliage/Private/ProceduralFoliage.cpp
--- a/Engine/Source/Runtime/Foliage/Private/ProceduralFoliage.cpp
+++ b/Engine/Source/Runtime/Foliage/Private/ProceduralFoliage.cpp
@@ -136,7 +136,7 @@ const UProceduralFoliageTile* UProceduralFoliage::GetRandomTile(int32 X, int32 Y
 		const float YRand = HashStream.FRand();
 		const int32 RandomNumber = (RAND_MAX * XRand / (YRand + 0.01f));
 		const int32 Idx = FMath::Clamp(RandomNumber % PrecomputedTiles.Num(), 0, PrecomputedTiles.Num() - 1);
-		return PrecomputedTiles[Idx];
+		return PrecomputedTiles[Idx].Get();
 	}
 
 	return nullptr;
diff --git a/Engine/Source/Runtime/Foliage/Public/ProceduralFoliage.h b/Engine/Source/Runtime/Foliage/Public/ProceduralFoliage.h
--- a/Engine/Source/Runtime/Foliage/Public/ProceduralFoliage.h
+++ b/Engine/Source/Runtime/Foliage/Public/ProceduralFoliage.h
@@ -82,8 +82,7 @@ private:
 
 	void SetClean();
 private:
-	UPROPERTY()
-	TArray<UProceduralFoliageTile*> PrecomputedTiles;
+	TArray<TWeakObjectPtr<UProceduralFoliageTile>> PrecomputedTiles;
 
 	FRandomStream RandomStream;
 };
