Reinstated the "Resource Size" asset tooltip in the content browser.  Added a facility to specify a localizable suffix on asset registry tags.
#ttp 325917 - Editor: Mouse over Assets in content browser is not showing resource size anymore.
#branch UE4
#proj Runtime.CoreUObject, Editor.ContentBrowser

diff --git a/Engine/Source/Runtime/CoreUObject/Private/UObject/Obj.cpp b/Engine/Source/Runtime/CoreUObject/Private/UObject/Obj.cpp
--- a/Engine/Source/Runtime/CoreUObject/Private/UObject/Obj.cpp
+++ b/Engine/Source/Runtime/CoreUObject/Private/UObject/Obj.cpp
@@ -828,11 +828,11 @@ void UObject::Serialize( FArchive& Ar )
 	}
 
 	// Serialize object properties which are defined in the class.
-	// Handle derived UClass objects (exact UClass objects are native only and shouldn't be touched)
-	if (Class != UClass::StaticClass())
-	{
-		SerializeScriptProperties(Ar);
-	}
+		// Handle derived UClass objects (exact UClass objects are native only and shouldn't be touched)
+		if (Class != UClass::StaticClass())
+		{
+				SerializeScriptProperties(Ar);
+			}
 
 	// Keep track of pending kill
 	if( Ar.IsTransacting() )
@@ -1167,7 +1167,7 @@ void UObject::GetAssetRegistryTags(TArray<FAssetRegistryTag>& OutTags) const
 {
 	// Add ResourceSize if non-zero. GetResourceSize is not const because many override implementations end up calling Serialize on this pointers.
 	SIZE_T ResourceSize = const_cast<UObject*>(this)->GetResourceSize(EResourceSizeMode::Exclusive);
-	if ( ResourceSize > 0 )
+	if ( ResourceSize != RESOURCE_SIZE_NONE )
 	{
 		OutTags.Add( FAssetRegistryTag("ResourceSize", FString::Printf(TEXT("%0.2f"), ResourceSize / 1024.f), FAssetRegistryTag::TT_Numerical) );
 	}
@@ -3614,9 +3614,9 @@ void StaticExit()
 		UObject* Obj = *It;
 		if (Obj && !Obj->IsA<UField>()) // Skip Structures, properties, etc.. They could be still necessary while GC.
 		{
-			// Mark as unreachable so purge phase will kill it.
+		// Mark as unreachable so purge phase will kill it.
 			It->SetFlags(RF_Unreachable);
-		}
+	}
 	}
 
 	// Fully purge all objects, not using time limit.
diff --git a/Engine/Source/Runtime/CoreUObject/Public/UObject/UObject.h b/Engine/Source/Runtime/CoreUObject/Public/UObject/UObject.h
--- a/Engine/Source/Runtime/CoreUObject/Public/UObject/UObject.h
+++ b/Engine/Source/Runtime/CoreUObject/Public/UObject/UObject.h
@@ -200,18 +200,18 @@ public:
 	virtual void PostLoad();
 
 	/**
-	* Instances components for objects being loaded from disk, if necessary.  Ensures that component references
-	* between nested components are fixed up correctly.
-	*
-	* @param	OuterInstanceGraph	when calling this method on subobjects, specifies the instancing graph which contains all instanced
-	*								subobjects and components for a subobject root.
-	*/
+	 * Instances components for objects being loaded from disk, if necessary.  Ensures that component references
+	 * between nested components are fixed up correctly.
+	 *
+	 * @param	OuterInstanceGraph	when calling this method on subobjects, specifies the instancing graph which contains all instanced
+	 *								subobjects and components for a subobject root.
+	 */
 	virtual void PostLoadSubobjects(FObjectInstancingGraph* OuterInstanceGraph);
-
+	
 	/**
-	* Called before destroying the object.  This is called immediately upon deciding to destroy the object, to allow the object to begin an
-	* asynchronous cleanup process.
-	*/
+	 * Called before destroying the object.  This is called immediately upon deciding to destroy the object, to allow the object to begin an
+	 * asynchronous cleanup process.
+	 */
 	virtual void BeginDestroy();
 
 	/**
@@ -419,18 +419,21 @@ public:
 		return false;
 	}
 
+	/** Special value meaning that the resource size is not defined */
+	static const SIZE_T RESOURCE_SIZE_NONE = static_cast<SIZE_T>(-1);
+
 	/**
 	 * Returns the size of the object/ resource for display to artists/ LDs in the Editor. The
-	 * default behavior is to return 0 which indicates that the resource shouldn't display its
-	 * size which is used to not confuse people by displaying small sizes for e.g. objects like
-	 * materials
+	 * default behavior is to return RESOURCE_SIZE_NONE which indicates that the resource shouldn't
+	 * display its size which is used to not confuse people by displaying small sizes
+	 * e.g. for objects like materials.
 	 *
 	 * @param	Type	Indicates which resource size should be returned
 	 * @return	Size of resource as to be displayed to artists/ LDs in the Editor.
 	 */
 	virtual SIZE_T GetResourceSize(EResourceSizeMode::Type Mode)
 	{
-		return 0;
+		return RESOURCE_SIZE_NONE;
 	}
 
 	/** 
