Support for loading objects with StaticLoadObject without specifying the object name (short package name is assumed to be the main asset name).

diff --git a/Engine/Source/Runtime/CoreUObject/Private/UObject/UObjectGlobals.cpp b/Engine/Source/Runtime/CoreUObject/Private/UObject/UObjectGlobals.cpp
--- a/Engine/Source/Runtime/CoreUObject/Private/UObject/UObjectGlobals.cpp
+++ b/Engine/Source/Runtime/CoreUObject/Private/UObject/UObjectGlobals.cpp
@@ -444,7 +444,7 @@ UPackage* FindPackage( UObject* InOuter, const TCHAR* PackageName )
 	{
 		InName = MakeUniqueObjectName( InOuter, UPackage::StaticClass() ).ToString();
 	}
-	ResolveName( InOuter, InName, 1, 0 );
+	ResolveName( InOuter, InName, true, false );
 
 	UPackage* Result = NULL;
 	if ( InName != TEXT("None") )
@@ -484,7 +484,7 @@ UPackage* CreatePackage( UObject* InOuter, const TCHAR* PackageName )
 		InName = MakeUniqueObjectName( InOuter, UPackage::StaticClass() ).ToString();
 	}
 
-	ResolveName( InOuter, InName, 1, 0 );
+	ResolveName( InOuter, InName, true, false );
 
 	UPackage* Result = NULL;
 	if ( InName.Len() == 0 )
@@ -570,9 +570,6 @@ bool ResolveName( UObject*& InPackage, FString& InOutName, bool Create, bool Thr
 	// Strip off the object class.
 	ConstructorHelpers::StripObjectClass( InOutName );
 
-	// Handle specified packages.
-	int32 i;
-
 	// if you're attempting to find an object in any package using a dotted name that isn't fully
 	// qualified (such as ObjectName.SubobjectName - notice no package name there), you normally call
 	// StaticFindObject and pass in ANY_PACKAGE as the value for InPackage.  When StaticFindObject calls ResolveName,
@@ -582,17 +579,20 @@ bool ResolveName( UObject*& InPackage, FString& InOutName, bool Create, bool Thr
 	// name couldn't be found, pass in ANY_PACKAGE as the value for InPackage to the call to FindObject<UObject>().
 	bool bSubobjectPath = false;
 
+	// Handle specified packages.
+	int32 DotIndex = INDEX_NONE;// InOutName.Find(TEXT("."), ESearchCase::CaseSensitive);
+
 	// to make parsing the name easier, replace the subobject delimiter with an extra dot
 	InOutName.ReplaceInline(SUBOBJECT_DELIMITER, TEXT(".."), ESearchCase::CaseSensitive);
-	while ((i = InOutName.Find(TEXT("."), ESearchCase::CaseSensitive)) != -1)
+	while ((DotIndex = InOutName.Find(TEXT("."), ESearchCase::CaseSensitive)) != INDEX_NONE)
 	{
-		FString PartialName = InOutName.Left(i);
+		FString PartialName = InOutName.Left(DotIndex);
 
 		// if the next part of InOutName ends in two dots, it indicates that the next object in the path name
 		// is not a top-level object (i.e. it's a subobject).  e.g. SomePackage.SomeGroup.SomeObject..Subobject
-		if ( InOutName.Mid(i+1,1) == TEXT(".") )
+		if (InOutName.Mid(DotIndex + 1, 1) == TEXT("."))
 		{
-			InOutName      = PartialName + InOutName.Mid(i+1);
+			InOutName = PartialName + InOutName.Mid(DotIndex + 1);
 			bSubobjectPath = true;
 			Create         = false;
 		}
@@ -632,7 +632,7 @@ bool ResolveName( UObject*& InPackage, FString& InOutName, bool Create, bool Thr
 
 			check(InPackage);
 		}
-		InOutName = InOutName.Mid(i+1);
+		InOutName = InOutName.Mid(DotIndex + 1);
 	}
 
 	return true;
@@ -698,71 +698,95 @@ bool ParseObject( const TCHAR* Stream, const TCHAR* Match, UClass* Class, UObjec
  *
  * @return The object that was loaded or found. NULL for a failure.
  */
-UObject* StaticLoadObject(UClass* ObjectClass, UObject* InOuter, const TCHAR* InName, const TCHAR* Filename, uint32 LoadFlags, UPackageMap* Sandbox, bool bAllowObjectReconciliation )
+UObject* StaticLoadObjectInternal(UClass* ObjectClass, UObject* InOuter, const TCHAR* InName, const TCHAR* Filename, uint32 LoadFlags, UPackageMap* Sandbox, bool bAllowObjectReconciliation)
 {
-	//UE_CLOG(GIsRoutingPostLoad, LogUObjectGlobals, Fatal, TEXT("Calling StaticLoadObject during PostLoad is forbidden."));
-
 	SCOPE_CYCLE_COUNTER(STAT_LoadObject);
 	check(ObjectClass);
 	check(InName);
 
 	FString StrName = InName;
+	UObject* Result = nullptr;
+	const bool bContainsObjectName = !!FCString::Strstr(InName, TEXT("."));
 
 	// break up the name into packages, returning the innermost name and its outer
 	ResolveName(InOuter, StrName, true, true);
 	if (InOuter)
 	{
-		if( bAllowObjectReconciliation && ((FApp::IsGame() && !GIsEditor && !IsRunningCommandlet()) 
+		if (bAllowObjectReconciliation && ((FApp::IsGame() && !GIsEditor && !IsRunningCommandlet())
 #if WITH_EDITOR
 			|| GIsImportingT3D
 #endif
 			))
 		{
-			if (UObject* Result = StaticFindObjectFast(ObjectClass, InOuter, *StrName))
-			{
-				return Result;
-			}
+			Result = StaticFindObjectFast(ObjectClass, InOuter, *StrName);
 		}
 
-		if( (FPlatformProperties::RequiresCookedData() || FPlatformProperties::IsServerOnly()) && GUseSeekFreeLoading )
+		if (!Result)
 		{
-			if ( (LoadFlags&LOAD_NoWarn) == 0 )
+			if ((FPlatformProperties::RequiresCookedData() || FPlatformProperties::IsServerOnly()) && GUseSeekFreeLoading)
 			{
-				UE_LOG(LogUObjectGlobals, Warning ,TEXT("StaticLoadObject for %s %s %s couldn't find object in memory!"),
-					*ObjectClass->GetName(),
-					*InOuter->GetName(),
-					*StrName);
+				// Warn only if we allow it and we're sure we're not going to re-try searching with the object name
+				if (bContainsObjectName && (LoadFlags&LOAD_NoWarn) == 0)
+				{
+					UE_LOG(LogUObjectGlobals, Warning, TEXT("StaticLoadObject for %s %s %s couldn't find object in memory!"),
+						*ObjectClass->GetName(),
+						*InOuter->GetName(),
+						*StrName);
+				}
 			}
-		}
-		else
-		{
-			// now that we have one asset per package, we load the entire package whenever a single object is requested
-			LoadPackage(NULL, *InOuter->GetOutermost()->GetName(), LoadFlags & ~LOAD_Verify);
-
-			// now, find the object in the package
-			if (UObject* Result = Result = StaticFindObjectFast(ObjectClass, InOuter, *StrName))
+			else
 			{
-				return Result;
-			}
+				// now that we have one asset per package, we load the entire package whenever a single object is requested
+				LoadPackage(NULL, *InOuter->GetOutermost()->GetName(), LoadFlags & ~LOAD_Verify);
 
-			// If the object was not found, check for a redirector and follow it if the class matches
-			UObjectRedirector* Redirector = FindObjectFast<UObjectRedirector>(InOuter, *StrName);
-			if ( Redirector && Redirector->DestinationObject && Redirector->DestinationObject->IsA(ObjectClass) )
-			{
-				return Redirector->DestinationObject;
+				// now, find the object in the package
+				Result = StaticFindObjectFast(ObjectClass, InOuter, *StrName);
+
+				// If the object was not found, check for a redirector and follow it if the class matches
+				if (!Result)
+				{
+					UObjectRedirector* Redirector = FindObjectFast<UObjectRedirector>(InOuter, *StrName);
+					if (Redirector && Redirector->DestinationObject && Redirector->DestinationObject->IsA(ObjectClass))
+					{
+						return Redirector->DestinationObject;
+					}
+				}
 			}
 		}
 	}
 
-	// we haven't created or found the object, error
-	FFormatNamedArguments Arguments;
-	Arguments.Add(TEXT("ClassName"), FText::FromString( ObjectClass->GetName() ));
-	Arguments.Add(TEXT("OuterName"), InOuter ? FText::FromString( InOuter->GetPathName() ) : NSLOCTEXT( "Core", "None", "None" ));
-	Arguments.Add(TEXT("ObjectName"), FText::FromString( StrName ));
-	const FString Error = FText::Format( NSLOCTEXT( "Core", "ObjectNotFound", "Failed to find object '{ClassName} {OuterName}.{ObjectName}'" ), Arguments ).ToString();
-	SafeLoadError(InOuter, LoadFlags, *Error, *Error);
+	if (!Result && !bContainsObjectName)
+	{
+		// Assume that the object we're trying to load is the main asset inside of the package 
+		// which usually has the same name as the short package name.
+		StrName = InName;
+		StrName += TEXT(".");
+		StrName += FPackageName::GetShortName(InName);
+		Result = StaticLoadObjectInternal(ObjectClass, InOuter, *StrName, Filename, LoadFlags, Sandbox, bAllowObjectReconciliation);
+	}
+
+	return Result;
+}
 
-	return nullptr;
+UObject* StaticLoadObject(UClass* ObjectClass, UObject* InOuter, const TCHAR* InName, const TCHAR* Filename, uint32 LoadFlags, UPackageMap* Sandbox, bool bAllowObjectReconciliation )
+{
+	//UE_CLOG(GIsRoutingPostLoad, LogUObjectGlobals, Fatal, TEXT("Calling StaticLoadObject during PostLoad is forbidden."));
+
+	UObject* Result = StaticLoadObjectInternal(ObjectClass, InOuter, InName, Filename, LoadFlags, Sandbox, bAllowObjectReconciliation);
+	if (!Result)
+	{
+		FString ObjectName = InName;
+		ResolveName(InOuter, ObjectName, true, true);
+
+		// we haven't created or found the object, error
+		FFormatNamedArguments Arguments;
+		Arguments.Add(TEXT("ClassName"), FText::FromString(ObjectClass->GetName()));
+		Arguments.Add(TEXT("OuterName"), InOuter ? FText::FromString(InOuter->GetPathName()) : NSLOCTEXT("Core", "None", "None"));
+		Arguments.Add(TEXT("ObjectName"), FText::FromString(ObjectName));
+		const FString Error = FText::Format(NSLOCTEXT("Core", "ObjectNotFound", "Failed to find object '{ClassName} {OuterName}.{ObjectName}'"), Arguments).ToString();
+		SafeLoadError(InOuter, LoadFlags, *Error, *Error);
+	}
+	return Result;
 }
 
 //
