More linker load stats, replaces expensive Sleep with ConditionalSleep (from Jaroslaw Surowiec)

diff --git a/Engine/Source/Runtime/CoreUObject/Private/UObject/LinkerLoad.cpp b/Engine/Source/Runtime/CoreUObject/Private/UObject/LinkerLoad.cpp
--- a/Engine/Source/Runtime/CoreUObject/Private/UObject/LinkerLoad.cpp
+++ b/Engine/Source/Runtime/CoreUObject/Private/UObject/LinkerLoad.cpp
@@ -721,6 +721,8 @@ bool ULinkerLoad::IsTimeLimitExceeded( const TCHAR* CurrentTask, int32 Granulari
  */
 ULinkerLoad::ELinkerStatus ULinkerLoad::CreateLoader()
 {
+	DECLARE_SCOPE_CYCLE_COUNTER( TEXT( "ULinkerLoad::CreateLoader" ), STAT_LinkerLoad_CreateLoader, STATGROUP_AsyncIOSystem );
+
 #if WITH_EDITOR
 
 	if (!LoadProgressScope)
@@ -753,11 +755,13 @@ ULinkerLoad::ELinkerStatus ULinkerLoad::CreateLoader()
 			if( PrecacheInfo->SynchronizationObject->GetValue() != 0 )
 			{
 				double StartTime = FPlatformTime::Seconds();
-				while (PrecacheInfo->SynchronizationObject->GetValue() != 0)
+			
+				FPlatformProcess::ConditionalSleep( [&]()
 				{
 					SHUTDOWN_IF_EXIT_REQUESTED;
-					FPlatformProcess::Sleep(0);
-				}
+					return PrecacheInfo->SynchronizationObject->GetValue() == 0;
+				} 
+				);
 				float WaitTime = FPlatformTime::Seconds() - StartTime;
 				UE_LOG(LogInit, Log, TEXT("Waited %.3f sec for async package '%s' to complete caching."), WaitTime, *Filename);
 			}
@@ -863,6 +867,8 @@ ULinkerLoad::ELinkerStatus ULinkerLoad::CreateLoader()
  */
 ULinkerLoad::ELinkerStatus ULinkerLoad::SerializePackageFileSummary()
 {
+	DECLARE_SCOPE_CYCLE_COUNTER( TEXT( "ULinkerLoad::SerializePackageFileSummary" ), STAT_LinkerLoad_SerializePackageFileSummary, STATGROUP_AsyncIOSystem );
+
 	if( bHasSerializedPackageFileSummary == false )
 	{
 #if WITH_EDITOR
@@ -1067,6 +1073,8 @@ ULinkerLoad::ELinkerStatus ULinkerLoad::SerializePackageFileSummary()
  */
 ULinkerLoad::ELinkerStatus ULinkerLoad::SerializeNameMap()
 {
+	DECLARE_SCOPE_CYCLE_COUNTER( TEXT( "ULinkerLoad::SerializeNameMap" ), STAT_LinkerLoad_SerializeNameMap, STATGROUP_AsyncIOSystem );
+
 	// The name map is the first item serialized. We wait till all the header information is read
 	// before any serialization. @todo async, @todo seamless: this could be spread out across name,
 	// import and export maps if the package file summary contained more detailed information on
@@ -1117,6 +1125,8 @@ ULinkerLoad::ELinkerStatus ULinkerLoad::SerializeNameMap()
  */
 ULinkerLoad::ELinkerStatus ULinkerLoad::SerializeImportMap()
 {
+	DECLARE_SCOPE_CYCLE_COUNTER( TEXT( "ULinkerLoad::SerializeImportMap" ), STAT_LinkerLoad_SerializeImportMap, STATGROUP_AsyncIOSystem );
+
 	if( ImportMapIndex == 0 && Summary.ImportCount > 0 )
 	{
 		Seek( Summary.ImportOffset );
@@ -1138,6 +1148,8 @@ ULinkerLoad::ELinkerStatus ULinkerLoad::SerializeImportMap()
  */
 ULinkerLoad::ELinkerStatus ULinkerLoad::FixupImportMap()
 {
+	DECLARE_SCOPE_CYCLE_COUNTER( TEXT( "ULinkerLoad::FixupImportMap" ), STAT_LinkerLoad_FixupImportMap, STATGROUP_AsyncIOSystem );
+
 	if( bHasFixedUpImportMap == false )
 	{
 #if WITH_EDITOR
@@ -1335,6 +1347,8 @@ ULinkerLoad::ELinkerStatus ULinkerLoad::FixupImportMap()
  */
 ULinkerLoad::ELinkerStatus ULinkerLoad::SerializeExportMap()
 {
+	DECLARE_SCOPE_CYCLE_COUNTER( TEXT( "ULinkerLoad::SerializeExportMap" ), STAT_LinkerLoad_SerializeExportMap, STATGROUP_AsyncIOSystem );
+
 	if( ExportMapIndex == 0 && Summary.ExportCount > 0 )
 	{
 		Seek( Summary.ExportOffset );
@@ -1353,6 +1367,8 @@ ULinkerLoad::ELinkerStatus ULinkerLoad::SerializeExportMap()
 
 ULinkerLoad::ELinkerStatus ULinkerLoad::RemapImports()
 {
+	DECLARE_SCOPE_CYCLE_COUNTER( TEXT( "ULinkerLoad::RemapImports" ), STAT_LinkerLoad_RemapImports, STATGROUP_AsyncIOSystem );
+
 #if !(UE_BUILD_SHIPPING || UE_BUILD_TEST)
 
 	for (int32 ImportIndex = 0; ImportIndex < ImportMap.Num(); ImportIndex++)
@@ -1394,6 +1410,8 @@ ULinkerLoad::ELinkerStatus ULinkerLoad::RemapImports()
  */
 ULinkerLoad::ELinkerStatus ULinkerLoad::StartTextureAllocation()
 {
+	DECLARE_SCOPE_CYCLE_COUNTER( TEXT( "ULinkerLoad::StartTextureAllocation" ), STAT_LinkerLoad_StartTextureAllocation, STATGROUP_AsyncIOSystem );
+
 	double StartTime = FPlatformTime::Seconds();
 	int32 NumAllocationsStarted = 0;
 	int32 NumAllocationsConsidered = 0;
@@ -1462,6 +1480,8 @@ ULinkerLoad::ELinkerStatus ULinkerLoad::StartTextureAllocation()
  */
 ULinkerLoad::ELinkerStatus ULinkerLoad::SerializeDependsMap()
 {
+	DECLARE_SCOPE_CYCLE_COUNTER( TEXT( "ULinkerLoad::SerializeDependsMap" ), STAT_LinkerLoad_SerializeDependsMap, STATGROUP_AsyncIOSystem );
+
 	// Skip serializing depends map if we are using seekfree loading
 	if( FPlatformProperties::RequiresCookedData() 
 	// or we are neither Editor nor commandlet
@@ -1585,6 +1605,8 @@ ULinkerLoad::ELinkerStatus ULinkerLoad::SerializeThumbnails( bool bForceEnableIn
  */
 ULinkerLoad::ELinkerStatus ULinkerLoad::CreateExportHash()
 {
+	DECLARE_SCOPE_CYCLE_COUNTER( TEXT( "ULinkerLoad::CreateExportHash" ), STAT_LinkerLoad_CreateExportHash, STATGROUP_AsyncIOSystem );
+
 	// Zero initialize hash on first iteration.
 	if( ExportHashIndex == 0 )
 	{
@@ -1617,6 +1639,8 @@ ULinkerLoad::ELinkerStatus ULinkerLoad::CreateExportHash()
  */
 ULinkerLoad::ELinkerStatus ULinkerLoad::FindExistingExports()
 {
+	DECLARE_SCOPE_CYCLE_COUNTER( TEXT( "ULinkerLoad::FindExistingExports" ), STAT_LinkerLoad_FindExistingExports, STATGROUP_AsyncIOSystem );
+
 	if( bHasFoundExistingExports == false )
 	{
 		// only look for existing exports in the editor after it has started up
@@ -1651,6 +1675,8 @@ ULinkerLoad::ELinkerStatus ULinkerLoad::FindExistingExports()
  */
 ULinkerLoad::ELinkerStatus ULinkerLoad::FinalizeCreation()
 {
+	DECLARE_SCOPE_CYCLE_COUNTER( TEXT( "ULinkerLoad::FinalizeCreation" ), STAT_LinkerLoad_FinalizeCreation, STATGROUP_AsyncIOSystem );
+
 	if( bHasFinishedInitialization == false )
 	{
 #if WITH_EDITOR
@@ -4251,6 +4277,8 @@ FName ULinkerLoad::FindNewNameForClass(FName OldClassName, bool bIsInstance)
 */
 ULinkerLoad::ELinkerStatus ULinkerLoad::FixupExportMap()
 {
+	DECLARE_SCOPE_CYCLE_COUNTER( TEXT( "ULinkerLoad::FixupExportMap" ), STAT_LinkerLoad_FixupExportMap, STATGROUP_AsyncIOSystem );
+
 	// No need to fixup exports if everything is cooked.
 	if (!FPlatformProperties::RequiresCookedData())
 	{
