Merging CIS fix

diff --git a/Engine/Source/Runtime/CoreUObject/Private/UObject/LinkerLoad.cpp b/Engine/Source/Runtime/CoreUObject/Private/UObject/LinkerLoad.cpp
--- a/Engine/Source/Runtime/CoreUObject/Private/UObject/LinkerLoad.cpp
+++ b/Engine/Source/Runtime/CoreUObject/Private/UObject/LinkerLoad.cpp
@@ -721,7 +721,7 @@ bool ULinkerLoad::IsTimeLimitExceeded( const TCHAR* CurrentTask, int32 Granulari
  */
 ULinkerLoad::ELinkerStatus ULinkerLoad::CreateLoader()
 {
-	DECLARE_SCOPE_CYCLE_COUNTER( TEXT( "ULinkerLoad::CreateLoader" ), STAT_LinkerLoad_CreateLoader, STATGROUP_AsyncIOSystem );
+	DECLARE_SCOPE_CYCLE_COUNTER( TEXT( "ULinkerLoad::CreateLoader" ), STAT_LinkerLoad_CreateLoader, STATGROUP_LinkerLoad );
 
 #if WITH_EDITOR
 
@@ -867,7 +867,7 @@ ULinkerLoad::ELinkerStatus ULinkerLoad::CreateLoader()
  */
 ULinkerLoad::ELinkerStatus ULinkerLoad::SerializePackageFileSummary()
 {
-	DECLARE_SCOPE_CYCLE_COUNTER( TEXT( "ULinkerLoad::SerializePackageFileSummary" ), STAT_LinkerLoad_SerializePackageFileSummary, STATGROUP_AsyncIOSystem );
+	DECLARE_SCOPE_CYCLE_COUNTER( TEXT( "ULinkerLoad::SerializePackageFileSummary" ), STAT_LinkerLoad_SerializePackageFileSummary, STATGROUP_LinkerLoad );
 
 	if( bHasSerializedPackageFileSummary == false )
 	{
@@ -1073,7 +1073,7 @@ ULinkerLoad::ELinkerStatus ULinkerLoad::SerializePackageFileSummary()
  */
 ULinkerLoad::ELinkerStatus ULinkerLoad::SerializeNameMap()
 {
-	DECLARE_SCOPE_CYCLE_COUNTER( TEXT( "ULinkerLoad::SerializeNameMap" ), STAT_LinkerLoad_SerializeNameMap, STATGROUP_AsyncIOSystem );
+	DECLARE_SCOPE_CYCLE_COUNTER( TEXT( "ULinkerLoad::SerializeNameMap" ), STAT_LinkerLoad_SerializeNameMap, STATGROUP_LinkerLoad );
 
 	// The name map is the first item serialized. We wait till all the header information is read
 	// before any serialization. @todo async, @todo seamless: this could be spread out across name,
@@ -1125,7 +1125,7 @@ ULinkerLoad::ELinkerStatus ULinkerLoad::SerializeNameMap()
  */
 ULinkerLoad::ELinkerStatus ULinkerLoad::SerializeImportMap()
 {
-	DECLARE_SCOPE_CYCLE_COUNTER( TEXT( "ULinkerLoad::SerializeImportMap" ), STAT_LinkerLoad_SerializeImportMap, STATGROUP_AsyncIOSystem );
+	DECLARE_SCOPE_CYCLE_COUNTER( TEXT( "ULinkerLoad::SerializeImportMap" ), STAT_LinkerLoad_SerializeImportMap, STATGROUP_LinkerLoad );
 
 	if( ImportMapIndex == 0 && Summary.ImportCount > 0 )
 	{
@@ -1148,7 +1148,7 @@ ULinkerLoad::ELinkerStatus ULinkerLoad::SerializeImportMap()
  */
 ULinkerLoad::ELinkerStatus ULinkerLoad::FixupImportMap()
 {
-	DECLARE_SCOPE_CYCLE_COUNTER( TEXT( "ULinkerLoad::FixupImportMap" ), STAT_LinkerLoad_FixupImportMap, STATGROUP_AsyncIOSystem );
+	DECLARE_SCOPE_CYCLE_COUNTER( TEXT( "ULinkerLoad::FixupImportMap" ), STAT_LinkerLoad_FixupImportMap, STATGROUP_LinkerLoad );
 
 	if( bHasFixedUpImportMap == false )
 	{
@@ -1347,7 +1347,7 @@ ULinkerLoad::ELinkerStatus ULinkerLoad::FixupImportMap()
  */
 ULinkerLoad::ELinkerStatus ULinkerLoad::SerializeExportMap()
 {
-	DECLARE_SCOPE_CYCLE_COUNTER( TEXT( "ULinkerLoad::SerializeExportMap" ), STAT_LinkerLoad_SerializeExportMap, STATGROUP_AsyncIOSystem );
+	DECLARE_SCOPE_CYCLE_COUNTER( TEXT( "ULinkerLoad::SerializeExportMap" ), STAT_LinkerLoad_SerializeExportMap, STATGROUP_LinkerLoad );
 
 	if( ExportMapIndex == 0 && Summary.ExportCount > 0 )
 	{
@@ -1367,7 +1367,7 @@ ULinkerLoad::ELinkerStatus ULinkerLoad::SerializeExportMap()
 
 ULinkerLoad::ELinkerStatus ULinkerLoad::RemapImports()
 {
-	DECLARE_SCOPE_CYCLE_COUNTER( TEXT( "ULinkerLoad::RemapImports" ), STAT_LinkerLoad_RemapImports, STATGROUP_AsyncIOSystem );
+	DECLARE_SCOPE_CYCLE_COUNTER( TEXT( "ULinkerLoad::RemapImports" ), STAT_LinkerLoad_RemapImports, STATGROUP_LinkerLoad );
 
 #if !(UE_BUILD_SHIPPING || UE_BUILD_TEST)
 
@@ -1410,7 +1410,7 @@ ULinkerLoad::ELinkerStatus ULinkerLoad::RemapImports()
  */
 ULinkerLoad::ELinkerStatus ULinkerLoad::StartTextureAllocation()
 {
-	DECLARE_SCOPE_CYCLE_COUNTER( TEXT( "ULinkerLoad::StartTextureAllocation" ), STAT_LinkerLoad_StartTextureAllocation, STATGROUP_AsyncIOSystem );
+	DECLARE_SCOPE_CYCLE_COUNTER( TEXT( "ULinkerLoad::StartTextureAllocation" ), STAT_LinkerLoad_StartTextureAllocation, STATGROUP_LinkerLoad );
 
 	double StartTime = FPlatformTime::Seconds();
 	int32 NumAllocationsStarted = 0;
@@ -1480,7 +1480,7 @@ ULinkerLoad::ELinkerStatus ULinkerLoad::StartTextureAllocation()
  */
 ULinkerLoad::ELinkerStatus ULinkerLoad::SerializeDependsMap()
 {
-	DECLARE_SCOPE_CYCLE_COUNTER( TEXT( "ULinkerLoad::SerializeDependsMap" ), STAT_LinkerLoad_SerializeDependsMap, STATGROUP_AsyncIOSystem );
+	DECLARE_SCOPE_CYCLE_COUNTER( TEXT( "ULinkerLoad::SerializeDependsMap" ), STAT_LinkerLoad_SerializeDependsMap, STATGROUP_LinkerLoad );
 
 	// Skip serializing depends map if we are using seekfree loading
 	if( FPlatformProperties::RequiresCookedData() 
@@ -1605,7 +1605,7 @@ ULinkerLoad::ELinkerStatus ULinkerLoad::SerializeThumbnails( bool bForceEnableIn
  */
 ULinkerLoad::ELinkerStatus ULinkerLoad::CreateExportHash()
 {
-	DECLARE_SCOPE_CYCLE_COUNTER( TEXT( "ULinkerLoad::CreateExportHash" ), STAT_LinkerLoad_CreateExportHash, STATGROUP_AsyncIOSystem );
+	DECLARE_SCOPE_CYCLE_COUNTER( TEXT( "ULinkerLoad::CreateExportHash" ), STAT_LinkerLoad_CreateExportHash, STATGROUP_LinkerLoad );
 
 	// Zero initialize hash on first iteration.
 	if( ExportHashIndex == 0 )
@@ -1639,7 +1639,7 @@ ULinkerLoad::ELinkerStatus ULinkerLoad::CreateExportHash()
  */
 ULinkerLoad::ELinkerStatus ULinkerLoad::FindExistingExports()
 {
-	DECLARE_SCOPE_CYCLE_COUNTER( TEXT( "ULinkerLoad::FindExistingExports" ), STAT_LinkerLoad_FindExistingExports, STATGROUP_AsyncIOSystem );
+	DECLARE_SCOPE_CYCLE_COUNTER( TEXT( "ULinkerLoad::FindExistingExports" ), STAT_LinkerLoad_FindExistingExports, STATGROUP_LinkerLoad );
 
 	if( bHasFoundExistingExports == false )
 	{
@@ -1675,7 +1675,7 @@ ULinkerLoad::ELinkerStatus ULinkerLoad::FindExistingExports()
  */
 ULinkerLoad::ELinkerStatus ULinkerLoad::FinalizeCreation()
 {
-	DECLARE_SCOPE_CYCLE_COUNTER( TEXT( "ULinkerLoad::FinalizeCreation" ), STAT_LinkerLoad_FinalizeCreation, STATGROUP_AsyncIOSystem );
+	DECLARE_SCOPE_CYCLE_COUNTER( TEXT( "ULinkerLoad::FinalizeCreation" ), STAT_LinkerLoad_FinalizeCreation, STATGROUP_LinkerLoad );
 
 	if( bHasFinishedInitialization == false )
 	{
@@ -4277,7 +4277,7 @@ FName ULinkerLoad::FindNewNameForClass(FName OldClassName, bool bIsInstance)
 */
 ULinkerLoad::ELinkerStatus ULinkerLoad::FixupExportMap()
 {
-	DECLARE_SCOPE_CYCLE_COUNTER( TEXT( "ULinkerLoad::FixupExportMap" ), STAT_LinkerLoad_FixupExportMap, STATGROUP_AsyncIOSystem );
+	DECLARE_SCOPE_CYCLE_COUNTER( TEXT( "ULinkerLoad::FixupExportMap" ), STAT_LinkerLoad_FixupExportMap, STATGROUP_LinkerLoad );
 
 	// No need to fixup exports if everything is cooked.
 	if (!FPlatformProperties::RequiresCookedData())
