Core - Replaces expensive Sleep with ConditionalSleep and added a few more stats

diff --git a/Engine/Source/Runtime/Core/Private/Serialization/AsyncIOSystemBase.cpp b/Engine/Source/Runtime/Core/Private/Serialization/AsyncIOSystemBase.cpp
--- a/Engine/Source/Runtime/Core/Private/Serialization/AsyncIOSystemBase.cpp
+++ b/Engine/Source/Runtime/Core/Private/Serialization/AsyncIOSystemBase.cpp
@@ -720,6 +720,8 @@ void FAsyncIOSystemBase::TickSingleThreaded()
 
 void FAsyncIOSystemBase::BlockTillAllRequestsFinished()
 {
+	DECLARE_SCOPE_CYCLE_COUNTER( TEXT( "FAsyncIOSystemBase::BlockTillAllRequestsFinished" ), STAT_FAsyncIOSystemBase_BlockTillAllRequestsFinished, STATGROUP_AsyncIOSystem );
+
 	// Block till all requests are fulfilled.
 	while( true ) 
 	{
@@ -737,7 +739,7 @@ void FAsyncIOSystemBase::BlockTillAllRequestsFinished()
 			SHUTDOWN_IF_EXIT_REQUESTED;
 
 			//@todo streaming: this should be replaced by waiting for an event.
-			FPlatformProcess::Sleep( 0.001f );
+			FPlatformProcess::SleepNoStats( 0.001f );
 		}
 	}
 }
diff --git a/Engine/Source/Runtime/CoreUObject/Private/Serialization/AsyncLoading.cpp b/Engine/Source/Runtime/CoreUObject/Private/Serialization/AsyncLoading.cpp
--- a/Engine/Source/Runtime/CoreUObject/Private/Serialization/AsyncLoading.cpp
+++ b/Engine/Source/Runtime/CoreUObject/Private/Serialization/AsyncLoading.cpp
@@ -1313,11 +1313,12 @@ FArchiveAsync::FArchiveAsync( const TCHAR* InFileName )
 void FArchiveAsync::FlushCache()
 {
 	// Wait on all outstanding requests.
-	while( PrecacheReadStatus[CURRENT].GetValue() || PrecacheReadStatus[NEXT].GetValue() )
+	FPlatformProcess::ConditionalSleep( [&]()
 	{
 		SHUTDOWN_IF_EXIT_REQUESTED;
-		FPlatformProcess::Sleep(0.0001);
-	}
+		return PrecacheReadStatus[CURRENT].GetValue()==0 && PrecacheReadStatus[NEXT].GetValue()==0;
+	} );
+
 	uint32 Delta = 0;
 
 	// Invalidate any precached data and free memory for current buffer.
@@ -1623,6 +1624,8 @@ bool FArchiveAsync::Precache( int64 RequestOffset, int64 RequestSize )
  */
 void FArchiveAsync::Serialize( void* Data, int64 Count )
 {
+	DECLARE_SCOPE_CYCLE_COUNTER( TEXT( "FArchiveAsync::Serialize" ), STAT_ArchiveAsync_Serialize, STATGROUP_AsyncLoad );
+
 	// Ensure we aren't reading beyond the end of the file
 	checkf( CurrentPos + Count <= TotalSize(), TEXT("Seeked past end of file %s (%lld / %lld)"), *FileName, CurrentPos + Count, TotalSize() );
 
@@ -1642,7 +1645,7 @@ void FArchiveAsync::Serialize( void* Data, int64 Count )
 			SHUTDOWN_IF_EXIT_REQUESTED;
 			if (FPlatformProcess::SupportsMultithreading())
 			{
-				FPlatformProcess::Sleep(0);
+				FPlatformProcess::SleepNoStats(0);
 			}
 			else
 			{
@@ -1668,7 +1671,7 @@ void FArchiveAsync::Serialize( void* Data, int64 Count )
 		}
 		if (FPlatformProcess::SupportsMultithreading())
 		{
-			FPlatformProcess::Sleep(0);
+			FPlatformProcess::SleepNoStats( 0 );
 		}
 		else
 		{
