Core - Replaces expensive Sleep with ConditionalSleep

diff --git a/Engine/Source/Runtime/CoreUObject/Private/Serialization/AsyncLoading.cpp b/Engine/Source/Runtime/CoreUObject/Private/Serialization/AsyncLoading.cpp
--- a/Engine/Source/Runtime/CoreUObject/Private/Serialization/AsyncLoading.cpp
+++ b/Engine/Source/Runtime/CoreUObject/Private/Serialization/AsyncLoading.cpp
@@ -1624,8 +1624,6 @@ bool FArchiveAsync::Precache( int64 RequestOffset, int64 RequestSize )
  */
 void FArchiveAsync::Serialize( void* Data, int64 Count )
 {
-	DECLARE_SCOPE_CYCLE_COUNTER( TEXT( "FArchiveAsync::Serialize" ), STAT_ArchiveAsync_Serialize, STATGROUP_AsyncLoad );
-
 	// Ensure we aren't reading beyond the end of the file
 	checkf( CurrentPos + Count <= TotalSize(), TEXT("Seeked past end of file %s (%lld / %lld)"), *FileName, CurrentPos + Count, TotalSize() );
 
@@ -1635,23 +1633,22 @@ void FArchiveAsync::Serialize( void* Data, int64 Count )
 	// Make sure serialization request fits entirely in already precached region.
 	if( !PrecacheBufferContainsRequest( CurrentPos, Count ) )
 	{
+		DECLARE_SCOPE_CYCLE_COUNTER( TEXT( "FArchiveAsync::Serialize.PrecacheBufferContainsRequest" ), STAT_ArchiveAsync_Serialize_PrecacheBufferContainsRequest, STATGROUP_AsyncLoad );
+
 		// Keep track of time we started to block.
 		StartTime	= FPlatformTime::Seconds();
 		bIOBlocked	= true;
 
 		// Busy wait for region to be precached.
-		while( !Precache( CurrentPos, Count ) )
+		FPlatformProcess::ConditionalSleep( [&]()
 		{
 			SHUTDOWN_IF_EXIT_REQUESTED;
-			if (FPlatformProcess::SupportsMultithreading())
-			{
-				FPlatformProcess::SleepNoStats(0);
-			}
-			else
+			if( !FPlatformProcess::SupportsMultithreading() )
 			{
 				FIOSystem::Get().TickSingleThreaded();
-			}	
-		}
+			}
+			return Precache( CurrentPos, Count );
+		} );
 
 		// There shouldn't be any outstanding read requests for the main buffer at this point.
 		check( PrecacheReadStatus[CURRENT].GetValue() == 0 );
@@ -1659,25 +1656,22 @@ void FArchiveAsync::Serialize( void* Data, int64 Count )
 	
 	// Make sure to wait till read request has finished before progressing. This can happen if PreCache interface
 	// is not being used for serialization.
-	while( PrecacheReadStatus[CURRENT].GetValue() != 0 )
+	FPlatformProcess::ConditionalSleep( [&]()
 	{
 		SHUTDOWN_IF_EXIT_REQUESTED;
 		// Only update StartTime if we haven't already started blocking I/O above.
 		if( !bIOBlocked )
 		{
 			// Keep track of time we started to block.
-			StartTime	= FPlatformTime::Seconds();
-			bIOBlocked	= true;
-		}
-		if (FPlatformProcess::SupportsMultithreading())
-		{
-			FPlatformProcess::SleepNoStats( 0 );
+			StartTime = FPlatformTime::Seconds();
+			bIOBlocked = true;
 		}
-		else
+		if( !FPlatformProcess::SupportsMultithreading() )
 		{
 			FIOSystem::Get().TickSingleThreaded();
-		}		
-	}
+		}
+		return PrecacheReadStatus[CURRENT].GetValue() == 0;
+	} );
 
 	// Update stats if we were blocked.
 #if STATS
