EXTERNAL_OBJECT_NAMES removed.

diff --git a/Engine/Source/Runtime/CoreUObject/Private/UObject/UObjectBase.cpp b/Engine/Source/Runtime/CoreUObject/Private/UObject/UObjectBase.cpp
--- a/Engine/Source/Runtime/CoreUObject/Private/UObject/UObjectBase.cpp
+++ b/Engine/Source/Runtime/CoreUObject/Private/UObject/UObjectBase.cpp
@@ -47,57 +47,6 @@ struct FPendingRegistrant
 static FPendingRegistrant* GFirstPendingRegistrant = NULL;
 static FPendingRegistrant* GLastPendingRegistrant = NULL;
 
-#if EXTERNAL_OBJECT_NAMES
-
-/**
- * Annotation for FNames
- */
-struct FNameAnnotation
-{
-	/**
-	 * default constructor
-	 * Default constructor must be the default item
-	 */
-	FNameAnnotation() :
-		Name(NAME_None)
-	{
-	}
-	/**
-	 * Determine if this Name is the default...which is NAME_None
-	 * @return true is this is NAME_None
-	 */
-	FORCEINLINE bool IsDefault()
-	{
-		return Name == NAME_None;
-	}
-
-	/**
-	 * Constructor 
-	 * @param InName name to assign
-	 */
-	FNameAnnotation(FName InName) :
-		Name(InName)
-	{
-	}
-
-	/**
-	 * Name for this object
-	 */
-	FName				Name; 
-
-};
-
-template <> struct TIsPODType<FNameAnnotation> { enum { Value = true }; };
-
-
-/**
- * Annotation to relate names to uobjects
- *
- */
-static FUObjectAnnotationDense<FNameAnnotation,false> NameAnnotation;
-
-#endif
-
 /**
  * Constructor used for bootstrapping
  * @param	InClass			possibly NULL, this gives the class of the new object, if known at this time
@@ -148,11 +97,7 @@ UObjectBase::~UObjectBase()
 
 const FName UObjectBase::GetFName() const
 {
-#if EXTERNAL_OBJECT_NAMES
-	return NameAnnotation.GetAnnotation(InternalIndex).Name;
-#else
 	return Name;
-#endif
 }
 
 
@@ -218,11 +163,7 @@ void UObjectBase::DeferredRegister(UClass *UClassStaticClass,const TCHAR* Packag
  */
 void UObjectBase::AddObject(FName InName)
 {
-#if EXTERNAL_OBJECT_NAMES
-	NameAnnotation.AddAnnotation(InternalIndex,InName);
-#else
 	Name = InName;
-#endif
 	GUObjectArray.AllocateUObjectIndex(this);
 	check(InName != NAME_None && InternalIndex >= 0);
 	HashObject(this);
@@ -240,11 +181,7 @@ void UObjectBase::LowLevelRename(FName NewName,UObject *NewOuter)
 	STAT(StatID = TStatId();) // reset the stat id since this thing now has a different name
 	UnhashObject(this);
 	check(InternalIndex >= 0);
-#if EXTERNAL_OBJECT_NAMES
-	NameAnnotation.AddAnnotation(InternalIndex,NewName);
-#else
 	Name = NewName;
-#endif
 	if (NewOuter)
 	{
 		Outer = NewOuter;
diff --git a/Engine/Source/Runtime/CoreUObject/Public/UObject/UObjectBase.h b/Engine/Source/Runtime/CoreUObject/Public/UObject/UObjectBase.h
--- a/Engine/Source/Runtime/CoreUObject/Public/UObject/UObjectBase.h
+++ b/Engine/Source/Runtime/CoreUObject/Public/UObject/UObjectBase.h
@@ -7,11 +7,6 @@
 #ifndef __UOBJECTBASE_H__
 #define __UOBJECTBASE_H__
 
-/**
- * The UObject FNames can be stored outside of the UObject proper. @todo this needs to be evaluated later.
- */
-#define EXTERNAL_OBJECT_NAMES (0)   
-
 DECLARE_DWORD_COUNTER_STAT_EXTERN(TEXT("STAT_UObjectsStatGroupTester"), STAT_UObjectsStatGroupTester, STATGROUP_UObjects, COREUOBJECT_API);
 
 class COREUOBJECT_API UObjectBase
@@ -211,10 +206,8 @@ private:
 	/** Class the object belongs to. */
 	UClass*							Class;
 
-#if !EXTERNAL_OBJECT_NAMES
 	/** Name of this object */
 	FName							Name;
-#endif
 
 	/** Object this object resides in. */
 	UObject*						Outer;
