Wrapping up 4.7 backwards compat code in versino checks
* Object version for marking BP generated class component's public
* Object version for updating creation method of ActorComponents

diff --git a/Engine/Source/Runtime/Core/Public/UObject/ObjectVersion.h b/Engine/Source/Runtime/Core/Public/UObject/ObjectVersion.h
--- a/Engine/Source/Runtime/Core/Public/UObject/ObjectVersion.h
+++ b/Engine/Source/Runtime/Core/Public/UObject/ObjectVersion.h
@@ -470,6 +470,10 @@ enum EUnrealEngineObjectUE4Version
 	VER_UE4_OPTIONALLY_CLEAR_GPU_EMITTERS_ON_INIT,
 	// Also store the Material guid with the landscape grass data
 	VER_UE4_SERIALIZE_LANDSCAPE_GRASS_DATA_MATERIAL_GUID,
+	// Make sure that all template components from blueprint generated classes are flagged as public
+	VER_UE4_BLUEPRINT_GENERATED_CLASS_COMPONENT_TEMPLATES_PUBLIC,
+	// Split out creation method on ActorComponents to distinguish between native, instance, and simple or user construction script
+	VER_UE4_ACTOR_COMPONENT_CREATION_METHOD,
 
 	// -----<new versions can be added before this line>-------------------------------------------------
 	// - this needs to be the last line (see note below)
diff --git a/Engine/Source/Runtime/CoreUObject/Private/UObject/LinkerLoad.cpp b/Engine/Source/Runtime/CoreUObject/Private/UObject/LinkerLoad.cpp
--- a/Engine/Source/Runtime/CoreUObject/Private/UObject/LinkerLoad.cpp
+++ b/Engine/Source/Runtime/CoreUObject/Private/UObject/LinkerLoad.cpp
@@ -4295,26 +4295,28 @@ ULinkerLoad::ELinkerStatus ULinkerLoad::FixupExportMap()
 
 			// ActorComponents outered to a BlueprintGeneratedClass (or even older ones that are outered to Blueprint) need to be marked RF_Public, but older content was 
 			// not created as such.  This updates the ExportTable such that they are correctly flagged when created and when other packages validate their imports.
-			// TODO: Wrap this in a version check if possible. Mark up for future removal when minimum version is incremented
-			if ((Export.ObjectFlags & RF_Public) == 0)
+			if (UE4Ver() < VER_UE4_BLUEPRINT_GENERATED_CLASS_COMPONENT_TEMPLATES_PUBLIC)
 			{
-				static const FName NAME_BlueprintGeneratedClass("BlueprintGeneratedClass");
-				static const FName NAME_Blueprint("Blueprint");
-				const FName OuterClassName = GetExportClassName(Export.OuterIndex);
-				if (OuterClassName == NAME_BlueprintGeneratedClass || OuterClassName == NAME_Blueprint)
+				if ((Export.ObjectFlags & RF_Public) == 0)
 				{
-					static const UClass* ActorComponentClass = FindObjectChecked<UClass>(ANY_PACKAGE, TEXT("ActorComponent"), true);
-					static const FString BPGeneratedClassPostfix(TEXT("_C"));
-					const FString NameClassString = NameClass.ToString();
-					UClass* Class = FindObject<UClass>(ANY_PACKAGE, *NameClassString);
-
-					// It is (obviously) a component if the class is a child of actor component
-					// and (almost certainly) a component if the class cannot be loaded but it ends in _C meaning it was generated from a blueprint
-					// However, it (probably) isn't safe to load the blueprint class, so we just check the _C and it is (probably) good enough
-					if (    ((Class != nullptr) && Class->IsChildOf(ActorComponentClass))
-						 || ((Class == nullptr) && NameClassString.EndsWith(BPGeneratedClassPostfix)))
+					static const FName NAME_BlueprintGeneratedClass("BlueprintGeneratedClass");
+					static const FName NAME_Blueprint("Blueprint");
+					const FName OuterClassName = GetExportClassName(Export.OuterIndex);
+					if (OuterClassName == NAME_BlueprintGeneratedClass || OuterClassName == NAME_Blueprint)
 					{
-						Export.ObjectFlags |= RF_Public;
+						static const UClass* ActorComponentClass = FindObjectChecked<UClass>(ANY_PACKAGE, TEXT("ActorComponent"), true);
+						static const FString BPGeneratedClassPostfix(TEXT("_C"));
+						const FString NameClassString = NameClass.ToString();
+						UClass* Class = FindObject<UClass>(ANY_PACKAGE, *NameClassString);
+
+						// It is (obviously) a component if the class is a child of actor component
+						// and (almost certainly) a component if the class cannot be loaded but it ends in _C meaning it was generated from a blueprint
+						// However, it (probably) isn't safe to load the blueprint class, so we just check the _C and it is (probably) good enough
+						if (    ((Class != nullptr) && Class->IsChildOf(ActorComponentClass))
+							 || ((Class == nullptr) && NameClassString.EndsWith(BPGeneratedClassPostfix)))
+						{
+							Export.ObjectFlags |= RF_Public;
+						}
 					}
 				}
 			}
diff --git a/Engine/Source/Runtime/Engine/Private/ActorComponent.cpp b/Engine/Source/Runtime/Engine/Private/ActorComponent.cpp
--- a/Engine/Source/Runtime/Engine/Private/ActorComponent.cpp
+++ b/Engine/Source/Runtime/Engine/Private/ActorComponent.cpp
@@ -127,8 +127,7 @@ void UActorComponent::PostLoad()
 {
 	Super::PostLoad();
 
-	// TODO: Wrap all this up with an engine version, for now just don't do it in cooked builds since we know they are good
-	if (!FPlatformProperties::RequiresCookedData())
+	if (GetLinkerUE4Version() < VER_UE4_ACTOR_COMPONENT_CREATION_METHOD)
 	{
 		if (bCreatedByConstructionScript_DEPRECATED)
 		{
@@ -141,7 +140,7 @@ void UActorComponent::PostLoad()
 
 		if (CreationMethod == EComponentCreationMethod::SimpleConstructionScript)
 		{
-			UBlueprintGeneratedClass* Class = Cast/*Checked*/<UBlueprintGeneratedClass>(GetOuter()->GetClass());
+			UBlueprintGeneratedClass* Class = CastChecked<UBlueprintGeneratedClass>(GetOuter()->GetClass());
 			while (Class)
 			{
 				USimpleConstructionScript* SCS = Class->SimpleConstructionScript;
