[INTEGRATE] Change 2451443 by [REDACTED_PERSON][REDACTED_HANDLE] on 2015/02/19 05:46:10

	Fix for RC blocker UE-10050 - crash loading a BulkData inside a transaction, seen when pasting an actor that references an object in a map with an AtmosphericFogComponent

diff --git a/Engine/Source/Runtime/CoreUObject/Private/Serialization/BulkData.cpp b/Engine/Source/Runtime/CoreUObject/Private/Serialization/BulkData.cpp
--- a/Engine/Source/Runtime/CoreUObject/Private/Serialization/BulkData.cpp
+++ b/Engine/Source/Runtime/CoreUObject/Private/Serialization/BulkData.cpp
@@ -652,11 +652,17 @@ void FUntypedBulkData::Serialize( FArchive& Ar, UObject* Owner, int32 Idx )
 				// Number of elements in array.
 				Ar << ElementCount;
 
-				// Make sure bulk data is loaded.
-				MakeSureBulkDataIsLoaded();
+				// Don't attempt to load or serialize BulkData if the current size is 0.
+				// This could be a newly constructed BulkData that has not yet been loaded, 
+				// and allocating 0 bytes now will cause a crash when we load.
+				if (GetBulkDataSize() > 0)
+				{
+					// Make sure bulk data is loaded.
+					MakeSureBulkDataIsLoaded();
 
-				// Serialize bulk data.
-				SerializeBulkData( Ar, BulkData );
+					// Serialize bulk data.
+					SerializeBulkData(Ar, BulkData);
+				}
 			}
 		}
 	}
