[INTEGRATE] Change 2455873  by [REDACTED_PERSON][REDACTED_HANDLE] on [REDACTED_PHONE]:50

                      Fixing Finishlinker spikes in cooked builds

diff --git a/Engine/Source/Runtime/CoreUObject/Private/UObject/LinkerLoad.cpp b/Engine/Source/Runtime/CoreUObject/Private/UObject/LinkerLoad.cpp
--- a/Engine/Source/Runtime/CoreUObject/Private/UObject/LinkerLoad.cpp
+++ b/Engine/Source/Runtime/CoreUObject/Private/UObject/LinkerLoad.cpp
@@ -869,7 +869,7 @@ ULinkerLoad::ELinkerStatus ULinkerLoad::SerializePackageFileSummary()
 {
 	DECLARE_SCOPE_CYCLE_COUNTER( TEXT( "ULinkerLoad::SerializePackageFileSummary" ), STAT_LinkerLoad_SerializePackageFileSummary, STATGROUP_LinkerLoad );
 
-	if( bHasSerializedPackageFileSummary == false )
+	if (bHasSerializedPackageFileSummary == false)
 	{
 #if WITH_EDITOR
 		LoadProgressScope->EnterProgressFrame(1);
@@ -931,26 +931,29 @@ ULinkerLoad::ELinkerStatus ULinkerLoad::SerializePackageFileSummary()
 		}
 
 #if PLATFORM_WINDOWS
-		// check if this package version stored the 4-byte magic post tag
-		// get the offset of the post tag
-		int64 MagicOffset = TotalSize() - sizeof(uint32);
-		// store the current file offset
-		int64 OriginalOffset = Tell();
-			
-		uint32 Tag = 0;
-			
-		// seek to the post tag and serialize it
-		Seek(MagicOffset);
-		*this << Tag;
-
-		if (Tag != PACKAGE_FILE_TAG)
+		if (!FPlatformProperties::RequiresCookedData())
 		{
-			UE_LOG(LogLinker, Warning, TEXT("Unable to load package (%s). Post Tag is not valid. File might be corrupted."), *Filename );
-			return LINKER_Failed;
-		}
+			// check if this package version stored the 4-byte magic post tag
+			// get the offset of the post tag
+			int64 MagicOffset = TotalSize() - sizeof(uint32);
+			// store the current file offset
+			int64 OriginalOffset = Tell();
 
-		// seek back to the position after the package summary
-		Seek(OriginalOffset);
+			uint32 Tag = 0;
+
+			// seek to the post tag and serialize it
+			Seek(MagicOffset);
+			*this << Tag;
+
+			if (Tag != PACKAGE_FILE_TAG)
+			{
+				UE_LOG(LogLinker, Warning, TEXT("Unable to load package (%s). Post Tag is not valid. File might be corrupted."), *Filename);
+				return LINKER_Failed;
+			}
+
+			// seek back to the position after the package summary
+			Seek(OriginalOffset);
+		}
 #endif // PLATFORM_WINDOWS
 
 		// Check custom versions.
