[INTEGRATE] Change 2460997  by [REDACTED_PERSON][REDACTED_HANDLE] on [REDACTED_PHONE]:46

                      Fixing crash when loading corrupted asset

diff --git a/Engine/Source/Runtime/CoreUObject/Private/UObject/Linker.cpp b/Engine/Source/Runtime/CoreUObject/Private/UObject/Linker.cpp
--- a/Engine/Source/Runtime/CoreUObject/Private/UObject/Linker.cpp
+++ b/Engine/Source/Runtime/CoreUObject/Private/UObject/Linker.cpp
@@ -669,7 +669,7 @@ ULinkerLoad* GetPackageLinker
 	}
 
 	// Verify compatibility.
-	if( CompatibleGuid && Result->Summary.Guid!=*CompatibleGuid )
+	if (Result && CompatibleGuid && Result->Summary.Guid != *CompatibleGuid)
 	{
 		FFormatNamedArguments Arguments;
 		Arguments.Add(TEXT("AssetName"), FText::FromString(InOuter->GetName()));
diff --git a/Engine/Source/Runtime/CoreUObject/Private/UObject/LinkerLoad.cpp b/Engine/Source/Runtime/CoreUObject/Private/UObject/LinkerLoad.cpp
--- a/Engine/Source/Runtime/CoreUObject/Private/UObject/LinkerLoad.cpp
+++ b/Engine/Source/Runtime/CoreUObject/Private/UObject/LinkerLoad.cpp
@@ -640,13 +640,15 @@ ULinkerLoad::ELinkerStatus ULinkerLoad::Tick( float InTimeLimit, bool bInUseTime
 			);
 	}
 
-#if WITH_EDITOR
 	if (Status == LINKER_Failed)
 	{
+		GObjPendingLoaders.Remove(LinkerRoot);
+#if WITH_EDITOR
+
 		delete LoadProgressScope;
-		LoadProgressScope = nullptr;
-	}
+		LoadProgressScope = nullptr;	
 #endif
+	}
 
 	// Return whether we completed or not.
 	return Status;
