- Minor code fixes for Xcode 6.3beta

diff --git a/Engine/Source/Editor/UnrealEd/Private/CookOnTheFlyServer.cpp b/Engine/Source/Editor/UnrealEd/Private/CookOnTheFlyServer.cpp
--- a/Engine/Source/Editor/UnrealEd/Private/CookOnTheFlyServer.cpp
+++ b/Engine/Source/Editor/UnrealEd/Private/CookOnTheFlyServer.cpp
@@ -927,7 +927,7 @@ uint32 UCookOnTheFlyServer::TickCookOnTheSide( const float TimeSlice, uint32 &Co
 
 		// prevent autosave from happening until we are finished cooking
 		// causes really bad hitches
-		if ( GUnrealEd && (&GUnrealEd->GetPackageAutoSaver()) )
+		if ( GUnrealEd )
 		{
 			const static float SecondsWarningTillAutosave = 10.0f;
 			GUnrealEd->GetPackageAutoSaver().ForceMinimumTimeTillAutoSave(SecondsWarningTillAutosave);
diff --git a/Engine/Source/Runtime/AIModule/Private/BehaviorTree/BehaviorTreeComponent.cpp b/Engine/Source/Runtime/AIModule/Private/BehaviorTree/BehaviorTreeComponent.cpp
--- a/Engine/Source/Runtime/AIModule/Private/BehaviorTree/BehaviorTreeComponent.cpp
+++ b/Engine/Source/Runtime/AIModule/Private/BehaviorTree/BehaviorTreeComponent.cpp
@@ -955,7 +955,7 @@ void UBehaviorTreeComponent::TickComponent(float DeltaTime, enum ELevelTick Tick
 	Super::TickComponent(DeltaTime, TickType, ThisTickFunction);
 	SCOPE_CYCLE_COUNTER(STAT_AI_BehaviorTree_Tick);
 
-	check(this && this->IsPendingKill() == false);
+	check(this != nullptr && this->IsPendingKill() == false);
 		
 	if (bRequestedFlowUpdate)
 	{
diff --git a/Engine/Source/Runtime/CoreUObject/Private/UObject/Obj.cpp b/Engine/Source/Runtime/CoreUObject/Private/UObject/Obj.cpp
--- a/Engine/Source/Runtime/CoreUObject/Private/UObject/Obj.cpp
+++ b/Engine/Source/Runtime/CoreUObject/Private/UObject/Obj.cpp
@@ -492,7 +492,7 @@ void UObject::FinishDestroy()
 FString UObject::GetDetailedInfo() const
 {
 	FString Result;  
-	if( this )
+	if( this != nullptr )
 	{
 		Result = GetDetailedInfoInternal();
 	}
@@ -1005,7 +1005,7 @@ bool UObject::CanCheckDefaultSubObjects(bool bForceCheck, bool& bResult)
 {
 	bool bCanCheck = true;
 	bResult = true;
-	if (!this)
+	if (this == nullptr)
 	{
 		bResult = false; // these aren't in a suitable spot in their lifetime for testing
 		bCanCheck = false;
diff --git a/Engine/Source/Runtime/CoreUObject/Private/UObject/UObjectBase.cpp b/Engine/Source/Runtime/CoreUObject/Private/UObject/UObjectBase.cpp
--- a/Engine/Source/Runtime/CoreUObject/Private/UObject/UObjectBase.cpp
+++ b/Engine/Source/Runtime/CoreUObject/Private/UObject/UObjectBase.cpp
@@ -211,7 +211,7 @@ void UObjectBase::SetClass(UClass* NewClass)
  */
 bool UObjectBase::IsValidLowLevel() const
 {
-	if( !this )
+	if( this == nullptr )
 	{
 		UE_LOG(LogUObjectBase, Warning, TEXT("NULL object") );
 		return false;
@@ -230,7 +230,7 @@ bool UObjectBase::IsValidLowLevelFast(bool bRecursive /*= true*/) const
 	const int32 AlignmentCheck = MIN_ALIGNMENT - 1;
 
 	// Check 'this' pointer before trying to access any of the Object's members
-	if (!this || (UPTRINT)this < 0x100)
+	if ((this == nullptr) || (UPTRINT)this < 0x100)
 	{
 		UE_LOG(LogUObjectBase, Error, TEXT("\'this\' pointer is invalid."));
 		return false;
diff --git a/Engine/Source/Runtime/CoreUObject/Private/UObject/UObjectBaseUtility.cpp b/Engine/Source/Runtime/CoreUObject/Private/UObject/UObjectBaseUtility.cpp
--- a/Engine/Source/Runtime/CoreUObject/Private/UObject/UObjectBaseUtility.cpp
+++ b/Engine/Source/Runtime/CoreUObject/Private/UObject/UObjectBaseUtility.cpp
@@ -72,7 +72,7 @@ void UObjectBaseUtility::GetPathName( const UObject* StopOuter, FString& ResultS
 FString UObjectBaseUtility::GetFullName( const UObject* StopOuter/*=NULL*/ ) const
 {
 	FString Result;  
-	if( this )
+	if( this != nullptr )
 	{
 		Result.Empty(128);
 		GetClass()->AppendName(Result);
