UE-8676 Crash when changing a Blueprint's vector variable to a transform variable
UE-9089 Better handling for UDS rename

Problem: When a property was of type on native struct (that is serialized binary) and the type was change to user defined struct (that is serialized using tags): the new structure tries to serialize tag from binary data of the old struct.
Solution: Use persistent guid to identify user defined structures (even after they are renamed).

diff --git a/Engine/Source/Editor/UnrealEd/Private/Kismet2/StructureEditorUtils.cpp b/Engine/Source/Editor/UnrealEd/Private/Kismet2/StructureEditorUtils.cpp
--- a/Engine/Source/Editor/UnrealEd/Private/Kismet2/StructureEditorUtils.cpp
+++ b/Engine/Source/Editor/UnrealEd/Private/Kismet2/StructureEditorUtils.cpp
@@ -37,6 +37,7 @@ UUserDefinedStruct* FStructureEditorUtils::CreateUserDefinedStruct(UObject* InPa
 		Struct->EditorData = NewObject<UUserDefinedStructEditorData>(Struct, NAME_None, RF_Transactional);
 		check(Struct->EditorData);
 
+		Struct->Guid = FGuid::NewGuid();
 		Struct->SetMetaData(TEXT("BlueprintType"), TEXT("true"));
 		Struct->Bind();
 		Struct->StaticLink(true);
diff --git a/Engine/Source/Runtime/Core/Public/UObject/ObjectVersion.h b/Engine/Source/Runtime/Core/Public/UObject/ObjectVersion.h
--- a/Engine/Source/Runtime/Core/Public/UObject/ObjectVersion.h
+++ b/Engine/Source/Runtime/Core/Public/UObject/ObjectVersion.h
@@ -476,6 +476,8 @@ enum EUnrealEngineObjectUE4Version
 	VER_UE4_ACTOR_COMPONENT_CREATION_METHOD,
 	// K2Node_Event now uses FMemberReference for handling references
 	VER_UE4_K2NODE_EVENT_MEMBER_REFERENCE,
+	// FPropertyTag stores GUID of struct
+	VER_UE4_STRUCT_GUID_IN_PROPERTY_TAG,
 
 	// -----<new versions can be added before this line>-------------------------------------------------
 	// - this needs to be the last line (see note below)
diff --git a/Engine/Source/Runtime/CoreUObject/Private/UObject/Class.cpp b/Engine/Source/Runtime/CoreUObject/Private/UObject/Class.cpp
--- a/Engine/Source/Runtime/CoreUObject/Private/UObject/Class.cpp
+++ b/Engine/Source/Runtime/CoreUObject/Private/UObject/Class.cpp
@@ -799,6 +799,16 @@ void UStruct::SerializeTaggedProperties(FArchive& Ar, uint8* Data, UStruct* Defa
 				break;
 			}
 
+			auto CanSerializeFromStructWithDifferentName = [](const FArchive& Ar, const FPropertyTag& PropertyTag, const UStructProperty* StructProperty)
+			{
+				if (Ar.UE4Ver() < VER_UE4_STRUCT_GUID_IN_PROPERTY_TAG)
+				{
+					// Old Implementation
+					return !StructProperty->UseBinaryOrNativeSerialization(Ar);
+				}
+				return PropertyTag.StructGuid.IsValid() && StructProperty && StructProperty->Struct && (PropertyTag.StructGuid == StructProperty->Struct->GetCustomGuid());
+			};
+
 			// Move to the next property to be serialized
 			if( AdvanceProperty && --RemainingArrayDim <= 0 )
 			{
@@ -1137,7 +1147,7 @@ void UStruct::SerializeTaggedProperties(FArchive& Ar, uint8* Data, UStruct* Defa
 				}
 			}
 			else if( Tag.Type==NAME_StructProperty && Tag.StructName!=CastChecked<UStructProperty>(Property)->Struct->GetFName() 
-				&& CastChecked<UStructProperty>(Property)->UseBinaryOrNativeSerialization(Ar) )
+				&& !CanSerializeFromStructWithDifferentName(Ar, Tag, CastChecked<UStructProperty>(Property)))
 			{
 				UE_LOG(LogClass, Warning, TEXT("Property %s of %s struct type mismatch %s/%s for package:  %s. If that property got renamed, add an ActiveStructRedirect."), *Tag.Name.ToString(), *GetName(), *Tag.StructName.ToString(), *CastChecked<UStructProperty>(Property)->Struct->GetName(), *Ar.GetArchiveName() );
 			}
@@ -2301,6 +2311,11 @@ void UScriptStruct::DestroyStruct(void* Dest, int32 ArrayDim) const
 
 void UScriptStruct::RecursivelyPreload() {}
 
+FGuid UScriptStruct::GetCustomGuid() const
+{
+	return FGuid();
+}
+
 IMPLEMENT_CORE_INTRINSIC_CLASS(UScriptStruct, UStruct,
 	{
 	}
diff --git a/Engine/Source/Runtime/CoreUObject/Public/UObject/Class.h b/Engine/Source/Runtime/CoreUObject/Public/UObject/Class.h
--- a/Engine/Source/Runtime/CoreUObject/Public/UObject/Class.h
+++ b/Engine/Source/Runtime/CoreUObject/Public/UObject/Class.h
@@ -1093,6 +1093,8 @@ public:
 	void ClearScriptStruct(void* Dest, int32 ArrayDim = 1) const;
 
 	virtual COREUOBJECT_API void RecursivelyPreload();
+
+	virtual COREUOBJECT_API FGuid GetCustomGuid() const;
 };
 
 class FStructOnScope
diff --git a/Engine/Source/Runtime/CoreUObject/Public/UObject/PropertyTag.h b/Engine/Source/Runtime/CoreUObject/Public/UObject/PropertyTag.h
--- a/Engine/Source/Runtime/CoreUObject/Public/UObject/PropertyTag.h
+++ b/Engine/Source/Runtime/CoreUObject/Public/UObject/PropertyTag.h
@@ -22,6 +22,7 @@ struct FPropertyTag
 	int32		Size;       // Property size.
 	int32		ArrayIndex;	// Index if an array; else 0.
 	int32		SizeOffset;	// location in stream of tag size member
+	FGuid	StructGuid;
 
 	// Constructors.
 	FPropertyTag()
@@ -53,6 +54,7 @@ struct FPropertyTag
 			if (UStructProperty* StructProperty = Cast<UStructProperty>(Property))
 			{
 				StructName = StructProperty->Struct->GetFName();
+				StructGuid = StructProperty->Struct->GetCustomGuid();
 			}
 			else if (UByteProperty* ByteProp = Cast<UByteProperty>(Property))
 			{
@@ -95,6 +97,10 @@ struct FPropertyTag
 		if (Tag.Type == NAME_StructProperty)
 		{
 			Ar << Tag.StructName;
+			if (Ar.UE4Ver() >= VER_UE4_STRUCT_GUID_IN_PROPERTY_TAG)
+			{
+				Ar << Tag.StructGuid;
+			}
 		}
 		// only need to serialize this for bools
 		else if (Tag.Type == NAME_BoolProperty)
diff --git a/Engine/Source/Runtime/Engine/Classes/Engine/UserDefinedStruct.h b/Engine/Source/Runtime/Engine/Classes/Engine/UserDefinedStruct.h
--- a/Engine/Source/Runtime/Engine/Classes/Engine/UserDefinedStruct.h
+++ b/Engine/Source/Runtime/Engine/Classes/Engine/UserDefinedStruct.h
@@ -40,14 +40,20 @@ public:
 	UObject* EditorData;
 #endif // WITH_EDITORONLY_DATA
 
+	UPROPERTY()
+	FGuid Guid;
+
 #if WITH_EDITOR
 	// UObject interface.
 	virtual void Serialize(FArchive& Ar) override;
 	virtual void PostDuplicate(bool bDuplicateForPIE) override;
 	virtual void GetAssetRegistryTags(TArray<FAssetRegistryTag>& OutTags) const override;
+	virtual void PostLoad() override;
 	// End of UObject interface.
 
 	void InitializeDefaultValue(uint8* StructData) const;
+
+	void ValidateGuid();
 #endif	// WITH_EDITOR
 
 	virtual void SerializeTaggedProperties(FArchive& Ar, uint8* Data, UStruct* DefaultsStruct, uint8* Defaults, const UObject* BreakRecursionIfFullyLoad=NULL) const override;
@@ -55,6 +61,7 @@ public:
 
 	// UScriptStruct interface.
 	virtual void RecursivelyPreload() override;
+	virtual FGuid GetCustomGuid() const override;
 	// End of  UScriptStruct interface.
 
 #if WITH_EDITOR
diff --git a/Engine/Source/Runtime/Engine/Private/UserDefinedStruct.cpp b/Engine/Source/Runtime/Engine/Private/UserDefinedStruct.cpp
--- a/Engine/Source/Runtime/Engine/Private/UserDefinedStruct.cpp
+++ b/Engine/Source/Runtime/Engine/Private/UserDefinedStruct.cpp
@@ -32,6 +32,10 @@ void UUserDefinedStruct::Serialize(FArchive& Ar)
 void UUserDefinedStruct::PostDuplicate(bool bDuplicateForPIE)
 {
 	Super::PostDuplicate(bDuplicateForPIE);
+	if (!bDuplicateForPIE)
+	{
+		Guid = FGuid::NewGuid();
+	}
 	if (!bDuplicateForPIE && (GetOuter() != GetTransientPackage()))
 	{
 		SetMetaData(TEXT("BlueprintType"), TEXT("true"));
@@ -39,6 +43,13 @@ void UUserDefinedStruct::PostDuplicate(bool bDuplicateForPIE)
 	}
 }
 
+void UUserDefinedStruct::PostLoad()
+{
+	Super::PostLoad();
+
+	ValidateGuid();
+}
+
 void UUserDefinedStruct::GetAssetRegistryTags(TArray<FAssetRegistryTag>& OutTags) const
 {
 	Super::GetAssetRegistryTags(OutTags);
@@ -59,6 +70,22 @@ void UUserDefinedStruct::InitializeDefaultValue(uint8* StructData) const
 	FStructureEditorUtils::Fill_MakeStructureDefaultValue(this, StructData);
 }
 
+void UUserDefinedStruct::ValidateGuid()
+{
+	// Backward compatibility:
+	// The guid is created in an deterministic way using existing name.
+	if (!Guid.IsValid() && (GetFName() != NAME_None))
+	{
+		const FString HashString = GetFName().ToString();
+		ensure(HashString.Len());
+
+		const uint32 BufferLength = HashString.Len() * sizeof(HashString[0]);
+		uint32 HashBuffer[5];
+		FSHA1::HashBuffer(*HashString, BufferLength, reinterpret_cast<uint8*>(HashBuffer));
+		Guid = FGuid(HashBuffer[1], HashBuffer[2], HashBuffer[3], HashBuffer[4]);
+	}
+}
+
 #endif	// WITH_EDITOR
 
 FString UUserDefinedStruct::PropertyNameToDisplayName(FName Name) const
@@ -146,4 +173,9 @@ void UUserDefinedStruct::RecursivelyPreload()
 			StaticLink(true);
 		}
 	}
+}
+
+FGuid UUserDefinedStruct::GetCustomGuid() const
+{
+	return Guid;
 }
\ No newline at end of file
