UE-10832 Deferred loading - placeholder for function

Test case: /QAGame/Content/Developers/maciejmroz/MCCycle/

diff --git a/Engine/Source/Runtime/CoreUObject/Private/Blueprint/BlueprintSupport.cpp b/Engine/Source/Runtime/CoreUObject/Private/Blueprint/BlueprintSupport.cpp
--- a/Engine/Source/Runtime/CoreUObject/Private/Blueprint/BlueprintSupport.cpp
+++ b/Engine/Source/Runtime/CoreUObject/Private/Blueprint/BlueprintSupport.cpp
@@ -3,6 +3,7 @@
 #include "CoreUObjectPrivate.h"
 #include "UObject/LinkerPlaceholderClass.h"
 #include "UObject/LinkerPlaceholderExportObject.h"
+#include "UObject/LinkerPlaceholderFunction.h"
 #include "Linker.h"
 #include "PropertyTag.h"
 #include "UObject/StructScriptLoader.h"
@@ -357,6 +358,34 @@ bool ULinkerLoad::DeferPotentialCircularImport(const int32 Index)
 
 					Import.XObject = Placeholder;
 				}
+				else if (ImportClass->IsChildOf<UFunction>())
+				{		
+					UObject* OuterObject = nullptr;
+					if (Import.OuterIndex.IsImport())
+					{
+						const int32 OuterImportIndex = Import.OuterIndex.ToImport();
+						FObjectImport& OuterImport = ImportMap[OuterImportIndex];
+						if (nullptr == OuterImport.XObject)
+						{
+							DeferPotentialCircularImport(OuterImportIndex);
+						}
+						OuterObject = OuterImport.XObject;
+					}
+					const bool bOuterPlaceholderExists = OuterObject && OuterObject->IsA<ULinkerPlaceholderClass>();
+					if (bOuterPlaceholderExists && (nullptr == Import.XObject))
+					{
+						UClass*   PlaceholderType = ULinkerPlaceholderFunction::StaticClass();
+
+						FName PlaceholderName(*FString::Printf(TEXT("PLACEHOLDER-FUNCTION_%s"), *Import.ObjectName.ToString()));
+						PlaceholderName = MakeUniqueObjectName(OuterObject, PlaceholderType, PlaceholderName);
+
+						ULinkerPlaceholderFunction* Placeholder = NewObject<ULinkerPlaceholderFunction>(OuterObject, PlaceholderType, PlaceholderName, RF_Public | RF_Transient);
+						Placeholder->ImportIndex = Index;
+						Placeholder->Bind();
+						Placeholder->StaticLink(/*bRelinkExistingProperties =*/true);
+						Import.XObject = Placeholder;
+					}
+				}
 			}
 		}
 
@@ -698,6 +727,29 @@ void ULinkerLoad::ResolveDeferredDependencies(UStruct* LoadStruct)
 				//       to the same function (for the same placeholder)
 				ResolveDependencyPlaceholder(PlaceholderClass, Cast<UClass>(LoadStruct));
 			}
+			else if (ULinkerPlaceholderFunction* PlaceholderFunction = Cast<ULinkerPlaceholderFunction>(Import.XObject))
+			{
+				DEFERRED_DEPENDENCY_CHECK(PlaceholderFunction->ImportIndex == ImportIndex);
+
+				auto OwnerClass = PlaceholderFunction->GetOwnerClass();
+				ResolveDependencyPlaceholder(OwnerClass, Cast<UClass>(LoadStruct));
+
+				UFunction* RealFuncObj = nullptr;
+				if ((Import.XObject != nullptr) && (Import.XObject != PlaceholderFunction))
+				{
+					RealFuncObj = CastChecked<UFunction>(Import.XObject);
+				}
+				else
+				{
+					Import.XObject = nullptr;
+					RealFuncObj = CastChecked<UFunction>(CreateImport(ImportIndex), ECastCheckedType::NullAllowed);
+				}
+				ensure(RealFuncObj != PlaceholderFunction);
+				if (RealFuncObj)
+				{
+					PlaceholderFunction->ReplaceTrackedReferences(RealFuncObj);
+				}
+			}
 			else if (UScriptStruct* StructObj = Cast<UScriptStruct>(Import.XObject))
 			{
 				// in case this is a user defined struct, we have to resolve any 
@@ -836,7 +888,7 @@ int32 ULinkerLoad::ResolveDependencyPlaceholder(UClass* PlaceholderIn, UClass* R
 	DEFERRED_DEPENDENCY_CHECK( (ReplacementCount > 0) || PlaceholderClass->HasReferences() || PlaceholderClass->HasBeenResolved() );
 
 	ReplacementCount += PlaceholderClass->ReplaceTrackedReferences(RealClassObj);
-	PlaceholderClass->MarkPendingKill(); // @TODO: ensure these are properly GC'd
+	// PlaceholderClass->MarkPendingKill(); // @TODO: ensure these are properly GC'd
 
 #if USE_DEFERRED_DEPENDENCY_CHECK_VERIFICATION_TESTS
 	UObject* PlaceholderObj = PlaceholderClass;
diff --git a/Engine/Source/Runtime/CoreUObject/Private/UObject/Class.cpp b/Engine/Source/Runtime/CoreUObject/Private/UObject/Class.cpp
--- a/Engine/Source/Runtime/CoreUObject/Private/UObject/Class.cpp
+++ b/Engine/Source/Runtime/CoreUObject/Private/UObject/Class.cpp
@@ -8,6 +8,7 @@
 #include "PropertyTag.h"
 #include "HotReloadInterface.h"
 #include "LinkerPlaceholderClass.h"
+#include "LinkerPlaceholderFunction.h"
 #include "StructScriptLoader.h"
 
 DECLARE_LOG_CATEGORY_EXTERN(LogScriptSerialization, Log, All);
@@ -1488,10 +1489,14 @@ bool UStruct::GetStringMetaDataHierarchical(const FName& Key, FString* OutValue)
 	static void HandlePlaceholderScriptRef(ScriptPointerType& ScriptPtr)
 	{
 		UObject*& ExprPtrRef = (UObject*&)ScriptPtr;
-		if (ULinkerPlaceholderClass* PlaceholderObj = Cast<ULinkerPlaceholderClass>(ExprPtrRef)) \
+		if (ULinkerPlaceholderClass* PlaceholderObj = Cast<ULinkerPlaceholderClass>(ExprPtrRef))
 		{
 			PlaceholderObj->AddReferencingScriptExpr((ULinkerPlaceholderClass**)(&ExprPtrRef));
 		}
+		else if (ULinkerPlaceholderFunction* PlaceholderFunc = Cast<ULinkerPlaceholderFunction>(ExprPtrRef))
+		{
+			PlaceholderFunc->AddReferencingScriptExpr((ULinkerPlaceholderFunction**)(&ExprPtrRef));
+		}
 	}
 
 	#define FIXUP_EXPR_OBJECT_POINTER(Type) \
diff --git a/Engine/Source/Runtime/CoreUObject/Private/UObject/LinkerPlaceholderBase.cpp b/Engine/Source/Runtime/CoreUObject/Private/UObject/LinkerPlaceholderBase.cpp
new file mode 100644
--- /dev/null
+++ b/Engine/Source/Runtime/CoreUObject/Private/UObject/LinkerPlaceholderBase.cpp
@@ -0,0 +1,55 @@
+// Copyright 1998-2015 Epic Games, Inc. All Rights Reserved.
+
+#include "CoreUObjectPrivate.h"
+#include "LinkerPlaceholderBase.h"
+#include "Blueprint/BlueprintSupport.h"
+
+//------------------------------------------------------------------------------
+void FLinkerPlaceholderBase::AddReferencingProperty(UProperty* ReferencingProperty)
+{
+#if USE_DEFERRED_DEPENDENCY_CHECK_VERIFICATION_TESTS
+	FObjectImport* PlaceholderImport = nullptr;
+	if (ULinkerLoad* PropertyLinker = ReferencingProperty->GetLinker())
+	{
+		for (FObjectImport& Import : PropertyLinker->ImportMap)
+		{
+			if (Import.XObject == this)
+			{
+				PlaceholderImport = &Import;
+				break;
+			}
+		}
+		check(GetOuter() == PropertyLinker->LinkerRoot);
+		check((PropertyLinker->LoadFlags & LOAD_DeferDependencyLoads) || FBlueprintSupport::IsResolvingDeferredDependenciesDisabled());
+	}
+	// if this check hits, then we're adding dependencies after we've 
+	// already resolved the placeholder (it won't be resolved again)
+	check(!bResolvedReferences);
+#endif // USE_DEFERRED_DEPENDENCY_CHECK_VERIFICATION_TESTS
+
+	ReferencingProperties.Add(ReferencingProperty);
+}
+
+//------------------------------------------------------------------------------
+bool FLinkerPlaceholderBase::HasReferences() const
+{
+	return (GetRefCount() > 0);
+}
+
+//------------------------------------------------------------------------------
+int32 FLinkerPlaceholderBase::GetRefCount() const
+{
+	return ReferencingProperties.Num();
+}
+
+//------------------------------------------------------------------------------
+bool FLinkerPlaceholderBase::HasBeenResolved() const
+{
+	return !HasReferences() && bResolvedReferences;
+}
+
+//------------------------------------------------------------------------------
+void FLinkerPlaceholderBase::RemovePropertyReference(UProperty* ReferencingProperty)
+{
+	ReferencingProperties.Remove(ReferencingProperty);
+}
diff --git a/Engine/Source/Runtime/CoreUObject/Private/UObject/LinkerPlaceholderBase.h b/Engine/Source/Runtime/CoreUObject/Private/UObject/LinkerPlaceholderBase.h
new file mode 100644
--- /dev/null
+++ b/Engine/Source/Runtime/CoreUObject/Private/UObject/LinkerPlaceholderBase.h
@@ -0,0 +1,71 @@
+// Copyright 1998-2014 Epic Games, Inc. All Rights Reserved.
+
+#pragma once
+
+class UProperty;
+
+class FLinkerPlaceholderBase
+{
+public:
+	FLinkerPlaceholderBase()
+		: ImportIndex(0)
+		, bResolvedReferences(false)
+	{}
+
+	virtual ~FLinkerPlaceholderBase()
+	{}
+
+	/**
+	 * Caches off the supplied property so that we can later replace it's use of
+	 * this class with another (real) class.
+	 * 
+	 * @param  ReferencingProperty    A property that uses and stores this class.
+	 */
+	void AddReferencingProperty(UProperty* ReferencingProperty);
+
+	/**
+	 * A query method that let's us check to see if this class is currently 
+	 * being referenced by anything (if this returns false, then a referencing 
+	 * property could have forgotten to add itself... or, we've replaced all
+	 * references).
+	 * 
+	 * @return True if this has anything stored in its ReferencingProperties container, otherwise false.
+	 */
+	bool HasReferences() const;
+
+	/**
+	 * Query method that retrieves the current number of KNOWN references to 
+	 * this placeholder class.
+	 * 
+	 * @return The number of references that this class is currently tracking.
+	 */
+	virtual int32 GetRefCount() const;
+
+	/**
+	 * Checks to see if 1) this placeholder has had RemoveTrackedReference() 
+	 * called on it, and 2) it doesn't have any more references that have since 
+	 * been added.
+	 * 
+	 * @return True if ReplaceTrackedReferences() has been ran, and no KNOWN references have been added.
+	 */
+	bool HasBeenResolved() const;
+
+	/**
+	 * Removes the specified property from this class's internal tracking list 
+	 * (which aims to keep track of properties utilizing this class).
+	 * 
+	 * @param  ReferencingProperty    A property that used to use this class, and now no longer does.
+	 */
+	void RemovePropertyReference(UProperty* ReferencingProperty);
+
+public:
+	/** Set by the ULinkerLoad that created this instance, tracks what import this was used in place of. */
+	int32 ImportIndex;
+
+protected:
+	/** Links to UProperties that are currently using this class */
+	TSet<UProperty*> ReferencingProperties;
+
+	/** Used to catch references that are added after we've already resolved all references */
+	bool bResolvedReferences;
+};
diff --git a/Engine/Source/Runtime/CoreUObject/Private/UObject/LinkerPlaceholderClass.cpp b/Engine/Source/Runtime/CoreUObject/Private/UObject/LinkerPlaceholderClass.cpp
--- a/Engine/Source/Runtime/CoreUObject/Private/UObject/LinkerPlaceholderClass.cpp
+++ b/Engine/Source/Runtime/CoreUObject/Private/UObject/LinkerPlaceholderClass.cpp
@@ -1,13 +1,11 @@
 // Copyright 1998-2015 Epic Games, Inc. All Rights Reserved.
 
 #include "CoreUObjectPrivate.h"
-#include "LinkerPlaceholderClass.h"
-#include "Blueprint/BlueprintSupport.h"
+#include "LinkerPlaceholderFunction.h"
 
 //------------------------------------------------------------------------------
 ULinkerPlaceholderClass::ULinkerPlaceholderClass(const FObjectInitializer& ObjectInitializer)
 	: Super(ObjectInitializer)
-	, bResolvedReferences(false)
 {
 }
 
@@ -63,32 +61,6 @@ void ULinkerPlaceholderClass::Bind()
 	ClassAddReferencedObjects = &ULinkerPlaceholderClass::AddReferencedObjects;
 }
 
-//------------------------------------------------------------------------------
-void ULinkerPlaceholderClass::AddReferencingProperty(UProperty* ReferencingProperty)
-{
-#if USE_DEFERRED_DEPENDENCY_CHECK_VERIFICATION_TESTS
-	FObjectImport* PlaceholderImport = nullptr;
-	if (ULinkerLoad* PropertyLinker = ReferencingProperty->GetLinker())
-	{
-		for (FObjectImport& Import : PropertyLinker->ImportMap)
-		{
-			if (Import.XObject == this)
-			{
-				PlaceholderImport = &Import;
-				break;
-			}
-		}
-		check(GetOuter() == PropertyLinker->LinkerRoot);
-		check((PropertyLinker->LoadFlags & LOAD_DeferDependencyLoads) || FBlueprintSupport::IsResolvingDeferredDependenciesDisabled());
-	}
-	// if this check hits, then we're adding dependencies after we've 
-	// already resolved the placeholder (it won't be resolved again)
-	check(!bResolvedReferences); 	
-#endif // USE_DEFERRED_DEPENDENCY_CHECK_VERIFICATION_TESTS
-
-	ReferencingProperties.Add(ReferencingProperty);
-}
-
 //------------------------------------------------------------------------------
 void ULinkerPlaceholderClass::AddReferencingScriptExpr(ULinkerPlaceholderClass** ExpressionPtr)
 {
@@ -99,30 +71,6 @@ void ULinkerPlaceholderClass::AddReferencingScriptExpr(ULinkerPlaceholderClass**
 	ReferencingScriptExpressions.Add((UClass**)ExpressionPtr);
 }
 
-//------------------------------------------------------------------------------
-bool ULinkerPlaceholderClass::HasReferences() const
-{
-	return (GetRefCount() > 0);
-}
-
-//------------------------------------------------------------------------------
-int32 ULinkerPlaceholderClass::GetRefCount() const
-{
-	return ReferencingProperties.Num() + ReferencingScriptExpressions.Num();
-}
-
-//------------------------------------------------------------------------------
-bool ULinkerPlaceholderClass::HasBeenResolved() const
-{
-	return !HasReferences() && bResolvedReferences;
-}
-
-//------------------------------------------------------------------------------
-void ULinkerPlaceholderClass::RemovePropertyReference(UProperty* ReferencingProperty)
-{
-	ReferencingProperties.Remove(ReferencingProperty);
-}
-
 //------------------------------------------------------------------------------
 int32 ULinkerPlaceholderClass::ReplaceTrackedReferences(UClass* ReplacementClass)
 {
@@ -170,3 +118,9 @@ int32 ULinkerPlaceholderClass::ReplaceTrackedReferences(UClass* ReplacementClass
 	bResolvedReferences = true;
 	return ReplacementCount;
 }
+
+//------------------------------------------------------------------------------
+int32 ULinkerPlaceholderClass::GetRefCount() const
+{
+	return FLinkerPlaceholderBase::GetRefCount() + ReferencingScriptExpressions.Num();
+}
diff --git a/Engine/Source/Runtime/CoreUObject/Private/UObject/LinkerPlaceholderClass.h b/Engine/Source/Runtime/CoreUObject/Private/UObject/LinkerPlaceholderClass.h
--- a/Engine/Source/Runtime/CoreUObject/Private/UObject/LinkerPlaceholderClass.h
+++ b/Engine/Source/Runtime/CoreUObject/Private/UObject/LinkerPlaceholderClass.h
@@ -3,6 +3,7 @@
 #pragma once
 #include "Class.h"
 #include "Set.h"
+#include "LinkerPlaceholderBase.h"
 
 // Forward declarations
 class FObjectInitializer;
@@ -14,7 +15,7 @@ class UObjectPropertyBase;
  * Holds on to references where this is currently being utilized, so we can 
  * easily replace references to it later (once the real class is available).
  */ 
-class ULinkerPlaceholderClass : public UClass
+class ULinkerPlaceholderClass : public UClass, public FLinkerPlaceholderBase
 {
 public:
 	DECLARE_CASTED_CLASS_INTRINSIC_NO_CTOR(ULinkerPlaceholderClass, UClass, /*TStaticFlags =*/0, CoreUObject, /*TStaticCastFlags =*/0, NO_API)
@@ -32,13 +33,9 @@ public:
  	virtual void Bind() override;
 	// End of UField interface.
 
-	/**
-	 * Caches off the supplied property so that we can later replace it's use of
-	 * this class with another (real) class.
-	 * 
-	 * @param  ReferencingProperty    A property that uses and stores this class.
-	 */
-	void AddReferencingProperty(UProperty* ReferencingProperty);
+	// FLinkerPlaceholderBase 
+	virtual int32 GetRefCount() const override;
+	// End of FLinkerPlaceholderBase
 
 	/**
 	 * Records a raw pointer, directly to the UClass* script expression (so that
@@ -54,41 +51,6 @@ public:
 	 */
 	void AddReferencingScriptExpr(ULinkerPlaceholderClass** ExpressionPtr);
 
-	/**
-	 * A query method that let's us check to see if this class is currently 
-	 * being referenced by anything (if this returns false, then a referencing 
-	 * property could have forgotten to add itself... or, we've replaced all
-	 * references).
-	 * 
-	 * @return True if this has anything stored in its ReferencingProperties container, otherwise false.
-	 */
-	bool HasReferences() const;
-
-	/**
-	 * Query method that retrieves the current number of KNOWN references to 
-	 * this placeholder class.
-	 * 
-	 * @return The number of references that this class is currently tracking.
-	 */
-	int32 GetRefCount() const;
-
-	/**
-	 * Checks to see if 1) this placeholder has had RemoveTrackedReference() 
-	 * called on it, and 2) it doesn't have any more references that have since 
-	 * been added.
-	 * 
-	 * @return True if ReplaceTrackedReferences() has been ran, and no KNOWN references have been added.
-	 */
-	bool HasBeenResolved() const;
-
-	/**
-	 * Removes the specified property from this class's internal tracking list 
-	 * (which aims to keep track of properties utilizing this class).
-	 * 
-	 * @param  ReferencingProperty    A property that used to use this class, and now no longer does.
-	 */
-	void RemovePropertyReference(UProperty* ReferencingProperty);
-
 	/**
 	 * Iterates over all referencing properties and attempts to replace their 
 	 * references to this class with a new (hopefully proper) class.
@@ -98,16 +60,7 @@ public:
 	 */
 	int32 ReplaceTrackedReferences(UClass* ReplacementClass);
 
-public:
-	/** Set by the ULinkerLoad that created this instance, tracks what import this was used in place of. */
-	int32 ImportIndex;
-
 private:
-	/** Links to UProperties that are currently using this class */
-	TSet<UProperty*> ReferencingProperties;
-
-	/** Used to catch references that are added after we've already resolved all references */
-	bool bResolvedReferences;
 
 	/** Points directly at UClass* refs that we're serialized in as part of script bytecode */
 	TSet<UClass**> ReferencingScriptExpressions;
diff --git a/Engine/Source/Runtime/CoreUObject/Private/UObject/LinkerPlaceholderFunction.cpp b/Engine/Source/Runtime/CoreUObject/Private/UObject/LinkerPlaceholderFunction.cpp
new file mode 100644
--- /dev/null
+++ b/Engine/Source/Runtime/CoreUObject/Private/UObject/LinkerPlaceholderFunction.cpp
@@ -0,0 +1,78 @@
+// Copyright 1998-2015 Epic Games, Inc. All Rights Reserved.
+
+#include "CoreUObjectPrivate.h"
+#include "LinkerPlaceholderFunction.h"
+
+//------------------------------------------------------------------------------
+ULinkerPlaceholderFunction::ULinkerPlaceholderFunction(const FObjectInitializer& ObjectInitializer)
+	: Super(ObjectInitializer)
+{
+}
+
+//------------------------------------------------------------------------------
+ULinkerPlaceholderFunction::~ULinkerPlaceholderFunction()
+{
+}
+
+//------------------------------------------------------------------------------
+IMPLEMENT_CORE_INTRINSIC_CLASS(ULinkerPlaceholderFunction, UFunction,
+	{
+	}
+);
+
+//------------------------------------------------------------------------------
+int32 ULinkerPlaceholderFunction::ReplaceTrackedReferences(UFunction* ReplacementObj)
+{
+	int32 ReplacementCount = 0;
+	for (UProperty* Property : ReferencingProperties)
+	{
+		if (auto DelegateProperty = Cast<UDelegateProperty>(Property))
+		{
+			if (DelegateProperty->SignatureFunction == this)
+			{
+				DelegateProperty->SignatureFunction = ReplacementObj;
+				++ReplacementCount;
+			}
+		}
+		else if (auto MulticastDelegateProperty = Cast<UMulticastDelegateProperty>(Property))
+		{
+			MulticastDelegateProperty->SignatureFunction = ReplacementObj;
+			++ReplacementCount;
+		}
+		else
+		{
+			ensure(false);
+		}
+	}
+	ReferencingProperties.Empty();
+
+	for (UFunction** ScriptRefPtr : ReferencingScriptExpressions)
+	{
+		UFunction*& ScriptRef = *ScriptRefPtr;
+		if (ScriptRef == this)
+		{
+			ScriptRef = ReplacementObj;
+			++ReplacementCount;
+		}
+	}
+	ReferencingScriptExpressions.Empty();
+
+	bResolvedReferences = true;
+	return ReplacementCount;
+}
+
+//------------------------------------------------------------------------------
+void ULinkerPlaceholderFunction::AddReferencingScriptExpr(ULinkerPlaceholderFunction** ExpressionPtr)
+{
+#if USE_DEFERRED_DEPENDENCY_CHECK_VERIFICATION_TESTS
+	check(*ExpressionPtr == this);
+#endif // USE_DEFERRED_DEPENDENCY_CHECK_VERIFICATION_TESTS
+
+	ReferencingScriptExpressions.Add((UFunction**)ExpressionPtr);
+}
+
+//------------------------------------------------------------------------------
+int32 ULinkerPlaceholderFunction::GetRefCount() const
+{
+	return FLinkerPlaceholderBase::GetRefCount() + ReferencingScriptExpressions.Num();
+}
\ No newline at end of file
diff --git a/Engine/Source/Runtime/CoreUObject/Private/UObject/LinkerPlaceholderFunction.h b/Engine/Source/Runtime/CoreUObject/Private/UObject/LinkerPlaceholderFunction.h
new file mode 100644
--- /dev/null
+++ b/Engine/Source/Runtime/CoreUObject/Private/UObject/LinkerPlaceholderFunction.h
@@ -0,0 +1,29 @@
+// Copyright 1998-2014 Epic Games, Inc. All Rights Reserved.
+
+#pragma once
+#include "LinkerPlaceholderClass.h"
+
+// Forward declarations
+class FObjectInitializer;
+class UObjectPropertyBase;
+
+class ULinkerPlaceholderFunction : public UFunction, public FLinkerPlaceholderBase
+{
+public:
+	DECLARE_CASTED_CLASS_INTRINSIC_NO_CTOR(ULinkerPlaceholderFunction, UFunction, /*TStaticFlags =*/0, CoreUObject, /*TStaticCastFlags =*/0, NO_API)
+
+	ULinkerPlaceholderFunction(const FObjectInitializer& ObjectInitializer);
+	virtual ~ULinkerPlaceholderFunction();
+
+	// FPlaceholderBase 
+	virtual int32 GetRefCount() const override;
+	// End of FPlaceholderBase
+
+	int32 ReplaceTrackedReferences(UFunction* ReplacementClass);
+
+	void AddReferencingScriptExpr(ULinkerPlaceholderFunction** ExpressionPtr);
+
+private:
+
+	TSet<UFunction**> ReferencingScriptExpressions;
+};
\ No newline at end of file
diff --git a/Engine/Source/Runtime/CoreUObject/Private/UObject/PropertyClass.cpp b/Engine/Source/Runtime/CoreUObject/Private/UObject/PropertyClass.cpp
--- a/Engine/Source/Runtime/CoreUObject/Private/UObject/PropertyClass.cpp
+++ b/Engine/Source/Runtime/CoreUObject/Private/UObject/PropertyClass.cpp
@@ -56,7 +56,7 @@ void UClassProperty::SetMetaClass(UClass* NewMetaClass)
 		NewPlaceholderClass->AddReferencingProperty(this);
 	}
 
-	if (ULinkerPlaceholderClass* OldPlaceholderClass = Cast<ULinkerPlaceholderClass>(NewMetaClass))
+	if (ULinkerPlaceholderClass* OldPlaceholderClass = Cast<ULinkerPlaceholderClass>(MetaClass))
 	{
 		OldPlaceholderClass->RemovePropertyReference(this);
 	}
diff --git a/Engine/Source/Runtime/CoreUObject/Private/UObject/PropertyDelegate.cpp b/Engine/Source/Runtime/CoreUObject/Private/UObject/PropertyDelegate.cpp
--- a/Engine/Source/Runtime/CoreUObject/Private/UObject/PropertyDelegate.cpp
+++ b/Engine/Source/Runtime/CoreUObject/Private/UObject/PropertyDelegate.cpp
@@ -2,6 +2,7 @@
 
 #include "CoreUObjectPrivate.h"
 #include "PropertyHelper.h"
+#include "LinkerPlaceholderFunction.h"
 
 /*-----------------------------------------------------------------------------
 	UDelegateProperty.
@@ -106,6 +107,16 @@ void UDelegateProperty::Serialize( FArchive& Ar )
 {
 	Super::Serialize( Ar );
 	Ar << SignatureFunction;
+
+#if USE_CIRCULAR_DEPENDENCY_LOAD_DEFERRING
+	if (Ar.IsLoading() || Ar.IsObjectReferenceCollector())
+	{
+		if (auto PlaceholderFunc = Cast<ULinkerPlaceholderFunction>(SignatureFunction))
+		{
+			PlaceholderFunc->AddReferencingProperty(this);
+		}
+	}
+#endif // USE_CIRCULAR_DEPENDENCY_LOAD_DEFERRING
 }
 
 bool UDelegateProperty::SameType(const UProperty* Other) const
@@ -113,6 +124,19 @@ bool UDelegateProperty::SameType(const UProperty* Other) const
 	return Super::SameType(Other) && (SignatureFunction == ((UDelegateProperty*)Other)->SignatureFunction);
 }
 
+void UDelegateProperty::BeginDestroy()
+{
+#if USE_CIRCULAR_DEPENDENCY_LOAD_DEFERRING
+	if (auto PlaceholderFunc = Cast<ULinkerPlaceholderFunction>(SignatureFunction))
+	{
+		PlaceholderFunc->RemovePropertyReference(this);
+	}
+#endif // USE_CIRCULAR_DEPENDENCY_LOAD_DEFERRING
+
+	Super::BeginDestroy();
+}
+
+
 IMPLEMENT_CORE_INTRINSIC_CLASS(UDelegateProperty, UProperty,
 	{
 		Class->EmitObjectReference(STRUCT_OFFSET(UDelegateProperty, SignatureFunction), TEXT("SignatureFunction"));
diff --git a/Engine/Source/Runtime/CoreUObject/Private/UObject/PropertyMulticastDelegate.cpp b/Engine/Source/Runtime/CoreUObject/Private/UObject/PropertyMulticastDelegate.cpp
--- a/Engine/Source/Runtime/CoreUObject/Private/UObject/PropertyMulticastDelegate.cpp
+++ b/Engine/Source/Runtime/CoreUObject/Private/UObject/PropertyMulticastDelegate.cpp
@@ -2,6 +2,7 @@
 
 #include "CoreUObjectPrivate.h"
 #include "PropertyHelper.h"
+#include "LinkerPlaceholderFunction.h"
 
 /*-----------------------------------------------------------------------------
 	UMulticastDelegateProperty.
@@ -247,11 +248,32 @@ const TCHAR* UMulticastDelegateProperty::ImportText_Remove( const TCHAR* Buffer,
 }
 
 
-
 void UMulticastDelegateProperty::Serialize( FArchive& Ar )
 {
 	Super::Serialize( Ar );
 	Ar << SignatureFunction;
+
+#if USE_CIRCULAR_DEPENDENCY_LOAD_DEFERRING
+	if (Ar.IsLoading() || Ar.IsObjectReferenceCollector())
+	{
+		if (auto PlaceholderFunc = Cast<ULinkerPlaceholderFunction>(SignatureFunction))
+		{
+			PlaceholderFunc->AddReferencingProperty(this);
+		}
+	}
+#endif // USE_CIRCULAR_DEPENDENCY_LOAD_DEFERRING
+}
+
+void UMulticastDelegateProperty::BeginDestroy()
+{
+#if USE_CIRCULAR_DEPENDENCY_LOAD_DEFERRING
+	if (auto PlaceholderFunc = Cast<ULinkerPlaceholderFunction>(SignatureFunction))
+	{
+		PlaceholderFunc->RemovePropertyReference(this);
+	}
+#endif // USE_CIRCULAR_DEPENDENCY_LOAD_DEFERRING
+
+	Super::BeginDestroy();
 }
 
 bool UMulticastDelegateProperty::SameType(const UProperty* Other) const
diff --git a/Engine/Source/Runtime/CoreUObject/Public/UObject/ScriptSerialization.h b/Engine/Source/Runtime/CoreUObject/Public/UObject/ScriptSerialization.h
--- a/Engine/Source/Runtime/CoreUObject/Public/UObject/ScriptSerialization.h
+++ b/Engine/Source/Runtime/CoreUObject/Public/UObject/ScriptSerialization.h
@@ -217,6 +217,7 @@
 		case EX_FinalFunction:
 		{
 			XFER_FUNC_POINTER;											// Stack node.
+			FIXUP_EXPR_OBJECT_POINTER(UStruct*);
 			while( SerializeExpr( iCode, Ar ) != EX_EndFunctionParms ); // Parms.
 			break;
 		}
@@ -229,6 +230,7 @@
 		case EX_CallMulticastDelegate:
 		{
 			XFER_FUNC_POINTER;											// Stack node.
+			FIXUP_EXPR_OBJECT_POINTER(UStruct*);
 			while( SerializeExpr( iCode, Ar ) != EX_EndFunctionParms ); // Parms.
 			break;
 		}
diff --git a/Engine/Source/Runtime/CoreUObject/Public/UObject/UnrealType.h b/Engine/Source/Runtime/CoreUObject/Public/UObject/UnrealType.h
--- a/Engine/Source/Runtime/CoreUObject/Public/UObject/UnrealType.h
+++ b/Engine/Source/Runtime/CoreUObject/Public/UObject/UnrealType.h
@@ -2606,6 +2606,7 @@ public:
 
 	// UObject interface
 	virtual void Serialize( FArchive& Ar ) override;
+	virtual void BeginDestroy() override;
 	// End of UObject interface
 
 	// UProperty interface
@@ -2652,6 +2653,7 @@ public:
 
 	// UObject interface
 	virtual void Serialize( FArchive& Ar ) override;
+	virtual void BeginDestroy() override;
 	// End of UObject interface
 
 	// UProperty interface
