[INTEGRATE] Change 2462032  by [REDACTED_PERSON][REDACTED_HANDLE] on [REDACTED_PHONE]:38

                     Async BulkData loading

diff --git a/Engine/Config/BaseEngine.ini b/Engine/Config/BaseEngine.ini
--- a/Engine/Config/BaseEngine.ini
+++ b/Engine/Config/BaseEngine.ini
@@ -1274,6 +1274,7 @@ CutdownPaths=%GAMEDIR%CutdownPackages
 ZeroEngineVersionWarning=True
 UseStrictEngineVersioning=True
 FlushStreamingOnGC=False
+MinBulkDataSizeForAsyncLoading=131072
 
 [Internationalization]
 +LocalizationPaths=../../../Engine/Content/Localization/Engine
@@ -2229,5 +2230,4 @@ NameForGameView5="GameView5"
 ; These are the defaults for HTML5 settings, and they need to be in the .ini since UBT reads the .ini settings, without instantiating the class
 [/Script/HTML5PlatformEditor.HTML5TargetSettings]
 HeapSizeDevelopment=1024
-HeapSizeShipping=1024
-DeployServerPort=53501
\ No newline at end of file
+HeapSizeShipping=1024DeployServerPort=53501
\ No newline at end of file
diff --git a/Engine/Source/Runtime/CoreUObject/Private/Serialization/BulkData.cpp b/Engine/Source/Runtime/CoreUObject/Private/Serialization/BulkData.cpp
--- a/Engine/Source/Runtime/CoreUObject/Private/Serialization/BulkData.cpp
+++ b/Engine/Source/Runtime/CoreUObject/Private/Serialization/BulkData.cpp
@@ -119,12 +119,22 @@ FUntypedBulkData::~FUntypedBulkData()
 	DECLARE_SCOPE_CYCLE_COUNTER(TEXT("FUntypedBulkData::~FUntypedBulkData"), STAT_UBD_Destructor, STATGROUP_Memory);
 
 	check( LockStatus == LOCKSTATUS_Unlocked );
+	if (SerializeFuture.IsValid())
+	{
+		WaitForAsyncLoading();
+	}
+
 	// Free memory.
-	if( bShouldFreeOnEmpty )
+	if (bShouldFreeOnEmpty)
 	{
-		FMemory::Free( BulkData );
+		FMemory::Free(BulkData);
 	}
-	BulkData = NULL;
+	if (BulkData != BulkDataAsync)
+	{
+		FMemory::Free(BulkDataAsync);
+	}
+	BulkData = nullptr;
+	BulkDataAsync = nullptr;
 	
 #if WITH_EDITOR
 	// Detach from archive.
@@ -332,6 +342,11 @@ bool FUntypedBulkData::IsBulkDataLoaded() const
 	return BulkData != NULL;
 }
 
+bool FUntypedBulkData::IsAsyncLoadingComplete()
+{
+	return SerializeFuture.IsValid() == false || SerializeFuture.WaitFor(FTimespan(0));
+}
+
 /**
 * Returns whether this bulk data is used
 * @return true if BULKDATA_Unused is not set
@@ -345,6 +360,12 @@ bool FUntypedBulkData::IsAvailableForUse() const
 	Data retrieval and manipulation.
 -----------------------------------------------------------------------------*/
 
+void FUntypedBulkData::ResetAsyncData()
+{
+	BulkDataAsync = nullptr;
+	SerializeFuture = TFuture<bool>();
+}
+
 /**
  * Retrieves a copy of the bulk data.
  *
@@ -386,6 +407,13 @@ void FUntypedBulkData::GetCopy( void** Dest, bool bDiscardInternalCopy )
 	// Passed in memory is NULL so we need to allocate some.
 	else
 	{
+		// Check for any pending async requests
+		if (!BulkData && SerializeFuture.IsValid())
+		{
+			WaitForAsyncLoading();
+			BulkData = BulkDataAsync;
+			ResetAsyncData();
+		}
 		// The data is already loaded so we can simply use a mempcy.
 		if( BulkData )
 		{
@@ -395,7 +423,8 @@ void FUntypedBulkData::GetCopy( void** Dest, bool bDiscardInternalCopy )
 			if( bDiscardInternalCopy && (CanLoadFromDisk()|| (BulkDataFlags & BULKDATA_SingleUse)) )
 			{
 				*Dest = BulkData;
-				BulkData = NULL;
+				BulkData = nullptr;
+				ResetAsyncData();
 			}
 			// Can't/ Don't discard so we need to allocate and copy.
 			else
@@ -607,6 +636,38 @@ bool FUntypedBulkData::ShouldFreeOnEmpty() const
 	Serialization.
 -----------------------------------------------------------------------------*/
 
+void FUntypedBulkData::StartSerializingBulkData(FArchive& Ar, UObject* Owner, int32 Idx, bool bPayloadInline)
+{
+	DECLARE_SCOPE_CYCLE_COUNTER(TEXT("FUntypedBulkData::StartSerializingBulkData"), STAT_UBD_Serialize, STATGROUP_Memory);
+	check(SerializeFuture.IsValid() == false);	
+
+	// Async
+	SerializeFuture = Async<bool>(EAsyncExecution::ThreadPool, [=]()
+	{
+		BulkDataAsync = GetBulkDataResourceMemory(Owner, Idx);
+		if (!BulkDataAsync)
+		{
+			BulkDataAsync = FMemory::Realloc(BulkDataAsync, GetBulkDataSize());
+		}
+
+		FArchive* Ar = IFileManager::Get().CreateFileReader(*Filename, FILEREAD_Silent);
+		checkf(Ar != NULL, TEXT("Attempted to load bulk data from an invalid filename '%s'."), *Filename);
+
+		// Seek to the beginning of the bulk data in the file.
+		Ar->Seek(BulkDataOffsetInFile);
+		SerializeBulkData(*Ar, BulkDataAsync);
+		delete Ar;
+
+		return true;
+	});
+
+	// Skip bulk data in this archive
+	if (bPayloadInline)
+	{
+		Ar.Seek(Ar.Tell() + BulkDataSizeOnDisk);
+	}
+}
+
 /**
 * Serialize function used to serialize this bulk data structure.
 *
@@ -706,25 +767,43 @@ void FUntypedBulkData::Serialize( FArchive& Ar, UObject* Owner, int32 Idx )
 				BulkDataOffsetInFile += Owner->GetLinker()->Summary.BulkDataStartOffset;
 			}
 
-// 			check( (bPayloadInline && BulkDataOffsetInFile == Ar.Tell()) || 
-// 				   (!bPayloadInline && BulkDataOffsetInFile > Ar.Tell()));
-			
+			// Minimum bulk data size to start async loading
+			static struct FMinBulkDataSizeForAsyncLoadingSetting
+			{
+				int32 Value;
+				FMinBulkDataSizeForAsyncLoadingSetting()
+					: Value(131072) // 128K by default
+				{
+					GConfig->GetInt(TEXT("Core.System"), TEXT("MinBulkDataSizeForAsyncLoading"), Value, GEngineIni);
+				}
+			} MinBulkDataSizeForAsyncLoading;
+
 			// We're allowing defered serialization.
 			if( Ar.IsAllowingLazyLoading() && Owner != NULL)
 			{
 				Linker = Owner->GetLinker();
-
+#if WITH_EDITOR
+				if (Linker)
+#else
+				if (Linker.IsValid())
+#endif
+				{
+					Filename = Linker->Filename;
+				}
 #if WITH_EDITOR
 				check(Linker);
 				Ar.AttachBulkData( Owner, this );
 				AttachedAr = &Ar;
 				
 #else
-				check(Linker.IsValid());
-				Filename = Linker->Filename;
+				check(Linker.IsValid());				
 #endif // WITH_EDITOR
-				// only skip over payload, if it's stored inline
-				if (bPayloadInline)
+				if (!Filename.IsEmpty() && GetBulkDataSize() > MinBulkDataSizeForAsyncLoading.Value && MinBulkDataSizeForAsyncLoading.Value >= 0)
+				{
+					// Start serializing immediately
+					StartSerializingBulkData(Ar, Owner, Idx, bPayloadInline);
+				}
+				else if (bPayloadInline)
 				{
 					// Force non-lazy loading of inline bulk data to prevent PostLoad spikes.
 					// Memory for bulk data can come from preallocated GPU-accessible resource memory or default to system memory
@@ -740,31 +819,42 @@ void FUntypedBulkData::Serialize( FArchive& Ar, UObject* Owner, int32 Idx )
 			// Serialize the bulk data right away.
 			else
 			{
-				// memory for bulk data can come from preallocated GPU-accessible resource memory or default to system memory
-				BulkData = GetBulkDataResourceMemory(Owner,Idx);
-				if( !BulkData )
+				if (Owner && Owner->GetLinker())
 				{
-					BulkData = FMemory::Realloc( BulkData, GetBulkDataSize() );
+					Filename = Owner->GetLinker()->Filename;
 				}
-				
-				if (bPayloadInline)
-				{
-					// if the payload is stored inline, just serialize it
-					SerializeBulkData( Ar, BulkData );
-				}
-				else
+				if (!Filename.IsEmpty() && GetBulkDataSize() > MinBulkDataSizeForAsyncLoading.Value && MinBulkDataSizeForAsyncLoading.Value >= 0)
 				{
-					// if the payload is NOT stored inline ...
-					
-					// store the current file offset
-					int64 CurOffset = Ar.Tell();
-					// seek to the location in the file where the payload is stored
-					Ar.Seek(BulkDataOffsetInFile);
-					// serialize the payload
-					SerializeBulkData( Ar, BulkData );
-					// seek to the location we came from
-					Ar.Seek(CurOffset);
+					StartSerializingBulkData(Ar, Owner, Idx, bPayloadInline);
 				}
+  			else
+  			{
+				  // memory for bulk data can come from preallocated GPU-accessible resource memory or default to system memory
+				  BulkData = GetBulkDataResourceMemory(Owner,Idx);
+				  if( !BulkData )
+				  {
+					  BulkData = FMemory::Realloc( BulkData, GetBulkDataSize() );
+				  }
+				  
+				  if (bPayloadInline)
+				  {
+					  // if the payload is stored inline, just serialize it
+					  SerializeBulkData( Ar, BulkData );
+				  }
+				  else
+				  {
+					  // if the payload is NOT stored inline ...
+					  
+					  // store the current file offset
+					  int64 CurOffset = Ar.Tell();
+					  // seek to the location in the file where the payload is stored
+					  Ar.Seek(BulkDataOffsetInFile);
+					  // serialize the payload
+					  SerializeBulkData( Ar, BulkData );
+					  // seek to the location we came from
+					  Ar.Seek(CurOffset);
+				  }
+  			}
 			}
 		}
 		// We're saving to the persistent archive.
@@ -973,16 +1063,17 @@ void FUntypedBulkData::Copy( const FUntypedBulkData& Other )
  */
 void FUntypedBulkData::InitializeMemberVariables()
 {
-	BulkDataFlags				= BULKDATA_None;
-	ElementCount				= 0;
-	BulkDataOffsetInFile		= INDEX_NONE;
-	BulkDataSizeOnDisk			= INDEX_NONE;
-	BulkData					= NULL;
-	LockStatus					= LOCKSTATUS_Unlocked;
-	bShouldFreeOnEmpty			= true;
-	Linker						= NULL;
+	BulkDataFlags = BULKDATA_None;
+	ElementCount = 0;
+	BulkDataOffsetInFile = INDEX_NONE;
+	BulkDataSizeOnDisk = INDEX_NONE;
+	BulkData = nullptr;
+	BulkDataAsync = nullptr;
+	LockStatus = LOCKSTATUS_Unlocked;
+	bShouldFreeOnEmpty = true;
+	Linker = nullptr;
 #if WITH_EDITOR
-	AttachedAr					= NULL;
+	AttachedAr = nullptr;
 #endif
 }
 
@@ -1103,18 +1194,50 @@ void FUntypedBulkData::MakeSureBulkDataIsLoaded()
 	// Nothing to do if data is already loaded.
 	if( !BulkData )
 	{
-		// Allocate memory for bulk data.
-		BulkData = FMemory::Malloc( GetBulkDataSize() );
-
-		// Only load if there is something to load. E.g. we might have just created the bulk data array
-		// in which case it starts out with a size of zero.
-		if( GetBulkDataSize() > 0 )
+		// Look for async request first
+		if (SerializeFuture.IsValid())
+		{
+			WaitForAsyncLoading();
+			BulkData = BulkDataAsync;
+			ResetAsyncData();
+		}
+		else
 		{
-			LoadDataIntoMemory( BulkData );
+			// Allocate memory for bulk data.
+			BulkData = FMemory::Malloc(GetBulkDataSize());
+
+			// Only load if there is something to load. E.g. we might have just created the bulk data array
+			// in which case it starts out with a size of zero.
+			if (GetBulkDataSize() > 0)
+			{
+				LoadDataIntoMemory(BulkData);
+			}
 		}
 	}
 }
 
+void FUntypedBulkData::WaitForAsyncLoading()
+{
+	check(SerializeFuture.IsValid());
+	DECLARE_SCOPE_CYCLE_COUNTER(TEXT("FUntypedBulkData::WaitForAsyncLoading"), STAT_UBD_WaitForAsyncLoading, STATGROUP_Memory);
+	while (!SerializeFuture.WaitFor(FTimespan(0, 0, 0, 0, 1000)))
+	{
+		UE_LOG(LogSerialization, Warning, TEXT("Waiting for %s bulk data (%lld) to be loaded longer than 1000ms"), *Filename, GetBulkDataSize());
+	}
+	check(BulkDataAsync);
+}
+
+bool FUntypedBulkData::FlushAsyncLoading(void* Dest)
+{
+	bool bIsLoadingAsync = SerializeFuture.IsValid();
+	if (bIsLoadingAsync)
+	{
+		WaitForAsyncLoading();
+		FMemory::Memcpy(Dest, BulkDataAsync, GetBulkDataSize());
+	}
+	return bIsLoadingAsync;
+}
+
 /**
  * Loads the data from disk into the specified memory block. This requires us still being attached to an
  * archive we can use for serialization.
@@ -1123,6 +1246,12 @@ void FUntypedBulkData::MakeSureBulkDataIsLoaded()
  */
 void FUntypedBulkData::LoadDataIntoMemory( void* Dest )
 {
+	// Try flushing async loading before attempting to load
+	if (FlushAsyncLoading(Dest))
+	{
+		return;
+	}
+
 #if WITH_EDITOR
 	checkf( AttachedAr, TEXT( "Attempted to load bulk data without an attached archive. Most likely the bulk data was loaded twice on console, which is not supported" ) );
 
diff --git a/Engine/Source/Runtime/CoreUObject/Public/Serialization/BulkData.h b/Engine/Source/Runtime/CoreUObject/Public/Serialization/BulkData.h
--- a/Engine/Source/Runtime/CoreUObject/Public/Serialization/BulkData.h
+++ b/Engine/Source/Runtime/CoreUObject/Public/Serialization/BulkData.h
@@ -1,9 +1,10 @@
 // Copyright 1998-2015 Epic Games, Inc. All Rights Reserved.
 
-
 #ifndef _UNBULKDATA_H
 #define _UNBULKDATA_H
 
+#include "Async/Async.h"
+
 /**
  * Flags serialized with the bulk data.
  */
@@ -177,6 +178,13 @@ struct COREUOBJECT_API FUntypedBulkData
 	 */
 	bool IsBulkDataLoaded() const;
 
+	/**
+	* Returns whether the bulk data asynchronous load has completed.
+	*
+	* @return true if bulk data has been loaded or async loading was not used to load this data, false otherwise
+	*/
+	bool IsAsyncLoadingComplete();
+
 	/**
 	* Returns whether this bulk data is used
 	* @return true if BULKDATA_Unused is not set
@@ -362,6 +370,18 @@ private:
 	 */
 	void LoadDataIntoMemory( void* Dest );
 
+	/** Starts serializing bulk data asynchronously */
+	void StartSerializingBulkData(FArchive& Ar, UObject* Owner, int32 Idx, bool bPayloadInline);
+
+	/** Flushes any pending async load of bulk data  and copies the data to Dest buffer*/
+	bool FlushAsyncLoading(void* Dest);
+
+	/** Waits until pending async load finishes */
+	void WaitForAsyncLoading();
+
+	/** Resets async loading state */
+	void ResetAsyncData();
+	
 	/*-----------------------------------------------------------------------------
 		Member variables.
 	-----------------------------------------------------------------------------*/
@@ -377,9 +397,13 @@ private:
 
 	/** Pointer to cached bulk data																						*/
 	void*				BulkData;
+	/** Pointer to cached async bulk data																						*/
+	void*				BulkDataAsync;
 	/** Current lock status																								*/
 	uint32				LockStatus;
-	
+	/** Async helper for loading bulk data on a separate thread */
+	TFuture<bool> SerializeFuture;
+
 protected:
 	/** true when data has been allocated internally by the bulk data and does not come from a preallocated resource	*/
 	bool				bShouldFreeOnEmpty;
