[INTEGRATE] Change 2462119  by [REDACTED_PERSON][REDACTED_HANDLE] on [REDACTED_PHONE]:42

                     Disabling async bulk data in non-cooked builds

diff --git a/Engine/Source/Runtime/CoreUObject/Private/Serialization/BulkData.cpp b/Engine/Source/Runtime/CoreUObject/Private/Serialization/BulkData.cpp
--- a/Engine/Source/Runtime/CoreUObject/Private/Serialization/BulkData.cpp
+++ b/Engine/Source/Runtime/CoreUObject/Private/Serialization/BulkData.cpp
@@ -798,7 +798,7 @@ void FUntypedBulkData::Serialize( FArchive& Ar, UObject* Owner, int32 Idx )
 #else
 				check(Linker.IsValid());				
 #endif // WITH_EDITOR
-				if (!Filename.IsEmpty() && GetBulkDataSize() > MinBulkDataSizeForAsyncLoading.Value && MinBulkDataSizeForAsyncLoading.Value >= 0)
+				if (FPlatformProperties::RequiresCookedData() && !Filename.IsEmpty() && GetBulkDataSize() > MinBulkDataSizeForAsyncLoading.Value && MinBulkDataSizeForAsyncLoading.Value >= 0)
 				{
 					// Start serializing immediately
 					StartSerializingBulkData(Ar, Owner, Idx, bPayloadInline);
@@ -823,7 +823,7 @@ void FUntypedBulkData::Serialize( FArchive& Ar, UObject* Owner, int32 Idx )
 				{
 					Filename = Owner->GetLinker()->Filename;
 				}
-				if (!Filename.IsEmpty() && GetBulkDataSize() > MinBulkDataSizeForAsyncLoading.Value && MinBulkDataSizeForAsyncLoading.Value >= 0)
+				if (FPlatformProperties::RequiresCookedData() && !Filename.IsEmpty() && GetBulkDataSize() > MinBulkDataSizeForAsyncLoading.Value && MinBulkDataSizeForAsyncLoading.Value >= 0)
 				{
 					StartSerializingBulkData(Ar, Owner, Idx, bPayloadInline);
 				}
