[INTEGRATE] Change 2462246  by [REDACTED_PERSON][REDACTED_HANDLE] on [REDACTED_PHONE]:21

                    Fixing bigger mem usage due to bulk data loading changes

diff --git a/Engine/Source/Runtime/CoreUObject/Private/Serialization/BulkData.cpp b/Engine/Source/Runtime/CoreUObject/Private/Serialization/BulkData.cpp
--- a/Engine/Source/Runtime/CoreUObject/Private/Serialization/BulkData.cpp
+++ b/Engine/Source/Runtime/CoreUObject/Private/Serialization/BulkData.cpp
@@ -798,22 +798,25 @@ void FUntypedBulkData::Serialize( FArchive& Ar, UObject* Owner, int32 Idx )
 #else
 				check(Linker.IsValid());				
 #endif // WITH_EDITOR
-				if (FPlatformProperties::RequiresCookedData() && !Filename.IsEmpty() && GetBulkDataSize() > MinBulkDataSizeForAsyncLoading.Value && MinBulkDataSizeForAsyncLoading.Value >= 0)
-				{
-					// Start serializing immediately
-					StartSerializingBulkData(Ar, Owner, Idx, bPayloadInline);
-				}
-				else if (bPayloadInline)
+				if (bPayloadInline)
 				{
-					// Force non-lazy loading of inline bulk data to prevent PostLoad spikes.
-					// Memory for bulk data can come from preallocated GPU-accessible resource memory or default to system memory
-					BulkData = GetBulkDataResourceMemory(Owner, Idx);
-					if (!BulkData)
+					if (FPlatformProperties::RequiresCookedData() && !Filename.IsEmpty() && GetBulkDataSize() > MinBulkDataSizeForAsyncLoading.Value && MinBulkDataSizeForAsyncLoading.Value >= 0)
 					{
-						BulkData = FMemory::Realloc(BulkData, GetBulkDataSize());
+						// Start serializing immediately
+						StartSerializingBulkData(Ar, Owner, Idx, bPayloadInline);
+					}
+					else
+					{
+						// Force non-lazy loading of inline bulk data to prevent PostLoad spikes.
+						// Memory for bulk data can come from preallocated GPU-accessible resource memory or default to system memory
+						BulkData = GetBulkDataResourceMemory(Owner, Idx);
+						if (!BulkData)
+						{
+							BulkData = FMemory::Realloc(BulkData, GetBulkDataSize());
+						}
+						// if the payload is stored inline, just serialize it
+						SerializeBulkData(Ar, BulkData);
 					}
-					// if the payload is stored inline, just serialize it
-					SerializeBulkData(Ar, BulkData);
 				}
 			}
 			// Serialize the bulk data right away.
