BulkData improvements:

- Added the ability to force stream bulk data
- Added SerializeElements to avoid the cost of SerializeElement for each element.

diff --git a/Engine/Source/Runtime/CoreUObject/Private/Serialization/BulkData.cpp b/Engine/Source/Runtime/CoreUObject/Private/Serialization/BulkData.cpp
--- a/Engine/Source/Runtime/CoreUObject/Private/Serialization/BulkData.cpp
+++ b/Engine/Source/Runtime/CoreUObject/Private/Serialization/BulkData.cpp
@@ -668,6 +668,26 @@ void FUntypedBulkData::StartSerializingBulkData(FArchive& Ar, UObject* Owner, in
 	}
 }
 
+bool FUntypedBulkData::ShouldStreamBulkData()
+{
+	// Minimum bulk data size to start async loading
+	static struct FMinBulkDataSizeForAsyncLoadingSetting
+	{
+		int32 Value;
+		FMinBulkDataSizeForAsyncLoadingSetting()
+			: Value(131072) // 128K by default
+		{
+			GConfig->GetInt(TEXT("Core.System"), TEXT("MinBulkDataSizeForAsyncLoading"), Value, GEngineIni);
+		}
+	} MinBulkDataSizeForAsyncLoading;
+
+	const bool bForceStream = !!(BulkDataFlags & BULKDATA_ForceStreamPayload);
+
+	return (FPlatformProperties::RequiresCookedData() && !Filename.IsEmpty() &&
+		(bForceStream || GetBulkDataSize() > MinBulkDataSizeForAsyncLoading.Value) &&
+		MinBulkDataSizeForAsyncLoading.Value >= 0);
+}
+
 /**
 * Serialize function used to serialize this bulk data structure.
 *
@@ -759,7 +779,7 @@ void FUntypedBulkData::Serialize( FArchive& Ar, UObject* Owner, int32 Idx )
 			Ar << BulkDataOffsetInFile;
 
 			// determine whether the payload is stored inline or at the end of the file
-			bool bPayloadInline = !(BulkDataFlags&BULKDATA_PayloadAtEndOfFile);
+			const bool bPayloadInline = !(BulkDataFlags&BULKDATA_PayloadAtEndOfFile);
 
 			// fix up the file offset, but only if not stored inline
 			if (Owner != NULL && Owner->GetLinker() && !bPayloadInline)
@@ -767,17 +787,6 @@ void FUntypedBulkData::Serialize( FArchive& Ar, UObject* Owner, int32 Idx )
 				BulkDataOffsetInFile += Owner->GetLinker()->Summary.BulkDataStartOffset;
 			}
 
-			// Minimum bulk data size to start async loading
-			static struct FMinBulkDataSizeForAsyncLoadingSetting
-			{
-				int32 Value;
-				FMinBulkDataSizeForAsyncLoadingSetting()
-					: Value(131072) // 128K by default
-				{
-					GConfig->GetInt(TEXT("Core.System"), TEXT("MinBulkDataSizeForAsyncLoading"), Value, GEngineIni);
-				}
-			} MinBulkDataSizeForAsyncLoading;
-
 			// We're allowing defered serialization.
 			if( Ar.IsAllowingLazyLoading() && Owner != NULL)
 			{
@@ -800,7 +809,7 @@ void FUntypedBulkData::Serialize( FArchive& Ar, UObject* Owner, int32 Idx )
 #endif // WITH_EDITOR
 				if (bPayloadInline)
 				{
-					if (FPlatformProperties::RequiresCookedData() && !Filename.IsEmpty() && GetBulkDataSize() > MinBulkDataSizeForAsyncLoading.Value && MinBulkDataSizeForAsyncLoading.Value >= 0)
+					if (ShouldStreamBulkData())
 					{
 						// Start serializing immediately
 						StartSerializingBulkData(Ar, Owner, Idx, bPayloadInline);
@@ -826,7 +835,7 @@ void FUntypedBulkData::Serialize( FArchive& Ar, UObject* Owner, int32 Idx )
 				{
 					Filename = Owner->GetLinker()->Filename;
 				}
-				if (FPlatformProperties::RequiresCookedData() && !Filename.IsEmpty() && GetBulkDataSize() > MinBulkDataSizeForAsyncLoading.Value && MinBulkDataSizeForAsyncLoading.Value >= 0)
+				if (ShouldStreamBulkData())
 				{
 					StartSerializingBulkData(Ar, Owner, Idx, bPayloadInline);
 				}
@@ -1080,6 +1089,15 @@ void FUntypedBulkData::InitializeMemberVariables()
 #endif
 }
 
+void FUntypedBulkData::SerializeElements(FArchive& Ar, void* Data)
+{
+	// Serialize each element individually.				
+	for (int32 ElementIndex = 0; ElementIndex < ElementCount; ElementIndex++)
+	{
+		SerializeElement(Ar, Data, ElementIndex);
+	}
+}
+
 /**
  * Serialize just the bulk data portion to/ from the passed in memory.
  *
@@ -1152,11 +1170,8 @@ void FUntypedBulkData::SerializeBulkData( FArchive& Ar, void* Data )
 				FMemoryReader MemoryReader( SerializedData, true );
 				MemoryReader.SetByteSwapping( Ar.ForceByteSwapping() );
 
-				// Serialize each element individually via memory reader.				
-				for( int32 ElementIndex=0; ElementIndex<ElementCount; ElementIndex++ )
-				{
-					SerializeElement( MemoryReader, Data, ElementIndex );
-				}
+				// Serialize each element individually via memory reader.
+				SerializeElements(MemoryReader, Data);
 			}
 			// Saving, data is uncompressed in memory and needs to be compressed.
 			else if( Ar.IsSaving() )
@@ -1165,11 +1180,8 @@ void FUntypedBulkData::SerializeBulkData( FArchive& Ar, void* Data )
 				FMemoryWriter MemoryWriter( SerializedData, true );
 				MemoryWriter.SetByteSwapping( Ar.ForceByteSwapping() );
 
-				// Serialize each element individually via memory writer.				
-				for( int32 ElementIndex=0; ElementIndex<ElementCount; ElementIndex++ )
-				{
-					SerializeElement( MemoryWriter, Data, ElementIndex );
-				}
+				// Serialize each element individually via memory writer.			
+				SerializeElements(MemoryWriter, Data);
 
 				// Serialize data with passed in archive and compress.
 				Ar.SerializeCompressed( SerializedData.GetData(), SerializedData.Num(), GetDecompressionFlags() );
@@ -1179,10 +1191,7 @@ void FUntypedBulkData::SerializeBulkData( FArchive& Ar, void* Data )
 		else
 		{
 			// We can use the passed in archive if we're not compressing the data.
-			for( int32 ElementIndex=0; ElementIndex<ElementCount; ElementIndex++ )
-			{
-				SerializeElement( Ar, Data, ElementIndex );
-			}
+			SerializeElements(Ar, Data);
 		}
 	}
 }
diff --git a/Engine/Source/Runtime/CoreUObject/Public/Serialization/BulkData.h b/Engine/Source/Runtime/CoreUObject/Public/Serialization/BulkData.h
--- a/Engine/Source/Runtime/CoreUObject/Public/Serialization/BulkData.h
+++ b/Engine/Source/Runtime/CoreUObject/Public/Serialization/BulkData.h
@@ -26,6 +26,8 @@ enum EBulkDataFlags
 	BULKDATA_ForceInlinePayload					= 1<<6,
 	/** Flag to check if either compression mode is specified						*/
 	BULKDATA_SerializeCompressed				= (BULKDATA_SerializeCompressedZLIB),
+	/** Forces the payload to be always streamed, regardless of its size */
+	BULKDATA_ForceStreamPayload = 1 << 7,
 
 };
 
@@ -312,6 +314,16 @@ struct COREUOBJECT_API FUntypedBulkData
 	-----------------------------------------------------------------------------*/
 
 protected:
+
+	/**
+	* Serializes all elements, a single element at a time, allowing backward compatible serialization
+	* and endian swapping to be performed.
+	*
+	* @param Ar			Archive to serialize with
+	* @param Data			Base pointer to data
+	*/
+	virtual void SerializeElements(FArchive& Ar, void* Data);
+
 	/**
 	 * Serializes a single element at a time, allowing backward compatible serialization
 	 * and endian swapping to be performed. Needs to be overloaded by derived classes.
@@ -382,6 +394,9 @@ private:
 	/** Resets async loading state */
 	void ResetAsyncData();
 	
+	/** Returns true if bulk data should be loaded asynchronously */
+	bool ShouldStreamBulkData();
+
 	/*-----------------------------------------------------------------------------
 		Member variables.
 	-----------------------------------------------------------------------------*/
