UE-10111: Hot Reload crash with Scene Component tied to Actor class

We've changed the way vtable ptr are obtained for UClasses. Right now we use special empty constructor to get it, so we don't have to use normal one which was causing troubles. This special vtable helper constructor is generated automatically for most cases, but sometimes it's not possible (e.g. no default constructor for members types). In such case there is a possibility to override that generated constructor and write a custom one. The signature is UCustomClass::UCustomClass(FVTableHelper& Helper) and it needs to call base class'es constructor with the same parameter. Please do not use these constructors for anything else. There is no guarantee the object will be in correct state after calling this.

You can still disable this functionality if something breaks, just set WITH_HOT_RELOAD_CTORS to 0.

diff --git a/Engine/Source/Developer/GameplayDebugger/Classes/GameplayDebuggingControllerComponent.h b/Engine/Source/Developer/GameplayDebugger/Classes/GameplayDebuggingControllerComponent.h
--- a/Engine/Source/Developer/GameplayDebugger/Classes/GameplayDebuggingControllerComponent.h
+++ b/Engine/Source/Developer/GameplayDebugger/Classes/GameplayDebuggingControllerComponent.h
@@ -17,6 +17,11 @@ class GAMEPLAYDEBUGGER_API UGameplayDebuggingControllerComponent : public UActor
 {
 	GENERATED_UCLASS_BODY()
 
+#if WITH_HOT_RELOAD && WITH_HOT_RELOAD_CTORS
+	/** DO NOT USE. This constructor is for internal usage only for hot-reload purposes. */
+	UGameplayDebuggingControllerComponent(FVTableHelper& Helper);
+#endif // WITH_HOT_RELOAD && WITH_HOT_RELOAD_CTORS
+
 	virtual void TickComponent(float DeltaTime, enum ELevelTick TickType, FActorComponentTickFunction *ThisTickFunction) override;
 	virtual void BeginDestroy() override;
 	virtual void OnRegister() override;
diff --git a/Engine/Source/Developer/GameplayDebugger/Private/GameplayDebuggingControllerComponent.cpp b/Engine/Source/Developer/GameplayDebugger/Private/GameplayDebuggingControllerComponent.cpp
--- a/Engine/Source/Developer/GameplayDebugger/Private/GameplayDebuggingControllerComponent.cpp
+++ b/Engine/Source/Developer/GameplayDebugger/Private/GameplayDebuggingControllerComponent.cpp
@@ -34,6 +34,15 @@ UGameplayDebuggingControllerComponent::UGameplayDebuggingControllerComponent(con
 #endif //!(UE_BUILD_SHIPPING || UE_BUILD_TEST)
 }
 
+#if WITH_HOT_RELOAD && WITH_HOT_RELOAD_CTORS
+UGameplayDebuggingControllerComponent::UGameplayDebuggingControllerComponent(FVTableHelper& Helper)
+	: Super(Helper)
+	, KeyPressActivationTime(0.0f)
+{
+
+}
+#endif // WITH_HOT_RELOAD && WITH_HOT_RELOAD_CTORS
+
 void UGameplayDebuggingControllerComponent::OnRegister()
 {
 	Super::OnRegister();
diff --git a/Engine/Source/Editor/UnrealEd/Private/HierarchicalLOD.cpp b/Engine/Source/Editor/UnrealEd/Private/HierarchicalLOD.cpp
--- a/Engine/Source/Editor/UnrealEd/Private/HierarchicalLOD.cpp
+++ b/Engine/Source/Editor/UnrealEd/Private/HierarchicalLOD.cpp
@@ -81,6 +81,14 @@ FHierarchicalLODBuilder::FHierarchicalLODBuilder(class UWorld* InWorld)
 :	World(InWorld)
 {}
 
+#if WITH_HOT_RELOAD && WITH_HOT_RELOAD_CTORS
+FHierarchicalLODBuilder::FHierarchicalLODBuilder()
+	: World(nullptr)
+{
+	EnsureRetrievingVTablePtr();
+}
+#endif // WITH_HOT_RELOAD && WITH_HOT_RELOAD_CTORS
+
 // build hierarchical cluster
 void FHierarchicalLODBuilder::Build()
 {
diff --git a/Engine/Source/Editor/UnrealEd/Public/HierarchicalLOD.h b/Engine/Source/Editor/UnrealEd/Public/HierarchicalLOD.h
--- a/Engine/Source/Editor/UnrealEd/Public/HierarchicalLOD.h
+++ b/Engine/Source/Editor/UnrealEd/Public/HierarchicalLOD.h
@@ -51,6 +51,11 @@ struct UNREALED_API FHierarchicalLODBuilder
 {
 	FHierarchicalLODBuilder(class UWorld* InWorld);
 
+#if WITH_HOT_RELOAD && WITH_HOT_RELOAD_CTORS
+	/** DO NOT USE. This constructor is for internal usage only for hot-reload purposes. */
+	FHierarchicalLODBuilder();
+#endif // WITH_HOT_RELOAD && WITH_HOT_RELOAD_CTORS
+
 	// build hierarchical cluster
 	void Build();
 
diff --git a/Engine/Source/Programs/UnrealHeaderTool/Private/CodeGenerator.cpp b/Engine/Source/Programs/UnrealHeaderTool/Private/CodeGenerator.cpp
--- a/Engine/Source/Programs/UnrealHeaderTool/Private/CodeGenerator.cpp
+++ b/Engine/Source/Programs/UnrealHeaderTool/Private/CodeGenerator.cpp
@@ -1867,14 +1867,95 @@ void ExportCopyConstructorDefinition(FStringOutputDevice& Out, FClass* Class, co
 	Out.Logf(TEXT("public:\r\n"));
 }
 
+#if WITH_HOT_RELOAD_CTORS
+/**
+ * Generates vtable helper constructor declaration.
+ *
+ * @param Out Output device to generate to.
+ * @param Class Class to generate constructor for.
+ * @param API API string for this constructor.
+ */
+void ExportVTableHelperConstructorDecl(FStringOutputDevice& Out, FClass* Class, const FString& API)
+{
+	auto* ClassData = GScriptHelper.FindClassData(Class);
+	if (!ClassData->bCustomVTableHelperConstructorDeclared)
+	{
+		Out.Logf(TEXT("%sDECLARE_VTABLE_PTR_HELPER_CTOR(%s_API, %s);" LINE_TERMINATOR), FCString::Spc(4), *API, NameLookupCPP.GetNameCPP(Class));
+	}
+}
+
+/**
+ * Generates vtable helper ctor caller.
+ *
+ * @param Out Output device to generate to.
+ * @param Class Class to generate for.
+ */
+void ExportVTableHelperCtorCaller(FStringOutputDevice& Out, FClass* Class)
+{
+	Out.Logf(TEXT("DEFINE_VTABLE_PTR_HELPER_CTOR_CALLER(%s);" LINE_TERMINATOR), NameLookupCPP.GetNameCPP(Class));
+}
+
+/**
+ * Generates vtable helper ctor caller dummy.
+ *
+ * @param Out Output device to generate to.
+ */
+void ExportVTableHelperCtorCallerDummy(FStringOutputDevice& Out)
+{
+	Out.Logf(TEXT("DEFINE_VTABLE_PTR_HELPER_CTOR_CALLER_DUMMY();" LINE_TERMINATOR));
+}
+
+/**
+ * Generates vtable helper constructor body.
+ *
+ * @param Out Output device to generate to.
+ * @param Class Class to generate constructor for.
+ * @param API API string for this constructor.
+ */
+void ExportVTableHelperConstructorBody(FStringOutputDevice& Out, FClass* Class)
+{
+	auto* ClassData = GScriptHelper.FindClassData(Class);
+	if (!ClassData->bCustomVTableHelperConstructorDeclared)
+	{
+		Out.Logf(TEXT("%sDEFINE_VTABLE_PTR_HELPER_CTOR(%s);" LINE_TERMINATOR), FCString::Spc(4), NameLookupCPP.GetNameCPP(Class));
+	}
+}
+
+/**
+ * Generates vtable helper caller and eventual constructor body.
+ *
+ * @param Out Output device to generate to.
+ * @param Class Class to generate for.
+ * @param API API string.
+ * @param bExportVTableConstructors Tells if we are going to export constructors as well.
+ */
+void ExportVTableHelperCtorAndCaller(FStringOutputDevice& Out, FClass* Class, const FString& API, bool bExportVTableConstructors)
+{
+	if (bExportVTableConstructors)
+	{
+		ExportVTableHelperConstructorDecl(Out, Class, API);
+		ExportVTableHelperCtorCaller(Out, Class);
+	}
+	else
+	{
+		ExportVTableHelperCtorCallerDummy(Out);
+	}
+}
+#endif // WITH_HOT_RELOAD_CTORS
+
 /**
  * Generates standard constructor declaration.
  *
  * @param Out Output device to generate to.
  * @param Class Class to generate constructor for.
  * @param API API string for this constructor.
+ * @param bExportVTableConstructors Export VTable constructors.
  */
-void ExportStandardConstructorsMacro(FStringOutputDevice& Out, FClass* Class, const FString& API)
+void ExportStandardConstructorsMacro(FStringOutputDevice& Out, FClass* Class, const FString& API
+#if WITH_HOT_RELOAD_CTORS
+		, bool bExportVTableConstructors
+#endif // WITH_HOT_RELOAD_CTORS
+)
 {
 	if (!Class->HasAnyClassFlags(CLASS_CustomConstructor))
 	{
@@ -1882,6 +1963,11 @@ void ExportStandardConstructorsMacro(FStringOutputDevice& Out, FClass* Class, co
 		Out.Logf(TEXT("%s%s_API %s(const FObjectInitializer& ObjectInitializer);\r\n"), FCString::Spc(4), *API, NameLookupCPP.GetNameCPP(Class));
 	}
 	Out.Logf(TEXT("%sDEFINE_DEFAULT_OBJECT_INITIALIZER_CONSTRUCTOR_CALL(%s)\r\n"), FCString::Spc(4), NameLookupCPP.GetNameCPP(Class));
+
+#if WITH_HOT_RELOAD_CTORS
+	ExportVTableHelperCtorAndCaller(Out, Class, API, bExportVTableConstructors);
+#endif // WITH_HOT_RELOAD_CTORS
+
 	ExportCopyConstructorDefinition(Out, Class, API);
 }
 
@@ -1957,10 +2043,20 @@ void ExportDefaultConstructorCallDefinition(FStringOutputDevice& Out, FClass* Cl
  * @param Out Output device to generate to.
  * @param Class Class to generate constructor for.
  * @param API API string for this constructor.
+ * @param bExportVTableConstructors Export VTable constructors.
  */
-void ExportEnhancedConstructorsMacro(FStringOutputDevice& Out, FClass* Class, const FString& API)
+void ExportEnhancedConstructorsMacro(FStringOutputDevice& Out, FClass* Class, const FString& API
+#if WITH_HOT_RELOAD_CTORS
+		, bool bExportVTableConstructors
+#endif // WITH_HOT_RELOAD_CTORS
+)
 {
 	ExportConstructorDefinition(Out, Class, API);
+	
+#if WITH_HOT_RELOAD_CTORS
+	ExportVTableHelperCtorAndCaller(Out, Class, API, bExportVTableConstructors);
+#endif // WITH_HOT_RELOAD_CTORS
+	
 	ExportDefaultConstructorCallDefinition(Out, Class);
 }
 
@@ -2015,8 +2111,18 @@ void FNativeClassHeaderGenerator::ExportConstructorsMacros(const FString& Constr
 	FString StdMacroName = ConstructorsMacroPrefix + TEXT("_STANDARD_CONSTRUCTORS");
 	FString EnhMacroName = ConstructorsMacroPrefix + TEXT("_ENHANCED_CONSTRUCTORS");
 
+#if WITH_HOT_RELOAD_CTORS
+	ExportStandardConstructorsMacro(StdMacro, Class, APIArg, bExportVTableConstructors);
+	ExportEnhancedConstructorsMacro(EnhMacro, Class, APIArg, bExportVTableConstructors);
+
+	if (bExportVTableConstructors)
+	{
+		ExportVTableHelperConstructorBody(GetGeneratedFunctionTextDevice(), Class);
+	}
+#else // WITH_HOT_RELOAD_CTORS
 	ExportStandardConstructorsMacro(StdMacro, Class, APIArg);
 	ExportEnhancedConstructorsMacro(EnhMacro, Class, APIArg);
+#endif // WITH_HOT_RELOAD_CTORS
 
 	GeneratedHeaderText.Log(*Macroize(*StdMacroName, *StdMacro));
 	GeneratedHeaderText.Log(*Macroize(*EnhMacroName, *EnhMacro));
@@ -4124,12 +4230,23 @@ FStringOutputDevice& FNativeClassHeaderGenerator::GetGeneratedFunctionTextDevice
 }
 
 // Constructor.
-FNativeClassHeaderGenerator::FNativeClassHeaderGenerator(UPackage* InPackage, const TArray<FUnrealSourceFile*>& SourceFiles, FClasses& AllClasses, bool InAllowSaveExportedHeaders)
+FNativeClassHeaderGenerator::FNativeClassHeaderGenerator(
+	UPackage* InPackage,
+	const TArray<FUnrealSourceFile*>& SourceFiles,
+	FClasses& AllClasses,
+	bool InAllowSaveExportedHeaders
+#if WITH_HOT_RELOAD_CTORS
+	, bool bExportVTableConstructors
+#endif // WITH_HOT_RELOAD_CTORS
+)
 	: API                                   (FPackageName::GetShortName(InPackage).ToUpper())
 	, Package                               (InPackage)
 	, bIsExportingForOffsetDeterminationOnly(false)
 	, bAllowSaveExportedHeaders             (InAllowSaveExportedHeaders)
 	, bFailIfGeneratedCodeChanges           (false)
+#if WITH_HOT_RELOAD_CTORS
+	, bExportVTableConstructors				(bExportVTableConstructors)
+#endif // WITH_HOT_RELOAD_CTORS
 {
 	const FString PackageName = FPackageName::GetShortName(Package);
 
@@ -4960,7 +5077,23 @@ ECompilationResult::Type UnrealHeaderTool_Main(const FString& ModuleInfoFilename
 					FClasses AllClasses(Package);
 					AllClasses.Validate();
 
+#if WITH_HOT_RELOAD_CTORS
+					static struct FUseVTableConstructorsCache
+					{
+						FUseVTableConstructorsCache()
+						{
+							bUseVTableConstructors = false;
+							GConfig->GetBool(TEXT("Core.System"), TEXT("UseVTableConstructors"), bUseVTableConstructors, GEngineIni);
+						}
+
+						bool bUseVTableConstructors;
+					} UseVTableConstructorsCache;
+
+					Result = FHeaderParser::ParseAllHeadersInside(AllClasses, GWarn, Package, Module, ScriptPlugins,
+						Module.ModuleType != EBuildModuleType::Game || UseVTableConstructorsCache.bUseVTableConstructors);
+#else // WITH_HOT_RELOAD_CTORS
 					Result = FHeaderParser::ParseAllHeadersInside(AllClasses, GWarn, Package, Module, ScriptPlugins);
+#endif // WITH_HOT_RELOAD_CTORS
 					if (Result != ECompilationResult::Succeeded)
 					{
 						++NumFailures;
diff --git a/Engine/Source/Programs/UnrealHeaderTool/Private/HeaderParser.cpp b/Engine/Source/Programs/UnrealHeaderTool/Private/HeaderParser.cpp
--- a/Engine/Source/Programs/UnrealHeaderTool/Private/HeaderParser.cpp
+++ b/Engine/Source/Programs/UnrealHeaderTool/Private/HeaderParser.cpp
@@ -6547,7 +6547,14 @@ TArray<FUnrealSourceFile*> GetSourceFilesWithInheritanceOrdering(UPackage* Curre
 }
 
 // Begins the process of exporting C++ class declarations for native classes in the specified package
-void FHeaderParser::ExportNativeHeaders( UPackage* CurrentPackage, FClasses& AllClasses, bool bAllowSaveExportedHeaders )
+void FHeaderParser::ExportNativeHeaders(
+	UPackage* CurrentPackage,
+	FClasses& AllClasses,
+	bool bAllowSaveExportedHeaders
+#if WITH_HOT_RELOAD_CTORS
+	, bool bExportVTableConstructors
+#endif // WITH_HOT_RELOAD_CTORS
+)
 {
 	// Build a list of header filenames
 	TArray<FString>	ClassHeaderFilenames;
@@ -6582,7 +6589,15 @@ void FHeaderParser::ExportNativeHeaders( UPackage* CurrentPackage, FClasses& All
 		}
 
 		// Export native class definitions to package header files.
-		FNativeClassHeaderGenerator(CurrentPackage, SourceFiles, AllClasses, bAllowSaveExportedHeaders);
+		FNativeClassHeaderGenerator(
+			CurrentPackage,
+			SourceFiles,
+			AllClasses,
+			bAllowSaveExportedHeaders
+#if WITH_HOT_RELOAD_CTORS
+			, bExportVTableConstructors
+#endif // WITH_HOT_RELOAD_CTORS
+		);
 	}
 }
 
@@ -6702,7 +6717,16 @@ void ExportClassTreeToScriptPlugins(const FClassTree* Node, const FManifestModul
 }
 
 // Parse all headers for classes that are inside CurrentPackage.
-ECompilationResult::Type FHeaderParser::ParseAllHeadersInside(FClasses& ModuleClasses, FFeedbackContext* Warn, UPackage* CurrentPackage, const FManifestModule& Module, TArray<IScriptGeneratorPluginInterface*>& ScriptPlugins)
+ECompilationResult::Type FHeaderParser::ParseAllHeadersInside(
+	FClasses& ModuleClasses,
+	FFeedbackContext* Warn,
+	UPackage* CurrentPackage,
+	const FManifestModule& Module,
+	TArray<IScriptGeneratorPluginInterface*>& ScriptPlugins
+#if WITH_HOT_RELOAD_CTORS
+	, bool bExportVTableConstructors
+#endif // WITH_HOT_RELOAD_CTORS
+)
 {
 	// Disable loading of objects outside of this package (or more exactly, objects which aren't UFields, CDO, or templates)
 	TGuardValue<bool> AutoRestoreVerifyObjectRefsFlag(GVerifyObjectReferencesOnly, true);
@@ -6748,7 +6772,14 @@ ECompilationResult::Type FHeaderParser::ParseAllHeadersInside(FClasses& ModuleCl
 			// from the feedback context.
 			Warn->SetContext(NULL);
 
-			ExportNativeHeaders(CurrentPackage, ModuleClasses, Module.SaveExportedHeaders);
+			ExportNativeHeaders(
+				CurrentPackage,
+				ModuleClasses,
+				Module.SaveExportedHeaders
+#if WITH_HOT_RELOAD_CTORS
+				, bExportVTableConstructors
+#endif // WITH_HOT_RELOAD_CTORS
+			);
 
 			// Done with header generation
 			if (HeaderParser.LinesParsed > 0)
@@ -6898,6 +6929,13 @@ void FHeaderParser::SimplifiedClassParse(const TCHAR* InBuffer, TArray<FSimplifi
 				bKeepPreprocessorDirectives = false;
 				bNotCPP = true;
 			}
+#if WITH_HOT_RELOAD_CTORS
+			else if (bIf && FParse::Command(&Str, TEXT("WITH_HOT_RELOAD")))
+			{
+				Target = &ClassHeaderTextStrippedOfCppText;
+				bKeepPreprocessorDirectives = false;
+			}
+#endif // WITH_HOT_RELOAD_CTORS
 			else if (bIf && (FParse::Command(&Str,TEXT("WITH_EDITORONLY_DATA")) || FParse::Command(&Str,TEXT("WITH_EDITOR"))))
 			{
 				Target = &ClassHeaderTextStrippedOfCppText;
@@ -7456,13 +7494,20 @@ bool FHeaderParser::TryToMatchConstructorParameterList(FToken Token)
 
 	auto* ClassData = GScriptHelper.FindClassData(GetCurrentClass());
 
-	ClassData->bConstructorDeclared = true;
+	bool bOICtor = false;
+#if WITH_HOT_RELOAD_CTORS
+	bool bVTCtor = false;
+#endif // WITH_HOT_RELOAD_CTORS
 
 	if (!ClassData->bDefaultConstructorDeclared && MatchSymbol(TEXT(")")))
 	{
 		ClassData->bDefaultConstructorDeclared = true;
 	}
-	else if (!ClassData->bObjectInitializerConstructorDeclared)
+	else if (!ClassData->bObjectInitializerConstructorDeclared
+#if WITH_HOT_RELOAD_CTORS
+		|| !ClassData->bCustomVTableHelperConstructorDeclared
+#endif // WITH_HOT_RELOAD_CTORS
+	)
 	{
 		FToken ObjectInitializerParamParsingToken;
 
@@ -7475,7 +7520,10 @@ bool FHeaderParser::TryToMatchConstructorParameterList(FToken Token)
 			// Template instantiation or additional parameter excludes ObjectInitializer constructor.
 			if (ObjectInitializerParamParsingToken.Matches(TEXT(",")) || ObjectInitializerParamParsingToken.Matches(TEXT("<")))
 			{
-				ClassData->bObjectInitializerConstructorDeclared = false;
+				bOICtor = false;
+#if WITH_HOT_RELOAD_CTORS
+				bVTCtor = false;
+#endif // WITH_HOT_RELOAD_CTORS
 				break;
 			}
 
@@ -7507,8 +7555,15 @@ bool FHeaderParser::TryToMatchConstructorParameterList(FToken Token)
 				|| ObjectInitializerParamParsingToken.Matches(TEXT("FPostConstructInitializeProperties")) // Deprecated, but left here, so it won't break legacy code.
 				)
 			{
-				ClassData->bObjectInitializerConstructorDeclared = true;
+				bOICtor = true;
 			}
+
+#if WITH_HOT_RELOAD_CTORS
+			if (ObjectInitializerParamParsingToken.Matches(TEXT("FVTableHelper")))
+			{
+				bVTCtor = true;
+			}
+#endif // WITH_HOT_RELOAD_CTORS
 		}
 
 		// Parse until finish.
@@ -7527,9 +7582,19 @@ bool FHeaderParser::TryToMatchConstructorParameterList(FToken Token)
 			}
 		}
 
-		ClassData->bObjectInitializerConstructorDeclared = ClassData->bObjectInitializerConstructorDeclared && bIsRef && bIsConst;
+		ClassData->bObjectInitializerConstructorDeclared = ClassData->bObjectInitializerConstructorDeclared || (bOICtor && bIsRef && bIsConst);
+#if WITH_HOT_RELOAD_CTORS
+		ClassData->bCustomVTableHelperConstructorDeclared = ClassData->bCustomVTableHelperConstructorDeclared || (bVTCtor && bIsRef);
+#endif // WITH_HOT_RELOAD_CTORS
 	}
 
+	ClassData->bConstructorDeclared =
+#if WITH_HOT_RELOAD_CTORS
+		ClassData->bConstructorDeclared || !bVTCtor;
+#else // WITH_HOT_RELOAD_CTORS
+		true;
+#endif // WITH_HOT_RELOAD_CTORS
+
 	// Optionally match semicolon.
 	if (!MatchSymbol(TEXT(";")))
 	{
diff --git a/Engine/Source/Programs/UnrealHeaderTool/Private/HeaderParser.h b/Engine/Source/Programs/UnrealHeaderTool/Private/HeaderParser.h
--- a/Engine/Source/Programs/UnrealHeaderTool/Private/HeaderParser.h
+++ b/Engine/Source/Programs/UnrealHeaderTool/Private/HeaderParser.h
@@ -173,7 +173,16 @@ public:
 	static void ComputeFunctionParametersSize(UClass* InClass);
 
 	// Parse all headers for classes that are inside LimitOuter.
-	static ECompilationResult::Type ParseAllHeadersInside(FClasses& ModuleClasses, FFeedbackContext* Warn, UPackage* LimitOuter, const FManifestModule& Module, TArray<class IScriptGeneratorPluginInterface*>& ScriptPlugins);
+	static ECompilationResult::Type ParseAllHeadersInside(
+		FClasses& ModuleClasses,
+		FFeedbackContext* Warn,
+		UPackage* LimitOuter,
+		const FManifestModule& Module,
+		TArray<class IScriptGeneratorPluginInterface*>& ScriptPlugins
+#if WITH_HOT_RELOAD_CTORS
+		, bool bExportVTableConstructors
+#endif // WITH_HOT_RELOAD_CTORS
+	);
 
 	// Performs a preliminary parse of the text in the specified buffer, pulling out:
 	//   Class name and parent class name
@@ -448,7 +457,14 @@ protected:
 	 * @param	CurrentPackage	the package being compiled
 	 * @param	AllClasses		the class tree for CurrentPackage
 	 */
-	static void ExportNativeHeaders( UPackage* CurrentPackage, FClasses& AllClasses, bool bAllowSaveExportedHeaders );
+	static void ExportNativeHeaders(
+		UPackage* CurrentPackage,
+		FClasses& AllClasses,
+		bool bAllowSaveExportedHeaders
+#if WITH_HOT_RELOAD_CTORS
+		, bool bExportVTableConstructors
+#endif // WITH_HOT_RELOAD_CTORS
+	);
 
 	// FContextSupplier interface.
 	virtual FString GetContext() override;
diff --git a/Engine/Source/Programs/UnrealHeaderTool/Private/NativeClassExporter.h b/Engine/Source/Programs/UnrealHeaderTool/Private/NativeClassExporter.h
--- a/Engine/Source/Programs/UnrealHeaderTool/Private/NativeClassExporter.h
+++ b/Engine/Source/Programs/UnrealHeaderTool/Private/NativeClassExporter.h
@@ -109,6 +109,11 @@ private:
 	/** If true, any change in the generated headers will result in UHT failure. */
 	bool bFailIfGeneratedCodeChanges;
 
+#if WITH_HOT_RELOAD_CTORS
+	/** If true, then generate vtable constructors and it's callers. */
+	bool bExportVTableConstructors;
+#endif // WITH_HOT_RELOAD_CTORS
+
 	/** All properties that need to be forward declared. */
 	TArray<UProperty*>	ForwardDeclarations;
 
@@ -516,7 +521,15 @@ private:
 public:
 
 	// Constructor
-	FNativeClassHeaderGenerator(UPackage* InPackage, const TArray<FUnrealSourceFile*>& SourceFiles, FClasses& AllClasses, bool InAllowSaveExportedHeaders);
+	FNativeClassHeaderGenerator(
+		UPackage* InPackage,
+		const TArray<FUnrealSourceFile*>& SourceFiles,
+		FClasses& AllClasses,
+		bool InAllowSaveExportedHeaders
+#if WITH_HOT_RELOAD_CTORS
+		, bool bExportVTableConstructors
+#endif // WITH_HOT_RELOAD_CTORS
+	);
 
 	/**
 	 * Gets string with function return type.
diff --git a/Engine/Source/Programs/UnrealHeaderTool/Private/ParserHelper.h b/Engine/Source/Programs/UnrealHeaderTool/Private/ParserHelper.h
--- a/Engine/Source/Programs/UnrealHeaderTool/Private/ParserHelper.h
+++ b/Engine/Source/Programs/UnrealHeaderTool/Private/ParserHelper.h
@@ -1366,6 +1366,9 @@ public:
 		, bConstructorDeclared(false)
 		, bDefaultConstructorDeclared(false)
 		, bObjectInitializerConstructorDeclared(false)
+#if WITH_HOT_RELOAD_CTORS
+		, bCustomVTableHelperConstructorDeclared(false)
+#endif // WITH_HOT_RELOAD_CTORS
 		, GeneratedBodyMacroAccessSpecifier(ACCESS_NotAnAccessSpecifier)
 	{
 	}
@@ -1643,6 +1646,11 @@ public:
 	// Is ObjectInitializer constructor (i.e. a constructor with only one parameter of type FObjectInitializer) declared?
 	bool bObjectInitializerConstructorDeclared;
 
+#if WITH_HOT_RELOAD_CTORS
+	// Is custom VTable helper constructor declared?
+	bool bCustomVTableHelperConstructorDeclared;
+#endif // WITH_HOT_RELOAD_CTORS
+
 	// GENERATED_BODY access specifier to preserve.
 	EAccessSpecifier GeneratedBodyMacroAccessSpecifier;
 };
diff --git a/Engine/Source/Runtime/AIModule/Classes/BehaviorTree/BehaviorTreeComponent.h b/Engine/Source/Runtime/AIModule/Classes/BehaviorTree/BehaviorTreeComponent.h
--- a/Engine/Source/Runtime/AIModule/Classes/BehaviorTree/BehaviorTreeComponent.h
+++ b/Engine/Source/Runtime/AIModule/Classes/BehaviorTree/BehaviorTreeComponent.h
@@ -85,6 +85,11 @@ protected:
 	bool TreeHasBeenStarted() const;
 
 public:
+#if WITH_HOT_RELOAD && WITH_HOT_RELOAD_CTORS
+	/** DO NOT USE. This constructor is for internal usage only for hot-reload purposes. */
+	UBehaviorTreeComponent(FVTableHelper& Helper);
+#endif // WITH_HOT_RELOAD && WITH_HOT_RELOAD_CTORS
+
 	virtual bool IsRunning() const override;
 	virtual bool IsPaused() const override;
 	virtual void Cleanup() override;
diff --git a/Engine/Source/Runtime/AIModule/Private/BehaviorTree/BehaviorTreeComponent.cpp b/Engine/Source/Runtime/AIModule/Private/BehaviorTree/BehaviorTreeComponent.cpp
--- a/Engine/Source/Runtime/AIModule/Private/BehaviorTree/BehaviorTreeComponent.cpp
+++ b/Engine/Source/Runtime/AIModule/Private/BehaviorTree/BehaviorTreeComponent.cpp
@@ -31,6 +31,15 @@ UBehaviorTreeComponent::UBehaviorTreeComponent(const FObjectInitializer& ObjectI
 	bIsPaused = false;
 }
 
+#if WITH_HOT_RELOAD && WITH_HOT_RELOAD_CTORS
+UBehaviorTreeComponent::UBehaviorTreeComponent(FVTableHelper& Helper)
+	: Super(Helper)
+	, SearchData(*this)
+{
+
+}
+#endif // WITH_HOT_RELOAD && WITH_HOT_RELOAD_CTORS
+
 void UBehaviorTreeComponent::UninitializeComponent()
 {
 	UBehaviorTreeManager* BTManager = UBehaviorTreeManager::GetCurrent(GetWorld());
diff --git a/Engine/Source/Runtime/Core/Private/Misc/ConfigCacheIni.cpp b/Engine/Source/Runtime/Core/Private/Misc/ConfigCacheIni.cpp
--- a/Engine/Source/Runtime/Core/Private/Misc/ConfigCacheIni.cpp
+++ b/Engine/Source/Runtime/Core/Private/Misc/ConfigCacheIni.cpp
@@ -1334,6 +1334,13 @@ FConfigCacheIni::FConfigCacheIni(EConfigCacheType InType)
 {
 }
 
+#if WITH_HOT_RELOAD && WITH_HOT_RELOAD_CTORS
+FConfigCacheIni::FConfigCacheIni()
+{
+	EnsureRetrievingVTablePtr();
+}
+#endif // WITH_HOT_RELOAD && WITH_HOT_RELOAD_CTORS
+
 FConfigCacheIni::~FConfigCacheIni()
 {
 	Flush( 1 );
diff --git a/Engine/Source/Runtime/Core/Private/Misc/CoreMisc.cpp b/Engine/Source/Runtime/Core/Private/Misc/CoreMisc.cpp
--- a/Engine/Source/Runtime/Core/Private/Misc/CoreMisc.cpp
+++ b/Engine/Source/Runtime/Core/Private/Misc/CoreMisc.cpp
@@ -17,7 +17,6 @@
 #include "Ticker.h"
 #include "DerivedDataCacheInterface.h"
 
-
 DEFINE_LOG_CATEGORY(LogSHA);
 DEFINE_LOG_CATEGORY(LogStats);
 DEFINE_LOG_CATEGORY(LogStreaming);
@@ -29,6 +28,7 @@ DEFINE_LOG_CATEGORY(LogLocalization);
 DEFINE_LOG_CATEGORY(LogLongPackageNames);
 DEFINE_LOG_CATEGORY(LogProcess);
 DEFINE_LOG_CATEGORY(LogLoad);
+DEFINE_LOG_CATEGORY(LogCore);
 
 
 /*-----------------------------------------------------------------------------
@@ -985,4 +985,13 @@ FBoolConfigValueHelper::FBoolConfigValueHelper(const TCHAR* Section, const TCHAR
 	: bValue(false)
 {
 	GConfig->GetBool(Section, Key, bValue, Filename);
-}
\ No newline at end of file
+}
+
+#if WITH_HOT_RELOAD && WITH_HOT_RELOAD_CTORS
+bool GIsRetrievingVTablePtr = false;
+
+void EnsureRetrievingVTablePtr()
+{
+	UE_CLOG(!GIsRetrievingVTablePtr, LogCore, Fatal, TEXT("This should be used only during vtable ptr retrieval process."));
+}
+#endif // WITH_HOT_RELOAD && WITH_HOT_RELOAD_CTORS
diff --git a/Engine/Source/Runtime/Core/Public/CoreGlobals.h b/Engine/Source/Runtime/Core/Public/CoreGlobals.h
--- a/Engine/Source/Runtime/Core/Public/CoreGlobals.h
+++ b/Engine/Source/Runtime/Core/Public/CoreGlobals.h
@@ -199,6 +199,11 @@ extern CORE_API bool GIsInitialLoad;
 /** true when we are routing ConditionalPostLoad/PostLoad to objects */
 extern CORE_API bool GIsRoutingPostLoad;
 
+#if WITH_HOT_RELOAD && WITH_HOT_RELOAD_CTORS
+/** true when we are retrieving VTablePtr from UClass */
+extern CORE_API bool GIsRetrievingVTablePtr;
+#endif // WITH_HOT_RELOAD && WITH_HOT_RELOAD_CTORS
+
 /** Steadily increasing frame counter. */
 extern CORE_API uint64 GFrameCounter;
 
@@ -275,4 +280,12 @@ extern CORE_API bool GPumpingMessagesOutsideOfMainLoop;
 extern CORE_API double GBlueprintCompileTime;
 
 /** Stack names from the VM to be unrolled when we assert */
-extern CORE_API TArray<FScriptTraceStackNode> GScriptStack;
\ No newline at end of file
+extern CORE_API TArray<FScriptTraceStackNode> GScriptStack;
+
+#if WITH_HOT_RELOAD && WITH_HOT_RELOAD_CTORS
+/**
+ * Ensures that current thread is during retrieval of vtable ptr
+ * of some UClass.
+ */
+CORE_API void EnsureRetrievingVTablePtr();
+#endif // WITH_HOT_RELOAD && WITH_HOT_RELOAD_CTORS
\ No newline at end of file
diff --git a/Engine/Source/Runtime/Core/Public/HAL/OutputDevices.h b/Engine/Source/Runtime/Core/Public/HAL/OutputDevices.h
--- a/Engine/Source/Runtime/Core/Public/HAL/OutputDevices.h
+++ b/Engine/Source/Runtime/Core/Public/HAL/OutputDevices.h
@@ -536,6 +536,7 @@ CORE_API DECLARE_LOG_CATEGORY_EXTERN(LogNetPackageMap, Warning, All);
 CORE_API DECLARE_LOG_CATEGORY_EXTERN(LogNetSerialization, Warning, All);
 CORE_API DECLARE_LOG_CATEGORY_EXTERN(LogMemory, Log, All);
 CORE_API DECLARE_LOG_CATEGORY_EXTERN(LogProfilingDebugging, Log, All);
+CORE_API DECLARE_LOG_CATEGORY_EXTERN(LogCore, Log, All);
 
 // Temporary log category, generally you should not check things in that use this
-CORE_API DECLARE_LOG_CATEGORY_EXTERN(LogTemp, Log, All);
+CORE_API DECLARE_LOG_CATEGORY_EXTERN(LogTemp, Log, All);
\ No newline at end of file
diff --git a/Engine/Source/Runtime/Core/Public/Misc/Build.h b/Engine/Source/Runtime/Core/Public/Misc/Build.h
--- a/Engine/Source/Runtime/Core/Public/Misc/Build.h
+++ b/Engine/Source/Runtime/Core/Public/Misc/Build.h
@@ -140,6 +140,11 @@
 	#define WITH_HOT_RELOAD (!IS_MONOLITHIC && !UE_BUILD_SHIPPING)
 #endif
 
+/**
+ * Whether we want to use special hot-reload empty constructors.
+ */
+#define WITH_HOT_RELOAD_CTORS 1
+
 /*----------------------------------------------------------------------------
 	Optional bridge options coming from UBT, do not modify directly!
 	If UBT doesn't set the value, it is assumed to be 0, and we set that here.
diff --git a/Engine/Source/Runtime/Core/Public/Misc/ConfigCacheIni.h b/Engine/Source/Runtime/Core/Public/Misc/ConfigCacheIni.h
--- a/Engine/Source/Runtime/Core/Public/Misc/ConfigCacheIni.h
+++ b/Engine/Source/Runtime/Core/Public/Misc/ConfigCacheIni.h
@@ -172,6 +172,12 @@ class CORE_API FConfigCacheIni : public TMap<FString,FConfigFile>
 public:
 	// Basic functions.
 	FConfigCacheIni(EConfigCacheType Type);
+
+#if WITH_HOT_RELOAD && WITH_HOT_RELOAD_CTORS
+	/** DO NOT USE. This constructor is for internal usage only for hot-reload purposes. */
+	FConfigCacheIni();
+#endif // WITH_HOT_RELOAD && WITH_HOT_RELOAD_CTORS
+
 	virtual ~FConfigCacheIni();
 
 	/**
diff --git a/Engine/Source/Runtime/Core/Public/Templates/SharedPointer.h b/Engine/Source/Runtime/Core/Public/Templates/SharedPointer.h
--- a/Engine/Source/Runtime/Core/Public/Templates/SharedPointer.h
+++ b/Engine/Source/Runtime/Core/Public/Templates/SharedPointer.h
@@ -150,14 +150,23 @@ public:
 		: Object( InObject )
 		, SharedReferenceCount( InObject )
 	{
-		// If the following assert goes off, it means a TSharedRef was initialized from a nullptr object pointer.
-		// Shared references must never be nullptr, so either pass a valid object or consider using TSharedPtr instead.
-		check( InObject != nullptr );
+		Init(InObject);
+	}
 
-		// If the object happens to be derived from TSharedFromThis, the following method
-		// will prime the object with a weak pointer to itself.
-		SharedPointerInternals::EnableSharedFromThis( this, InObject, InObject );
+#if WITH_HOT_RELOAD && WITH_HOT_RELOAD_CTORS
+	/**
+	 * Constructs default shared reference that owns the default object for specified type.
+	 *
+	 * Used internally only. Please do not use!
+	 */
+	TSharedRef()
+		: Object(new ObjectType())
+		, SharedReferenceCount(Object)
+	{
+		EnsureRetrievingVTablePtr();
+		Init(Object);
 	}
+#endif // WITH_HOT_RELOAD && WITH_HOT_RELOAD_CTORS
 
 	/**
 	 * Constructs a shared reference using a proxy reference to a raw pointer. (See MakeShareable())
@@ -347,6 +356,17 @@ public:
 	}
 
 private:
+	template<class OtherType>
+	void Init(OtherType* InObject)
+	{
+		// If the following assert goes off, it means a TSharedRef was initialized from a nullptr object pointer.
+		// Shared references must never be nullptr, so either pass a valid object or consider using TSharedPtr instead.
+		check(InObject != nullptr);
+
+		// If the object happens to be derived from TSharedFromThis, the following method
+		// will prime the object with a weak pointer to itself.
+		SharedPointerInternals::EnableSharedFromThis(this, InObject, InObject);
+	}
 
 	/**
 	 * Converts a shared pointer to a shared reference.  The pointer *must* be valid or an assertion will trigger.
diff --git a/Engine/Source/Runtime/CoreUObject/Classes/Object.h b/Engine/Source/Runtime/CoreUObject/Classes/Object.h
--- a/Engine/Source/Runtime/CoreUObject/Classes/Object.h
+++ b/Engine/Source/Runtime/CoreUObject/Classes/Object.h
@@ -882,6 +882,10 @@ public:
 	 * Default UObject constructor.
 	 */
 	UObject(const FObjectInitializer& ObjectInitializer);
+#if WITH_HOT_RELOAD && WITH_HOT_RELOAD_CTORS
+	/** DO NOT USE. This constructor is for internal usage only for hot-reload purposes. */
+	UObject(FVTableHelper& Helper);
+#endif // WITH_HOT_RELOAD && WITH_HOT_RELOAD_CTORS
 
 	//=============================================================================
 	// K2 support functions.
diff --git a/Engine/Source/Runtime/CoreUObject/Private/UObject/Class.cpp b/Engine/Source/Runtime/CoreUObject/Private/UObject/Class.cpp
--- a/Engine/Source/Runtime/CoreUObject/Private/UObject/Class.cpp
+++ b/Engine/Source/Runtime/CoreUObject/Private/UObject/Class.cpp
@@ -2424,9 +2424,13 @@ class FRestoreClassInfo: public FRestoreForUObjectOverwrite
 	/** Saved ClassCastFlags **/
 	EClassCastFlags	CastFlags;
 	/** Saved ClassConstructor **/
-	void			(*Constructor)(const FObjectInitializer&);
+	UClass::ClassConstructorType Constructor;
+#if WITH_HOT_RELOAD && WITH_HOT_RELOAD_CTORS
+	/** Saved ClassVTableHelperCtorCaller **/
+	UClass::ClassVTableHelperCtorCallerType ClassVTableHelperCtorCaller;
+#endif // WITH_HOT_RELOAD && WITH_HOT_RELOAD_CTORS
 	/** Saved ClassConstructor **/
-	void			(*AddReferencedObjects)(UObject*, FReferenceCollector&);
+	UClass::ClassAddReferencedObjectsType AddReferencedObjects;
 	/** Saved NativeFunctionLookupTable. */
 	TArray<FNativeFunctionLookup> NativeFunctionLookupTable;
 public:
@@ -2444,6 +2448,9 @@ public:
 		Flags(Save->ClassFlags & CLASS_Abstract),
 		CastFlags(Save->ClassCastFlags),
 		Constructor(Save->ClassConstructor),
+#if WITH_HOT_RELOAD && WITH_HOT_RELOAD_CTORS
+		ClassVTableHelperCtorCaller(Save->ClassVTableHelperCtorCaller),
+#endif // WITH_HOT_RELOAD && WITH_HOT_RELOAD_CTORS
 		AddReferencedObjects(Save->ClassAddReferencedObjects),
 		NativeFunctionLookupTable(Save->NativeFunctionLookupTable)
 	{
@@ -2458,6 +2465,9 @@ public:
 		Target->ClassFlags |= Flags;
 		Target->ClassCastFlags |= CastFlags;
 		Target->ClassConstructor = Constructor;
+#if WITH_HOT_RELOAD && WITH_HOT_RELOAD_CTORS
+		Target->ClassVTableHelperCtorCaller = ClassVTableHelperCtorCaller;
+#endif // WITH_HOT_RELOAD && WITH_HOT_RELOAD_CTORS
 		Target->ClassAddReferencedObjects = AddReferencedObjects;
 		Target->NativeFunctionLookupTable = NativeFunctionLookupTable;
 	}
@@ -2665,7 +2675,11 @@ void UClass::Bind()
 	}
 
 	UClass* SuperClass = GetSuperClass();
-	if (SuperClass && (ClassConstructor == NULL || ClassAddReferencedObjects == NULL))
+	if (SuperClass && (ClassConstructor == nullptr || ClassAddReferencedObjects == nullptr
+#if WITH_HOT_RELOAD && WITH_HOT_RELOAD_CTORS
+		|| ClassVTableHelperCtorCaller == nullptr
+#endif // WITH_HOT_RELOAD && WITH_HOT_RELOAD_CTORS
+		))
 	{
 		// Chase down constructor in parent class.
 		SuperClass->Bind();
@@ -2673,6 +2687,12 @@ void UClass::Bind()
 		{
 			ClassConstructor = SuperClass->ClassConstructor;
 		}
+#if WITH_HOT_RELOAD && WITH_HOT_RELOAD_CTORS
+		if (!ClassVTableHelperCtorCaller)
+		{
+			ClassVTableHelperCtorCaller = SuperClass->ClassVTableHelperCtorCaller;
+		}
+#endif // WITH_HOT_RELOAD && WITH_HOT_RELOAD_CTORS
 		if (!ClassAddReferencedObjects)
 		{
 			ClassAddReferencedObjects = SuperClass->ClassAddReferencedObjects;
@@ -3170,7 +3190,10 @@ FArchive& operator<<(FArchive& Ar, FImplementedInterface& A)
 
 void UClass::PurgeClass(bool bRecompilingOnLoad)
 {
-	ClassConstructor = NULL;
+	ClassConstructor = nullptr;
+#if WITH_HOT_RELOAD && WITH_HOT_RELOAD_CTORS
+	ClassVTableHelperCtorCaller = nullptr;
+#endif // WITH_HOT_RELOAD && WITH_HOT_RELOAD_CTORS
 	ClassFlags = 0;
 	ClassCastFlags = 0;
 	ClassUnique = 0;
@@ -3321,11 +3344,17 @@ UClass::UClass
 	EClassCastFlags	InClassCastFlags,
 	const TCHAR*    InConfigName,
 	EObjectFlags	InFlags,
-	void			(*InClassConstructor)(const FObjectInitializer&),
-	void			(*InClassAddReferencedObjects)(UObject*, class FReferenceCollector&)
+	ClassConstructorType InClassConstructor,
+#if WITH_HOT_RELOAD && WITH_HOT_RELOAD_CTORS
+	ClassVTableHelperCtorCallerType InClassVTableHelperCtorCaller,
+#endif // WITH_HOT_RELOAD && WITH_HOT_RELOAD_CTORS
+	ClassAddReferencedObjectsType InClassAddReferencedObjects
 )
 :	UStruct					( EC_StaticConstructor, InSize, InFlags )
 ,	ClassConstructor		( InClassConstructor )
+#if WITH_HOT_RELOAD && WITH_HOT_RELOAD_CTORS
+,	ClassVTableHelperCtorCaller(InClassVTableHelperCtorCaller)
+#endif // WITH_HOT_RELOAD && WITH_HOT_RELOAD_CTORS
 ,	ClassAddReferencedObjects( InClassAddReferencedObjects )
 ,	ClassFlags				( InClassFlags | CLASS_Native )
 ,	ClassCastFlags			( InClassCastFlags )
@@ -3348,8 +3377,11 @@ bool UClass::HotReloadPrivateStaticClass(
 	uint32			InClassFlags,
 	EClassCastFlags	InClassCastFlags,
 	const TCHAR*    InConfigName,
-	void			(*InClassConstructor)(const FObjectInitializer&),
-	void			(*InAddReferencedObjects)(UObject*, class FReferenceCollector&),
+	ClassConstructorType InClassConstructor,
+#if WITH_HOT_RELOAD_CTORS
+	ClassVTableHelperCtorCallerType InClassVTableHelperCtorCaller,
+#endif // WITH_HOT_RELOAD_CTORS
+	ClassAddReferencedObjectsType InClassAddReferencedObjects,
 	class UClass* TClass_Super_StaticClass,
 	class UClass* TClass_WithinClass_StaticClass
 	)
@@ -3367,9 +3399,12 @@ bool UClass::HotReloadPrivateStaticClass(
 	//@todo safe? ClassFlags = InClassFlags | CLASS_Native;
 	//@todo safe? ClassCastFlags = InClassCastFlags;
 	//@todo safe? ClassConfigName = InConfigName;
-	void(*OldClassConstructor)(const FObjectInitializer&) = ClassConstructor;
+	ClassConstructorType OldClassConstructor = ClassConstructor;
 	ClassConstructor = InClassConstructor;
-	ClassAddReferencedObjects = InAddReferencedObjects;
+#if WITH_HOT_RELOAD_CTORS
+	ClassVTableHelperCtorCaller = InClassVTableHelperCtorCaller;
+#endif // WITH_HOT_RELOAD_CTORS
+	ClassAddReferencedObjects = InClassAddReferencedObjects;
 	/* No recursive ::StaticClass calls allowed. Setup extras. */
 	/* @todo safe? 
 	if (TClass_Super_StaticClass != this)
@@ -3385,7 +3420,33 @@ bool UClass::HotReloadPrivateStaticClass(
 
 	UE_LOG(LogClass, Verbose, TEXT("Attempting to change VTable for class %s."),*GetName());
 	ClassWithin = UPackage::StaticClass();  // We are just avoiding error checks with this...we don't care about this temp object other than to get the vtable.
+
+#if WITH_HOT_RELOAD_CTORS
+	static struct FUseVTableConstructorsCache
+	{
+		FUseVTableConstructorsCache()
+		{
+			bUseVTableConstructors = false;
+			GConfig->GetBool(TEXT("Core.System"), TEXT("UseVTableConstructors"), bUseVTableConstructors, GEngineIni);
+		}
+
+		bool bUseVTableConstructors;
+	} UseVTableConstructorsCache;
+
+	UObject* TempObjectForVTable = nullptr;
+	if (UseVTableConstructorsCache.bUseVTableConstructors)
+	{
+		TGuardValue<bool> Guard(GIsRetrievingVTablePtr, true);
+		auto Helper = FVTableHelper();
+		TempObjectForVTable = ClassVTableHelperCtorCaller(Helper);
+	}
+	else
+	{
+		TempObjectForVTable = StaticConstructObject_Internal(this, GetTransientPackage(), NAME_None, RF_NeedLoad | RF_ClassDefaultObject | RF_TagGarbageTemp);
+	}
+#else // WITH_HOT_RELOAD_CTORS
 	UObject* TempObjectForVTable = StaticConstructObject_Internal(this, GetTransientPackage(), NAME_None, RF_NeedLoad | RF_ClassDefaultObject | RF_TagGarbageTemp);
+#endif // WITH_HOT_RELOAD_CTORS
 
 	if( !TempObjectForVTable->IsRooted() )
 	{
@@ -3417,6 +3478,9 @@ bool UClass::HotReloadPrivateStaticClass(
 				if (Class->ClassConstructor == OldClassConstructor)
 				{
 					Class->ClassConstructor = ClassConstructor;
+#if WITH_HOT_RELOAD_CTORS
+					Class->ClassVTableHelperCtorCaller = ClassVTableHelperCtorCaller;
+#endif // WITH_HOT_RELOAD_CTORS
 					Class->ClassAddReferencedObjects = ClassAddReferencedObjects;
 					CountClass++;
 				}
diff --git a/Engine/Source/Runtime/CoreUObject/Private/UObject/LinkerPlaceholderClass.cpp b/Engine/Source/Runtime/CoreUObject/Private/UObject/LinkerPlaceholderClass.cpp
--- a/Engine/Source/Runtime/CoreUObject/Private/UObject/LinkerPlaceholderClass.cpp
+++ b/Engine/Source/Runtime/CoreUObject/Private/UObject/LinkerPlaceholderClass.cpp
@@ -56,6 +56,9 @@ void ULinkerPlaceholderClass::PostInitProperties()
 void ULinkerPlaceholderClass::Bind()
 {
 	ClassConstructor = InternalConstructor<ULinkerPlaceholderClass>;
+#if WITH_HOT_RELOAD && WITH_HOT_RELOAD_CTORS
+	ClassVTableHelperCtorCaller = InternalVTableHelperCtorCaller<ULinkerPlaceholderClass>;
+#endif // WITH_HOT_RELOAD && WITH_HOT_RELOAD_CTORS
 	Super::Bind();
 
 	ClassAddReferencedObjects = &ULinkerPlaceholderClass::AddReferencedObjects;
diff --git a/Engine/Source/Runtime/CoreUObject/Private/UObject/Obj.cpp b/Engine/Source/Runtime/CoreUObject/Private/UObject/Obj.cpp
--- a/Engine/Source/Runtime/CoreUObject/Private/UObject/Obj.cpp
+++ b/Engine/Source/Runtime/CoreUObject/Private/UObject/Obj.cpp
@@ -50,8 +50,36 @@ static UPackage*			GObjTransientPkg								= NULL;
 UObject::UObject( EStaticConstructor, EObjectFlags InFlags )
 : UObjectBaseUtility(InFlags | RF_Native | RF_RootSet)
 {
+#if WITH_HOT_RELOAD && WITH_HOT_RELOAD_CTORS
+	EnsureNotRetrievingVTablePtr();
+#endif // WITH_HOT_RELOAD && WITH_HOT_RELOAD_CTORS
 }
 
+#if WITH_HOT_RELOAD && WITH_HOT_RELOAD_CTORS
+UObject::UObject(FVTableHelper& Helper)
+{
+	EnsureRetrievingVTablePtr();
+
+	static struct FUseVTableConstructorsCache
+	{
+		FUseVTableConstructorsCache()
+		{
+			bUseVTableConstructors = false;
+			GConfig->GetBool(TEXT("Core.System"), TEXT("UseVTableConstructors"), bUseVTableConstructors, GEngineIni);
+		}
+
+		bool bUseVTableConstructors;
+	} UseVTableConstructorsCache;
+
+	UE_CLOG(!UseVTableConstructorsCache.bUseVTableConstructors, LogCore, Fatal, TEXT("This constructor is disabled."));
+}
+
+void UObject::EnsureNotRetrievingVTablePtr() const
+{
+	UE_CLOG(GIsRetrievingVTablePtr, LogCore, Fatal, TEXT("We are currently retrieving VTable ptr. Please use FVTableHelper constructor instead."));
+}
+#endif // WITH_HOT_RELOAD && WITH_HOT_RELOAD_CTORS
+
 UObject* UObject::CreateDefaultSubobject(FName SubobjectFName, UClass* ReturnType, UClass* ClassToCreateByDefault, bool bIsRequired, bool bAbstract, bool bIsTransient)
 {
 	UE_CLOG(!GIsInConstructor, LogObj, Fatal, TEXT("CreateDefultSubobject can only be used inside of UObject constructors. UObject constructing subobjects cannot be created using new or placement new operator."));
diff --git a/Engine/Source/Runtime/CoreUObject/Private/UObject/SavePackage.cpp b/Engine/Source/Runtime/CoreUObject/Private/UObject/SavePackage.cpp
--- a/Engine/Source/Runtime/CoreUObject/Private/UObject/SavePackage.cpp
+++ b/Engine/Source/Runtime/CoreUObject/Private/UObject/SavePackage.cpp
@@ -2217,6 +2217,9 @@ static bool ValidateConformCompatibility(UPackage* NewPackage, ULinkerLoad* OldL
 		if (OldClass != NULL && NewClass != NULL && OldClass->HasAnyFlags(RF_Native) && NewClass->HasAnyFlags(RF_Native))
 		{
 			OldClass->ClassConstructor = NewClass->ClassConstructor;
+#if WITH_HOT_RELOAD && WITH_HOT_RELOAD_CTORS
+			OldClass->ClassVTableHelperCtorCaller = NewClass->ClassVTableHelperCtorCaller;
+#endif // WITH_HOT_RELOAD && WITH_HOT_RELOAD_CTORS
 			OldClass->ClassAddReferencedObjects = NewClass->ClassAddReferencedObjects;
 		}
 	}
diff --git a/Engine/Source/Runtime/CoreUObject/Private/UObject/UObjectGlobals.cpp b/Engine/Source/Runtime/CoreUObject/Private/UObject/UObjectGlobals.cpp
--- a/Engine/Source/Runtime/CoreUObject/Private/UObject/UObjectGlobals.cpp
+++ b/Engine/Source/Runtime/CoreUObject/Private/UObject/UObjectGlobals.cpp
@@ -1914,6 +1914,10 @@ void UObject::PostInitProperties()
 
 UObject::UObject()
 {
+#if WITH_HOT_RELOAD && WITH_HOT_RELOAD_CTORS
+	EnsureNotRetrievingVTablePtr();
+#endif // WITH_HOT_RELOAD && WITH_HOT_RELOAD_CTORS
+
 	FObjectInitializer* ObjectInitializerPtr = FTlsObjectInitializers::Top();
 	UE_CLOG(!ObjectInitializerPtr, LogUObjectGlobals, Fatal, TEXT("%s is not being constructed with either NewObject, NewNamedObject or ConstructObject."), *GetName());
 	FObjectInitializer& ObjectInitializer = *ObjectInitializerPtr;
@@ -1924,6 +1928,10 @@ UObject::UObject()
 
 UObject::UObject(const FObjectInitializer& ObjectInitializer)
 {
+#if WITH_HOT_RELOAD && WITH_HOT_RELOAD_CTORS
+	EnsureNotRetrievingVTablePtr();
+#endif // WITH_HOT_RELOAD && WITH_HOT_RELOAD_CTORS
+
 	check(!ObjectInitializer.Obj || ObjectInitializer.Obj == this);
 	const_cast<FObjectInitializer&>(ObjectInitializer).Obj = this;
 	const_cast<FObjectInitializer&>(ObjectInitializer).FinalizeSubobjectClassInitialization();
diff --git a/Engine/Source/Runtime/CoreUObject/Public/UObject/Class.h b/Engine/Source/Runtime/CoreUObject/Public/UObject/Class.h
--- a/Engine/Source/Runtime/CoreUObject/Public/UObject/Class.h
+++ b/Engine/Source/Runtime/CoreUObject/Public/UObject/Class.h
@@ -1749,9 +1749,18 @@ public:
 
 	friend class FRestoreClassInfo;
 
-	void(*ClassConstructor)(const FObjectInitializer&);
+	typedef void		(*ClassConstructorType)				(const FObjectInitializer&);
+#if WITH_HOT_RELOAD && WITH_HOT_RELOAD_CTORS
+	typedef UObject*	(*ClassVTableHelperCtorCallerType)	(FVTableHelper& Helper);
+#endif // WITH_HOT_RELOAD && WITH_HOT_RELOAD_CTORS
+	typedef void		(*ClassAddReferencedObjectsType)	(UObject*, class FReferenceCollector&);
+
+	ClassConstructorType ClassConstructor;
+#if WITH_HOT_RELOAD && WITH_HOT_RELOAD_CTORS
+	ClassVTableHelperCtorCallerType ClassVTableHelperCtorCaller;
+#endif // WITH_HOT_RELOAD && WITH_HOT_RELOAD_CTORS
 	/** Pointer to a static AddReferencedObjects method. */
-	void(*ClassAddReferencedObjects)(UObject*, class FReferenceCollector&);
+	ClassAddReferencedObjectsType ClassAddReferencedObjects;
 
 	// Class flags; See EClassFlags for more information
 	uint32 ClassFlags;
@@ -1853,8 +1862,11 @@ public:
 	UClass(const FObjectInitializer& ObjectInitializer);
 	explicit UClass(const FObjectInitializer& ObjectInitializer, UClass* InSuperClass);
 	UClass( EStaticConstructor, FName InName, uint32 InSize, uint32 InClassFlags, EClassCastFlags InClassCastFlags,
-		const TCHAR* InClassConfigName, EObjectFlags InFlags, void(*InClassConstructor)(const FObjectInitializer&),
-		void(*InClassAddReferencedObjects)(UObject*, class FReferenceCollector&));
+		const TCHAR* InClassConfigName, EObjectFlags InFlags, ClassConstructorType InClassConstructor,
+#if WITH_HOT_RELOAD && WITH_HOT_RELOAD_CTORS
+		ClassVTableHelperCtorCallerType InClassVTableHelperCtorCaller,
+#endif // WITH_HOT_RELOAD && WITH_HOT_RELOAD_CTORS
+		ClassAddReferencedObjectsType InClassAddReferencedObjects);
 
 #if WITH_HOT_RELOAD
 	/**
@@ -1872,8 +1884,11 @@ public:
 		uint32			InClassFlags,
 		EClassCastFlags	InClassCastFlags,
 		const TCHAR*    InConfigName,
-		void			(*InClassConstructor)(const FObjectInitializer&),
-		void			(*InAddReferencedObjects)(UObject*, class FReferenceCollector&),
+		ClassConstructorType InClassConstructor,
+#if WITH_HOT_RELOAD_CTORS
+		ClassVTableHelperCtorCallerType InClassVTableHelperCtorCaller,
+#endif // WITH_HOT_RELOAD_CTORS
+		ClassAddReferencedObjectsType InClassAddReferencedObjects,
 		class UClass* TClass_Super_StaticClass,
 		class UClass* TClass_WithinClass_StaticClass
 		);
@@ -2255,6 +2270,17 @@ void InternalConstructor( const FObjectInitializer& X )
 	T::__DefaultConstructor(X);
 }
 
+#if WITH_HOT_RELOAD && WITH_HOT_RELOAD_CTORS
+/**
+ * Helper template to call the vtable ctor caller for a class
+ */
+template<class T>
+UObject* InternalVTableHelperCtorCaller(FVTableHelper& Helper)
+{
+	return T::__VTableCtorCaller(Helper);
+}
+#endif // WITH_HOT_RELOAD && WITH_HOT_RELOAD_CTORS
+
 COREUOBJECT_API void InitializePrivateStaticClass(
 	class UClass* TClass_Super_StaticClass,
 	class UClass* TClass_PrivateStaticClass,
@@ -2290,7 +2316,10 @@ void GetPrivateStaticClassBody( const TCHAR* PackageName, const TCHAR* Name, UCl
 				TClass::StaticClassFlags,
 				TClass::StaticClassCastFlags(),
 				TClass::StaticConfigName(),
-				(void(*)(const FObjectInitializer&))InternalConstructor<TClass>,
+				(UClass::ClassConstructorType)InternalConstructor<TClass>,
+#if WITH_HOT_RELOAD_CTORS
+				(UClass::ClassVTableHelperCtorCallerType)InternalVTableHelperCtorCaller<TClass>,
+#endif // WITH_HOT_RELOAD_CTORS
 				&TClass::AddReferencedObjects,
 				TClass::Super::StaticClass(),
 				TClass::WithinClass::StaticClass()
@@ -2318,7 +2347,10 @@ void GetPrivateStaticClassBody( const TCHAR* PackageName, const TCHAR* Name, UCl
 		TClass::StaticClassCastFlags(),
 		TClass::StaticConfigName(),
 		EObjectFlags(RF_Public | RF_Standalone | RF_Transient | RF_Native | RF_RootSet),
-		(void(*)(const FObjectInitializer&))InternalConstructor<TClass>,
+		(UClass::ClassConstructorType)InternalConstructor<TClass>,
+#if WITH_HOT_RELOAD && WITH_HOT_RELOAD_CTORS
+		(UClass::ClassVTableHelperCtorCallerType)InternalVTableHelperCtorCaller<TClass>,
+#endif // WITH_HOT_RELOAD && WITH_HOT_RELOAD_CTORS
 		&TClass::AddReferencedObjects
 		);
 	check(ReturnClass);
diff --git a/Engine/Source/Runtime/CoreUObject/Public/UObject/ObjectBase.h b/Engine/Source/Runtime/CoreUObject/Public/UObject/ObjectBase.h
--- a/Engine/Source/Runtime/CoreUObject/Public/UObject/ObjectBase.h
+++ b/Engine/Source/Runtime/CoreUObject/Public/UObject/ObjectBase.h
@@ -98,6 +98,19 @@ enum EStaticConstructor				{EC_StaticConstructor};
 enum EInternal						{EC_InternalUseOnlyConstructor};
 enum ECppProperty					{EC_CppProperty};
 
+#if WITH_HOT_RELOAD && WITH_HOT_RELOAD_CTORS
+/** DO NOT USE. Helper class to invoke specialized hot-reload constructor. */
+class FVTableHelper
+{
+public:
+	/** DO NOT USE. This constructor is for internal usage only for hot-reload purposes. */
+	COREUOBJECT_API FVTableHelper()
+	{
+		EnsureRetrievingVTablePtr();
+	}
+};
+#endif // WITH_HOT_RELOAD && WITH_HOT_RELOAD_CTORS
+
 /** Empty API definition.  Used as a placeholder parameter when no DLL export/import API is needed for a UObject class */
 #define NO_API
 
@@ -1157,11 +1170,11 @@ namespace UM
 #define RELAY_CONSTRUCTOR(TClass, TSuperClass) TClass(const FObjectInitializer& ObjectInitializer) : TSuperClass(ObjectInitializer) {}
 
 #if !USE_COMPILED_IN_NATIVES
-	#define COMPILED_IN_FLAGS(TStaticFlags) (TStaticFlags& ~(CLASS_Intrinsic))
-	#define COMPILED_IN_INTRINSIC 0
+#define COMPILED_IN_FLAGS(TStaticFlags) (TStaticFlags& ~(CLASS_Intrinsic))
+#define COMPILED_IN_INTRINSIC 0
 #else
-	#define COMPILED_IN_FLAGS(TStaticFlags) (TStaticFlags | CLASS_Intrinsic)
-	#define COMPILED_IN_INTRINSIC 1
+#define COMPILED_IN_FLAGS(TStaticFlags) (TStaticFlags | CLASS_Intrinsic)
+#define COMPILED_IN_INTRINSIC 1
 #endif
 
 #define DECLARE_SERIALIZER( TClass ) \
@@ -1220,47 +1233,140 @@ public: \
 #define DEFINE_DEFAULT_OBJECT_INITIALIZER_CONSTRUCTOR_CALL(TClass) \
 	static void __DefaultConstructor(const FObjectInitializer& X) { new((EInternal*)X.GetObj())TClass(X); }
 
-#define DECLARE_CLASS_INTRINSIC_NO_CTOR(TClass,TSuperClass,TStaticFlags,TPackage) \
-	DECLARE_CLASS(TClass, TSuperClass, TStaticFlags | CLASS_Intrinsic, CASTCLASS_None, TPackage, NO_API) \
-	enum { IsIntrinsic = 1 }; \
-	static void StaticRegisterNatives##TClass() {} \
-	DECLARE_SERIALIZER(TClass) \
-	DEFINE_DEFAULT_OBJECT_INITIALIZER_CONSTRUCTOR_CALL(TClass)
-
-#define DECLARE_CLASS_INTRINSIC(TClass,TSuperClass,TStaticFlags,TPackage) \
-	DECLARE_CLASS(TClass,TSuperClass,TStaticFlags|CLASS_Intrinsic,CASTCLASS_None,TPackage,NO_API ) \
-	RELAY_CONSTRUCTOR(TClass, TSuperClass) \
-	enum {IsIntrinsic=1}; \
-	static void StaticRegisterNatives##TClass() {} \
-	DECLARE_SERIALIZER(TClass) \
-	DEFINE_DEFAULT_OBJECT_INITIALIZER_CONSTRUCTOR_CALL(TClass)
-
-#define DECLARE_CASTED_CLASS_INTRINSIC_WITH_API_NO_CTOR( TClass, TSuperClass, TStaticFlags, TPackage, TStaticCastFlags, TRequiredAPI ) \
-	DECLARE_CLASS(TClass, TSuperClass, TStaticFlags | CLASS_Intrinsic, TStaticCastFlags, TPackage, TRequiredAPI) \
-	enum { IsIntrinsic = 1 }; \
-	static void StaticRegisterNatives##TClass() {} \
-	DECLARE_SERIALIZER(TClass) \
-	DEFINE_DEFAULT_OBJECT_INITIALIZER_CONSTRUCTOR_CALL(TClass)
-
-#define DECLARE_CASTED_CLASS_INTRINSIC_WITH_API( TClass, TSuperClass, TStaticFlags, TPackage, TStaticCastFlags, TRequiredAPI ) \
-	DECLARE_CLASS(TClass,TSuperClass,TStaticFlags|CLASS_Intrinsic,TStaticCastFlags,TPackage,TRequiredAPI ) \
-	RELAY_CONSTRUCTOR(TClass, TSuperClass) \
-	enum {IsIntrinsic=1}; \
-	static void StaticRegisterNatives##TClass() {} \
-	DECLARE_SERIALIZER(TClass) \
-	DEFINE_DEFAULT_OBJECT_INITIALIZER_CONSTRUCTOR_CALL(TClass)
+#if WITH_HOT_RELOAD && WITH_HOT_RELOAD_CTORS
+	#define DECLARE_VTABLE_PTR_HELPER_CTOR(API, TClass) \
+		/** DO NOT USE. This constructor is for internal usage only for hot-reload purposes. */ \
+		API TClass(FVTableHelper& Helper);
 
-#define DECLARE_CASTED_CLASS_INTRINSIC( TClass, TSuperClass, TStaticFlags, TPackage, TStaticCastFlags ) \
-	DECLARE_CASTED_CLASS_INTRINSIC_WITH_API( TClass, TSuperClass, TStaticFlags, TPackage, TStaticCastFlags, NO_API) \
+	#define DEFINE_VTABLE_PTR_HELPER_CTOR(TClass) \
+		TClass::TClass(FVTableHelper& Helper) : Super(Helper) {};
+
+	#define DEFINE_VTABLE_PTR_HELPER_CTOR_CALLER(TClass) \
+		static UObject* __VTableCtorCaller(FVTableHelper& Helper) \
+		{ \
+			return new (EC_InternalUseOnlyConstructor, (UObject*)GetTransientPackage(), NAME_None, RF_NeedLoad | RF_ClassDefaultObject | RF_TagGarbageTemp) TClass(Helper); \
+		}
 
-#define DECLARE_CASTED_CLASS_INTRINSIC_NO_CTOR( TClass, TSuperClass, TStaticFlags, TPackage, TStaticCastFlags, TRequiredAPI ) \
-	DECLARE_CLASS(TClass,TSuperClass,TStaticFlags|CLASS_Intrinsic,TStaticCastFlags,TPackage, TRequiredAPI ) \
-	enum {IsIntrinsic=1}; \
-	static void StaticRegisterNatives##TClass() {} \
-	DECLARE_SERIALIZER(TClass) \
-	DEFINE_DEFAULT_OBJECT_INITIALIZER_CONSTRUCTOR_CALL(TClass)
+	#define DEFINE_VTABLE_PTR_HELPER_CTOR_CALLER_DUMMY() \
+		static UObject* __VTableCtorCaller(FVTableHelper& Helper) \
+		{ \
+			return nullptr; \
+		}
+
+	#define DECLARE_CLASS_INTRINSIC_NO_CTOR(TClass,TSuperClass,TStaticFlags,TPackage) \
+		DECLARE_CLASS(TClass, TSuperClass, TStaticFlags | CLASS_Intrinsic, CASTCLASS_None, TPackage, NO_API) \
+		enum { IsIntrinsic = 1 }; \
+		static void StaticRegisterNatives##TClass() {} \
+		DECLARE_SERIALIZER(TClass) \
+		DEFINE_DEFAULT_OBJECT_INITIALIZER_CONSTRUCTOR_CALL(TClass) \
+		static UObject* __VTableCtorCaller(FVTableHelper& Helper) \
+		{ \
+			return new (EC_InternalUseOnlyConstructor, (UObject*)GetTransientPackage(), NAME_None, RF_NeedLoad | RF_ClassDefaultObject | RF_TagGarbageTemp) TClass(Helper); \
+		}
+
+	#define DECLARE_CLASS_INTRINSIC(TClass,TSuperClass,TStaticFlags,TPackage) \
+		DECLARE_CLASS(TClass,TSuperClass,TStaticFlags|CLASS_Intrinsic,CASTCLASS_None,TPackage,NO_API ) \
+		RELAY_CONSTRUCTOR(TClass, TSuperClass) \
+		/** DO NOT USE. This constructor is for internal usage only for hot-reload purposes. */ \
+		TClass(FVTableHelper& Helper) : Super(Helper) {}; \
+		enum {IsIntrinsic=1}; \
+		static void StaticRegisterNatives##TClass() {} \
+		DECLARE_SERIALIZER(TClass) \
+		DEFINE_DEFAULT_OBJECT_INITIALIZER_CONSTRUCTOR_CALL(TClass) \
+		static UObject* __VTableCtorCaller(FVTableHelper& Helper) \
+		{ \
+			return new (EC_InternalUseOnlyConstructor, (UObject*)GetTransientPackage(), NAME_None, RF_NeedLoad | RF_ClassDefaultObject | RF_TagGarbageTemp) TClass(Helper); \
+		}
+
+	#define DECLARE_CASTED_CLASS_INTRINSIC_WITH_API_NO_CTOR( TClass, TSuperClass, TStaticFlags, TPackage, TStaticCastFlags, TRequiredAPI ) \
+		DECLARE_CLASS(TClass, TSuperClass, TStaticFlags | CLASS_Intrinsic, TStaticCastFlags, TPackage, TRequiredAPI) \
+		enum { IsIntrinsic = 1 }; \
+		static void StaticRegisterNatives##TClass() {} \
+		DECLARE_SERIALIZER(TClass) \
+		DEFINE_DEFAULT_OBJECT_INITIALIZER_CONSTRUCTOR_CALL(TClass) \
+		static UObject* __VTableCtorCaller(FVTableHelper& Helper) \
+		{ \
+			return new (EC_InternalUseOnlyConstructor, (UObject*)GetTransientPackage(), NAME_None, RF_NeedLoad | RF_ClassDefaultObject | RF_TagGarbageTemp) TClass(Helper); \
+		}
+
+	#define DECLARE_CASTED_CLASS_INTRINSIC_WITH_API( TClass, TSuperClass, TStaticFlags, TPackage, TStaticCastFlags, TRequiredAPI ) \
+		DECLARE_CLASS(TClass,TSuperClass,TStaticFlags|CLASS_Intrinsic,TStaticCastFlags,TPackage,TRequiredAPI ) \
+		RELAY_CONSTRUCTOR(TClass, TSuperClass) \
+		/** DO NOT USE. This constructor is for internal usage only for hot-reload purposes. */ \
+		TClass(FVTableHelper& Helper) : Super(Helper) {}; \
+		enum {IsIntrinsic=1}; \
+		static void StaticRegisterNatives##TClass() {} \
+		DECLARE_SERIALIZER(TClass) \
+		DEFINE_DEFAULT_OBJECT_INITIALIZER_CONSTRUCTOR_CALL(TClass) \
+		static UObject* __VTableCtorCaller(FVTableHelper& Helper) \
+		{ \
+			return new (EC_InternalUseOnlyConstructor, (UObject*)GetTransientPackage(), NAME_None, RF_NeedLoad | RF_ClassDefaultObject | RF_TagGarbageTemp) TClass(Helper); \
+		}
+
+	#define DECLARE_CASTED_CLASS_INTRINSIC_NO_CTOR_NO_VTABLE_CTOR( TClass, TSuperClass, TStaticFlags, TPackage, TStaticCastFlags, TRequiredAPI ) \
+		DECLARE_CLASS(TClass,TSuperClass,TStaticFlags|CLASS_Intrinsic,TStaticCastFlags,TPackage, TRequiredAPI ) \
+		enum {IsIntrinsic=1}; \
+		static void StaticRegisterNatives##TClass() {} \
+		DECLARE_SERIALIZER(TClass) \
+		DEFINE_DEFAULT_OBJECT_INITIALIZER_CONSTRUCTOR_CALL(TClass) \
+		static UObject* __VTableCtorCaller(FVTableHelper& Helper) \
+		{ \
+			return new (EC_InternalUseOnlyConstructor, (UObject*)GetTransientPackage(), NAME_None, RF_NeedLoad | RF_ClassDefaultObject | RF_TagGarbageTemp) TClass(Helper); \
+		}
+
+	#define DECLARE_CASTED_CLASS_INTRINSIC_NO_CTOR( TClass, TSuperClass, TStaticFlags, TPackage, TStaticCastFlags, TRequiredAPI ) \
+		DECLARE_CASTED_CLASS_INTRINSIC_NO_CTOR_NO_VTABLE_CTOR( TClass, TSuperClass, TStaticFlags, TPackage, TStaticCastFlags, TRequiredAPI ) \
+		/** DO NOT USE. This constructor is for internal usage only for hot-reload purposes. */ \
+		TClass(FVTableHelper& Helper) : Super(Helper) {}; \
+
+#else // WITH_HOT_RELOAD && WITH_HOT_RELOAD_CTORS
+	#define DECLARE_VTABLE_PTR_HELPER_CTOR(API, TClass)
+	#define DEFINE_VTABLE_PTR_HELPER_CTOR(TClass)
+	#define DEFINE_VTABLE_PTR_HELPER_CTOR_CALLER(TClass)
+	#define DEFINE_VTABLE_PTR_HELPER_CTOR_CALLER_DUMMY()
+
+	#define DECLARE_CLASS_INTRINSIC_NO_CTOR(TClass,TSuperClass,TStaticFlags,TPackage) \
+		DECLARE_CLASS(TClass, TSuperClass, TStaticFlags | CLASS_Intrinsic, CASTCLASS_None, TPackage, NO_API) \
+		enum { IsIntrinsic = 1 }; \
+		static void StaticRegisterNatives##TClass() {} \
+		DECLARE_SERIALIZER(TClass) \
+		DEFINE_DEFAULT_OBJECT_INITIALIZER_CONSTRUCTOR_CALL(TClass) \
+
+	#define DECLARE_CLASS_INTRINSIC(TClass,TSuperClass,TStaticFlags,TPackage) \
+		DECLARE_CLASS(TClass,TSuperClass,TStaticFlags|CLASS_Intrinsic,CASTCLASS_None,TPackage,NO_API ) \
+		RELAY_CONSTRUCTOR(TClass, TSuperClass) \
+		enum {IsIntrinsic=1}; \
+		static void StaticRegisterNatives##TClass() {} \
+		DECLARE_SERIALIZER(TClass) \
+		DEFINE_DEFAULT_OBJECT_INITIALIZER_CONSTRUCTOR_CALL(TClass) \
+
+	#define DECLARE_CASTED_CLASS_INTRINSIC_WITH_API_NO_CTOR( TClass, TSuperClass, TStaticFlags, TPackage, TStaticCastFlags, TRequiredAPI ) \
+		DECLARE_CLASS(TClass, TSuperClass, TStaticFlags | CLASS_Intrinsic, TStaticCastFlags, TPackage, TRequiredAPI) \
+		enum { IsIntrinsic = 1 }; \
+		static void StaticRegisterNatives##TClass() {} \
+		DECLARE_SERIALIZER(TClass) \
+		DEFINE_DEFAULT_OBJECT_INITIALIZER_CONSTRUCTOR_CALL(TClass) \
+
+	#define DECLARE_CASTED_CLASS_INTRINSIC_WITH_API( TClass, TSuperClass, TStaticFlags, TPackage, TStaticCastFlags, TRequiredAPI ) \
+		DECLARE_CLASS(TClass,TSuperClass,TStaticFlags|CLASS_Intrinsic,TStaticCastFlags,TPackage,TRequiredAPI ) \
+		RELAY_CONSTRUCTOR(TClass, TSuperClass) \
+		enum {IsIntrinsic=1}; \
+		static void StaticRegisterNatives##TClass() {} \
+		DECLARE_SERIALIZER(TClass) \
+		DEFINE_DEFAULT_OBJECT_INITIALIZER_CONSTRUCTOR_CALL(TClass) \
+
+	#define DECLARE_CASTED_CLASS_INTRINSIC_NO_CTOR( TClass, TSuperClass, TStaticFlags, TPackage, TStaticCastFlags, TRequiredAPI ) \
+		DECLARE_CLASS(TClass,TSuperClass,TStaticFlags|CLASS_Intrinsic,TStaticCastFlags,TPackage, TRequiredAPI ) \
+		enum {IsIntrinsic=1}; \
+		static void StaticRegisterNatives##TClass() {} \
+		DECLARE_SERIALIZER(TClass) \
+		DEFINE_DEFAULT_OBJECT_INITIALIZER_CONSTRUCTOR_CALL(TClass) \
+
+#endif // WITH_HOT_RELOAD && WITH_HOT_RELOAD_CTORS
 
 
+#define DECLARE_CASTED_CLASS_INTRINSIC( TClass, TSuperClass, TStaticFlags, TPackage, TStaticCastFlags ) \
+	DECLARE_CASTED_CLASS_INTRINSIC_WITH_API( TClass, TSuperClass, TStaticFlags, TPackage, TStaticCastFlags, NO_API) \
 
 // Declare that objects of class being defined reside within objects of the specified class.
 #define DECLARE_WITHIN( TWithinClass ) \
diff --git a/Engine/Source/Runtime/CoreUObject/Public/UObject/UObject.h b/Engine/Source/Runtime/CoreUObject/Public/UObject/UObject.h
--- a/Engine/Source/Runtime/CoreUObject/Public/UObject/UObject.h
+++ b/Engine/Source/Runtime/CoreUObject/Public/UObject/UObject.h
@@ -37,6 +37,12 @@ class COREUOBJECT_API UObject : public UObjectBaseUtility
 	// Declarations.
 	DECLARE_CLASS(UObject,UObject,CLASS_Abstract|CLASS_NoExport,CASTCLASS_None,CoreUObject,NO_API)
 	DEFINE_DEFAULT_OBJECT_INITIALIZER_CONSTRUCTOR_CALL(UObject)
+#if WITH_HOT_RELOAD && WITH_HOT_RELOAD_CTORS
+	static UObject* __VTableCtorCaller(FVTableHelper& Helper)
+	{
+		return new (EC_InternalUseOnlyConstructor, (UObject*)GetTransientPackage(), NAME_None, RF_NeedLoad | RF_ClassDefaultObject | RF_TagGarbageTemp) UObject(Helper);
+	}
+#endif // WITH_HOT_RELOAD && WITH_HOT_RELOAD_CTORS
 
 	typedef UObject WithinClass;
 	static const TCHAR* StaticConfigName() {return TEXT("Engine");}
@@ -45,6 +51,10 @@ class COREUOBJECT_API UObject : public UObjectBaseUtility
 	UObject();
 	UObject(const FObjectInitializer& ObjectInitializer);
 	UObject( EStaticConstructor, EObjectFlags InFlags );
+#if WITH_HOT_RELOAD && WITH_HOT_RELOAD_CTORS
+	/** DO NOT USE. This constructor is for internal usage only for hot-reload purposes. */
+	UObject(FVTableHelper& Helper);
+#endif // WITH_HOT_RELOAD && WITH_HOT_RELOAD_CTORS
 
 	static void StaticRegisterNativesUObject() 
 	{
@@ -760,6 +770,15 @@ private:
 	 * Saves just the section(s) for this class into the given ini file for the class (with just the changes from base)
 	 */
 	void UpdateSingleSectionOfConfigFile(const FString& ConfigIniName);
+
+#if WITH_HOT_RELOAD && WITH_HOT_RELOAD_CTORS
+	/**
+	 * Ensures that current thread is NOT during vtable ptr retrieval process
+	 * of some UClass.
+	 */
+	void EnsureNotRetrievingVTablePtr() const;
+#endif // WITH_HOT_RELOAD && WITH_HOT_RELOAD_CTORS
+
 public:
 	
 	/**
diff --git a/Engine/Source/Runtime/CoreUObject/Public/UObject/UnrealType.h b/Engine/Source/Runtime/CoreUObject/Public/UObject/UnrealType.h
--- a/Engine/Source/Runtime/CoreUObject/Public/UObject/UnrealType.h
+++ b/Engine/Source/Runtime/CoreUObject/Public/UObject/UnrealType.h
@@ -940,6 +940,11 @@ public:
 		SetElementSize();
 	}
 
+#if WITH_HOT_RELOAD && WITH_HOT_RELOAD_CTORS
+	/** DO NOT USE. This constructor is for internal usage only for hot-reload purposes. */
+	TProperty(FVTableHelper& Helper) : Super(Helper) {};
+#endif // WITH_HOT_RELOAD && WITH_HOT_RELOAD_CTORS
+
 	// UHT interface
 	virtual FString GetCPPType( FString* ExtendedTypeText=NULL, uint32 CPPExportFlags=0 ) const override
 	{
@@ -1038,6 +1043,11 @@ public:
 	{
 	}
 
+#if WITH_HOT_RELOAD && WITH_HOT_RELOAD_CTORS
+	/** DO NOT USE. This constructor is for internal usage only for hot-reload purposes. */
+	TProperty_WithEqualityAndSerializer(FVTableHelper& Helper) : Super(Helper) {};
+#endif // WITH_HOT_RELOAD && WITH_HOT_RELOAD_CTORS
+
 	// UProperty interface.
 	virtual bool Identical( const void* A, const void* B, uint32 PortFlags=0 ) const override
 	{
@@ -1196,6 +1206,11 @@ public:
 	{
 	}
 
+#if WITH_HOT_RELOAD && WITH_HOT_RELOAD_CTORS
+	/** DO NOT USE. This constructor is for internal usage only for hot-reload purposes. */
+	TProperty_Numeric(FVTableHelper& Helper) : Super(Helper) {};
+#endif // WITH_HOT_RELOAD && WITH_HOT_RELOAD_CTORS
+
 	FString GetCPPTypeForwardDeclaration() const override
 	{
 		return FString();
@@ -1754,6 +1769,11 @@ public:
 		this->PropertyClass = InClass;
 	}
 
+#if WITH_HOT_RELOAD && WITH_HOT_RELOAD_CTORS
+	/** DO NOT USE. This constructor is for internal usage only for hot-reload purposes. */
+	TUObjectPropertyBase(FVTableHelper& Helper) : Super(Helper) {};
+#endif // WITH_HOT_RELOAD && WITH_HOT_RELOAD_CTORS
+
 	// UProperty interface.
 	virtual bool ContainsObjectReference() const override
 	{
diff --git a/Engine/Source/Runtime/Engine/Classes/Engine/DestructibleFractureSettings.h b/Engine/Source/Runtime/Engine/Classes/Engine/DestructibleFractureSettings.h
--- a/Engine/Source/Runtime/Engine/Classes/Engine/DestructibleFractureSettings.h
+++ b/Engine/Source/Runtime/Engine/Classes/Engine/DestructibleFractureSettings.h
@@ -155,6 +155,11 @@ class UDestructibleFractureSettings
 {
 	GENERATED_UCLASS_BODY()
 
+#if WITH_HOT_RELOAD && WITH_HOT_RELOAD_CTORS
+	/** DO NOT USE. This constructor is for internal usage only for hot-reload purposes. */
+	UDestructibleFractureSettings(FVTableHelper& Helper);
+#endif // WITH_HOT_RELOAD && WITH_HOT_RELOAD_CTORS
+
 	/** The number of voronoi cell sites. */
 	UPROPERTY(EditAnywhere, Category = Voronoi, meta = (ClampMin = "1", UIMin = "1"))
 	int32									CellSiteCount;
diff --git a/Engine/Source/Runtime/Engine/Classes/Engine/Level.h b/Engine/Source/Runtime/Engine/Classes/Engine/Level.h
--- a/Engine/Source/Runtime/Engine/Classes/Engine/Level.h
+++ b/Engine/Source/Runtime/Engine/Classes/Engine/Level.h
@@ -495,6 +495,15 @@ public:
 	// Constructor.
 	ENGINE_API void Initialize(const FURL& InURL);
 	ULevel(const FObjectInitializer& ObjectInitializer);
+
+#if WITH_HOT_RELOAD && WITH_HOT_RELOAD_CTORS
+	/** DO NOT USE. This constructor is for internal usage only for hot-reload purposes. */
+	ULevel(FVTableHelper& Helper)
+		: Super(Helper)
+		, Actors(this)
+	{}
+#endif // WITH_HOT_RELOAD && WITH_HOT_RELOAD_CTORS
+
 	~ULevel();
 
 	// Begin UObject interface.
diff --git a/Engine/Source/Runtime/Engine/Classes/Engine/LocalPlayer.h b/Engine/Source/Runtime/Engine/Classes/Engine/LocalPlayer.h
--- a/Engine/Source/Runtime/Engine/Classes/Engine/LocalPlayer.h
+++ b/Engine/Source/Runtime/Engine/Classes/Engine/LocalPlayer.h
@@ -140,6 +140,11 @@ class ENGINE_API ULocalPlayer : public UPlayer
 {
 	GENERATED_UCLASS_BODY()
 
+#if WITH_HOT_RELOAD && WITH_HOT_RELOAD_CTORS
+	/** DO NOT USE. This constructor is for internal usage only for hot-reload purposes. */
+	ULocalPlayer(FVTableHelper& Helper) : Super(Helper), SlateOperations(FReply::Unhandled()) {}
+#endif // WITH_HOT_RELOAD && WITH_HOT_RELOAD_CTORS
+
 	/** The FUniqueNetId which this player is associated with. */
 	TSharedPtr<class FUniqueNetId> CachedUniqueNetId;
 
diff --git a/Engine/Source/Runtime/Engine/Classes/Engine/Polys.h b/Engine/Source/Runtime/Engine/Classes/Engine/Polys.h
--- a/Engine/Source/Runtime/Engine/Classes/Engine/Polys.h
+++ b/Engine/Source/Runtime/Engine/Classes/Engine/Polys.h
@@ -591,6 +591,14 @@ class UPolys : public UObject
 	, Element( this )
 	{}
 
+#if WITH_HOT_RELOAD && WITH_HOT_RELOAD_CTORS
+	/** DO NOT USE. This constructor is for internal usage only for hot-reload purposes. */
+	UPolys(FVTableHelper& Helper)
+		: Super(Helper)
+		, Element(this)
+	{}
+#endif // WITH_HOT_RELOAD && WITH_HOT_RELOAD_CTORS
+
 	// Begin UObject Interface
 	ENGINE_API virtual bool Modify(bool bAlwaysMarkDirty = false) override;	
 	ENGINE_API virtual void Serialize( FArchive& Ar ) override;
diff --git a/Engine/Source/Runtime/Engine/Private/DestructibleFractureSettings.cpp b/Engine/Source/Runtime/Engine/Private/DestructibleFractureSettings.cpp
--- a/Engine/Source/Runtime/Engine/Private/DestructibleFractureSettings.cpp
+++ b/Engine/Source/Runtime/Engine/Private/DestructibleFractureSettings.cpp
@@ -234,6 +234,14 @@ UDestructibleFractureSettings::UDestructibleFractureSettings(const FObjectInitia
 	OriginalSubmeshCount = 0;
 }
 
+#if WITH_HOT_RELOAD && WITH_HOT_RELOAD_CTORS
+UDestructibleFractureSettings::UDestructibleFractureSettings(FVTableHelper& Helper)
+	: Super(Helper)
+{
+
+}
+#endif // WITH_HOT_RELOAD && WITH_HOT_RELOAD_CTORS
+
 void UDestructibleFractureSettings::Serialize(FArchive& Ar)
 {
 	Super::Serialize(Ar);
diff --git a/Engine/Source/Runtime/Engine/Private/Model.cpp b/Engine/Source/Runtime/Engine/Private/Model.cpp
--- a/Engine/Source/Runtime/Engine/Private/Model.cpp
+++ b/Engine/Source/Runtime/Engine/Private/Model.cpp
@@ -514,6 +514,20 @@ UModel::UModel(const FObjectInitializer& ObjectInitializer)
 
 }
 
+#if WITH_HOT_RELOAD && WITH_HOT_RELOAD_CTORS
+UModel::UModel(FVTableHelper& Helper)
+	: Super(Helper)
+	, Nodes(this)
+	, Verts(this)
+	, Vectors(this)
+	, Points(this)
+	, Surfs(this)
+	, VertexBuffer(this)
+{
+
+}
+#endif // WITH_HOT_RELOAD && WITH_HOT_RELOAD_CTORS
+
 void UModel::Initialize(ABrush* Owner, bool InRootOutside)
 {
 	LightingGuid = FGuid::NewGuid();
diff --git a/Engine/Source/Runtime/Engine/Public/GenericOctree.h b/Engine/Source/Runtime/Engine/Public/GenericOctree.h
--- a/Engine/Source/Runtime/Engine/Public/GenericOctree.h
+++ b/Engine/Source/Runtime/Engine/Public/GenericOctree.h
@@ -731,6 +731,11 @@ public:
 
 	/** Initialization constructor. */
 	TOctree(const FVector& InOrigin,float InExtent);
+
+#if WITH_HOT_RELOAD && WITH_HOT_RELOAD_CTORS
+	/** DO NOT USE. This constructor is for internal usage only for hot-reload purposes. */
+	TOctree();
+#endif // WITH_HOT_RELOAD && WITH_HOT_RELOAD_CTORS
 			
 private:
 
diff --git a/Engine/Source/Runtime/Engine/Public/GenericOctree.inl b/Engine/Source/Runtime/Engine/Public/GenericOctree.inl
--- a/Engine/Source/Runtime/Engine/Public/GenericOctree.inl
+++ b/Engine/Source/Runtime/Engine/Public/GenericOctree.inl
@@ -353,6 +353,15 @@ TOctree<ElementType,OctreeSemantics>::TOctree(const FVector& InOrigin,float InEx
 {
 }
 
+#if WITH_HOT_RELOAD && WITH_HOT_RELOAD_CTORS
+template<typename ElementType, typename OctreeSemantics>
+TOctree<ElementType, OctreeSemantics>::TOctree()
+	: RootNode(nullptr)
+{
+	EnsureRetrievingVTablePtr();
+}
+#endif // WITH_HOT_RELOAD && WITH_HOT_RELOAD_CTORS
+
 template<typename ElementType,typename OctreeSemantics>
 void TOctree<ElementType,OctreeSemantics>::ApplyOffset(const FVector& InOffset)
 {
diff --git a/Engine/Source/Runtime/Engine/Public/GenericQuadTree.h b/Engine/Source/Runtime/Engine/Public/GenericQuadTree.h
--- a/Engine/Source/Runtime/Engine/Public/GenericQuadTree.h
+++ b/Engine/Source/Runtime/Engine/Public/GenericQuadTree.h
@@ -6,6 +6,12 @@ class TQuadTree
 {
 	typedef TQuadTree<ElementType, NodeCapacity> TreeType;
 public:
+
+#if WITH_HOT_RELOAD && WITH_HOT_RELOAD_CTORS
+	/** DO NOT USE. This constructor is for internal usage only for hot-reload purposes. */
+	TQuadTree();
+#endif // WITH_HOT_RELOAD && WITH_HOT_RELOAD_CTORS
+
 	TQuadTree(const FBox2D& InBox);
 
 	/** Inserts an object of type ElementType with an associated 2D box of size Box (log n)*/
@@ -128,6 +134,14 @@ TQuadTree<ElementType, NodeCapacity>::TQuadTree(const FBox2D& Box)
 	SubTrees[0] = SubTrees[1] = SubTrees[2] = SubTrees[3] = nullptr;
 }
 
+#if WITH_HOT_RELOAD && WITH_HOT_RELOAD_CTORS
+template <typename ElementType, int32 NodeCapacity>
+TQuadTree<ElementType, NodeCapacity>::TQuadTree()
+{
+	EnsureRetrievingVTablePtr();
+}
+#endif // WITH_HOT_RELOAD && WITH_HOT_RELOAD_CTORS
+
 template <typename ElementType, int32 NodeCapacity>
 TQuadTree<ElementType, NodeCapacity>::~TQuadTree()
 {
diff --git a/Engine/Source/Runtime/Engine/Public/Model.h b/Engine/Source/Runtime/Engine/Public/Model.h
--- a/Engine/Source/Runtime/Engine/Public/Model.h
+++ b/Engine/Source/Runtime/Engine/Public/Model.h
@@ -343,7 +343,14 @@ enum {MAX_NODES  = 65536};
 enum {MAX_POINTS = 128000};
 class UModel : public UObject
 {
-	DECLARE_CASTED_CLASS_INTRINSIC_NO_CTOR(UModel,UObject,0,Engine,0,ENGINE_API)
+#if WITH_HOT_RELOAD && WITH_HOT_RELOAD_CTORS
+	DECLARE_CASTED_CLASS_INTRINSIC_NO_CTOR_NO_VTABLE_CTOR(UModel, UObject, 0, Engine, 0, ENGINE_API)
+
+	/** DO NOT USE. This constructor is for internal usage only for hot-reload purposes. */
+	UModel(FVTableHelper& Helper);
+#else // WITH_HOT_RELOAD && WITH_HOT_RELOAD_CTORS
+	DECLARE_CASTED_CLASS_INTRINSIC_NO_CTOR(UModel, UObject, 0, Engine, 0, ENGINE_API)
+#endif // WITH_HOT_RELOAD && WITH_HOT_RELOAD_CTORS
 
 #if WITH_EDITOR
 	// Arrays and subobjects.
diff --git a/Engine/Source/Runtime/Engine/Public/Net/DataBunch.h b/Engine/Source/Runtime/Engine/Public/Net/DataBunch.h
--- a/Engine/Source/Runtime/Engine/Public/Net/DataBunch.h
+++ b/Engine/Source/Runtime/Engine/Public/Net/DataBunch.h
@@ -54,7 +54,7 @@ public:
 #endif
 
 	// Functions.
-	FOutBunch();
+	ENGINE_API FOutBunch();
 	ENGINE_API FOutBunch( class UChannel* InChannel, bool bClose );
 	ENGINE_API FOutBunch( UPackageMap * PackageMap, int64 InMaxBits = 1024 );
 
diff --git a/Engine/Source/Runtime/Engine/Public/ShowFlags.h b/Engine/Source/Runtime/Engine/Public/ShowFlags.h
--- a/Engine/Source/Runtime/Engine/Public/ShowFlags.h
+++ b/Engine/Source/Runtime/Engine/Public/ShowFlags.h
@@ -135,85 +135,17 @@
 	/** this constructor defines the default view settings */
 	FEngineShowFlags(EShowFlagInitMode InitMode)
 	{
-		if(InitMode == ESFIM_All0)
-		{
-			FMemory::Memset(this, 0x00, sizeof(*this));
-			return;
-		}
-
-		// Most flags are on by default. With the following line we only need disable flags.
-		FMemory::Memset(this, 0xff, sizeof(*this));
-
-		// The following code sets what should be off by default.
-		SetVisualizeHDR(false);
-		SetOverrideDiffuseAndSpecular(false);
-		SetReflectionOverride(false);
-		SetLpvLightingOnly(false);
-		SetVisualizeBuffer(false);
-		SetVectorFields(false);
-		SetGBufferHints(false);
-		SetCompositeEditorPrimitives(InitMode == ESFIM_Editor);
-		SetTestImage(false);
-		SetVisualizeDOF(false);
-		SetVisualizeAdaptiveDOF(false);
-		SetVertexColors(false);
-		SetVisualizeMotionBlur(false);
-		SetSelectionOutline(false);
-		SetDebugAI(false);
-		SetNavigation(false);
-		SetLightComplexity(false);
-		SetShaderComplexity(false);
-		SetStationaryLightOverlap(false);
-		SetLightMapDensity(false);
-		SetVisualizeLPV(false);
-		SetStreamingBounds(false);
-		SetConstraints(false);
-		SetCameraFrustums(false);
-		SetAudioRadius(InitMode == ESFIM_Editor);
-		SetBSPSplit(false);
-		SetBrushes(false);
-		SetEditor(InitMode == ESFIM_Editor);
-		SetLargeVertices(false);
-		SetGrid(InitMode == ESFIM_Editor);
-		SetMeshEdges(false);
-		SetCover(false);
-		SetSplines(InitMode == ESFIM_Editor);
-		SetSelection(InitMode == ESFIM_Editor);
-		SetModeWidgets(InitMode == ESFIM_Editor);
-		SetBounds(false);
-		SetHitProxies(false);
-		SetPropertyColoration(false);
-		SetLightInfluences(false);
-		SetPivot(InitMode == ESFIM_Editor);
-		SetShadowFrustums(false);
-		SetWireframe(false);
-		SetLightRadius(InitMode == ESFIM_Editor);
-		SetVolumes(InitMode == ESFIM_Editor);
-		SetGame(InitMode != ESFIM_Editor);
-		SetLevelColoration(false);
-		SetCollision(false);
-		SetCollisionPawn(false);
-		SetCollisionVisibility(false);
-		SetCameraAspectRatioBars(false);
-		SetCameraSafeFrames(false);
-		SetVisualizeOutOfBoundsPixels(false);
-		SetHighResScreenshotMask(false);
-		SetHMDDistortion(true); // only effective if a HMD is in use
-		SetStereoRendering(false);
-		SetDistanceCulledPrimitives(InitMode == ESFIM_Editor);
-		SetVisualizeLightCulling(false);
-		SetPrecomputedVisibilityCells(false);
-		SetVolumeLightingSamples(false);
-		// we enable it manually on the editor view ports
-		SetSnap(false);
-		SetVisualizeMeshDistanceFields(false);
-		SetVisualizeDistanceFieldAO(false);
-		SetVisualizeDistanceFieldGI(false);
-		SetVisualizeSSR(false);
-		SetForceGBuffer(false);
-		SetVisualizeSSS(false);
-		SetVisualizeBloom(false);
+		Init(InitMode);
 	}
 
+#if WITH_HOT_RELOAD && WITH_HOT_RELOAD_CTORS
+	/** DO NOT USE. This constructor is for internal usage only for hot-reload purposes. */
+	FEngineShowFlags()
+	{
+		EnsureRetrievingVTablePtr();
+		Init(ESFIM_Game);
+	}
+#endif // WITH_HOT_RELOAD && WITH_HOT_RELOAD_CTORS
+
 	static bool CanBeToggledInEditor(const TCHAR* Name)
 	{
@@ -326,6 +258,89 @@
 	}
 
 private:
+	/** Initializes the struct. */
+	void Init(EShowFlagInitMode InitMode)
+	{
+		if (InitMode == ESFIM_All0)
+		{
+			FMemory::Memset(this, 0x00, sizeof(*this));
+			return;
+		}
 
+		// Most flags are on by default. With the following line we only need disable flags.
+		FMemory::Memset(this, 0xff, sizeof(*this));
+
+		// The following code sets what should be off by default.
+		SetVisualizeHDR(false);
+		SetOverrideDiffuseAndSpecular(false);
+		SetReflectionOverride(false);
+		SetLpvLightingOnly(false);
+		SetVisualizeBuffer(false);
+		SetVectorFields(false);
+		SetGBufferHints(false);
+		SetCompositeEditorPrimitives(InitMode == ESFIM_Editor);
+		SetTestImage(false);
+		SetVisualizeDOF(false);
+		SetVisualizeAdaptiveDOF(false);
+		SetVertexColors(false);
+		SetVisualizeMotionBlur(false);
+		SetSelectionOutline(false);
+		SetDebugAI(false);
+		SetNavigation(false);
+		SetLightComplexity(false);
+		SetShaderComplexity(false);
+		SetStationaryLightOverlap(false);
+		SetLightMapDensity(false);
+		SetVisualizeLPV(false);
+		SetStreamingBounds(false);
+		SetConstraints(false);
+		SetCameraFrustums(false);
+		SetAudioRadius(InitMode == ESFIM_Editor);
+		SetBSPSplit(false);
+		SetBrushes(false);
+		SetEditor(InitMode == ESFIM_Editor);
+		SetLargeVertices(false);
+		SetGrid(InitMode == ESFIM_Editor);
+		SetMeshEdges(false);
+		SetCover(false);
+		SetSplines(InitMode == ESFIM_Editor);
+		SetSelection(InitMode == ESFIM_Editor);
+		SetModeWidgets(InitMode == ESFIM_Editor);
+		SetBounds(false);
+		SetHitProxies(false);
+		SetPropertyColoration(false);
+		SetLightInfluences(false);
+		SetPivot(InitMode == ESFIM_Editor);
+		SetShadowFrustums(false);
+		SetWireframe(false);
+		SetLightRadius(InitMode == ESFIM_Editor);
+		SetVolumes(InitMode == ESFIM_Editor);
+		SetGame(InitMode != ESFIM_Editor);
+		SetLevelColoration(false);
+		SetCollision(false);
+		SetCollisionPawn(false);
+		SetCollisionVisibility(false);
+		SetCameraAspectRatioBars(false);
+		SetCameraSafeFrames(false);
+		SetVisualizeOutOfBoundsPixels(false);
+		SetHighResScreenshotMask(false);
+		SetHMDDistortion(true); // only effective if a HMD is in use
+		SetStereoRendering(false);
+		SetDistanceCulledPrimitives(InitMode == ESFIM_Editor);
+		SetVisualizeLightCulling(false);
+		SetPrecomputedVisibilityCells(false);
+		SetVolumeLightingSamples(false);
+		// we enable it manually on the editor view ports
+		SetSnap(false);
+		SetVisualizeMeshDistanceFields(false);
+		SetVisualizeDistanceFieldAO(false);
+		SetVisualizeDistanceFieldGI(false);
+		SetVisualizeSSR(false);
+		SetForceGBuffer(false);
+		SetVisualizeSSS(false);
+		SetVisualizeBloom(false);
+	}
+
+
 	/**
 	 * Allows to iterate through all flags (if you want to change or read the actual flag settings you need to pass in your own instance of FEngineShowFlags).
diff --git a/Engine/Source/Runtime/Landscape/Classes/LandscapeComponent.h b/Engine/Source/Runtime/Landscape/Classes/LandscapeComponent.h
--- a/Engine/Source/Runtime/Landscape/Classes/LandscapeComponent.h
+++ b/Engine/Source/Runtime/Landscape/Classes/LandscapeComponent.h
@@ -2,6 +2,9 @@
 
 #pragma once
 
+#include "LightMap.h"
+#include "ShadowMap.h"
+
 #include "SceneTypes.h"
 #include "StaticLighting.h"
 #include "Components/PrimitiveComponent.h"
