Making asserts in UObject constructors more verbose

diff --git a/Engine/Source/Runtime/CoreUObject/Private/UObject/UObjectGlobals.cpp b/Engine/Source/Runtime/CoreUObject/Private/UObject/UObjectGlobals.cpp
--- a/Engine/Source/Runtime/CoreUObject/Private/UObject/UObjectGlobals.cpp
+++ b/Engine/Source/Runtime/CoreUObject/Private/UObject/UObjectGlobals.cpp
@@ -1921,7 +1921,7 @@ UObject::UObject()
 	FObjectInitializer* ObjectInitializerPtr = FTlsObjectInitializers::Top();
 	UE_CLOG(!ObjectInitializerPtr, LogUObjectGlobals, Fatal, TEXT("%s is not being constructed with either NewObject, NewNamedObject or ConstructObject."), *GetName());
 	FObjectInitializer& ObjectInitializer = *ObjectInitializerPtr;
-	check(!ObjectInitializer.Obj || ObjectInitializer.Obj == this);
+	UE_CLOG(ObjectInitializer.Obj != nullptr && ObjectInitializer.Obj != this, LogUObjectGlobals, Fatal, TEXT("UObject() constructor called but it's not the object that's currently being constructed with NewObject. Maybe you trying to construct it on the stack which is not supported."));
 	const_cast<FObjectInitializer&>(ObjectInitializer).Obj = this;
 	const_cast<FObjectInitializer&>(ObjectInitializer).FinalizeSubobjectClassInitialization();
 }
@@ -1932,7 +1932,7 @@ UObject::UObject(const FObjectInitializer& ObjectInitializer)
 	EnsureNotRetrievingVTablePtr();
 #endif // WITH_HOT_RELOAD && WITH_HOT_RELOAD_CTORS
 
-	check(!ObjectInitializer.Obj || ObjectInitializer.Obj == this);
+	UE_CLOG(ObjectInitializer.Obj != nullptr && ObjectInitializer.Obj != this, LogUObjectGlobals, Fatal, TEXT("UObject(const FObjectInitializer&) constructor called but it's not the object that's currently being constructed with NewObject. Maybe you trying to construct it on the stack which is not supported."));
 	const_cast<FObjectInitializer&>(ObjectInitializer).Obj = this;
 	const_cast<FObjectInitializer&>(ObjectInitializer).FinalizeSubobjectClassInitialization();
 }
