Adding copy constructors to FGCObject

diff --git a/Engine/Source/Runtime/CoreUObject/Public/UObject/GCObject.h b/Engine/Source/Runtime/CoreUObject/Public/UObject/GCObject.h
--- a/Engine/Source/Runtime/CoreUObject/Public/UObject/GCObject.h
+++ b/Engine/Source/Runtime/CoreUObject/Public/UObject/GCObject.h
@@ -66,6 +66,18 @@ public:
  */
 class COREUOBJECT_API FGCObject
 {
+	void Init()
+	{
+		// Some objects can get created after the engine started shutting down (lazy init of singletons etc).
+		if (!GIsRequestingExit)
+		{
+			StaticInit();
+			check(GGCObjectReferencer);
+			// Add this instance to the referencer's list
+			GGCObjectReferencer->AddObject(this);
+		}
+	}
+
 public:
 	/**
 	 * The static object referencer object that is shared across all
@@ -91,14 +103,19 @@ public:
 	 */
 	FGCObject(void)
 	{
-		// Some objects can get created after the engine started shutting down (lazy init of singletons etc).
-		if (!GIsRequestingExit)
-		{
-			StaticInit();
-			check(GGCObjectReferencer);
-			// Add this instance to the referencer's list
-			GGCObjectReferencer->AddObject(this);
-		}
+		Init();
+	}
+
+	/** Copy constructor */
+	FGCObject(FGCObject const&)
+	{
+		Init();
+	}
+
+	/** Move constructor */
+	FGCObject(FGCObject&&)
+	{
+		Init();
 	}
 
 	/**
