Fix for FLinkerPlaceholderBase with USE_DEFERRED_DEPENDENCY_CHECK_VERIFICATION_TESTS

diff --git a/Engine/Source/Runtime/CoreUObject/Private/UObject/LinkerPlaceholderBase.cpp b/Engine/Source/Runtime/CoreUObject/Private/UObject/LinkerPlaceholderBase.cpp
--- a/Engine/Source/Runtime/CoreUObject/Private/UObject/LinkerPlaceholderBase.cpp
+++ b/Engine/Source/Runtime/CoreUObject/Private/UObject/LinkerPlaceholderBase.cpp
@@ -8,18 +8,20 @@
 void FLinkerPlaceholderBase::AddReferencingProperty(UProperty* ReferencingProperty)
 {
 #if USE_DEFERRED_DEPENDENCY_CHECK_VERIFICATION_TESTS
+	UObject* ThisAsObject = GetPlaceholderAsUObject();
+	check(ThisAsObject);
 	FObjectImport* PlaceholderImport = nullptr;
 	if (ULinkerLoad* PropertyLinker = ReferencingProperty->GetLinker())
 	{
 		for (FObjectImport& Import : PropertyLinker->ImportMap)
 		{
-			if (Import.XObject == this)
+			if (Import.XObject == ThisAsObject)
 			{
 				PlaceholderImport = &Import;
 				break;
 			}
 		}
-		check(GetOuter() == PropertyLinker->LinkerRoot);
+		check(ThisAsObject->GetOuter() == PropertyLinker->LinkerRoot);
 		check((PropertyLinker->LoadFlags & LOAD_DeferDependencyLoads) || FBlueprintSupport::IsResolvingDeferredDependenciesDisabled());
 	}
 	// if this check hits, then we're adding dependencies after we've 
diff --git a/Engine/Source/Runtime/CoreUObject/Private/UObject/LinkerPlaceholderBase.h b/Engine/Source/Runtime/CoreUObject/Private/UObject/LinkerPlaceholderBase.h
--- a/Engine/Source/Runtime/CoreUObject/Private/UObject/LinkerPlaceholderBase.h
+++ b/Engine/Source/Runtime/CoreUObject/Private/UObject/LinkerPlaceholderBase.h
@@ -58,6 +58,8 @@ public:
 	 */
 	void RemovePropertyReference(UProperty* ReferencingProperty);
 
+	virtual UObject* GetPlaceholderAsUObject() PURE_VIRTUAL(FLinkerPlaceholderBase::GetPlaceholderUObject, return nullptr;);
+
 public:
 	/** Set by the ULinkerLoad that created this instance, tracks what import this was used in place of. */
 	int32 ImportIndex;
diff --git a/Engine/Source/Runtime/CoreUObject/Private/UObject/LinkerPlaceholderClass.h b/Engine/Source/Runtime/CoreUObject/Private/UObject/LinkerPlaceholderClass.h
--- a/Engine/Source/Runtime/CoreUObject/Private/UObject/LinkerPlaceholderClass.h
+++ b/Engine/Source/Runtime/CoreUObject/Private/UObject/LinkerPlaceholderClass.h
@@ -35,6 +35,10 @@ public:
 
 	// FLinkerPlaceholderBase 
 	virtual int32 GetRefCount() const override;
+	virtual UObject* GetPlaceholderAsUObject() override
+	{
+		return this;
+	}
 	// End of FLinkerPlaceholderBase
 
 	/**
diff --git a/Engine/Source/Runtime/CoreUObject/Private/UObject/LinkerPlaceholderFunction.h b/Engine/Source/Runtime/CoreUObject/Private/UObject/LinkerPlaceholderFunction.h
--- a/Engine/Source/Runtime/CoreUObject/Private/UObject/LinkerPlaceholderFunction.h
+++ b/Engine/Source/Runtime/CoreUObject/Private/UObject/LinkerPlaceholderFunction.h
@@ -17,6 +17,10 @@ public:
 
 	// FPlaceholderBase 
 	virtual int32 GetRefCount() const override;
+	virtual UObject* GetPlaceholderAsUObject() override
+	{
+		return this;
+	}
 	// End of FPlaceholderBase
 
 	int32 ReplaceTrackedReferences(UFunction* ReplacementClass);
