Changed "WITH_HOT_RELOAD && WITH_HOT_RELOAD_CTORS" to "WITH_HOT_RELOAD_CTORS".

diff --git a/Engine/Source/Developer/GameplayDebugger/Classes/GameplayDebuggingControllerComponent.h b/Engine/Source/Developer/GameplayDebugger/Classes/GameplayDebuggingControllerComponent.h
--- a/Engine/Source/Developer/GameplayDebugger/Classes/GameplayDebuggingControllerComponent.h
+++ b/Engine/Source/Developer/GameplayDebugger/Classes/GameplayDebuggingControllerComponent.h
@@ -17,10 +17,10 @@ class GAMEPLAYDEBUGGER_API UGameplayDebuggingControllerComponent : public UActor
 {
 	GENERATED_UCLASS_BODY()
 
-#if WITH_HOT_RELOAD && WITH_HOT_RELOAD_CTORS
+#if WITH_HOT_RELOAD_CTORS
 	/** DO NOT USE. This constructor is for internal usage only for hot-reload purposes. */
 	UGameplayDebuggingControllerComponent(FVTableHelper& Helper);
-#endif // WITH_HOT_RELOAD && WITH_HOT_RELOAD_CTORS
+#endif // WITH_HOT_RELOAD_CTORS
 
 	virtual void TickComponent(float DeltaTime, enum ELevelTick TickType, FActorComponentTickFunction *ThisTickFunction) override;
 	virtual void BeginDestroy() override;
diff --git a/Engine/Source/Developer/GameplayDebugger/Private/GameplayDebuggingControllerComponent.cpp b/Engine/Source/Developer/GameplayDebugger/Private/GameplayDebuggingControllerComponent.cpp
--- a/Engine/Source/Developer/GameplayDebugger/Private/GameplayDebuggingControllerComponent.cpp
+++ b/Engine/Source/Developer/GameplayDebugger/Private/GameplayDebuggingControllerComponent.cpp
@@ -34,14 +34,14 @@ UGameplayDebuggingControllerComponent::UGameplayDebuggingControllerComponent(con
 #endif //!(UE_BUILD_SHIPPING || UE_BUILD_TEST)
 }
 
-#if WITH_HOT_RELOAD && WITH_HOT_RELOAD_CTORS
+#if WITH_HOT_RELOAD_CTORS
 UGameplayDebuggingControllerComponent::UGameplayDebuggingControllerComponent(FVTableHelper& Helper)
 	: Super(Helper)
 	, KeyPressActivationTime(0.0f)
 {
 
 }
-#endif // WITH_HOT_RELOAD && WITH_HOT_RELOAD_CTORS
+#endif // WITH_HOT_RELOAD_CTORS
 
 void UGameplayDebuggingControllerComponent::OnRegister()
 {
diff --git a/Engine/Source/Editor/UnrealEd/Private/HierarchicalLOD.cpp b/Engine/Source/Editor/UnrealEd/Private/HierarchicalLOD.cpp
--- a/Engine/Source/Editor/UnrealEd/Private/HierarchicalLOD.cpp
+++ b/Engine/Source/Editor/UnrealEd/Private/HierarchicalLOD.cpp
@@ -81,13 +81,13 @@ FHierarchicalLODBuilder::FHierarchicalLODBuilder(class UWorld* InWorld)
 :	World(InWorld)
 {}
 
-#if WITH_HOT_RELOAD && WITH_HOT_RELOAD_CTORS
+#if WITH_HOT_RELOAD_CTORS
 FHierarchicalLODBuilder::FHierarchicalLODBuilder()
 	: World(nullptr)
 {
 	EnsureRetrievingVTablePtr();
 }
-#endif // WITH_HOT_RELOAD && WITH_HOT_RELOAD_CTORS
+#endif // WITH_HOT_RELOAD_CTORS
 
 // build hierarchical cluster
 void FHierarchicalLODBuilder::Build()
diff --git a/Engine/Source/Editor/UnrealEd/Public/HierarchicalLOD.h b/Engine/Source/Editor/UnrealEd/Public/HierarchicalLOD.h
--- a/Engine/Source/Editor/UnrealEd/Public/HierarchicalLOD.h
+++ b/Engine/Source/Editor/UnrealEd/Public/HierarchicalLOD.h
@@ -51,10 +51,10 @@ struct UNREALED_API FHierarchicalLODBuilder
 {
 	FHierarchicalLODBuilder(class UWorld* InWorld);
 
-#if WITH_HOT_RELOAD && WITH_HOT_RELOAD_CTORS
+#if WITH_HOT_RELOAD_CTORS
 	/** DO NOT USE. This constructor is for internal usage only for hot-reload purposes. */
 	FHierarchicalLODBuilder();
-#endif // WITH_HOT_RELOAD && WITH_HOT_RELOAD_CTORS
+#endif // WITH_HOT_RELOAD_CTORS
 
 	// build hierarchical cluster
 	void Build();
diff --git a/Engine/Source/Runtime/AIModule/Classes/BehaviorTree/BehaviorTreeComponent.h b/Engine/Source/Runtime/AIModule/Classes/BehaviorTree/BehaviorTreeComponent.h
--- a/Engine/Source/Runtime/AIModule/Classes/BehaviorTree/BehaviorTreeComponent.h
+++ b/Engine/Source/Runtime/AIModule/Classes/BehaviorTree/BehaviorTreeComponent.h
@@ -85,10 +85,10 @@ protected:
 	bool TreeHasBeenStarted() const;
 
 public:
-#if WITH_HOT_RELOAD && WITH_HOT_RELOAD_CTORS
+#if WITH_HOT_RELOAD_CTORS
 	/** DO NOT USE. This constructor is for internal usage only for hot-reload purposes. */
 	UBehaviorTreeComponent(FVTableHelper& Helper);
-#endif // WITH_HOT_RELOAD && WITH_HOT_RELOAD_CTORS
+#endif // WITH_HOT_RELOAD_CTORS
 
 	virtual bool IsRunning() const override;
 	virtual bool IsPaused() const override;
diff --git a/Engine/Source/Runtime/AIModule/Private/BehaviorTree/BehaviorTreeComponent.cpp b/Engine/Source/Runtime/AIModule/Private/BehaviorTree/BehaviorTreeComponent.cpp
--- a/Engine/Source/Runtime/AIModule/Private/BehaviorTree/BehaviorTreeComponent.cpp
+++ b/Engine/Source/Runtime/AIModule/Private/BehaviorTree/BehaviorTreeComponent.cpp
@@ -31,14 +31,14 @@ UBehaviorTreeComponent::UBehaviorTreeComponent(const FObjectInitializer& ObjectI
 	bIsPaused = false;
 }
 
-#if WITH_HOT_RELOAD && WITH_HOT_RELOAD_CTORS
+#if WITH_HOT_RELOAD_CTORS
 UBehaviorTreeComponent::UBehaviorTreeComponent(FVTableHelper& Helper)
 	: Super(Helper)
 	, SearchData(*this)
 {
 
 }
-#endif // WITH_HOT_RELOAD && WITH_HOT_RELOAD_CTORS
+#endif // WITH_HOT_RELOAD_CTORS
 
 void UBehaviorTreeComponent::UninitializeComponent()
 {
diff --git a/Engine/Source/Runtime/Core/Private/Misc/ConfigCacheIni.cpp b/Engine/Source/Runtime/Core/Private/Misc/ConfigCacheIni.cpp
--- a/Engine/Source/Runtime/Core/Private/Misc/ConfigCacheIni.cpp
+++ b/Engine/Source/Runtime/Core/Private/Misc/ConfigCacheIni.cpp
@@ -1334,12 +1334,12 @@ FConfigCacheIni::FConfigCacheIni(EConfigCacheType InType)
 {
 }
 
-#if WITH_HOT_RELOAD && WITH_HOT_RELOAD_CTORS
+#if WITH_HOT_RELOAD_CTORS
 FConfigCacheIni::FConfigCacheIni()
 {
 	EnsureRetrievingVTablePtr();
 }
-#endif // WITH_HOT_RELOAD && WITH_HOT_RELOAD_CTORS
+#endif // WITH_HOT_RELOAD_CTORS
 
 FConfigCacheIni::~FConfigCacheIni()
 {
diff --git a/Engine/Source/Runtime/Core/Private/Misc/CoreMisc.cpp b/Engine/Source/Runtime/Core/Private/Misc/CoreMisc.cpp
--- a/Engine/Source/Runtime/Core/Private/Misc/CoreMisc.cpp
+++ b/Engine/Source/Runtime/Core/Private/Misc/CoreMisc.cpp
@@ -987,11 +987,11 @@ FBoolConfigValueHelper::FBoolConfigValueHelper(const TCHAR* Section, const TCHAR
 	GConfig->GetBool(Section, Key, bValue, Filename);
 }
 
-#if WITH_HOT_RELOAD && WITH_HOT_RELOAD_CTORS
+#if WITH_HOT_RELOAD_CTORS
 bool GIsRetrievingVTablePtr = false;
 
 void EnsureRetrievingVTablePtr()
 {
 	UE_CLOG(!GIsRetrievingVTablePtr, LogCore, Fatal, TEXT("This should be used only during vtable ptr retrieval process."));
 }
-#endif // WITH_HOT_RELOAD && WITH_HOT_RELOAD_CTORS
+#endif // WITH_HOT_RELOAD_CTORS
diff --git a/Engine/Source/Runtime/Core/Public/CoreGlobals.h b/Engine/Source/Runtime/Core/Public/CoreGlobals.h
--- a/Engine/Source/Runtime/Core/Public/CoreGlobals.h
+++ b/Engine/Source/Runtime/Core/Public/CoreGlobals.h
@@ -199,10 +199,10 @@ extern CORE_API bool GIsInitialLoad;
 /** true when we are routing ConditionalPostLoad/PostLoad to objects */
 extern CORE_API bool GIsRoutingPostLoad;
 
-#if WITH_HOT_RELOAD && WITH_HOT_RELOAD_CTORS
+#if WITH_HOT_RELOAD_CTORS
 /** true when we are retrieving VTablePtr from UClass */
 extern CORE_API bool GIsRetrievingVTablePtr;
-#endif // WITH_HOT_RELOAD && WITH_HOT_RELOAD_CTORS
+#endif // WITH_HOT_RELOAD_CTORS
 
 /** Steadily increasing frame counter. */
 extern CORE_API uint64 GFrameCounter;
@@ -282,10 +282,10 @@ extern CORE_API double GBlueprintCompileTime;
 /** Stack names from the VM to be unrolled when we assert */
 extern CORE_API TArray<FScriptTraceStackNode> GScriptStack;
 
-#if WITH_HOT_RELOAD && WITH_HOT_RELOAD_CTORS
+#if WITH_HOT_RELOAD_CTORS
 /**
  * Ensures that current thread is during retrieval of vtable ptr
  * of some UClass.
  */
 CORE_API void EnsureRetrievingVTablePtr();
-#endif // WITH_HOT_RELOAD && WITH_HOT_RELOAD_CTORS
\ No newline at end of file
+#endif // WITH_HOT_RELOAD_CTORS
\ No newline at end of file
diff --git a/Engine/Source/Runtime/Core/Public/Misc/ConfigCacheIni.h b/Engine/Source/Runtime/Core/Public/Misc/ConfigCacheIni.h
--- a/Engine/Source/Runtime/Core/Public/Misc/ConfigCacheIni.h
+++ b/Engine/Source/Runtime/Core/Public/Misc/ConfigCacheIni.h
@@ -173,10 +173,10 @@ public:
 	// Basic functions.
 	FConfigCacheIni(EConfigCacheType Type);
 
-#if WITH_HOT_RELOAD && WITH_HOT_RELOAD_CTORS
+#if WITH_HOT_RELOAD_CTORS
 	/** DO NOT USE. This constructor is for internal usage only for hot-reload purposes. */
 	FConfigCacheIni();
-#endif // WITH_HOT_RELOAD && WITH_HOT_RELOAD_CTORS
+#endif // WITH_HOT_RELOAD_CTORS
 
 	virtual ~FConfigCacheIni();
 
diff --git a/Engine/Source/Runtime/Core/Public/Templates/SharedPointer.h b/Engine/Source/Runtime/Core/Public/Templates/SharedPointer.h
--- a/Engine/Source/Runtime/Core/Public/Templates/SharedPointer.h
+++ b/Engine/Source/Runtime/Core/Public/Templates/SharedPointer.h
@@ -153,7 +153,7 @@ public:
 		Init(InObject);
 	}
 
-#if WITH_HOT_RELOAD && WITH_HOT_RELOAD_CTORS
+#if WITH_HOT_RELOAD_CTORS
 	/**
 	 * Constructs default shared reference that owns the default object for specified type.
 	 *
@@ -166,7 +166,7 @@ public:
 		EnsureRetrievingVTablePtr();
 		Init(Object);
 	}
-#endif // WITH_HOT_RELOAD && WITH_HOT_RELOAD_CTORS
+#endif // WITH_HOT_RELOAD_CTORS
 
 	/**
 	 * Constructs a shared reference using a proxy reference to a raw pointer. (See MakeShareable())
diff --git a/Engine/Source/Runtime/CoreUObject/Classes/Object.h b/Engine/Source/Runtime/CoreUObject/Classes/Object.h
--- a/Engine/Source/Runtime/CoreUObject/Classes/Object.h
+++ b/Engine/Source/Runtime/CoreUObject/Classes/Object.h
@@ -882,10 +882,10 @@ public:
 	 * Default UObject constructor.
 	 */
 	UObject(const FObjectInitializer& ObjectInitializer);
-#if WITH_HOT_RELOAD && WITH_HOT_RELOAD_CTORS
+#if WITH_HOT_RELOAD_CTORS
 	/** DO NOT USE. This constructor is for internal usage only for hot-reload purposes. */
 	UObject(FVTableHelper& Helper);
-#endif // WITH_HOT_RELOAD && WITH_HOT_RELOAD_CTORS
+#endif // WITH_HOT_RELOAD_CTORS
 
 	//=============================================================================
 	// K2 support functions.
diff --git a/Engine/Source/Runtime/CoreUObject/Private/UObject/Class.cpp b/Engine/Source/Runtime/CoreUObject/Private/UObject/Class.cpp
--- a/Engine/Source/Runtime/CoreUObject/Private/UObject/Class.cpp
+++ b/Engine/Source/Runtime/CoreUObject/Private/UObject/Class.cpp
@@ -2425,10 +2425,10 @@ class FRestoreClassInfo: public FRestoreForUObjectOverwrite
 	EClassCastFlags	CastFlags;
 	/** Saved ClassConstructor **/
 	UClass::ClassConstructorType Constructor;
-#if WITH_HOT_RELOAD && WITH_HOT_RELOAD_CTORS
+#if WITH_HOT_RELOAD_CTORS
 	/** Saved ClassVTableHelperCtorCaller **/
 	UClass::ClassVTableHelperCtorCallerType ClassVTableHelperCtorCaller;
-#endif // WITH_HOT_RELOAD && WITH_HOT_RELOAD_CTORS
+#endif // WITH_HOT_RELOAD_CTORS
 	/** Saved ClassConstructor **/
 	UClass::ClassAddReferencedObjectsType AddReferencedObjects;
 	/** Saved NativeFunctionLookupTable. */
@@ -2448,9 +2448,9 @@ public:
 		Flags(Save->ClassFlags & CLASS_Abstract),
 		CastFlags(Save->ClassCastFlags),
 		Constructor(Save->ClassConstructor),
-#if WITH_HOT_RELOAD && WITH_HOT_RELOAD_CTORS
+#if WITH_HOT_RELOAD_CTORS
 		ClassVTableHelperCtorCaller(Save->ClassVTableHelperCtorCaller),
-#endif // WITH_HOT_RELOAD && WITH_HOT_RELOAD_CTORS
+#endif // WITH_HOT_RELOAD_CTORS
 		AddReferencedObjects(Save->ClassAddReferencedObjects),
 		NativeFunctionLookupTable(Save->NativeFunctionLookupTable)
 	{
@@ -2465,9 +2465,9 @@ public:
 		Target->ClassFlags |= Flags;
 		Target->ClassCastFlags |= CastFlags;
 		Target->ClassConstructor = Constructor;
-#if WITH_HOT_RELOAD && WITH_HOT_RELOAD_CTORS
+#if WITH_HOT_RELOAD_CTORS
 		Target->ClassVTableHelperCtorCaller = ClassVTableHelperCtorCaller;
-#endif // WITH_HOT_RELOAD && WITH_HOT_RELOAD_CTORS
+#endif // WITH_HOT_RELOAD_CTORS
 		Target->ClassAddReferencedObjects = AddReferencedObjects;
 		Target->NativeFunctionLookupTable = NativeFunctionLookupTable;
 	}
@@ -2676,9 +2676,9 @@ void UClass::Bind()
 
 	UClass* SuperClass = GetSuperClass();
 	if (SuperClass && (ClassConstructor == nullptr || ClassAddReferencedObjects == nullptr
-#if WITH_HOT_RELOAD && WITH_HOT_RELOAD_CTORS
+#if WITH_HOT_RELOAD_CTORS
 		|| ClassVTableHelperCtorCaller == nullptr
-#endif // WITH_HOT_RELOAD && WITH_HOT_RELOAD_CTORS
+#endif // WITH_HOT_RELOAD_CTORS
 		))
 	{
 		// Chase down constructor in parent class.
@@ -2687,12 +2687,12 @@ void UClass::Bind()
 		{
 			ClassConstructor = SuperClass->ClassConstructor;
 		}
-#if WITH_HOT_RELOAD && WITH_HOT_RELOAD_CTORS
+#if WITH_HOT_RELOAD_CTORS
 		if (!ClassVTableHelperCtorCaller)
 		{
 			ClassVTableHelperCtorCaller = SuperClass->ClassVTableHelperCtorCaller;
 		}
-#endif // WITH_HOT_RELOAD && WITH_HOT_RELOAD_CTORS
+#endif // WITH_HOT_RELOAD_CTORS
 		if (!ClassAddReferencedObjects)
 		{
 			ClassAddReferencedObjects = SuperClass->ClassAddReferencedObjects;
@@ -3191,9 +3191,9 @@ FArchive& operator<<(FArchive& Ar, FImplementedInterface& A)
 void UClass::PurgeClass(bool bRecompilingOnLoad)
 {
 	ClassConstructor = nullptr;
-#if WITH_HOT_RELOAD && WITH_HOT_RELOAD_CTORS
+#if WITH_HOT_RELOAD_CTORS
 	ClassVTableHelperCtorCaller = nullptr;
-#endif // WITH_HOT_RELOAD && WITH_HOT_RELOAD_CTORS
+#endif // WITH_HOT_RELOAD_CTORS
 	ClassFlags = 0;
 	ClassCastFlags = 0;
 	ClassUnique = 0;
@@ -3345,16 +3345,16 @@ UClass::UClass
 	const TCHAR*    InConfigName,
 	EObjectFlags	InFlags,
 	ClassConstructorType InClassConstructor,
-#if WITH_HOT_RELOAD && WITH_HOT_RELOAD_CTORS
+#if WITH_HOT_RELOAD_CTORS
 	ClassVTableHelperCtorCallerType InClassVTableHelperCtorCaller,
-#endif // WITH_HOT_RELOAD && WITH_HOT_RELOAD_CTORS
+#endif // WITH_HOT_RELOAD_CTORS
 	ClassAddReferencedObjectsType InClassAddReferencedObjects
 )
 :	UStruct					( EC_StaticConstructor, InSize, InFlags )
 ,	ClassConstructor		( InClassConstructor )
-#if WITH_HOT_RELOAD && WITH_HOT_RELOAD_CTORS
+#if WITH_HOT_RELOAD_CTORS
 ,	ClassVTableHelperCtorCaller(InClassVTableHelperCtorCaller)
-#endif // WITH_HOT_RELOAD && WITH_HOT_RELOAD_CTORS
+#endif // WITH_HOT_RELOAD_CTORS
 ,	ClassAddReferencedObjects( InClassAddReferencedObjects )
 ,	ClassFlags				( InClassFlags | CLASS_Native )
 ,	ClassCastFlags			( InClassCastFlags )
diff --git a/Engine/Source/Runtime/CoreUObject/Private/UObject/LinkerPlaceholderClass.cpp b/Engine/Source/Runtime/CoreUObject/Private/UObject/LinkerPlaceholderClass.cpp
--- a/Engine/Source/Runtime/CoreUObject/Private/UObject/LinkerPlaceholderClass.cpp
+++ b/Engine/Source/Runtime/CoreUObject/Private/UObject/LinkerPlaceholderClass.cpp
@@ -56,9 +56,9 @@ void ULinkerPlaceholderClass::PostInitProperties()
 void ULinkerPlaceholderClass::Bind()
 {
 	ClassConstructor = InternalConstructor<ULinkerPlaceholderClass>;
-#if WITH_HOT_RELOAD && WITH_HOT_RELOAD_CTORS
+#if WITH_HOT_RELOAD_CTORS
 	ClassVTableHelperCtorCaller = InternalVTableHelperCtorCaller<ULinkerPlaceholderClass>;
-#endif // WITH_HOT_RELOAD && WITH_HOT_RELOAD_CTORS
+#endif // WITH_HOT_RELOAD_CTORS
 	Super::Bind();
 
 	ClassAddReferencedObjects = &ULinkerPlaceholderClass::AddReferencedObjects;
diff --git a/Engine/Source/Runtime/CoreUObject/Private/UObject/Obj.cpp b/Engine/Source/Runtime/CoreUObject/Private/UObject/Obj.cpp
--- a/Engine/Source/Runtime/CoreUObject/Private/UObject/Obj.cpp
+++ b/Engine/Source/Runtime/CoreUObject/Private/UObject/Obj.cpp
@@ -50,12 +50,12 @@ static UPackage*			GObjTransientPkg								= NULL;
 UObject::UObject( EStaticConstructor, EObjectFlags InFlags )
 : UObjectBaseUtility(InFlags | RF_Native | RF_RootSet)
 {
-#if WITH_HOT_RELOAD && WITH_HOT_RELOAD_CTORS
+#if WITH_HOT_RELOAD_CTORS
 	EnsureNotRetrievingVTablePtr();
-#endif // WITH_HOT_RELOAD && WITH_HOT_RELOAD_CTORS
+#endif // WITH_HOT_RELOAD_CTORS
 }
 
-#if WITH_HOT_RELOAD && WITH_HOT_RELOAD_CTORS
+#if WITH_HOT_RELOAD_CTORS
 UObject::UObject(FVTableHelper& Helper)
 {
 	EnsureRetrievingVTablePtr();
@@ -78,7 +78,7 @@ void UObject::EnsureNotRetrievingVTablePtr() const
 {
 	UE_CLOG(GIsRetrievingVTablePtr, LogCore, Fatal, TEXT("We are currently retrieving VTable ptr. Please use FVTableHelper constructor instead."));
 }
-#endif // WITH_HOT_RELOAD && WITH_HOT_RELOAD_CTORS
+#endif // WITH_HOT_RELOAD_CTORS
 
 UObject* UObject::CreateDefaultSubobject(FName SubobjectFName, UClass* ReturnType, UClass* ClassToCreateByDefault, bool bIsRequired, bool bAbstract, bool bIsTransient)
 {
diff --git a/Engine/Source/Runtime/CoreUObject/Private/UObject/SavePackage.cpp b/Engine/Source/Runtime/CoreUObject/Private/UObject/SavePackage.cpp
--- a/Engine/Source/Runtime/CoreUObject/Private/UObject/SavePackage.cpp
+++ b/Engine/Source/Runtime/CoreUObject/Private/UObject/SavePackage.cpp
@@ -2217,9 +2217,9 @@ static bool ValidateConformCompatibility(UPackage* NewPackage, ULinkerLoad* OldL
 		if (OldClass != NULL && NewClass != NULL && OldClass->HasAnyFlags(RF_Native) && NewClass->HasAnyFlags(RF_Native))
 		{
 			OldClass->ClassConstructor = NewClass->ClassConstructor;
-#if WITH_HOT_RELOAD && WITH_HOT_RELOAD_CTORS
+#if WITH_HOT_RELOAD_CTORS
 			OldClass->ClassVTableHelperCtorCaller = NewClass->ClassVTableHelperCtorCaller;
-#endif // WITH_HOT_RELOAD && WITH_HOT_RELOAD_CTORS
+#endif // WITH_HOT_RELOAD_CTORS
 			OldClass->ClassAddReferencedObjects = NewClass->ClassAddReferencedObjects;
 		}
 	}
diff --git a/Engine/Source/Runtime/CoreUObject/Private/UObject/UObjectGlobals.cpp b/Engine/Source/Runtime/CoreUObject/Private/UObject/UObjectGlobals.cpp
--- a/Engine/Source/Runtime/CoreUObject/Private/UObject/UObjectGlobals.cpp
+++ b/Engine/Source/Runtime/CoreUObject/Private/UObject/UObjectGlobals.cpp
@@ -1914,9 +1914,9 @@ void UObject::PostInitProperties()
 
 UObject::UObject()
 {
-#if WITH_HOT_RELOAD && WITH_HOT_RELOAD_CTORS
+#if WITH_HOT_RELOAD_CTORS
 	EnsureNotRetrievingVTablePtr();
-#endif // WITH_HOT_RELOAD && WITH_HOT_RELOAD_CTORS
+#endif // WITH_HOT_RELOAD_CTORS
 
 	FObjectInitializer* ObjectInitializerPtr = FTlsObjectInitializers::Top();
 	UE_CLOG(!ObjectInitializerPtr, LogUObjectGlobals, Fatal, TEXT("%s is not being constructed with either NewObject, NewNamedObject or ConstructObject."), *GetName());
@@ -1928,9 +1928,9 @@ UObject::UObject()
 
 UObject::UObject(const FObjectInitializer& ObjectInitializer)
 {
-#if WITH_HOT_RELOAD && WITH_HOT_RELOAD_CTORS
+#if WITH_HOT_RELOAD_CTORS
 	EnsureNotRetrievingVTablePtr();
-#endif // WITH_HOT_RELOAD && WITH_HOT_RELOAD_CTORS
+#endif // WITH_HOT_RELOAD_CTORS
 
 	UE_CLOG(ObjectInitializer.Obj != nullptr && ObjectInitializer.Obj != this, LogUObjectGlobals, Fatal, TEXT("UObject(const FObjectInitializer&) constructor called but it's not the object that's currently being constructed with NewObject. Maybe you trying to construct it on the stack which is not supported."));
 	const_cast<FObjectInitializer&>(ObjectInitializer).Obj = this;
diff --git a/Engine/Source/Runtime/CoreUObject/Public/UObject/Class.h b/Engine/Source/Runtime/CoreUObject/Public/UObject/Class.h
--- a/Engine/Source/Runtime/CoreUObject/Public/UObject/Class.h
+++ b/Engine/Source/Runtime/CoreUObject/Public/UObject/Class.h
@@ -1750,15 +1750,15 @@ public:
 	friend class FRestoreClassInfo;
 
 	typedef void		(*ClassConstructorType)				(const FObjectInitializer&);
-#if WITH_HOT_RELOAD && WITH_HOT_RELOAD_CTORS
+#if WITH_HOT_RELOAD_CTORS
 	typedef UObject*	(*ClassVTableHelperCtorCallerType)	(FVTableHelper& Helper);
-#endif // WITH_HOT_RELOAD && WITH_HOT_RELOAD_CTORS
+#endif // WITH_HOT_RELOAD_CTORS
 	typedef void		(*ClassAddReferencedObjectsType)	(UObject*, class FReferenceCollector&);
 
 	ClassConstructorType ClassConstructor;
-#if WITH_HOT_RELOAD && WITH_HOT_RELOAD_CTORS
+#if WITH_HOT_RELOAD_CTORS
 	ClassVTableHelperCtorCallerType ClassVTableHelperCtorCaller;
-#endif // WITH_HOT_RELOAD && WITH_HOT_RELOAD_CTORS
+#endif // WITH_HOT_RELOAD_CTORS
 	/** Pointer to a static AddReferencedObjects method. */
 	ClassAddReferencedObjectsType ClassAddReferencedObjects;
 
@@ -1863,9 +1863,9 @@ public:
 	explicit UClass(const FObjectInitializer& ObjectInitializer, UClass* InSuperClass);
 	UClass( EStaticConstructor, FName InName, uint32 InSize, uint32 InClassFlags, EClassCastFlags InClassCastFlags,
 		const TCHAR* InClassConfigName, EObjectFlags InFlags, ClassConstructorType InClassConstructor,
-#if WITH_HOT_RELOAD && WITH_HOT_RELOAD_CTORS
+#if WITH_HOT_RELOAD_CTORS
 		ClassVTableHelperCtorCallerType InClassVTableHelperCtorCaller,
-#endif // WITH_HOT_RELOAD && WITH_HOT_RELOAD_CTORS
+#endif // WITH_HOT_RELOAD_CTORS
 		ClassAddReferencedObjectsType InClassAddReferencedObjects);
 
 #if WITH_HOT_RELOAD
@@ -2270,7 +2270,7 @@ void InternalConstructor( const FObjectInitializer& X )
 	T::__DefaultConstructor(X);
 }
 
-#if WITH_HOT_RELOAD && WITH_HOT_RELOAD_CTORS
+#if WITH_HOT_RELOAD_CTORS
 /**
  * Helper template to call the vtable ctor caller for a class
  */
@@ -2279,7 +2279,7 @@ UObject* InternalVTableHelperCtorCaller(FVTableHelper& Helper)
 {
 	return T::__VTableCtorCaller(Helper);
 }
-#endif // WITH_HOT_RELOAD && WITH_HOT_RELOAD_CTORS
+#endif // WITH_HOT_RELOAD_CTORS
 
 COREUOBJECT_API void InitializePrivateStaticClass(
 	class UClass* TClass_Super_StaticClass,
@@ -2348,9 +2348,9 @@ void GetPrivateStaticClassBody( const TCHAR* PackageName, const TCHAR* Name, UCl
 		TClass::StaticConfigName(),
 		EObjectFlags(RF_Public | RF_Standalone | RF_Transient | RF_Native | RF_RootSet),
 		(UClass::ClassConstructorType)InternalConstructor<TClass>,
-#if WITH_HOT_RELOAD && WITH_HOT_RELOAD_CTORS
+#if WITH_HOT_RELOAD_CTORS
 		(UClass::ClassVTableHelperCtorCallerType)InternalVTableHelperCtorCaller<TClass>,
-#endif // WITH_HOT_RELOAD && WITH_HOT_RELOAD_CTORS
+#endif // WITH_HOT_RELOAD_CTORS
 		&TClass::AddReferencedObjects
 		);
 	check(ReturnClass);
diff --git a/Engine/Source/Runtime/CoreUObject/Public/UObject/ObjectBase.h b/Engine/Source/Runtime/CoreUObject/Public/UObject/ObjectBase.h
--- a/Engine/Source/Runtime/CoreUObject/Public/UObject/ObjectBase.h
+++ b/Engine/Source/Runtime/CoreUObject/Public/UObject/ObjectBase.h
@@ -98,7 +98,7 @@ enum EStaticConstructor				{EC_StaticConstructor};
 enum EInternal						{EC_InternalUseOnlyConstructor};
 enum ECppProperty					{EC_CppProperty};
 
-#if WITH_HOT_RELOAD && WITH_HOT_RELOAD_CTORS
+#if WITH_HOT_RELOAD_CTORS
 /** DO NOT USE. Helper class to invoke specialized hot-reload constructor. */
 class FVTableHelper
 {
@@ -109,7 +109,7 @@ public:
 		EnsureRetrievingVTablePtr();
 	}
 };
-#endif // WITH_HOT_RELOAD && WITH_HOT_RELOAD_CTORS
+#endif // WITH_HOT_RELOAD_CTORS
 
 /** Empty API definition.  Used as a placeholder parameter when no DLL export/import API is needed for a UObject class */
 #define NO_API
@@ -1233,7 +1233,7 @@ public: \
 #define DEFINE_DEFAULT_OBJECT_INITIALIZER_CONSTRUCTOR_CALL(TClass) \
 	static void __DefaultConstructor(const FObjectInitializer& X) { new((EInternal*)X.GetObj())TClass(X); }
 
-#if WITH_HOT_RELOAD && WITH_HOT_RELOAD_CTORS
+#if WITH_HOT_RELOAD_CTORS
 	#define DECLARE_VTABLE_PTR_HELPER_CTOR(API, TClass) \
 		/** DO NOT USE. This constructor is for internal usage only for hot-reload purposes. */ \
 		API TClass(FVTableHelper& Helper);
@@ -1241,18 +1241,23 @@ public: \
 	#define DEFINE_VTABLE_PTR_HELPER_CTOR(TClass) \
 		TClass::TClass(FVTableHelper& Helper) : Super(Helper) {};
 
-	#define DEFINE_VTABLE_PTR_HELPER_CTOR_CALLER(TClass) \
-		static UObject* __VTableCtorCaller(FVTableHelper& Helper) \
-		{ \
-			return new (EC_InternalUseOnlyConstructor, (UObject*)GetTransientPackage(), NAME_None, RF_NeedLoad | RF_ClassDefaultObject | RF_TagGarbageTemp) TClass(Helper); \
-		}
-
 	#define DEFINE_VTABLE_PTR_HELPER_CTOR_CALLER_DUMMY() \
 		static UObject* __VTableCtorCaller(FVTableHelper& Helper) \
 		{ \
 			return nullptr; \
 		}
 
+	#if WITH_HOT_RELOAD
+		#define DEFINE_VTABLE_PTR_HELPER_CTOR_CALLER(TClass) \
+			static UObject* __VTableCtorCaller(FVTableHelper& Helper) \
+			{ \
+				return new (EC_InternalUseOnlyConstructor, (UObject*)GetTransientPackage(), NAME_None, RF_NeedLoad | RF_ClassDefaultObject | RF_TagGarbageTemp) TClass(Helper); \
+			}
+	#else // WITH_HOT_RELOAD
+		#define DEFINE_VTABLE_PTR_HELPER_CTOR_CALLER(TClass) \
+			DEFINE_VTABLE_PTR_HELPER_CTOR_CALLER_DUMMY()
+	#endif // WITH_HOT_RELOAD
+
 	#define DECLARE_CLASS_INTRINSIC_NO_CTOR(TClass,TSuperClass,TStaticFlags,TPackage) \
 		DECLARE_CLASS(TClass, TSuperClass, TStaticFlags | CLASS_Intrinsic, CASTCLASS_None, TPackage, NO_API) \
 		enum { IsIntrinsic = 1 }; \
@@ -1319,7 +1324,7 @@ public: \
 		/** DO NOT USE. This constructor is for internal usage only for hot-reload purposes. */ \
 		TClass(FVTableHelper& Helper) : Super(Helper) {}; \
 
-#else // WITH_HOT_RELOAD && WITH_HOT_RELOAD_CTORS
+#else // WITH_HOT_RELOAD_CTORS
 	#define DECLARE_VTABLE_PTR_HELPER_CTOR(API, TClass)
 	#define DEFINE_VTABLE_PTR_HELPER_CTOR(TClass)
 	#define DEFINE_VTABLE_PTR_HELPER_CTOR_CALLER(TClass)
@@ -1362,7 +1367,7 @@ public: \
 		DECLARE_SERIALIZER(TClass) \
 		DEFINE_DEFAULT_OBJECT_INITIALIZER_CONSTRUCTOR_CALL(TClass) \
 
-#endif // WITH_HOT_RELOAD && WITH_HOT_RELOAD_CTORS
+#endif // WITH_HOT_RELOAD_CTORS
 
 
 #define DECLARE_CASTED_CLASS_INTRINSIC( TClass, TSuperClass, TStaticFlags, TPackage, TStaticCastFlags ) \
diff --git a/Engine/Source/Runtime/CoreUObject/Public/UObject/UObject.h b/Engine/Source/Runtime/CoreUObject/Public/UObject/UObject.h
--- a/Engine/Source/Runtime/CoreUObject/Public/UObject/UObject.h
+++ b/Engine/Source/Runtime/CoreUObject/Public/UObject/UObject.h
@@ -37,12 +37,12 @@ class COREUOBJECT_API UObject : public UObjectBaseUtility
 	// Declarations.
 	DECLARE_CLASS(UObject,UObject,CLASS_Abstract|CLASS_NoExport,CASTCLASS_None,CoreUObject,NO_API)
 	DEFINE_DEFAULT_OBJECT_INITIALIZER_CONSTRUCTOR_CALL(UObject)
-#if WITH_HOT_RELOAD && WITH_HOT_RELOAD_CTORS
+#if WITH_HOT_RELOAD_CTORS
 	static UObject* __VTableCtorCaller(FVTableHelper& Helper)
 	{
 		return new (EC_InternalUseOnlyConstructor, (UObject*)GetTransientPackage(), NAME_None, RF_NeedLoad | RF_ClassDefaultObject | RF_TagGarbageTemp) UObject(Helper);
 	}
-#endif // WITH_HOT_RELOAD && WITH_HOT_RELOAD_CTORS
+#endif // WITH_HOT_RELOAD_CTORS
 
 	typedef UObject WithinClass;
 	static const TCHAR* StaticConfigName() {return TEXT("Engine");}
@@ -51,10 +51,10 @@ class COREUOBJECT_API UObject : public UObjectBaseUtility
 	UObject();
 	UObject(const FObjectInitializer& ObjectInitializer);
 	UObject( EStaticConstructor, EObjectFlags InFlags );
-#if WITH_HOT_RELOAD && WITH_HOT_RELOAD_CTORS
+#if WITH_HOT_RELOAD_CTORS
 	/** DO NOT USE. This constructor is for internal usage only for hot-reload purposes. */
 	UObject(FVTableHelper& Helper);
-#endif // WITH_HOT_RELOAD && WITH_HOT_RELOAD_CTORS
+#endif // WITH_HOT_RELOAD_CTORS
 
 	static void StaticRegisterNativesUObject() 
 	{
@@ -771,13 +771,13 @@ private:
 	 */
 	void UpdateSingleSectionOfConfigFile(const FString& ConfigIniName);
 
-#if WITH_HOT_RELOAD && WITH_HOT_RELOAD_CTORS
+#if WITH_HOT_RELOAD_CTORS
 	/**
 	 * Ensures that current thread is NOT during vtable ptr retrieval process
 	 * of some UClass.
 	 */
 	void EnsureNotRetrievingVTablePtr() const;
-#endif // WITH_HOT_RELOAD && WITH_HOT_RELOAD_CTORS
+#endif // WITH_HOT_RELOAD_CTORS
 
 public:
 	
diff --git a/Engine/Source/Runtime/CoreUObject/Public/UObject/UnrealType.h b/Engine/Source/Runtime/CoreUObject/Public/UObject/UnrealType.h
--- a/Engine/Source/Runtime/CoreUObject/Public/UObject/UnrealType.h
+++ b/Engine/Source/Runtime/CoreUObject/Public/UObject/UnrealType.h
@@ -940,10 +940,10 @@ public:
 		SetElementSize();
 	}
 
-#if WITH_HOT_RELOAD && WITH_HOT_RELOAD_CTORS
+#if WITH_HOT_RELOAD_CTORS
 	/** DO NOT USE. This constructor is for internal usage only for hot-reload purposes. */
 	TProperty(FVTableHelper& Helper) : Super(Helper) {};
-#endif // WITH_HOT_RELOAD && WITH_HOT_RELOAD_CTORS
+#endif // WITH_HOT_RELOAD_CTORS
 
 	// UHT interface
 	virtual FString GetCPPType( FString* ExtendedTypeText=NULL, uint32 CPPExportFlags=0 ) const override
@@ -1043,10 +1043,10 @@ public:
 	{
 	}
 
-#if WITH_HOT_RELOAD && WITH_HOT_RELOAD_CTORS
+#if WITH_HOT_RELOAD_CTORS
 	/** DO NOT USE. This constructor is for internal usage only for hot-reload purposes. */
 	TProperty_WithEqualityAndSerializer(FVTableHelper& Helper) : Super(Helper) {};
-#endif // WITH_HOT_RELOAD && WITH_HOT_RELOAD_CTORS
+#endif // WITH_HOT_RELOAD_CTORS
 
 	// UProperty interface.
 	virtual bool Identical( const void* A, const void* B, uint32 PortFlags=0 ) const override
@@ -1206,10 +1206,10 @@ public:
 	{
 	}
 
-#if WITH_HOT_RELOAD && WITH_HOT_RELOAD_CTORS
+#if WITH_HOT_RELOAD_CTORS
 	/** DO NOT USE. This constructor is for internal usage only for hot-reload purposes. */
 	TProperty_Numeric(FVTableHelper& Helper) : Super(Helper) {};
-#endif // WITH_HOT_RELOAD && WITH_HOT_RELOAD_CTORS
+#endif // WITH_HOT_RELOAD_CTORS
 
 	FString GetCPPTypeForwardDeclaration() const override
 	{
@@ -1769,10 +1769,10 @@ public:
 		this->PropertyClass = InClass;
 	}
 
-#if WITH_HOT_RELOAD && WITH_HOT_RELOAD_CTORS
+#if WITH_HOT_RELOAD_CTORS
 	/** DO NOT USE. This constructor is for internal usage only for hot-reload purposes. */
 	TUObjectPropertyBase(FVTableHelper& Helper) : Super(Helper) {};
-#endif // WITH_HOT_RELOAD && WITH_HOT_RELOAD_CTORS
+#endif // WITH_HOT_RELOAD_CTORS
 
 	// UProperty interface.
 	virtual bool ContainsObjectReference() const override
diff --git a/Engine/Source/Runtime/Engine/Classes/Engine/DestructibleFractureSettings.h b/Engine/Source/Runtime/Engine/Classes/Engine/DestructibleFractureSettings.h
--- a/Engine/Source/Runtime/Engine/Classes/Engine/DestructibleFractureSettings.h
+++ b/Engine/Source/Runtime/Engine/Classes/Engine/DestructibleFractureSettings.h
@@ -155,10 +155,10 @@ class UDestructibleFractureSettings
 {
 	GENERATED_UCLASS_BODY()
 
-#if WITH_HOT_RELOAD && WITH_HOT_RELOAD_CTORS
+#if WITH_HOT_RELOAD_CTORS
 	/** DO NOT USE. This constructor is for internal usage only for hot-reload purposes. */
 	UDestructibleFractureSettings(FVTableHelper& Helper);
-#endif // WITH_HOT_RELOAD && WITH_HOT_RELOAD_CTORS
+#endif // WITH_HOT_RELOAD_CTORS
 
 	/** The number of voronoi cell sites. */
 	UPROPERTY(EditAnywhere, Category = Voronoi, meta = (ClampMin = "1", UIMin = "1"))
diff --git a/Engine/Source/Runtime/Engine/Classes/Engine/Level.h b/Engine/Source/Runtime/Engine/Classes/Engine/Level.h
--- a/Engine/Source/Runtime/Engine/Classes/Engine/Level.h
+++ b/Engine/Source/Runtime/Engine/Classes/Engine/Level.h
@@ -496,13 +496,13 @@ public:
 	ENGINE_API void Initialize(const FURL& InURL);
 	ULevel(const FObjectInitializer& ObjectInitializer);
 
-#if WITH_HOT_RELOAD && WITH_HOT_RELOAD_CTORS
+#if WITH_HOT_RELOAD_CTORS
 	/** DO NOT USE. This constructor is for internal usage only for hot-reload purposes. */
 	ULevel(FVTableHelper& Helper)
 		: Super(Helper)
 		, Actors(this)
 	{}
-#endif // WITH_HOT_RELOAD && WITH_HOT_RELOAD_CTORS
+#endif // WITH_HOT_RELOAD_CTORS
 
 	~ULevel();
 
diff --git a/Engine/Source/Runtime/Engine/Classes/Engine/LocalPlayer.h b/Engine/Source/Runtime/Engine/Classes/Engine/LocalPlayer.h
--- a/Engine/Source/Runtime/Engine/Classes/Engine/LocalPlayer.h
+++ b/Engine/Source/Runtime/Engine/Classes/Engine/LocalPlayer.h
@@ -140,10 +140,10 @@ class ENGINE_API ULocalPlayer : public UPlayer
 {
 	GENERATED_UCLASS_BODY()
 
-#if WITH_HOT_RELOAD && WITH_HOT_RELOAD_CTORS
+#if WITH_HOT_RELOAD_CTORS
 	/** DO NOT USE. This constructor is for internal usage only for hot-reload purposes. */
 	ULocalPlayer(FVTableHelper& Helper) : Super(Helper), SlateOperations(FReply::Unhandled()) {}
-#endif // WITH_HOT_RELOAD && WITH_HOT_RELOAD_CTORS
+#endif // WITH_HOT_RELOAD_CTORS
 
 	/** The FUniqueNetId which this player is associated with. */
 	TSharedPtr<class FUniqueNetId> CachedUniqueNetId;
diff --git a/Engine/Source/Runtime/Engine/Classes/Engine/Polys.h b/Engine/Source/Runtime/Engine/Classes/Engine/Polys.h
--- a/Engine/Source/Runtime/Engine/Classes/Engine/Polys.h
+++ b/Engine/Source/Runtime/Engine/Classes/Engine/Polys.h
@@ -591,13 +591,13 @@ class UPolys : public UObject
 	, Element( this )
 	{}
 
-#if WITH_HOT_RELOAD && WITH_HOT_RELOAD_CTORS
+#if WITH_HOT_RELOAD_CTORS
 	/** DO NOT USE. This constructor is for internal usage only for hot-reload purposes. */
 	UPolys(FVTableHelper& Helper)
 		: Super(Helper)
 		, Element(this)
 	{}
-#endif // WITH_HOT_RELOAD && WITH_HOT_RELOAD_CTORS
+#endif // WITH_HOT_RELOAD_CTORS
 
 	// Begin UObject Interface
 	ENGINE_API virtual bool Modify(bool bAlwaysMarkDirty = false) override;	
diff --git a/Engine/Source/Runtime/Engine/Private/DestructibleFractureSettings.cpp b/Engine/Source/Runtime/Engine/Private/DestructibleFractureSettings.cpp
--- a/Engine/Source/Runtime/Engine/Private/DestructibleFractureSettings.cpp
+++ b/Engine/Source/Runtime/Engine/Private/DestructibleFractureSettings.cpp
@@ -234,13 +234,13 @@ UDestructibleFractureSettings::UDestructibleFractureSettings(const FObjectInitia
 	OriginalSubmeshCount = 0;
 }
 
-#if WITH_HOT_RELOAD && WITH_HOT_RELOAD_CTORS
+#if WITH_HOT_RELOAD_CTORS
 UDestructibleFractureSettings::UDestructibleFractureSettings(FVTableHelper& Helper)
 	: Super(Helper)
 {
 
 }
-#endif // WITH_HOT_RELOAD && WITH_HOT_RELOAD_CTORS
+#endif // WITH_HOT_RELOAD_CTORS
 
 void UDestructibleFractureSettings::Serialize(FArchive& Ar)
 {
diff --git a/Engine/Source/Runtime/Engine/Private/Model.cpp b/Engine/Source/Runtime/Engine/Private/Model.cpp
--- a/Engine/Source/Runtime/Engine/Private/Model.cpp
+++ b/Engine/Source/Runtime/Engine/Private/Model.cpp
@@ -532,7 +532,7 @@ UModel::UModel(const FObjectInitializer& ObjectInitializer)
 
 }
 
-#if WITH_HOT_RELOAD && WITH_HOT_RELOAD_CTORS
+#if WITH_HOT_RELOAD_CTORS
 UModel::UModel(FVTableHelper& Helper)
 	: Super(Helper)
 	, Nodes(this)
@@ -544,7 +544,7 @@ UModel::UModel(FVTableHelper& Helper)
 {
 
 }
-#endif // WITH_HOT_RELOAD && WITH_HOT_RELOAD_CTORS
+#endif // WITH_HOT_RELOAD_CTORS
 
 void UModel::Initialize(ABrush* Owner, bool InRootOutside)
 {
diff --git a/Engine/Source/Runtime/Engine/Public/GenericOctree.h b/Engine/Source/Runtime/Engine/Public/GenericOctree.h
--- a/Engine/Source/Runtime/Engine/Public/GenericOctree.h
+++ b/Engine/Source/Runtime/Engine/Public/GenericOctree.h
@@ -732,10 +732,10 @@ public:
 	/** Initialization constructor. */
 	TOctree(const FVector& InOrigin,float InExtent);
 
-#if WITH_HOT_RELOAD && WITH_HOT_RELOAD_CTORS
+#if WITH_HOT_RELOAD_CTORS
 	/** DO NOT USE. This constructor is for internal usage only for hot-reload purposes. */
 	TOctree();
-#endif // WITH_HOT_RELOAD && WITH_HOT_RELOAD_CTORS
+#endif // WITH_HOT_RELOAD_CTORS
 			
 private:
 
diff --git a/Engine/Source/Runtime/Engine/Public/GenericOctree.inl b/Engine/Source/Runtime/Engine/Public/GenericOctree.inl
--- a/Engine/Source/Runtime/Engine/Public/GenericOctree.inl
+++ b/Engine/Source/Runtime/Engine/Public/GenericOctree.inl
@@ -353,14 +353,14 @@ TOctree<ElementType,OctreeSemantics>::TOctree(const FVector& InOrigin,float InEx
 {
 }
 
-#if WITH_HOT_RELOAD && WITH_HOT_RELOAD_CTORS
+#if WITH_HOT_RELOAD_CTORS
 template<typename ElementType, typename OctreeSemantics>
 TOctree<ElementType, OctreeSemantics>::TOctree()
 	: RootNode(nullptr)
 {
 	EnsureRetrievingVTablePtr();
 }
-#endif // WITH_HOT_RELOAD && WITH_HOT_RELOAD_CTORS
+#endif // WITH_HOT_RELOAD_CTORS
 
 template<typename ElementType,typename OctreeSemantics>
 void TOctree<ElementType,OctreeSemantics>::ApplyOffset(const FVector& InOffset)
diff --git a/Engine/Source/Runtime/Engine/Public/GenericQuadTree.h b/Engine/Source/Runtime/Engine/Public/GenericQuadTree.h
--- a/Engine/Source/Runtime/Engine/Public/GenericQuadTree.h
+++ b/Engine/Source/Runtime/Engine/Public/GenericQuadTree.h
@@ -7,10 +7,10 @@ class TQuadTree
 	typedef TQuadTree<ElementType, NodeCapacity> TreeType;
 public:
 
-#if WITH_HOT_RELOAD && WITH_HOT_RELOAD_CTORS
+#if WITH_HOT_RELOAD_CTORS
 	/** DO NOT USE. This constructor is for internal usage only for hot-reload purposes. */
 	TQuadTree();
-#endif // WITH_HOT_RELOAD && WITH_HOT_RELOAD_CTORS
+#endif // WITH_HOT_RELOAD_CTORS
 
 	TQuadTree(const FBox2D& InBox);
 
@@ -134,13 +134,13 @@ TQuadTree<ElementType, NodeCapacity>::TQuadTree(const FBox2D& Box)
 	SubTrees[0] = SubTrees[1] = SubTrees[2] = SubTrees[3] = nullptr;
 }
 
-#if WITH_HOT_RELOAD && WITH_HOT_RELOAD_CTORS
+#if WITH_HOT_RELOAD_CTORS
 template <typename ElementType, int32 NodeCapacity>
 TQuadTree<ElementType, NodeCapacity>::TQuadTree()
 {
 	EnsureRetrievingVTablePtr();
 }
-#endif // WITH_HOT_RELOAD && WITH_HOT_RELOAD_CTORS
+#endif // WITH_HOT_RELOAD_CTORS
 
 template <typename ElementType, int32 NodeCapacity>
 TQuadTree<ElementType, NodeCapacity>::~TQuadTree()
diff --git a/Engine/Source/Runtime/Engine/Public/Model.h b/Engine/Source/Runtime/Engine/Public/Model.h
--- a/Engine/Source/Runtime/Engine/Public/Model.h
+++ b/Engine/Source/Runtime/Engine/Public/Model.h
@@ -343,14 +343,14 @@ enum {MAX_NODES  = 65536};
 enum {MAX_POINTS = 128000};
 class UModel : public UObject
 {
-#if WITH_HOT_RELOAD && WITH_HOT_RELOAD_CTORS
+#if WITH_HOT_RELOAD_CTORS
 	DECLARE_CASTED_CLASS_INTRINSIC_NO_CTOR_NO_VTABLE_CTOR(UModel, UObject, 0, Engine, 0, ENGINE_API)
 
 	/** DO NOT USE. This constructor is for internal usage only for hot-reload purposes. */
 	UModel(FVTableHelper& Helper);
-#else // WITH_HOT_RELOAD && WITH_HOT_RELOAD_CTORS
+#else // WITH_HOT_RELOAD_CTORS
 	DECLARE_CASTED_CLASS_INTRINSIC_NO_CTOR(UModel, UObject, 0, Engine, 0, ENGINE_API)
-#endif // WITH_HOT_RELOAD && WITH_HOT_RELOAD_CTORS
+#endif // WITH_HOT_RELOAD_CTORS
 
 #if WITH_EDITOR
 	// Arrays and subobjects.
diff --git a/Engine/Source/Runtime/Engine/Public/ShowFlags.h b/Engine/Source/Runtime/Engine/Public/ShowFlags.h
--- a/Engine/Source/Runtime/Engine/Public/ShowFlags.h
+++ b/Engine/Source/Runtime/Engine/Public/ShowFlags.h
@@ -138,14 +138,14 @@ struct FEngineShowFlags
 		Init(InitMode);
 	}
 
-#if WITH_HOT_RELOAD && WITH_HOT_RELOAD_CTORS
+#if WITH_HOT_RELOAD_CTORS
 	/** DO NOT USE. This constructor is for internal usage only for hot-reload purposes. */
 	FEngineShowFlags()
 	{
 		EnsureRetrievingVTablePtr();
 		Init(ESFIM_Game);
 	}
-#endif // WITH_HOT_RELOAD && WITH_HOT_RELOAD_CTORS
+#endif // WITH_HOT_RELOAD_CTORS
 
 	static bool CanBeToggledInEditor(const TCHAR* Name)
 	{
