[UE-11029] various circular dependency caused by structures with references and inheritence. This fixes CDOs of child types being 0'd and also fixes load order issues stemming from structures with references to child types, nested in parent types (which triggered the code path indicated in the todo)

diff --git a/Engine/Source/Runtime/CoreUObject/Private/Blueprint/BlueprintSupport.cpp b/Engine/Source/Runtime/CoreUObject/Private/Blueprint/BlueprintSupport.cpp
--- a/Engine/Source/Runtime/CoreUObject/Private/Blueprint/BlueprintSupport.cpp
+++ b/Engine/Source/Runtime/CoreUObject/Private/Blueprint/BlueprintSupport.cpp
@@ -798,7 +798,7 @@ void ULinkerLoad::ResolveDeferredDependencies(UStruct* LoadStruct)
 		{
 			// there shouldn't be any deferred dependencies (belonging to this 
 			// linker) that need to be resolved by this point
-			DEFERRED_DEPENDENCY_CHECK(!PlaceholderClass->HasReferences() && PlaceholderClass->IsPendingKill());
+			DEFERRED_DEPENDENCY_CHECK(!PlaceholderClass->HasReferences());
 		}
 	}
 #endif // USE_DEFERRED_DEPENDENCY_CHECK_VERIFICATION_TESTS
@@ -1167,6 +1167,24 @@ void ULinkerLoad::ResolveDeferredExports(UClass* LoadClass)
 			BlueprintCDO->SetLinker(this, DeferredCDOIndex, /*bShouldDetatchExisting =*/false);
 			BlueprintCDO->SetFlags(OldFlags);
 		}
+
+		// Here we could make sure we're up to date with our parent. Alternatively we could guaranatee that
+		// SerializeDefaultObjects was called before we construct other default objects, but that is not easy to
+		// do. Simply moving the call to SerializeDefaultObject next to the call to create the CDO means
+		// the stream order changes...
+		DEFERRED_DEPENDENCY_CHECK(BlueprintCDO->GetClass() == LoadClass);
+		if (BlueprintCDO->HasAnyFlags(RF_NeedLoad) &&
+			!LoadClass->GetSuperClass()->HasAnyFlags(RF_Native))
+		{
+			for (UProperty* P = LoadClass->PropertyLink; P; P = P->PropertyLinkNext)
+			{
+				if (P->IsInContainer(LoadClass->GetSuperClass()))
+				{
+					P->CopyCompleteValue_InContainer(BlueprintCDO, LoadClass->GetSuperClass()->GetDefaultObject(false));
+				}
+			}
+		}
+
 		// should load the CDO (ensuring that it has been serialized in by the 
 		// time we get to class regeneration)
 		//
diff --git a/Engine/Source/Runtime/CoreUObject/Private/UObject/UObjectGlobals.cpp b/Engine/Source/Runtime/CoreUObject/Private/UObject/UObjectGlobals.cpp
--- a/Engine/Source/Runtime/CoreUObject/Private/UObject/UObjectGlobals.cpp
+++ b/Engine/Source/Runtime/CoreUObject/Private/UObject/UObjectGlobals.cpp
@@ -923,11 +923,7 @@ UPackage* LoadPackageInternal(UPackage* InOuter, const TCHAR* InLongPackageName,
 		// if this linker already has the DeferDependencyLoads flag, then we're
 		// already loading it earlier up the load chain (don't let it invoke any
 		// deeper loads that may introduce a circular dependency)
-		if ((Linker->LoadFlags & LOAD_DeferDependencyLoads))
-		{
-			DoNotLoadExportsFlags |= LOAD_DeferDependencyLoads;
-		}
-		// @TODO: what of cases where DeferDependencyLoads was specified, but not already set on the Linker?
+		DoNotLoadExportsFlags |= LOAD_DeferDependencyLoads;
 #endif // USE_CIRCULAR_DEPENDENCY_LOAD_DEFERRING
 
 		if ((LoadFlags & DoNotLoadExportsFlags) == 0)
