[UE-10708] Making sure we don't mark a package as fully loaded until all its exports have been properly loaded.

diff --git a/Engine/Source/Runtime/CoreUObject/Private/UObject/UObjectGlobals.cpp b/Engine/Source/Runtime/CoreUObject/Private/UObject/UObjectGlobals.cpp
--- a/Engine/Source/Runtime/CoreUObject/Private/UObject/UObjectGlobals.cpp
+++ b/Engine/Source/Runtime/CoreUObject/Private/UObject/UObjectGlobals.cpp
@@ -866,6 +866,8 @@ UPackage* LoadPackageInternal(UPackage* InOuter, const TCHAR* InLongPackageName,
 	// Try to load.
 	BeginLoad();
 
+	bool bFullyLoadSkipped = false;
+
 	SlowTask.EnterProgressFrame(30);
 	{
 		// Keep track of start time.
@@ -940,6 +942,10 @@ UPackage* LoadPackageInternal(UPackage* InOuter, const TCHAR* InLongPackageName,
 
 			Linker->SetSerializedProperty(OldSerializedProperty);
 		}
+		else
+		{
+			bFullyLoadSkipped = true;
+		}
 
 		SlowTask.EnterProgressFrame(30);
 
@@ -1007,8 +1013,11 @@ UPackage* LoadPackageInternal(UPackage* InOuter, const TCHAR* InLongPackageName,
 		// We no longer need the linker. Passing in NULL would reset all loaders so we need to check for that.
 		ResetLoaders( Result );
 	}
-	// Mark package as loaded.
-	Result->SetFlags(RF_WasLoaded);
+	if (!bFullyLoadSkipped)
+	{
+		// Mark package as loaded.
+		Result->SetFlags(RF_WasLoaded);
+	}
 
 	return Result;
 }
