Disabling async bulk data loading when multithreading is not supported

diff --git a/Engine/Source/Runtime/CoreUObject/Private/Serialization/BulkData.cpp b/Engine/Source/Runtime/CoreUObject/Private/Serialization/BulkData.cpp
--- a/Engine/Source/Runtime/CoreUObject/Private/Serialization/BulkData.cpp
+++ b/Engine/Source/Runtime/CoreUObject/Private/Serialization/BulkData.cpp
@@ -684,6 +684,7 @@ bool FUntypedBulkData::ShouldStreamBulkData()
 	const bool bForceStream = !!(BulkDataFlags & BULKDATA_ForceStreamPayload);
 
 	return (FPlatformProperties::RequiresCookedData() && !Filename.IsEmpty() &&
+		FPlatformProcess::SupportsMultithreading() &&
 		(bForceStream || GetBulkDataSize() > MinBulkDataSizeForAsyncLoading.Value) &&
 		MinBulkDataSizeForAsyncLoading.Value >= 0);
 }
