Fixing async loading code not handling FinalizeLinker failures properly resulting in infinite load times.

diff --git a/Engine/Source/Runtime/CoreUObject/Private/Serialization/AsyncLoading.cpp b/Engine/Source/Runtime/CoreUObject/Private/Serialization/AsyncLoading.cpp
--- a/Engine/Source/Runtime/CoreUObject/Private/Serialization/AsyncLoading.cpp
+++ b/Engine/Source/Runtime/CoreUObject/Private/Serialization/AsyncLoading.cpp
@@ -460,6 +460,7 @@ EAsyncPackageState::Type FAsyncPackage::CreateLinker()
  */
 EAsyncPackageState::Type FAsyncPackage::FinishLinker()
 {
+	EAsyncPackageState::Type Result = EAsyncPackageState::Complete;
 	if( !Linker->HasFinishedInitialization() )
 	{
 		SCOPE_CYCLE_COUNTER(STAT_FAsyncPackage_FinishLinker);
@@ -469,15 +470,22 @@ EAsyncPackageState::Type FAsyncPackage::FinishLinker()
 		const float RemainingTimeLimit = TimeLimit - (float)(FPlatformTime::Seconds() - TickStartTime);
 
 		// Operation still pending if Tick returns false
-		if (Linker->Tick(RemainingTimeLimit, bUseTimeLimit, bUseFullTimeLimit) != ULinkerLoad::LINKER_Loaded)
+		ULinkerLoad::ELinkerStatus LinkerResult = Linker->Tick(RemainingTimeLimit, bUseTimeLimit, bUseFullTimeLimit);
+		if (Result != ULinkerLoad::LINKER_Loaded)
 		{
 			// Give up remainder of timeslice if there is one to give up.
 			GiveUpTimeSlice();
-			return EAsyncPackageState::TimeOut;
+			Result = EAsyncPackageState::TimeOut;
+			if (LinkerResult == ULinkerLoad::LINKER_Failed)
+			{
+				// If linker failed we exit with EAsyncPackageState::TimeOut to skip all the remaining steps.
+				// The error will be handled as bLoadHasFailed will be true.
+				bLoadHasFailed = true;
+			}
 		}
 	}
 
-	return EAsyncPackageState::Complete;
+	return Result;
 }
 
 /**
