Fixed a small bug in FinalizeLinker error checking code

diff --git a/Engine/Source/Runtime/CoreUObject/Private/Serialization/AsyncLoading.cpp b/Engine/Source/Runtime/CoreUObject/Private/Serialization/AsyncLoading.cpp
--- a/Engine/Source/Runtime/CoreUObject/Private/Serialization/AsyncLoading.cpp
+++ b/Engine/Source/Runtime/CoreUObject/Private/Serialization/AsyncLoading.cpp
@@ -471,7 +471,7 @@ EAsyncPackageState::Type FAsyncPackage::FinishLinker()
 
 		// Operation still pending if Tick returns false
 		ULinkerLoad::ELinkerStatus LinkerResult = Linker->Tick(RemainingTimeLimit, bUseTimeLimit, bUseFullTimeLimit);
-		if (Result != ULinkerLoad::LINKER_Loaded)
+		if (LinkerResult != ULinkerLoad::LINKER_Loaded)
 		{
 			// Give up remainder of timeslice if there is one to give up.
 			GiveUpTimeSlice();
