[UE-9824] Moving placeholder property-value tracking/resolving from LinkerPlaceholderExportObject to the shared FLinkerPlaceholderBase so that import placeholders can use it as well (when a UObjectProperty sets a value referencing one). Also, genericizing placeholder import handling (in BlueprintSupport.cpp) so that we can simplify/centralize handling for various placeholder types (classes vs. functions).

diff --git a/Engine/Source/Runtime/CoreUObject/Private/Blueprint/BlueprintSupport.cpp b/Engine/Source/Runtime/CoreUObject/Private/Blueprint/BlueprintSupport.cpp
--- a/Engine/Source/Runtime/CoreUObject/Private/Blueprint/BlueprintSupport.cpp
+++ b/Engine/Source/Runtime/CoreUObject/Private/Blueprint/BlueprintSupport.cpp
@@ -332,6 +332,10 @@ bool ULinkerLoad::DeferPotentialCircularImport(const int32 Index)
 		{
 			if (const UClass* ImportClass = FindObject<UClass>(ClassPackage, *Import.ClassName.ToString()))
 			{
+				UObject* PlaceholderOuter = LinkerRoot;
+				FString  PlaceholderNamePrefix = TEXT("PLACEHOLDER_");
+				UClass*  PlaceholderType  = nullptr;
+
 				bool const bIsBlueprintClass = ImportClass->IsChildOf<UClass>();
 				// @TODO: if we could see if the related package is created 
 				//        (without loading it further), AND it already has this 
@@ -339,10 +343,30 @@ bool ULinkerLoad::DeferPotentialCircularImport(const int32 Index)
 				//        than allocating a placeholder
 				if (bIsBlueprintClass)
 				{
-					UPackage* PlaceholderOuter = LinkerRoot;
-					UClass*   PlaceholderType  = ULinkerPlaceholderClass::StaticClass();
+					PlaceholderNamePrefix = TEXT("PLACEHOLDER-CLASS_");
+					PlaceholderType = ULinkerPlaceholderClass::StaticClass();
+				}
+				else if (ImportClass->IsChildOf<UFunction>())
+				{		
+					if (Import.OuterIndex.IsImport())
+					{
+						const int32 OuterImportIndex = Import.OuterIndex.ToImport();
+						// @TODO: if the sole reason why we have ULinkerPlaceholderFunction 
+						//        is that it's outer is a placeholder, then we 
+						//        could instead log it (with the placeholder) as 
+						//        a referencer, and then move the function later
+						if (DeferPotentialCircularImport(OuterImportIndex))
+						{
+							PlaceholderNamePrefix = TEXT("PLACEHOLDER-FUNCTION_");
+							PlaceholderOuter = ImportMap[OuterImportIndex].XObject;
+							PlaceholderType  = ULinkerPlaceholderFunction::StaticClass();
+						}
+					}
+				}
 
-					FName PlaceholderName(*FString::Printf(TEXT("PLACEHOLDER-CLASS_%s"), *Import.ObjectName.ToString()));
+				if (PlaceholderType != nullptr)
+				{
+					FName PlaceholderName(*FString::Printf(TEXT("%s_%s"), *PlaceholderNamePrefix, *Import.ObjectName.ToString()));
 					PlaceholderName = MakeUniqueObjectName(PlaceholderOuter, PlaceholderType, PlaceholderName);
 
 					ULinkerPlaceholderClass* Placeholder = NewObject<ULinkerPlaceholderClass>(PlaceholderOuter, PlaceholderType, PlaceholderName, RF_Public | RF_Transient);
@@ -350,7 +374,7 @@ bool ULinkerLoad::DeferPotentialCircularImport(const int32 Index)
 					// easily look it up in the import map, given the 
 					// placeholder (needed, to find the corresponding import for 
 					// ResolvingDeferredPlaceholder)
-					Placeholder->ImportIndex = Index;
+					Placeholder->PackageIndex = FPackageIndex::FromImport(Index);
 					// make sure the class is fully formed (has its 
 					// ClassAddReferencedObjects/ClassConstructor members set)
 					Placeholder->Bind();
@@ -358,34 +382,6 @@ bool ULinkerLoad::DeferPotentialCircularImport(const int32 Index)
 
 					Import.XObject = Placeholder;
 				}
-				else if (ImportClass->IsChildOf<UFunction>())
-				{		
-					UObject* OuterObject = nullptr;
-					if (Import.OuterIndex.IsImport())
-					{
-						const int32 OuterImportIndex = Import.OuterIndex.ToImport();
-						FObjectImport& OuterImport = ImportMap[OuterImportIndex];
-						if (nullptr == OuterImport.XObject)
-						{
-							DeferPotentialCircularImport(OuterImportIndex);
-						}
-						OuterObject = OuterImport.XObject;
-					}
-					const bool bOuterPlaceholderExists = OuterObject && OuterObject->IsA<ULinkerPlaceholderClass>();
-					if (bOuterPlaceholderExists && (nullptr == Import.XObject))
-					{
-						UClass*   PlaceholderType = ULinkerPlaceholderFunction::StaticClass();
-
-						FName PlaceholderName(*FString::Printf(TEXT("PLACEHOLDER-FUNCTION_%s"), *Import.ObjectName.ToString()));
-						PlaceholderName = MakeUniqueObjectName(OuterObject, PlaceholderType, PlaceholderName);
-
-						ULinkerPlaceholderFunction* Placeholder = NewObject<ULinkerPlaceholderFunction>(OuterObject, PlaceholderType, PlaceholderName, RF_Public | RF_Transient);
-						Placeholder->ImportIndex = Index;
-						Placeholder->Bind();
-						Placeholder->StaticLink(/*bRelinkExistingProperties =*/true);
-						Import.XObject = Placeholder;
-					}
-				}
 			}
 		}
 
@@ -554,6 +550,8 @@ bool ULinkerLoad::DeferExportCreation(const int32 Index)
 	PlaceholderName = MakeUniqueObjectName(PlaceholderOuter, PlaceholderType, PlaceholderName);
 
 	ULinkerPlaceholderExportObject* Placeholder = NewObject<ULinkerPlaceholderExportObject>(PlaceholderOuter, PlaceholderType, PlaceholderName, RF_Public | RF_Transient);
+	Placeholder->PackageIndex = FPackageIndex::FromExport(Index);
+
 	Export.Object = Placeholder;
 #endif // USE_CIRCULAR_DEPENDENCY_LOAD_DEFERRING
 
@@ -719,7 +717,7 @@ void ULinkerLoad::ResolveDeferredDependencies(UStruct* LoadStruct)
 
 			if (ULinkerPlaceholderClass* PlaceholderClass = Cast<ULinkerPlaceholderClass>(Import.XObject))
 			{
-				DEFERRED_DEPENDENCY_CHECK(PlaceholderClass->ImportIndex == ImportIndex);
+				DEFERRED_DEPENDENCY_CHECK(PlaceholderClass->PackageIndex.ToImport() == ImportIndex);
 
 				// NOTE: we don't check that this resolve successfully replaced any
 				//       references (by the return value), because this resolve 
@@ -729,26 +727,13 @@ void ULinkerLoad::ResolveDeferredDependencies(UStruct* LoadStruct)
 			}
 			else if (ULinkerPlaceholderFunction* PlaceholderFunction = Cast<ULinkerPlaceholderFunction>(Import.XObject))
 			{
-				DEFERRED_DEPENDENCY_CHECK(PlaceholderFunction->ImportIndex == ImportIndex);
-
-				auto OwnerClass = PlaceholderFunction->GetOwnerClass();
-				ResolveDependencyPlaceholder(OwnerClass, Cast<UClass>(LoadStruct));
-
-				UFunction* RealFuncObj = nullptr;
-				if ((Import.XObject != nullptr) && (Import.XObject != PlaceholderFunction))
-				{
-					RealFuncObj = CastChecked<UFunction>(Import.XObject);
-				}
-				else
+				if (ULinkerPlaceholderClass* PlaceholderOwner = Cast<ULinkerPlaceholderClass>(PlaceholderFunction->GetOwnerClass()))
 				{
-					Import.XObject = nullptr;
-					RealFuncObj = CastChecked<UFunction>(CreateImport(ImportIndex), ECastCheckedType::NullAllowed);
-				}
-				ensure(RealFuncObj != PlaceholderFunction);
-				if (RealFuncObj)
-				{
-					PlaceholderFunction->ReplaceTrackedReferences(RealFuncObj);
+					ResolveDependencyPlaceholder(PlaceholderOwner, Cast<UClass>(LoadStruct));
 				}
+
+				DEFERRED_DEPENDENCY_CHECK(PlaceholderFunction->PackageIndex.ToImport() == ImportIndex);
+				ResolveDependencyPlaceholder(PlaceholderFunction, Cast<UClass>(LoadStruct)); 
 			}
 			else if (UScriptStruct* StructObj = Cast<UScriptStruct>(Import.XObject))
 			{
@@ -798,7 +783,7 @@ void ULinkerLoad::ResolveDeferredDependencies(UStruct* LoadStruct)
 		{
 			// there shouldn't be any deferred dependencies (belonging to this 
 			// linker) that need to be resolved by this point
-			DEFERRED_DEPENDENCY_CHECK(!PlaceholderClass->HasReferences());
+			DEFERRED_DEPENDENCY_CHECK(!PlaceholderClass->HasKnownReferences());
 		}
 	}
 #endif // USE_DEFERRED_DEPENDENCY_CHECK_VERIFICATION_TESTS
@@ -824,7 +809,7 @@ bool ULinkerLoad::HasUnresolvedDependencies() const
 #endif // USE_CIRCULAR_DEPENDENCY_LOAD_DEFERRING
 }
 
-int32 ULinkerLoad::ResolveDependencyPlaceholder(UClass* PlaceholderIn, UClass* ReferencingClass)
+int32 ULinkerLoad::ResolveDependencyPlaceholder(FLinkerPlaceholderBase* PlaceholderIn, UClass* ReferencingClass)
 {
 #if USE_CIRCULAR_DEPENDENCY_LOAD_DEFERRING
 #if USE_DEFERRED_DEPENDENCY_CHECK_VERIFICATION_TESTS
@@ -835,21 +820,21 @@ int32 ULinkerLoad::ResolveDependencyPlaceholder(UClass* PlaceholderIn, UClass* R
 #endif // USE_DEFERRED_DEPENDENCY_CHECK_VERIFICATION_TESTS
 
 	TGuardValue<uint32>  LoadFlagsGuard(LoadFlags, (LoadFlags & ~LOAD_DeferDependencyLoads));
- 	TGuardValue<UClass*> ResolvingClassGuard(ResolvingDeferredPlaceholder, PlaceholderIn);
-
-	DEFERRED_DEPENDENCY_CHECK(Cast<ULinkerPlaceholderClass>(PlaceholderIn) != nullptr);
-	DEFERRED_DEPENDENCY_CHECK(PlaceholderIn->GetOuter() == LinkerRoot);
+	TGuardValue<FLinkerPlaceholderBase*> ResolvingClassGuard(ResolvingDeferredPlaceholder, PlaceholderIn);
 
-	ULinkerPlaceholderClass* PlaceholderClass = (ULinkerPlaceholderClass*)PlaceholderIn;
+	UObject* PlaceholderObj = PlaceholderIn->GetPlaceholderAsUObject();
+	DEFERRED_DEPENDENCY_CHECK(PlaceholderObj != nullptr);
+	DEFERRED_DEPENDENCY_CHECK(PlaceholderObj->GetOutermost() == LinkerRoot);
+	DEFERRED_DEPENDENCY_CHECK(PlaceholderIn->PackageIndex.IsImport());
 	
-	int32 const ImportIndex = PlaceholderClass->ImportIndex;
+	int32 const ImportIndex = PlaceholderIn->PackageIndex.ToImport();
 	FObjectImport& Import = ImportMap[ImportIndex];
 	
-	UClass* RealClassObj = nullptr;
-	if ((Import.XObject != nullptr) && (Import.XObject != PlaceholderClass))
+	UObject* RealImportObj = nullptr;
+	if ((Import.XObject != nullptr) && (Import.XObject != PlaceholderObj))
 	{
 		DEFERRED_DEPENDENCY_CHECK(ResolvingDeferredPlaceholder == PlaceholderIn);
-		RealClassObj = CastChecked<UClass>(Import.XObject);
+		RealImportObj = Import.XObject;
 	}
 	else 
 	{
@@ -860,20 +845,22 @@ int32 ULinkerLoad::ResolveDependencyPlaceholder(UClass* PlaceholderIn, UClass* R
 		//       continue to load a package already started up the stack and you 
 		//       could end up in another ResolveDependencyPlaceholder() for some  
 		//       other placeholder before this one has completely finished resolving
-		RealClassObj = CastChecked<UClass>(CreateImport(ImportIndex), ECastCheckedType::NullAllowed);
+		RealImportObj = CreateImport(ImportIndex);
 	}
 
-	DEFERRED_DEPENDENCY_CHECK(RealClassObj == nullptr || RealClassObj->HasAnyFlags(RF_LoadCompleted));
+	DEFERRED_DEPENDENCY_CHECK(RealImportObj != PlaceholderObj);
+	DEFERRED_DEPENDENCY_CHECK(RealImportObj == nullptr || RealImportObj->HasAnyFlags(RF_LoadCompleted));
 
 	int32 ReplacementCount = 0;
 	if (ReferencingClass != nullptr)
 	{
+		// @TODO: roll this into ULinkerPlaceholderClass's ResolveAllPlaceholderReferences()
 		for (FImplementedInterface& Interface : ReferencingClass->Interfaces)
 		{
-			if (Interface.Class == PlaceholderClass)
+			if (Interface.Class == PlaceholderObj)
 			{
 				++ReplacementCount;
-				Interface.Class = RealClassObj;
+				Interface.Class = CastChecked<UClass>(RealImportObj);
 			}
 		}
 	}
@@ -885,13 +872,11 @@ int32 ULinkerLoad::ResolveDependencyPlaceholder(UClass* PlaceholderIn, UClass* R
 	// doesn't have any known references (and it hasn't already been resolved in
 	// some recursive call), then there is something out there still using this
 	// placeholder class
-	DEFERRED_DEPENDENCY_CHECK( (ReplacementCount > 0) || PlaceholderClass->HasReferences() || PlaceholderClass->HasBeenResolved() );
+	DEFERRED_DEPENDENCY_CHECK( (ReplacementCount > 0) || PlaceholderIn->HasKnownReferences() || PlaceholderIn->HasBeenFullyResolved() );
 
-	ReplacementCount += PlaceholderClass->ReplaceTrackedReferences(RealClassObj);
-	// PlaceholderClass->MarkPendingKill(); // @TODO: ensure these are properly GC'd
+	ReplacementCount += PlaceholderIn->ResolveAllPlaceholderReferences(RealImportObj);
 
 #if USE_DEFERRED_DEPENDENCY_CHECK_VERIFICATION_TESTS
-	UObject* PlaceholderObj = PlaceholderClass;
 	// @TODO: not an actual method, but would be nice to circumvent the need for bIsAsyncLoadRef below
 	//FAsyncObjectsReferencer::Get().RemoveObject(PlaceholderObj);
 
@@ -1124,7 +1109,7 @@ void ULinkerLoad::ResolveDeferredExports(UClass* LoadClass)
 				//       assert on it), because this could have only been created as 
 				//       part of the LoadAllObjects() pass (not for any specific 
 				//       container object).
-				PlaceholderExport->ReplaceReferencingObjectValues(LoadClass, ExportObj);
+				PlaceholderExport->ResolveAllPlaceholderReferences(ExportObj);
 				PlaceholderExport->MarkPendingKill();
 
 				// if we hadn't used a ULinkerPlaceholderExportObject in place of 
diff --git a/Engine/Source/Runtime/CoreUObject/Private/UObject/Class.cpp b/Engine/Source/Runtime/CoreUObject/Private/UObject/Class.cpp
--- a/Engine/Source/Runtime/CoreUObject/Private/UObject/Class.cpp
+++ b/Engine/Source/Runtime/CoreUObject/Private/UObject/Class.cpp
@@ -1491,11 +1491,11 @@ bool UStruct::GetStringMetaDataHierarchical(const FName& Key, FString* OutValue)
 		UObject*& ExprPtrRef = (UObject*&)ScriptPtr;
 		if (ULinkerPlaceholderClass* PlaceholderObj = Cast<ULinkerPlaceholderClass>(ExprPtrRef))
 		{
-			PlaceholderObj->AddReferencingScriptExpr((ULinkerPlaceholderClass**)(&ExprPtrRef));
+			PlaceholderObj->AddReferencingScriptExpr((UClass**)(&ExprPtrRef));
 		}
 		else if (ULinkerPlaceholderFunction* PlaceholderFunc = Cast<ULinkerPlaceholderFunction>(ExprPtrRef))
 		{
-			PlaceholderFunc->AddReferencingScriptExpr((ULinkerPlaceholderFunction**)(&ExprPtrRef));
+			PlaceholderFunc->AddReferencingScriptExpr((UFunction**)(&ExprPtrRef));
 		}
 	}
 
diff --git a/Engine/Source/Runtime/CoreUObject/Private/UObject/LinkerLoad.cpp b/Engine/Source/Runtime/CoreUObject/Private/UObject/LinkerLoad.cpp
--- a/Engine/Source/Runtime/CoreUObject/Private/UObject/LinkerLoad.cpp
+++ b/Engine/Source/Runtime/CoreUObject/Private/UObject/LinkerLoad.cpp
@@ -2917,7 +2917,7 @@ void ULinkerLoad::Preload( UObject* Object )
 			SCOPE_CYCLE_COUNTER(STAT_LinkerPreload);
 			FScopeCycleCounterUObject PreloadScope(Object, GET_STATID(STAT_LinkerPreload));
 			UClass* Cls = NULL;
-
+			
 			// If this is a struct, make sure that its parent struct is completely loaded
 			if( UStruct* Struct = dynamic_cast<UStruct*>(Object) )
 			{
@@ -2961,6 +2961,11 @@ void ULinkerLoad::Preload( UObject* Object )
 
 				{
 					SCOPE_CYCLE_COUNTER(STAT_LinkerSerialize);
+#if USE_CIRCULAR_DEPENDENCY_LOAD_DEFERRING
+					// communicate with FLinkerPlaceholderBase, what object is currently serializing in
+					FScopedPlaceholderContainerTracker SerializingObjTracker(Object);
+#endif // USE_CIRCULAR_DEPENDENCY_LOAD_DEFERRING
+
 					if (Object->HasAnyFlags(RF_ClassDefaultObject))
 					{
 #if USE_CIRCULAR_DEPENDENCY_LOAD_DEFERRING
diff --git a/Engine/Source/Runtime/CoreUObject/Private/UObject/LinkerPlaceholderBase.cpp b/Engine/Source/Runtime/CoreUObject/Private/UObject/LinkerPlaceholderBase.cpp
--- a/Engine/Source/Runtime/CoreUObject/Private/UObject/LinkerPlaceholderBase.cpp
+++ b/Engine/Source/Runtime/CoreUObject/Private/UObject/LinkerPlaceholderBase.cpp
@@ -4,54 +4,367 @@
 #include "LinkerPlaceholderBase.h"
 #include "Blueprint/BlueprintSupport.h"
 
+/*******************************************************************************
+ *FPlaceholderContainerTracker
+ ******************************************************************************/
+
+/**  */
+struct FPlaceholderContainerTracker : TThreadSingleton<FPlaceholderContainerTracker>
+{
+	TArray<UObject*> PerspectiveReferencerStack;
+};
+
 //------------------------------------------------------------------------------
-void FLinkerPlaceholderBase::AddReferencingProperty(UProperty* ReferencingProperty)
+FScopedPlaceholderContainerTracker::FScopedPlaceholderContainerTracker(UObject* InPlaceholderContainerCandidate)
+	: PlaceholderReferencerCandidate(InPlaceholderContainerCandidate)
 {
+	FPlaceholderContainerTracker::Get().PerspectiveReferencerStack.Add(InPlaceholderContainerCandidate);
+}
+
+//------------------------------------------------------------------------------
+FScopedPlaceholderContainerTracker::~FScopedPlaceholderContainerTracker()
+{
+	UObject* StackTop = FPlaceholderContainerTracker::Get().PerspectiveReferencerStack.Pop();
 #if USE_DEFERRED_DEPENDENCY_CHECK_VERIFICATION_TESTS
-	UObject* ThisAsObject = GetPlaceholderAsUObject();
-	check(ThisAsObject);
-	FObjectImport* PlaceholderImport = nullptr;
-	if (ULinkerLoad* PropertyLinker = ReferencingProperty->GetLinker())
+	check(StackTop == PlaceholderReferencerCandidate);
+#endif // USE_DEFERRED_DEPENDENCY_CHECK_VERIFICATION_TESTS
+}
+
+/*******************************************************************************
+ * FLinkerPlaceholderBase
+ ******************************************************************************/
+
+namespace LinkerPlaceholderObjectImpl
+{
+	/**
+	 * Traces the specified property up its outer chain, all the way up to the 
+	 * OwnerClass. Constructs a UProperty path along the way, so that we can use
+	 * traverse a property tree down, from the class owner down.
+	 * 
+	 * @param  LeafProperty	The property you want a path for.
+	 * @param  ChainOut		Output array that tracks the specified property back up its outer chain (first entry should be what was passed in LeafProperty). 
+	 */
+	static void BuildPropertyChain(const UProperty* LeafProperty, TArray<const UProperty*>& ChainOut);
+
+	/**
+	 * A recursive method that replaces all leaf references to this object with 
+	 * the supplied ReplacementValue.
+	 *
+	 * This function recurses the property chain (from class owner down) because 
+	 * at the time of AddReferencingPropertyValue() we cannot know/record the 
+	 * address/index of array properties (as they may change during array 
+	 * re-allocation or compaction). So we must follow the property chain and 
+	 * check every UArrayProperty member for references to this (hence, the need
+	 * for this recursive function).
+	 * 
+	 * @param  PropertyChain    An ascending outer chain, where the property at index zero is the leaf (referencer) property.
+	 * @param  ChainIndex		An index into the PropertyChain that this call should start at and iterate DOWN to zero.
+	 * @param  ValueAddress     The memory address of the value corresponding to the property at ChainIndex.
+	 * @param  OldValue
+	 * @param  ReplacementValue	The new object to replace all references to this with.
+	 * @return The number of references that were replaced.
+	 */
+	static int32 ResolvePlaceholderValues(const TArray<const UProperty*>& PropertyChain, int32 ChainIndex, uint8* ValueAddress, UObject* OldValue, UObject* ReplacementValue);
+}
+
+//------------------------------------------------------------------------------
+static void LinkerPlaceholderObjectImpl::BuildPropertyChain(const UProperty* LeafProperty, TArray<const UProperty*>& ChainOut)
+{
+	ChainOut.Empty();
+	ChainOut.Add(LeafProperty);
+
+	UClass* ClassOwner = LeafProperty->GetOwnerClass();
+
+	UObject* PropertyOuter = LeafProperty->GetOuter();
+	while ((PropertyOuter != nullptr) && (PropertyOuter != ClassOwner))
 	{
-		for (FObjectImport& Import : PropertyLinker->ImportMap)
+		if (const UProperty* PropertyOwner = Cast<const UProperty>(PropertyOuter))
 		{
-			if (Import.XObject == ThisAsObject)
+			ChainOut.Add(PropertyOwner);
+		}
+		PropertyOuter = PropertyOuter->GetOuter();
+	}
+}
+
+//------------------------------------------------------------------------------
+static int32 LinkerPlaceholderObjectImpl::ResolvePlaceholderValues(const TArray<const UProperty*>& PropertyChain, int32 ChainIndex, uint8* ValueAddress, UObject* OldValue, UObject* ReplacementValue)
+{
+	int32 ReplacementCount = 0;
+
+	for (int32 PropertyIndex = ChainIndex; PropertyIndex >= 0; --PropertyIndex)
+	{
+		const UProperty* Property = PropertyChain[PropertyIndex];
+		if (PropertyIndex == 0)
+		{
+#if USE_DEFERRED_DEPENDENCY_CHECK_VERIFICATION_TESTS
+			check(Property->IsA<UObjectProperty>());
+#endif // USE_DEFERRED_DEPENDENCY_CHECK_VERIFICATION_TESTS
+
+			const UObjectProperty* ReferencingProperty = (const UObjectProperty*)Property;
+			if (ReferencingProperty->GetObjectPropertyValue(ValueAddress) == OldValue)
 			{
-				PlaceholderImport = &Import;
-				break;
+				// @TODO: use an FArchiver with ReferencingProperty->SerializeItem() 
+				//        so that we can utilize CheckValidObject()
+				ReferencingProperty->SetObjectPropertyValue(ValueAddress, ReplacementValue);
+				// @TODO: unfortunately, this is currently protected
+				//ReferencingProperty->CheckValidObject(ValueAddress);
+
+				++ReplacementCount;
 			}
 		}
-		check(ThisAsObject->GetOuter() == PropertyLinker->LinkerRoot);
-		check((PropertyLinker->LoadFlags & LOAD_DeferDependencyLoads) || FBlueprintSupport::IsResolvingDeferredDependenciesDisabled());
+		else if (const UArrayProperty* ArrayProperty = Cast<const UArrayProperty>(Property))
+		{
+#if USE_DEFERRED_DEPENDENCY_CHECK_VERIFICATION_TESTS
+			const UProperty* NextProperty = PropertyChain[PropertyIndex - 1];
+			check(NextProperty == ArrayProperty->Inner);
+#endif // USE_DEFERRED_DEPENDENCY_CHECK_VERIFICATION_TESTS
+
+			// because we can't know which array entry was set with a reference 
+			// to this object, we have to comb through them all
+			FScriptArrayHelper ArrayHelper(ArrayProperty, ValueAddress);
+			for (int32 ArrayIndex = 0; ArrayIndex < ArrayHelper.Num(); ++ArrayIndex)
+			{
+				uint8* MemberAddress = ArrayHelper.GetRawPtr(ArrayIndex);
+				ReplacementCount += ResolvePlaceholderValues(PropertyChain, PropertyIndex - 1, MemberAddress, OldValue, ReplacementValue);
+			}
+
+			// the above recursive call chewed through the rest of the
+			// PropertyChain, no need to keep on here
+			break;
+		}
+		else
+		{
+			const UProperty* NextProperty = PropertyChain[PropertyIndex - 1];
+#if USE_DEFERRED_DEPENDENCY_CHECK_VERIFICATION_TESTS
+			check(NextProperty->GetOuter() == Property);
+#endif // USE_DEFERRED_DEPENDENCY_CHECK_VERIFICATION_TESTS
+			ValueAddress = Property->ContainerPtrToValuePtr<uint8>(ValueAddress, /*ArrayIndex =*/0);
+		}
+	}
+
+	return ReplacementCount;
+}
+
+//------------------------------------------------------------------------------
+FLinkerPlaceholderBase::FLinkerPlaceholderBase()
+	: bResolveWasInvoked(false)
+{
+}
+
+//------------------------------------------------------------------------------
+FLinkerPlaceholderBase::~FLinkerPlaceholderBase()
+{
+#if USE_DEFERRED_DEPENDENCY_CHECK_VERIFICATION_TESTS
+	check(!HasKnownReferences());
+#endif // USE_DEFERRED_DEPENDENCY_CHECK_VERIFICATION_TESTS
+}
+
+//------------------------------------------------------------------------------
+bool FLinkerPlaceholderBase::AddReferencingPropertyValue(const UObjectProperty* ReferencingProperty, void* DataPtr)
+{
+	TArray<UObject*>& PossibleReferencers = FPlaceholderContainerTracker::Get().PerspectiveReferencerStack;
+
+	UObject* FoundReferencer = nullptr;
+	// iterate backwards because this is meant to act as a stack, where the last
+	// entry is most likely the one we're looking for
+	for (int32 CandidateIndex = PossibleReferencers.Num() - 1; CandidateIndex >= 0; --CandidateIndex)
+	{
+		UObject* ReferencerCandidate = PossibleReferencers[CandidateIndex];
+
+		UClass* CandidateClass = ReferencerCandidate->GetClass();
+		UClass* PropOwnerClass = ReferencingProperty->GetOwnerClass();
+
+		if (CandidateClass->IsChildOf(PropOwnerClass))
+		{
+			FoundReferencer = ReferencerCandidate;
+			break;
+		}
 	}
-	// if this check hits, then we're adding dependencies after we've 
-	// already resolved the placeholder (it won't be resolved again)
-	check(!bResolvedReferences);
+
+#if USE_DEFERRED_DEPENDENCY_CHECK_VERIFICATION_TESTS
+	check(ReferencingProperty->GetObjectPropertyValue(DataPtr) == GetPlaceholderAsUObject());
+	check(FoundReferencer != nullptr);
 #endif // USE_DEFERRED_DEPENDENCY_CHECK_VERIFICATION_TESTS
 
-	ReferencingProperties.Add(ReferencingProperty);
+	if (FoundReferencer != nullptr)
+	{
+#if USE_DEFERRED_DEPENDENCY_CHECK_VERIFICATION_TESTS
+		// @TODO: verify that DataPtr belongs to FoundReferencer
+// 		uint8* ContainerStart = (uint8*)FoundReferencer;
+// 		uint8* ContainerEnd   = ContainerStart + FoundReferencer->GetClass()->GetStructureSize();
+// 		// check that we picked the right container object, and that DataPtr exists within it 
+// 		check(DataPtr >= ContainerStart && DataPtr <= ContainerEnd);
+#endif // USE_DEFERRED_DEPENDENCY_CHECK_VERIFICATION_TESTS
+
+		ReferencingContainers.FindOrAdd(FoundReferencer).Add(ReferencingProperty);
+	}
+	return (FoundReferencer != nullptr);
+}
+
+//------------------------------------------------------------------------------
+bool FLinkerPlaceholderBase::HasKnownReferences() const
+{
+	return (ReferencingContainers.Num() > 0);
+}
+
+//------------------------------------------------------------------------------
+int32 FLinkerPlaceholderBase::ResolveAllPlaceholderReferences(UObject* ReplacementObj)
+{
+	int32 ReplacementCount = ResolvePlaceholderPropertyValues(ReplacementObj);
+	ReferencingContainers.Empty();
+
+	MarkAsResolved();
+	return ReplacementCount;
+}
+
+//------------------------------------------------------------------------------
+bool FLinkerPlaceholderBase::HasBeenFullyResolved() const
+{
+	return IsMarkedResolved() && !HasKnownReferences();
 }
 
 //------------------------------------------------------------------------------
-bool FLinkerPlaceholderBase::HasReferences() const
+bool FLinkerPlaceholderBase::IsMarkedResolved() const
 {
-	return (GetRefCount() > 0);
+	return bResolveWasInvoked;
 }
 
 //------------------------------------------------------------------------------
-int32 FLinkerPlaceholderBase::GetRefCount() const
+void FLinkerPlaceholderBase::MarkAsResolved()
 {
-	return ReferencingProperties.Num();
+	bResolveWasInvoked = true;
+}
+ 
+//------------------------------------------------------------------------------
+int32 FLinkerPlaceholderBase::ResolvePlaceholderPropertyValues(UObject* NewObjectValue)
+{
+	int32 ResolvedTotal = 0;
+
+	UObject* ThisAsUObject = GetPlaceholderAsUObject();
+	for (auto& ReferencingPair : ReferencingContainers)
+	{
+		TWeakObjectPtr<UObject> ContainerPtr = ReferencingPair.Key;
+		if (!ContainerPtr.IsValid())
+		{
+			continue;
+		}
+		UObject* Container = ContainerPtr.Get();
+
+		for (const UObjectProperty* Property : ReferencingPair.Value)
+		{
+#if USE_DEFERRED_DEPENDENCY_CHECK_VERIFICATION_TESTS
+			check(Property->GetOwnerClass() == Container->GetClass());
+#endif // USE_DEFERRED_DEPENDENCY_CHECK_VERIFICATION_TESTS
+
+			TArray<const UProperty*> PropertyChain;
+			LinkerPlaceholderObjectImpl::BuildPropertyChain(Property, PropertyChain);
+			const UProperty* OutermostProperty = PropertyChain.Last();
+
+			uint8* OutermostAddress = OutermostProperty->ContainerPtrToValuePtr<uint8>((uint8*)Container, /*ArrayIndex =*/0);
+			int32 ResolvedCount = LinkerPlaceholderObjectImpl::ResolvePlaceholderValues(PropertyChain, PropertyChain.Num() - 1, OutermostAddress, ThisAsUObject, NewObjectValue);
+			ResolvedTotal += ResolvedCount;
+
+#if USE_DEFERRED_DEPENDENCY_CHECK_VERIFICATION_TESTS
+			// we expect that (because we have had ReferencingProperties added) 
+			// there should be at least one reference that is resolved... if 
+			// there were none, then a property could have changed its value 
+			// after it was set to this
+			// 
+			// NOTE: this may seem it can be resolved by properties removing 
+			//       themselves from ReferencingProperties, but certain 
+			//       properties may be the inner of a UArrayProperty (meaning 
+			//       there could be multiple references per property)... we'd 
+			//       have to inc/decrement a property ref-count to resolve that 
+			//       scenario
+			check(ResolvedCount > 0);
+#endif // USE_DEFERRED_DEPENDENCY_CHECK_VERIFICATION_TESTS
+		}
+	}
+
+	return ResolvedTotal;
 }
 
+/*******************************************************************************
+ * TLinkerImportPlaceholder<UClass>
+ ******************************************************************************/
+
 //------------------------------------------------------------------------------
-bool FLinkerPlaceholderBase::HasBeenResolved() const
+template<>
+int32 TLinkerImportPlaceholder<UClass>::ResolvePropertyReferences(UClass* ReplacementClass)
 {
-	return !HasReferences() && bResolvedReferences;
+	int32 ReplacementCount = 0;
+	UClass* PlaceholderClass = CastChecked<UClass>(GetPlaceholderAsUObject());
+
+	for (UProperty* Property : ReferencingProperties)
+	{
+		if (UObjectPropertyBase* BaseObjProperty = Cast<UObjectPropertyBase>(Property))
+		{
+			if (BaseObjProperty->PropertyClass == PlaceholderClass)
+			{
+				BaseObjProperty->PropertyClass = ReplacementClass;
+				++ReplacementCount;
+			}
+
+			UClassProperty* ClassProperty = Cast<UClassProperty>(BaseObjProperty);
+			if ((ClassProperty != nullptr) && (ClassProperty->MetaClass == PlaceholderClass))
+			{
+				ClassProperty->MetaClass = ReplacementClass;
+				++ReplacementCount;
+			}
+		}	
+		else if (UInterfaceProperty* InterfaceProp = Cast<UInterfaceProperty>(Property))
+		{
+			if (InterfaceProp->InterfaceClass == PlaceholderClass)
+			{
+				InterfaceProp->InterfaceClass = ReplacementClass;
+				++ReplacementCount;
+			}
+		}
+		else
+		{
+			ensureMsgf(false, TEXT("Unhandled property type: %s"), *Property->GetClass()->GetName());
+		}
+	}
+
+	ReferencingProperties.Empty();
+	return ReplacementCount;
 }
 
+/*******************************************************************************
+ * TLinkerImportPlaceholder<UFunction>
+ ******************************************************************************/
+
 //------------------------------------------------------------------------------
-void FLinkerPlaceholderBase::RemovePropertyReference(UProperty* ReferencingProperty)
+template<>
+int32 TLinkerImportPlaceholder<UFunction>::ResolvePropertyReferences(UFunction* ReplacementFunc)
 {
-	ReferencingProperties.Remove(ReferencingProperty);
+	int32 ReplacementCount = 0;
+	UFunction* PlaceholderFunc = CastChecked<UFunction>(GetPlaceholderAsUObject());
+
+	for (UProperty* Property : ReferencingProperties)
+	{
+		if (UDelegateProperty* DelegateProperty = Cast<UDelegateProperty>(Property))
+		{
+			if (DelegateProperty->SignatureFunction == PlaceholderFunc)
+			{
+				DelegateProperty->SignatureFunction = ReplacementFunc;
+				++ReplacementCount;
+			}
+		}
+		else if (UMulticastDelegateProperty* MulticastDelegateProperty = Cast<UMulticastDelegateProperty>(Property))
+		{
+			if (MulticastDelegateProperty->SignatureFunction == PlaceholderFunc)
+			{
+				MulticastDelegateProperty->SignatureFunction = ReplacementFunc;
+				++ReplacementCount;
+			}
+		}
+		else
+		{
+			ensureMsgf(false, TEXT("Unhandled property type: %s"), *Property->GetClass()->GetName());
+		}
+	}
+
+	ReferencingProperties.Empty();
+	return ReplacementCount;
 }
diff --git a/Engine/Source/Runtime/CoreUObject/Private/UObject/LinkerPlaceholderBase.h b/Engine/Source/Runtime/CoreUObject/Private/UObject/LinkerPlaceholderBase.h
--- a/Engine/Source/Runtime/CoreUObject/Private/UObject/LinkerPlaceholderBase.h
+++ b/Engine/Source/Runtime/CoreUObject/Private/UObject/LinkerPlaceholderBase.h
@@ -2,26 +2,53 @@
 
 #pragma once
 
+// Forward Declarations
 class UProperty;
+class UObjectProperty;
 
+/*******************************************************************************
+ *FPlaceholderContainerTracker
+ ******************************************************************************/
+/** 
+ * To track placeholder property values, we need to know the root container 
+ * instance that is set with the placeholder value (so we can reset it later). 
+ * This here is designed to track objects that are actively being preloaded
+ * (serialized in); so we have the container on hand, when a UObjectProperty 
+ * value is set with a placeholder.
+ */
+struct FScopedPlaceholderContainerTracker
+{
+public:
+	 FScopedPlaceholderContainerTracker(UObject* PerspectivePlaceholderReferencer);
+	~FScopedPlaceholderContainerTracker();
+
+private:
+	UObject* PlaceholderReferencerCandidate;
+};
+
+/*******************************************************************************
+ * TLinkerPlaceholderBase<>
+ ******************************************************************************/
+ 
 class FLinkerPlaceholderBase
 {
 public:
-	FLinkerPlaceholderBase()
-		: ImportIndex(0)
-		, bResolvedReferences(false)
-	{}
+	FLinkerPlaceholderBase();
+	virtual ~FLinkerPlaceholderBase();
 
-	virtual ~FLinkerPlaceholderBase()
-	{}
+	/** Set by the ULinkerLoad that created this instance, tracks what import/export this was used in place of. */
+	FPackageIndex PackageIndex;
 
 	/**
-	 * Caches off the supplied property so that we can later replace it's use of
-	 * this class with another (real) class.
+	 * Attempts to find and store the referencing container object (along with  
+	 * the specified property), so that we can replace the reference at a 
+	 * later point. Can fail if the container could not be found.
 	 * 
-	 * @param  ReferencingProperty    A property that uses and stores this class.
+	 * @param  ReferencingProperty	The property whose object-value is referencing this.
+	 * @param  DataPtr				Not saved off (as it can change), but used to verify that we pick the correct container.
+	 * @return True if we successfully found a container object and are now tracking it, otherwise false.
 	 */
-	void AddReferencingProperty(UProperty* ReferencingProperty);
+	bool AddReferencingPropertyValue(const UObjectProperty* ReferencingProperty, void* DataPtr);
 
 	/**
 	 * A query method that let's us check to see if this class is currently 
@@ -31,43 +58,137 @@ public:
 	 * 
 	 * @return True if this has anything stored in its ReferencingProperties container, otherwise false.
 	 */
-	bool HasReferences() const;
+	virtual bool HasKnownReferences() const;
 
 	/**
-	 * Query method that retrieves the current number of KNOWN references to 
-	 * this placeholder class.
+	 * Iterates over all known referencers and attempts to replace their 
+	 * references to this with a new (hopefully proper) UObject.
 	 * 
-	 * @return The number of references that this class is currently tracking.
+	 * @param  ReplacementObj	The object that you want all references to this replaced with.
+	 * @return The number of references that were successfully replaced.
 	 */
-	virtual int32 GetRefCount() const;
+	virtual int32 ResolveAllPlaceholderReferences(UObject* ReplacementObj);
 
 	/**
-	 * Checks to see if 1) this placeholder has had RemoveTrackedReference() 
-	 * called on it, and 2) it doesn't have any more references that have since 
-	 * been added.
+	 * Checks to see if 1) this placeholder has had 
+	 * ResolveAllPlaceholderReferences() called on it, and 2) it doesn't have 
+	 * any more references that have since been added.
 	 * 
-	 * @return True if ReplaceTrackedReferences() has been ran, and no KNOWN references have been added.
+	 * @return True if ResolveAllPlaceholderReferences() has been ran, and no KNOWN references have been added since.
 	 */
-	bool HasBeenResolved() const;
+	bool HasBeenFullyResolved() const;
+	
+	/**
+	 * Some of our internal validation checks rely on UObject compares (between 
+	 * this placeholder and other values). Since it is expected that this is 
+	 * sub-classed by a UObject, then we provide it this pure-virtual for the
+	 * sub-class to implement (and return itself).
+	 * 
+	 * @return This as a UObject (null unless overridden by a sub-class).
+	 */
+	virtual UObject* GetPlaceholderAsUObject()
+		PURE_VIRTUAL(FLinkerPlaceholderBase::GetPlaceholderAsUObject, return nullptr;)
 
+protected:
 	/**
-	 * Removes the specified property from this class's internal tracking list 
-	 * (which aims to keep track of properties utilizing this class).
+	 * Checks to see if ResolveAllPlaceholderReferences() has been called on 
+	 * this placeholder.
 	 * 
-	 * @param  ReferencingProperty    A property that used to use this class, and now no longer does.
+	 * @return True if ResolveAllPlaceholderReferences() has been called, otherwise false.
 	 */
-	void RemovePropertyReference(UProperty* ReferencingProperty);
+	bool IsMarkedResolved() const;
 
-	virtual UObject* GetPlaceholderAsUObject() PURE_VIRTUAL(FLinkerPlaceholderBase::GetPlaceholderUObject, return nullptr;);
+	/**
+	 * Flags this placeholder as resolved (so that IsMarkedResolved() and 
+	 * HasBeenFullyResolved() can return true).
+	 */
+	void MarkAsResolved();
+
+private:
+	/**
+	 * Iterates through ReferencingContainers and replaces any (KNOWN) 
+	 * references to this placeholder (with 
+	 * 
+	 * @param  ReplacementObj    
+	 * @return 
+	 */
+	int32 ResolvePlaceholderPropertyValues(UObject* ReplacementObj);
+	
+	/** Used to catch references that are added after we've already resolved all references */
+	bool bResolveWasInvoked;
 
+	/** Tracks container objects that have property values set to reference this placeholder (references that need to be replaced later) */
+	TMap< TWeakObjectPtr<UObject>, TSet<const UObjectProperty*> > ReferencingContainers;
+};
+
+/*******************************************************************************
+ * TLinkerImportPlaceholderBase<>
+ ******************************************************************************/
+ 
+template<class PlaceholderType>
+class TLinkerImportPlaceholder : public FLinkerPlaceholderBase
+{
 public:
-	/** Set by the ULinkerLoad that created this instance, tracks what import this was used in place of. */
-	int32 ImportIndex;
+	// FLinkerPlaceholderBase interface
+	virtual bool  HasKnownReferences() const override;
+	virtual int32 ResolveAllPlaceholderReferences(UObject* ReplacementObj) override;
+	// End FLinkerPlaceholderBase interface
 
-protected:
-	/** Links to UProperties that are currently using this class */
+	/**
+	 * Records the supplied property so that we can later replace it's 
+	 * references to this placeholder with another (real) object.
+	 * 
+	 * @param  ReferencingProperty	A property that references this placeholder.
+	 */
+	void AddReferencingProperty(UProperty* ReferencingProperty);
+
+	/**
+	 * Removes the specified property from this class's internal tracking list 
+	 * (which aims to keep track of properties utilizing this placeholder).
+	 * 
+	 * @param  ReferencingProperty	A property that used to use this placeholder and now no longer does.
+	 */
+	void RemoveReferencingProperty(UProperty* ReferencingProperty);
+
+	/**
+	 * Records a raw pointer, directly to the UObject* script expression (so 
+	 * that we can switch-out its value in ResolveScriptReferences). 
+	 *
+	 * NOTE: We don't worry about creating some kind of weak ref to the script 
+	 *       pointer (or facilitate a way for this tracked reference to be 
+	 *       removed). We're not worried about the script ref being deleted 
+	 *       before we call ResolveScriptReferences (because we expect that to
+	 *       do this all within the same frame; before GC can be ran).
+	 * 
+	 * @param  ExpressionPtr    A direct pointer to the UObject* that is now referencing this placeholder.
+	 */
+	void AddReferencingScriptExpr(PlaceholderType** ExpressionPtr);
+
+private:
+	/**
+	 * Iterates through all known ReferencingProperties and replaces references
+	 * to this placeholder with the supplied replacement object.
+	 * 
+	 * @param  ReplacementObj	The object that you want all references to this replaced with.
+	 * @return The number of references that were successfully replaced.
+	 */
+	int32 ResolvePropertyReferences(PlaceholderType* ReplacementObj);
+
+	/**
+	 * Iterates through all known ReferencingScriptExpressions and replaces 
+	 * references to this placeholder with the specified replacement object.
+	 * 
+	 * @param  ReplacementObj	The object that you want all references to this replaced with.
+	 * @return The number of references that were successfully replaced.
+	 */
+	int32 ResolveScriptReferences(PlaceholderType* ReplacementObj);
+
+	/** Links to UProperties that are currently directly using this placeholder */
 	TSet<UProperty*> ReferencingProperties;
 
-	/** Used to catch references that are added after we've already resolved all references */
-	bool bResolvedReferences;
+	/** Points directly at UObject* refs that were serialized in as part of script bytecode */
+	TSet<PlaceholderType**> ReferencingScriptExpressions;
 };
+
+// Templatized implementation
+#include "LinkerPlaceholderBase.inl" 
diff --git a/Engine/Source/Runtime/CoreUObject/Private/UObject/LinkerPlaceholderBase.inl b/Engine/Source/Runtime/CoreUObject/Private/UObject/LinkerPlaceholderBase.inl
new file mode 100644
--- /dev/null
+++ b/Engine/Source/Runtime/CoreUObject/Private/UObject/LinkerPlaceholderBase.inl
@@ -0,0 +1,102 @@
+// Copyright 1998-2014 Epic Games, Inc. All Rights Reserved.
+
+#pragma once
+
+/*******************************************************************************
+ * TLinkerImportPlaceholder<>
+ ******************************************************************************/
+
+//------------------------------------------------------------------------------
+template<class PlaceholderType>
+int32 TLinkerImportPlaceholder<PlaceholderType>::ResolveAllPlaceholderReferences(UObject* ReplacementObj)
+{
+	PlaceholderType* TypeCheckedReplacement = CastChecked<PlaceholderType>(ReplacementObj);
+
+	int32 ReplacementCount = ResolvePropertyReferences(TypeCheckedReplacement);
+	ReplacementCount += ResolveScriptReferences(TypeCheckedReplacement);
+	ReplacementCount += FLinkerPlaceholderBase::ResolveAllPlaceholderReferences(ReplacementObj);
+
+	return ReplacementCount;
+}
+
+//------------------------------------------------------------------------------
+template<class PlaceholderType>
+bool TLinkerImportPlaceholder<PlaceholderType>::HasKnownReferences() const
+{
+	return FLinkerPlaceholderBase::HasKnownReferences() || (ReferencingProperties.Num() > 0) || (ReferencingScriptExpressions.Num() > 0);
+}
+
+//------------------------------------------------------------------------------
+template<class PlaceholderType>
+void TLinkerImportPlaceholder<PlaceholderType>::AddReferencingProperty(UProperty* ReferencingProperty)
+{
+#if USE_DEFERRED_DEPENDENCY_CHECK_VERIFICATION_TESTS
+	UObject* ThisAsObject = GetPlaceholderAsUObject();
+	check(ThisAsObject != nullptr);
+
+	FObjectImport* PlaceholderImport = nullptr;
+	if (ULinkerLoad* PropertyLinker = ReferencingProperty->GetLinker())
+	{
+		for (FObjectImport& Import : PropertyLinker->ImportMap)
+		{
+			if (Import.XObject == ThisAsObject)
+			{
+				PlaceholderImport = &Import;
+				break;
+			}
+		}
+		check(ThisAsObject->GetOutermost() == PropertyLinker->LinkerRoot);
+		check((PropertyLinker->LoadFlags & LOAD_DeferDependencyLoads) || FBlueprintSupport::IsResolvingDeferredDependenciesDisabled());
+	}
+	// if this check hits, then we're adding dependencies after we've 
+	// already resolved the placeholder (it won't be resolved again)
+	check(!IsMarkedResolved());
+#endif // USE_DEFERRED_DEPENDENCY_CHECK_VERIFICATION_TESTS
+
+	ReferencingProperties.Add(ReferencingProperty);
+}
+
+//------------------------------------------------------------------------------
+template<class PlaceholderType>
+void TLinkerImportPlaceholder<PlaceholderType>::RemoveReferencingProperty(UProperty* ReferencingProperty)
+{
+	ReferencingProperties.Remove(ReferencingProperty);
+}
+
+//------------------------------------------------------------------------------
+template<class PlaceholderType>
+void TLinkerImportPlaceholder<PlaceholderType>::AddReferencingScriptExpr(PlaceholderType** ExpressionPtr)
+{
+#if USE_DEFERRED_DEPENDENCY_CHECK_VERIFICATION_TESTS
+	check(*ExpressionPtr == GetPlaceholderAsUObject());
+#endif // USE_DEFERRED_DEPENDENCY_CHECK_VERIFICATION_TESTS
+
+	ReferencingScriptExpressions.Add(ExpressionPtr);
+}
+
+//------------------------------------------------------------------------------
+template<class PlaceholderType>
+int32 TLinkerImportPlaceholder<PlaceholderType>::ResolvePropertyReferences(PlaceholderType* ReplacementObj) 
+	// requires template specialization (technically not "pure virtual"):
+	PURE_VIRTUAL(TLinkerImportPlaceholder<PlaceholderType>::ResolvePropertyReferences, return 0;)
+
+//------------------------------------------------------------------------------
+template<class PlaceholderType>
+int32 TLinkerImportPlaceholder<PlaceholderType>::ResolveScriptReferences(PlaceholderType* ReplacementObj)
+{
+	int32 ReplacementCount = 0;
+	PlaceholderType* PlaceholderObj = CastChecked<PlaceholderType>(GetPlaceholderAsUObject());
+
+	for (PlaceholderType** ScriptRefPtr : ReferencingScriptExpressions)
+	{
+		PlaceholderType*& ScriptRef = *ScriptRefPtr;
+		if (ScriptRef == PlaceholderObj)
+		{
+			ScriptRef = ReplacementObj;
+			++ReplacementCount;
+		}
+	}
+
+	ReferencingScriptExpressions.Empty();
+	return ReplacementCount;
+}
diff --git a/Engine/Source/Runtime/CoreUObject/Private/UObject/LinkerPlaceholderClass.cpp b/Engine/Source/Runtime/CoreUObject/Private/UObject/LinkerPlaceholderClass.cpp
--- a/Engine/Source/Runtime/CoreUObject/Private/UObject/LinkerPlaceholderClass.cpp
+++ b/Engine/Source/Runtime/CoreUObject/Private/UObject/LinkerPlaceholderClass.cpp
@@ -9,27 +9,14 @@ ULinkerPlaceholderClass::ULinkerPlaceholderClass(const FObjectInitializer& Objec
 {
 }
 
-//------------------------------------------------------------------------------
-ULinkerPlaceholderClass::~ULinkerPlaceholderClass()
-{
-#if USE_DEFERRED_DEPENDENCY_CHECK_VERIFICATION_TESTS
-	check(!HasReferences());
-#endif // USE_DEFERRED_DEPENDENCY_CHECK_VERIFICATION_TESTS
-
-	// by this point, we really shouldn't have any properties left (they should 
-	// have all got replaced), but just in case (so things don't blow up with a
-	// missing class)...
-	ReplaceTrackedReferences(UObject::StaticClass());
-}
-
 //------------------------------------------------------------------------------
 IMPLEMENT_CORE_INTRINSIC_CLASS(ULinkerPlaceholderClass, UClass, 
 	{
 		Class->ClassAddReferencedObjects = &ULinkerPlaceholderClass::AddReferencedObjects;
 		
-		// @TODO: use the Class->Emit...() functions here to aid garbage 
-		//        collection, so it has information on what class variables 
-		//        hold onto UObject references
+	// @TODO: use the Class->Emit...() functions here to aid garbage 
+	//        collection, so it has information on what class variables 
+	//        hold onto UObject references
 	}
 );
 
@@ -52,6 +39,22 @@ void ULinkerPlaceholderClass::PostInitProperties()
 	}
 }
 
+//------------------------------------------------------------------------------
+void ULinkerPlaceholderClass::BeginDestroy()
+{
+#if USE_DEFERRED_DEPENDENCY_CHECK_VERIFICATION_TESTS
+	check(IsMarkedResolved() || HasAnyFlags(RF_ClassDefaultObject));
+	check(!HasKnownReferences());
+#endif // USE_DEFERRED_DEPENDENCY_CHECK_VERIFICATION_TESTS
+
+	// by this point, we really shouldn't have any properties left (they should 
+	// have all got replaced), but just in case (so things don't blow up with a
+	// missing class)...
+	ResolveAllPlaceholderReferences(UObject::StaticClass());
+
+	Super::BeginDestroy();
+}
+
 //------------------------------------------------------------------------------
 void ULinkerPlaceholderClass::Bind()
 {
@@ -64,66 +67,4 @@ void ULinkerPlaceholderClass::Bind()
 	ClassAddReferencedObjects = &ULinkerPlaceholderClass::AddReferencedObjects;
 }
 
-//------------------------------------------------------------------------------
-void ULinkerPlaceholderClass::AddReferencingScriptExpr(ULinkerPlaceholderClass** ExpressionPtr)
-{
-#if USE_DEFERRED_DEPENDENCY_CHECK_VERIFICATION_TESTS
-	check(*ExpressionPtr == this);
-#endif // USE_DEFERRED_DEPENDENCY_CHECK_VERIFICATION_TESTS
-
-	ReferencingScriptExpressions.Add((UClass**)ExpressionPtr);
-}
-
-//------------------------------------------------------------------------------
-int32 ULinkerPlaceholderClass::ReplaceTrackedReferences(UClass* ReplacementClass)
-{
-	int32 ReplacementCount = 0;
-
-	for (UProperty* Property : ReferencingProperties)
-	{
-		if (UObjectPropertyBase* ObjProperty = Cast<UObjectPropertyBase>(Property))
-		{
-			if (ObjProperty->PropertyClass == this)
-			{
-				ObjProperty->PropertyClass = ReplacementClass;
-				++ReplacementCount;
-			}
 
-			UClassProperty* ClassProperty = Cast<UClassProperty>(ObjProperty);
-			if ((ClassProperty != nullptr) && (ClassProperty->MetaClass == this))
-			{
-				ClassProperty->MetaClass = ReplacementClass;
-				++ReplacementCount;
-			}
-		}	
-		else if (UInterfaceProperty* InterfaceProp = Cast<UInterfaceProperty>(Property))
-		{
-			if (InterfaceProp->InterfaceClass == this)
-			{
-				InterfaceProp->InterfaceClass = ReplacementClass;
-				++ReplacementCount;
-			}
-		}
-	}
-	ReferencingProperties.Empty();
-
-	for (UClass** ScriptRefPtr : ReferencingScriptExpressions)
-	{
-		UClass*& ScriptRef = *ScriptRefPtr;
-		if (ScriptRef == this)
-		{
-			ScriptRef = ReplacementClass;
-			++ReplacementCount;
-		}
-	}
-	ReferencingScriptExpressions.Empty();
-
-	bResolvedReferences = true;
-	return ReplacementCount;
-}
-
-//------------------------------------------------------------------------------
-int32 ULinkerPlaceholderClass::GetRefCount() const
-{
-	return FLinkerPlaceholderBase::GetRefCount() + ReferencingScriptExpressions.Num();
-}
diff --git a/Engine/Source/Runtime/CoreUObject/Private/UObject/LinkerPlaceholderClass.h b/Engine/Source/Runtime/CoreUObject/Private/UObject/LinkerPlaceholderClass.h
--- a/Engine/Source/Runtime/CoreUObject/Private/UObject/LinkerPlaceholderClass.h
+++ b/Engine/Source/Runtime/CoreUObject/Private/UObject/LinkerPlaceholderClass.h
@@ -2,12 +2,11 @@
 
 #pragma once
 #include "Class.h"
-#include "Set.h"
 #include "LinkerPlaceholderBase.h"
+#include "Set.h"
 
 // Forward declarations
 class FObjectInitializer;
-class UObjectPropertyBase;
 
 /**  
  * A utility class for the deferred dependency loader, used to stub in temporary
@@ -15,57 +14,25 @@ class UObjectPropertyBase;
  * Holds on to references where this is currently being utilized, so we can 
  * easily replace references to it later (once the real class is available).
  */ 
-class ULinkerPlaceholderClass : public UClass, public FLinkerPlaceholderBase
+class ULinkerPlaceholderClass : public UClass, public TLinkerImportPlaceholder<UClass>
 {
 public:
 	DECLARE_CASTED_CLASS_INTRINSIC_NO_CTOR(ULinkerPlaceholderClass, UClass, /*TStaticFlags =*/0, CoreUObject, /*TStaticCastFlags =*/0, NO_API)
 
 	ULinkerPlaceholderClass(const FObjectInitializer& ObjectInitializer);
-	virtual ~ULinkerPlaceholderClass();
-
-	static void AddReferencedObjects(UObject* InThis, FReferenceCollector& Collector);
 
 	// UObject interface.
+	static  void AddReferencedObjects(UObject* InThis, FReferenceCollector& Collector);
  	virtual void PostInitProperties() override;
+	virtual void BeginDestroy() override;
 	// End of UObject interface.
 
 	// UField interface.
  	virtual void Bind() override;
 	// End of UField interface.
 
-	// FLinkerPlaceholderBase 
-	virtual int32 GetRefCount() const override;
-	virtual UObject* GetPlaceholderAsUObject() override
-	{
-		return this;
-	}
-	// End of FLinkerPlaceholderBase
-
-	/**
-	 * Records a raw pointer, directly to the UClass* script expression (so that
-	 * we can switch-out its value in ReplaceTrackedReferences). 
-	 *
-	 * NOTE: We don't worry about creating some kind of weak ref to the script 
-	 *       pointer (or facilitate a way for this tracked reference to be 
-	 *       removed). We're not worried about the script ref being deleted 
-	 *       before we call ReplaceTrackedReferences (because we expect that we
-	 *       do this all within the same frame; before GC can be ran).
-	 * 
-	 * @param  ExpressionPtr    A direct pointer to the UClass* that is now referencing this placeholder.
-	 */
-	void AddReferencingScriptExpr(ULinkerPlaceholderClass** ExpressionPtr);
-
-	/**
-	 * Iterates over all referencing properties and attempts to replace their 
-	 * references to this class with a new (hopefully proper) class.
-	 * 
-	 * @param  ReplacementClass    The class that you want all references to this class replaced with.
-	 * @return The number of references that were successfully replaced.
-	 */
-	int32 ReplaceTrackedReferences(UClass* ReplacementClass);
-
-private:
-
-	/** Points directly at UClass* refs that we're serialized in as part of script bytecode */
-	TSet<UClass**> ReferencingScriptExpressions;
+	// FLinkerPlaceholderBase interface 
+	virtual UObject* GetPlaceholderAsUObject() override { return this; }
+	// End of FLinkerPlaceholderBase interface
 }; 
+
diff --git a/Engine/Source/Runtime/CoreUObject/Private/UObject/LinkerPlaceholderExportObject.cpp b/Engine/Source/Runtime/CoreUObject/Private/UObject/LinkerPlaceholderExportObject.cpp
--- a/Engine/Source/Runtime/CoreUObject/Private/UObject/LinkerPlaceholderExportObject.cpp
+++ b/Engine/Source/Runtime/CoreUObject/Private/UObject/LinkerPlaceholderExportObject.cpp
@@ -2,14 +2,10 @@
 
 #include "CoreUObjectPrivate.h"
 #include "LinkerPlaceholderExportObject.h"
-#include "Blueprint/BlueprintSupport.h"
 
 //------------------------------------------------------------------------------
 ULinkerPlaceholderExportObject::ULinkerPlaceholderExportObject(const FObjectInitializer& ObjectInitializer)
 	: Super(ObjectInitializer)
-#if USE_DEFERRED_DEPENDENCY_CHECK_VERIFICATION_TESTS
-	, bWasResolved(false)
-#endif // USE_DEFERRED_DEPENDENCY_CHECK_VERIFICATION_TESTS
 {
 }
 
@@ -23,138 +19,9 @@ IMPLEMENT_CORE_INTRINSIC_CLASS(ULinkerPlaceholderExportObject, UObject,
 void ULinkerPlaceholderExportObject::BeginDestroy()
 {
 #if USE_DEFERRED_DEPENDENCY_CHECK_VERIFICATION_TESTS
-	check(bWasResolved || HasAnyFlags(RF_ClassDefaultObject));
-	check(ReferencingProperties.Num() == 0);
+	check(IsMarkedResolved() || HasAnyFlags(RF_ClassDefaultObject));
+	check(!HasKnownReferences());
 #endif // USE_DEFERRED_DEPENDENCY_CHECK_VERIFICATION_TESTS
 
 	Super::BeginDestroy();
 }
-
-//------------------------------------------------------------------------------
-void ULinkerPlaceholderExportObject::AddReferencingProperty(const UObjectProperty* ReferencingProperty)
-{
-#if USE_DEFERRED_DEPENDENCY_CHECK_VERIFICATION_TESTS
-	check(bWasResolved != true);
-	check(!ReferencingProperty->GetLinker() || ReferencingProperty->GetLinker()->LinkerRoot == GetOuter());
-#endif // USE_DEFERRED_DEPENDENCY_CHECK_VERIFICATION_TESTS
-
-	ReferencingProperties.Add(ReferencingProperty);
-}
-
-//------------------------------------------------------------------------------
-bool ULinkerPlaceholderExportObject::IsReferencedBy(const UObjectProperty* Property)
-{
-	return ReferencingProperties.Contains(Property);
-}
-
-//------------------------------------------------------------------------------
-void ULinkerPlaceholderExportObject::ReplaceReferencingObjectValues(UClass* Container, UObject* NewObjectValue)
-{
-	UClass* ClassOwner = Container->GetClass();
-	for (TWeakObjectPtr<UObjectProperty> PropertyPtr : ReferencingProperties)
-	{
-		if (!PropertyPtr.IsValid())
-		{
-			continue;
-		}
-		const UObjectProperty* Property = PropertyPtr.Get();
-
-#if USE_DEFERRED_DEPENDENCY_CHECK_VERIFICATION_TESTS
-		check(Property->GetOwnerClass() == ClassOwner);
-#endif // USE_DEFERRED_DEPENDENCY_CHECK_VERIFICATION_TESTS
-
-		TArray<const UProperty*> PropertyChain;
-		PropertyChain.Add(Property);
-		const UProperty* OutermostProperty = Property;
-
-		UObject* PropertyOuter = Property->GetOuter();
-		while ((PropertyOuter != nullptr) && (PropertyOuter != ClassOwner))
-		{
-			if (const UProperty* PropertyOwner = Cast<const UProperty>(PropertyOuter))
-			{
-				OutermostProperty = PropertyOwner;
-				PropertyChain.Add(PropertyOwner);
-			}
-			PropertyOuter = PropertyOuter->GetOuter();
-		}
-
-		uint8* OutermostAddress = OutermostProperty->ContainerPtrToValuePtr<uint8>((uint8*)Container, /*ArrayIndex =*/0);
-		int32 ResolvedCount = ResolvePlaceholderValues(PropertyChain, PropertyChain.Num() - 1, OutermostAddress, NewObjectValue);
-
-#if USE_DEFERRED_DEPENDENCY_CHECK_VERIFICATION_TESTS
-		// we expect that (because we have had ReferencingProperties added) 
-		// there should be at least one reference that is resolved... if this 
-		// there were none, then a property could have changed its value after
-		// it was set to this
-		// 
-		// NOTE: this may seem it can be resolved by properties removing 
-		//       themselves from ReferencingProperties, but certain properties
-		//       may be the inner of a UArrayProperty (meaning there could be
-		//       multiple references per property)... we'd have to inc/decrement
-		//       a property ref-count to resolve that scenario
-		check(ResolvedCount > 0);
-#endif // USE_DEFERRED_DEPENDENCY_CHECK_VERIFICATION_TESTS
-	}
-
-	ReferencingProperties.Empty();
-#if USE_DEFERRED_DEPENDENCY_CHECK_VERIFICATION_TESTS
-	bWasResolved = true;
-#endif // USE_DEFERRED_DEPENDENCY_CHECK_VERIFICATION_TESTS
-}
-
-//------------------------------------------------------------------------------
-int32 ULinkerPlaceholderExportObject::ResolvePlaceholderValues(const TArray<const UProperty*>& PropertyChain, int32 ChainIndex, uint8* ValueAddress, UObject* ReplacementValue)
-{
-	int32 ReplacementCount = 0;
-
-	for (int32 PropertyIndex = ChainIndex; PropertyIndex >= 0; --PropertyIndex)
-	{
-		const UProperty* Property = PropertyChain[PropertyIndex];
-		if (PropertyIndex == 0)
-		{
-#if USE_DEFERRED_DEPENDENCY_CHECK_VERIFICATION_TESTS
-			check(Property->IsA<UObjectProperty>());
-#endif // USE_DEFERRED_DEPENDENCY_CHECK_VERIFICATION_TESTS
-
-			const UObjectProperty* ReferencingProperty = (const UObjectProperty*)Property;
-			if (ReferencingProperty->GetObjectPropertyValue(ValueAddress) == this)
-			{
-				ReferencingProperty->SetObjectPropertyValue(ValueAddress, ReplacementValue);
-				// @TODO: unfortunately, this is currently protected
-				//ReferencingProperty->CheckValidObject(ValueAddress);
-
-				++ReplacementCount;
-			}
-		}
-		else if (const UArrayProperty* ArrayProperty = Cast<const UArrayProperty>(Property))
-		{
-#if USE_DEFERRED_DEPENDENCY_CHECK_VERIFICATION_TESTS
-			const UProperty* NextProperty = PropertyChain[PropertyIndex - 1];
-			check(NextProperty == ArrayProperty->Inner);
-#endif // USE_DEFERRED_DEPENDENCY_CHECK_VERIFICATION_TESTS
-
-			// because we can't know which array entry was set with a reference 
-			// to this object, we have to comb through them all
-			FScriptArrayHelper ArrayHelper(ArrayProperty, ValueAddress);
-			for (int32 ArrayIndex = 0; ArrayIndex < ArrayHelper.Num(); ++ArrayIndex)
-			{
-				uint8* MemberAddress = ArrayHelper.GetRawPtr(ArrayIndex);
-				ReplacementCount += ResolvePlaceholderValues(PropertyChain, PropertyIndex - 1, MemberAddress, ReplacementValue);
-			}
-
-			// the above recursive call chewed through the rest of the
-			// PropertyChain, no need to keep on here
-			break;
-		}
-		else
-		{
-			const UProperty* NextProperty = PropertyChain[PropertyIndex - 1];
-#if USE_DEFERRED_DEPENDENCY_CHECK_VERIFICATION_TESTS
-			check(NextProperty->GetOuter() == Property);
-#endif // USE_DEFERRED_DEPENDENCY_CHECK_VERIFICATION_TESTS
-			ValueAddress = Property->ContainerPtrToValuePtr<uint8>(ValueAddress, /*ArrayIndex =*/0);
-		}
-	}
-
-	return ReplacementCount;
-}
diff --git a/Engine/Source/Runtime/CoreUObject/Private/UObject/LinkerPlaceholderExportObject.h b/Engine/Source/Runtime/CoreUObject/Private/UObject/LinkerPlaceholderExportObject.h
--- a/Engine/Source/Runtime/CoreUObject/Private/UObject/LinkerPlaceholderExportObject.h
+++ b/Engine/Source/Runtime/CoreUObject/Private/UObject/LinkerPlaceholderExportObject.h
@@ -2,18 +2,18 @@
 
 #pragma once
 #include "Object.h"
+#include "LinkerPlaceholderBase.h" // for FLinkerPlaceholderBase
 #include "Set.h"
 
 // Forward declarations
 class FObjectInitializer;
-class UObjectProperty;
 
 /**  
  * A utility class for the deferred dependency loader, used to stub in temporary
  * export objects so we don't spawn any Blueprint class instances before the 
  * class is fully regenerated.
  */ 
-class ULinkerPlaceholderExportObject : public UObject
+class ULinkerPlaceholderExportObject : public UObject, public FLinkerPlaceholderBase
 {
 public:
 	DECLARE_CASTED_CLASS_INTRINSIC_NO_CTOR(ULinkerPlaceholderExportObject, UObject, /*TStaticFlags =*/0, CoreUObject, /*TStaticCastFlags =*/0, NO_API)
@@ -24,69 +24,7 @@ public:
 	virtual void BeginDestroy() override;
 	// End of UObject interface
 
-	/**
-	 * Logs the supplied property as referencing this class (meaning an instance  
-	 * is referencing this object through that property's value).
-	 * 
-	 * @param  ReferencingProperty    The property whose object-value is pointing to this class.
-	 */
-	void AddReferencingProperty(const UObjectProperty* ReferencingProperty);
-
-	/**
-	 * Checks to see if the specified property is being tracked by this as a
-	 * referencer.
-	 * 
-	 * @param  Property    The property you want to check for.
-	 * @return True if the property was added as a referencer through AddReferencingProperty(), otherwise false.
-	 */
-	bool IsReferencedBy(const UObjectProperty* Property);
-
-	/**
-	 * Resolves all referencing object values (assuming they were all made as 
-	 * part of the specified Container object).
-	 *
-	 * NOTE: there is a very strict assumption here, that any property values 
-	 *       that were set to this (and recorded through AddReferencingProperty)
-	 *       were properties belonging to the specified Container object (if not
-	 *       there could be dangling, unresolved references).
-	 * 
-	 * @param  Container		The object whose properties may be referencing this temporary object.
-	 * @param  NewObjectValue   The object which should replace all references to this. 
-	 */
-	void ReplaceReferencingObjectValues(UClass* Container, UObject* NewObjectValue);
-
-private:
-	/**
-	 * A recursive method that replaces all leaf references to this object with 
-	 * the supplied ReplacementValue.
-	 *
-	 * This function recurses the property chain (from class owner down) because 
-	 * at the time of AddReferencingProperty() we cannot know/record the 
-	 * address/index of array properties (as they may change during array 
-	 * re-allocation or compaction). So we must follow the property chain and 
-	 * check every UArrayProperty member for references to this (hence, the need
-	 * for this recursive function).
-	 * 
-	 * @param  PropertyChain    An ascending outer chain, where the property at index zero is the leaf (referencer) property.
-	 * @param  ChainIndex		An index into the PropertyChain that this call should start at and iterate DOWN to zero.
-	 * @param  ValueAddress     The memory address of the value corresponding to the property at ChainIndex.
-	 * @param  ReplacementValue	The new object to replace all references to this with.
-	 * @return The number of references that were replaced.
-	 */
-	int32 ResolvePlaceholderValues(const TArray<const UProperty*>& PropertyChain, int32 ChainIndex, uint8* ValueAddress, UObject* ReplacementValue);
-
-	/**
-	 * A list of KNOWN referencing properties (meaning an object instance 
-	 * references this object through these properties).
-	 *
-	 * @TODO: consider making this a map with a ref-counter, so we can remove 
-	 *        property references (as they may be replaced prior to resolve)
-	 */
-	TSet< TWeakObjectPtr<UObjectProperty> > ReferencingProperties;
-
-#if USE_DEFERRED_DEPENDENCY_CHECK_VERIFICATION_TESTS
-public:
-	/** Used to catch any references that were added after ULinkerLoad has already done a resolve pass. */
-	bool bWasResolved;
-#endif // USE_DEFERRED_DEPENDENCY_CHECK_VERIFICATION_TESTS
+	// FLinkerPlaceholderBase interface
+	virtual UObject* GetPlaceholderAsUObject() override { return this; }
+	// End of FLinkerPlaceholderBase interface
 }; 
diff --git a/Engine/Source/Runtime/CoreUObject/Private/UObject/LinkerPlaceholderFunction.cpp b/Engine/Source/Runtime/CoreUObject/Private/UObject/LinkerPlaceholderFunction.cpp
--- a/Engine/Source/Runtime/CoreUObject/Private/UObject/LinkerPlaceholderFunction.cpp
+++ b/Engine/Source/Runtime/CoreUObject/Private/UObject/LinkerPlaceholderFunction.cpp
@@ -9,70 +9,8 @@ ULinkerPlaceholderFunction::ULinkerPlaceholderFunction(const FObjectInitializer&
 {
 }
 
-//------------------------------------------------------------------------------
-ULinkerPlaceholderFunction::~ULinkerPlaceholderFunction()
-{
-}
-
 //------------------------------------------------------------------------------
 IMPLEMENT_CORE_INTRINSIC_CLASS(ULinkerPlaceholderFunction, UFunction,
 	{
 	}
 );
-
-//------------------------------------------------------------------------------
-int32 ULinkerPlaceholderFunction::ReplaceTrackedReferences(UFunction* ReplacementObj)
-{
-	int32 ReplacementCount = 0;
-	for (UProperty* Property : ReferencingProperties)
-	{
-		if (auto DelegateProperty = Cast<UDelegateProperty>(Property))
-		{
-			if (DelegateProperty->SignatureFunction == this)
-			{
-				DelegateProperty->SignatureFunction = ReplacementObj;
-				++ReplacementCount;
-			}
-		}
-		else if (auto MulticastDelegateProperty = Cast<UMulticastDelegateProperty>(Property))
-		{
-			MulticastDelegateProperty->SignatureFunction = ReplacementObj;
-			++ReplacementCount;
-		}
-		else
-		{
-			ensure(false);
-		}
-	}
-	ReferencingProperties.Empty();
-
-	for (UFunction** ScriptRefPtr : ReferencingScriptExpressions)
-	{
-		UFunction*& ScriptRef = *ScriptRefPtr;
-		if (ScriptRef == this)
-		{
-			ScriptRef = ReplacementObj;
-			++ReplacementCount;
-		}
-	}
-	ReferencingScriptExpressions.Empty();
-
-	bResolvedReferences = true;
-	return ReplacementCount;
-}
-
-//------------------------------------------------------------------------------
-void ULinkerPlaceholderFunction::AddReferencingScriptExpr(ULinkerPlaceholderFunction** ExpressionPtr)
-{
-#if USE_DEFERRED_DEPENDENCY_CHECK_VERIFICATION_TESTS
-	check(*ExpressionPtr == this);
-#endif // USE_DEFERRED_DEPENDENCY_CHECK_VERIFICATION_TESTS
-
-	ReferencingScriptExpressions.Add((UFunction**)ExpressionPtr);
-}
-
-//------------------------------------------------------------------------------
-int32 ULinkerPlaceholderFunction::GetRefCount() const
-{
-	return FLinkerPlaceholderBase::GetRefCount() + ReferencingScriptExpressions.Num();
-}
\ No newline at end of file
diff --git a/Engine/Source/Runtime/CoreUObject/Private/UObject/LinkerPlaceholderFunction.h b/Engine/Source/Runtime/CoreUObject/Private/UObject/LinkerPlaceholderFunction.h
--- a/Engine/Source/Runtime/CoreUObject/Private/UObject/LinkerPlaceholderFunction.h
+++ b/Engine/Source/Runtime/CoreUObject/Private/UObject/LinkerPlaceholderFunction.h
@@ -1,33 +1,21 @@
 // Copyright 1998-2014 Epic Games, Inc. All Rights Reserved.
 
 #pragma once
-#include "LinkerPlaceholderClass.h"
+#include "Class.h" // for UFunction
+#include "LinkerPlaceholderBase.h"
 
 // Forward declarations
 class FObjectInitializer;
-class UObjectPropertyBase;
 
-class ULinkerPlaceholderFunction : public UFunction, public FLinkerPlaceholderBase
+/**  */
+class ULinkerPlaceholderFunction : public UFunction, public TLinkerImportPlaceholder<UFunction>
 {
 public:
 	DECLARE_CASTED_CLASS_INTRINSIC_NO_CTOR(ULinkerPlaceholderFunction, UFunction, /*TStaticFlags =*/0, CoreUObject, /*TStaticCastFlags =*/0, NO_API)
 
 	ULinkerPlaceholderFunction(const FObjectInitializer& ObjectInitializer);
-	virtual ~ULinkerPlaceholderFunction();
 
-	// FPlaceholderBase 
-	virtual int32 GetRefCount() const override;
-	virtual UObject* GetPlaceholderAsUObject() override
-	{
-		return this;
-	}
-	// End of FPlaceholderBase
-
-	int32 ReplaceTrackedReferences(UFunction* ReplacementClass);
-
-	void AddReferencingScriptExpr(ULinkerPlaceholderFunction** ExpressionPtr);
-
-private:
-
-	TSet<UFunction**> ReferencingScriptExpressions;
-};
\ No newline at end of file
+	// FLinkerPlaceholderBase interface 
+	virtual UObject* GetPlaceholderAsUObject() override { return this; }
+	// End of FLinkerPlaceholderBase interface
+};
diff --git a/Engine/Source/Runtime/CoreUObject/Private/UObject/PropertyBaseObject.cpp b/Engine/Source/Runtime/CoreUObject/Private/UObject/PropertyBaseObject.cpp
--- a/Engine/Source/Runtime/CoreUObject/Private/UObject/PropertyBaseObject.cpp
+++ b/Engine/Source/Runtime/CoreUObject/Private/UObject/PropertyBaseObject.cpp
@@ -14,7 +14,7 @@ void UObjectPropertyBase::BeginDestroy()
 #if USE_CIRCULAR_DEPENDENCY_LOAD_DEFERRING
 	if (ULinkerPlaceholderClass* PlaceholderClass = Cast<ULinkerPlaceholderClass>(PropertyClass))
 	{
-		PlaceholderClass->RemovePropertyReference(this);
+		PlaceholderClass->RemoveReferencingProperty(this);
 	}
 #endif // USE_CIRCULAR_DEPENDENCY_LOAD_DEFERRING
 
@@ -110,7 +110,7 @@ void UObjectPropertyBase::SetPropertyClass(UClass* NewPropertyClass)
 	
 	if (ULinkerPlaceholderClass* OldPlaceholderClass = Cast<ULinkerPlaceholderClass>(PropertyClass))
 	{
-		OldPlaceholderClass->RemovePropertyReference(this);
+		OldPlaceholderClass->RemoveReferencingProperty(this);
 	}
 	PropertyClass = NewPropertyClass;
 }
diff --git a/Engine/Source/Runtime/CoreUObject/Private/UObject/PropertyClass.cpp b/Engine/Source/Runtime/CoreUObject/Private/UObject/PropertyClass.cpp
--- a/Engine/Source/Runtime/CoreUObject/Private/UObject/PropertyClass.cpp
+++ b/Engine/Source/Runtime/CoreUObject/Private/UObject/PropertyClass.cpp
@@ -13,7 +13,7 @@ void UClassProperty::BeginDestroy()
 #if USE_CIRCULAR_DEPENDENCY_LOAD_DEFERRING
 	if (ULinkerPlaceholderClass* PlaceholderClass = Cast<ULinkerPlaceholderClass>(MetaClass))
 	{
-		PlaceholderClass->RemovePropertyReference(this);
+		PlaceholderClass->RemoveReferencingProperty(this);
 	}
 #endif // USE_CIRCULAR_DEPENDENCY_LOAD_DEFERRING
 
@@ -58,7 +58,7 @@ void UClassProperty::SetMetaClass(UClass* NewMetaClass)
 
 	if (ULinkerPlaceholderClass* OldPlaceholderClass = Cast<ULinkerPlaceholderClass>(MetaClass))
 	{
-		OldPlaceholderClass->RemovePropertyReference(this);
+		OldPlaceholderClass->RemoveReferencingProperty(this);
 	}
 	MetaClass = NewMetaClass;
 }
diff --git a/Engine/Source/Runtime/CoreUObject/Private/UObject/PropertyDelegate.cpp b/Engine/Source/Runtime/CoreUObject/Private/UObject/PropertyDelegate.cpp
--- a/Engine/Source/Runtime/CoreUObject/Private/UObject/PropertyDelegate.cpp
+++ b/Engine/Source/Runtime/CoreUObject/Private/UObject/PropertyDelegate.cpp
@@ -129,7 +129,7 @@ void UDelegateProperty::BeginDestroy()
 #if USE_CIRCULAR_DEPENDENCY_LOAD_DEFERRING
 	if (auto PlaceholderFunc = Cast<ULinkerPlaceholderFunction>(SignatureFunction))
 	{
-		PlaceholderFunc->RemovePropertyReference(this);
+		PlaceholderFunc->RemoveReferencingProperty(this);
 	}
 #endif // USE_CIRCULAR_DEPENDENCY_LOAD_DEFERRING
 
diff --git a/Engine/Source/Runtime/CoreUObject/Private/UObject/PropertyInterface.cpp b/Engine/Source/Runtime/CoreUObject/Private/UObject/PropertyInterface.cpp
--- a/Engine/Source/Runtime/CoreUObject/Private/UObject/PropertyInterface.cpp
+++ b/Engine/Source/Runtime/CoreUObject/Private/UObject/PropertyInterface.cpp
@@ -12,7 +12,7 @@ void UInterfaceProperty::BeginDestroy()
 #if USE_CIRCULAR_DEPENDENCY_LOAD_DEFERRING
 	if (ULinkerPlaceholderClass* PlaceholderClass = Cast<ULinkerPlaceholderClass>(InterfaceClass))
 	{
-		PlaceholderClass->RemovePropertyReference(this);
+		PlaceholderClass->RemoveReferencingProperty(this);
 	}
 #endif // USE_CIRCULAR_DEPENDENCY_LOAD_DEFERRING
 
@@ -249,7 +249,7 @@ void UInterfaceProperty::SetInterfaceClass(UClass* NewInterfaceClass)
 
 	if (ULinkerPlaceholderClass* OldPlaceholderClass = Cast<ULinkerPlaceholderClass>(InterfaceClass))
 	{
-		OldPlaceholderClass->RemovePropertyReference(this);
+		OldPlaceholderClass->RemoveReferencingProperty(this);
 	}
 	InterfaceClass = NewInterfaceClass;
 }
diff --git a/Engine/Source/Runtime/CoreUObject/Private/UObject/PropertyMulticastDelegate.cpp b/Engine/Source/Runtime/CoreUObject/Private/UObject/PropertyMulticastDelegate.cpp
--- a/Engine/Source/Runtime/CoreUObject/Private/UObject/PropertyMulticastDelegate.cpp
+++ b/Engine/Source/Runtime/CoreUObject/Private/UObject/PropertyMulticastDelegate.cpp
@@ -269,7 +269,7 @@ void UMulticastDelegateProperty::BeginDestroy()
 #if USE_CIRCULAR_DEPENDENCY_LOAD_DEFERRING
 	if (auto PlaceholderFunc = Cast<ULinkerPlaceholderFunction>(SignatureFunction))
 	{
-		PlaceholderFunc->RemovePropertyReference(this);
+		PlaceholderFunc->RemoveReferencingProperty(this);
 	}
 #endif // USE_CIRCULAR_DEPENDENCY_LOAD_DEFERRING
 
diff --git a/Engine/Source/Runtime/CoreUObject/Private/UObject/PropertyObject.cpp b/Engine/Source/Runtime/CoreUObject/Private/UObject/PropertyObject.cpp
--- a/Engine/Source/Runtime/CoreUObject/Private/UObject/PropertyObject.cpp
+++ b/Engine/Source/Runtime/CoreUObject/Private/UObject/PropertyObject.cpp
@@ -2,6 +2,7 @@
 
 #include "CoreUObjectPrivate.h"
 #include "LinkerPlaceholderExportObject.h"
+#include "LinkerPlaceholderClass.h"
 
 /*-----------------------------------------------------------------------------
 	UObjectProperty.
@@ -31,10 +32,16 @@ void UObjectProperty::SerializeItem( FArchive& Ar, void* Value, void const* Defa
 	UObject* CurrentValue = GetObjectPropertyValue(Value);
 	if (ObjectValue != CurrentValue)
 	{
+		SetObjectPropertyValue(Value, ObjectValue);
+
 #if USE_CIRCULAR_DEPENDENCY_LOAD_DEFERRING
 		if (ULinkerPlaceholderExportObject* PlaceholderVal = Cast<ULinkerPlaceholderExportObject>(ObjectValue))
 		{
-			PlaceholderVal->AddReferencingProperty(this);
+			PlaceholderVal->AddReferencingPropertyValue(this, Value);
+		}
+		else if (ULinkerPlaceholderClass* PlaceholderClass = Cast<ULinkerPlaceholderClass>(ObjectValue))
+		{
+			PlaceholderClass->AddReferencingPropertyValue(this, Value);
 		}
 		// NOTE: we don't remove this from CurrentValue if it is a 
 		//       ULinkerPlaceholderExportObject; this is because this property 
@@ -49,20 +56,12 @@ void UObjectProperty::SerializeItem( FArchive& Ar, void* Value, void const* Defa
 		//        to accommodate this (as it depends on finding itself as the set value)
 #endif // USE_CIRCULAR_DEPENDENCY_LOAD_DEFERRING
 
-		SetObjectPropertyValue(Value, ObjectValue);
 		CheckValidObject(Value);
 	}
 }
 
 void UObjectProperty::SetObjectPropertyValue(void* PropertyValueAddress, UObject* Value) const
 {
-#if USE_DEFERRED_DEPENDENCY_CHECK_VERIFICATION_TESTS
-	if (ULinkerPlaceholderExportObject* PlaceholderVal = Cast<ULinkerPlaceholderExportObject>(Value))
-	{
-		check(PlaceholderVal->IsReferencedBy(this));
-	}
-#endif // USE_DEFERRED_DEPENDENCY_CHECK_VERIFICATION_TESTS
-
 	SetPropertyValue(PropertyValueAddress, Value);
 }
 
diff --git a/Engine/Source/Runtime/CoreUObject/Public/UObject/Linker.h b/Engine/Source/Runtime/CoreUObject/Public/UObject/Linker.h
--- a/Engine/Source/Runtime/CoreUObject/Public/UObject/Linker.h
+++ b/Engine/Source/Runtime/CoreUObject/Public/UObject/Linker.h
@@ -1891,7 +1891,7 @@ private:
 	 * @param  ReferencingClass	The (Blueprint) class that was loading, while we deferred dependencies (now referencing the placeholder).
 	 * @return The number of placeholder references replaced (could be none, if this was recursively resolved).
 	 */
-	int32 ResolveDependencyPlaceholder(UClass* Placeholder, UClass* ReferencingClass = nullptr);
+	int32 ResolveDependencyPlaceholder(class FLinkerPlaceholderBase* Placeholder, UClass* ReferencingClass = nullptr);
 
 	/**
 	 * Query method to help catch recursive behavior. When this returns true, a 
@@ -1983,7 +1983,7 @@ public:
 	 * to make sure this placeholder is completely resolved before continuing on 
 	 * to the next.
 	 */
-	UClass* ResolvingDeferredPlaceholder;
+	class FLinkerPlaceholderBase* ResolvingDeferredPlaceholder;
 #endif // USE_CIRCULAR_DEPENDENCY_LOAD_DEFERRING
 
 
