Fixes to TLazyObjectPtr, style improvements and new features:
New define: PLATFORM_COMPILER_HAS_DEFAULT_FUNCTION_TEMPLATE_ARGUMENTS, indicating whether function templates can have default template args or whether workarounds are needed. This is effectively a global version of ENABLE_TFUNCTIONREF_WORKAROUND from TFunction. Only needed for VS 2012, once we drop that we can remove this.
Now constructs from T* instead of (templated) U*. It cast to a T* anyway, so would have failed to compile with anything that couldn't already bind to a T* arg, so the template was pointless.
As a result no constructor from TYPE_OF_NULL is needed (calls the T* constructor with NULL), so it has been removed. nullptr also now works!
Comparison operators have been made symmetrical (TLazyObjectPtr can now be the right-hand argument as well) and comparisons with nullptr now compile
Implicit upcasting is now enabled in the constructor and assignment, so "TLazyObjectPtr<Base> mybaselazyptr = TLazyObjectPtr<Derived>()" now works, as does passing a TLazyObjectPtr<Derived> to a function that expects a TLazyObjectPtr<Base>. As this is always safe, this works even if the object isn't currently loaded!
The constructor from a TWeakObjectPtr has been removed because it didn't even compile

diff --git a/Engine/Source/Runtime/Core/Public/HAL/Platform.h b/Engine/Source/Runtime/Core/Public/HAL/Platform.h
--- a/Engine/Source/Runtime/Core/Public/HAL/Platform.h
+++ b/Engine/Source/Runtime/Core/Public/HAL/Platform.h
@@ -165,6 +165,9 @@
 #ifndef PLATFORM_COMPILER_HAS_VARIADIC_TEMPLATES
 	#define PLATFORM_COMPILER_HAS_VARIADIC_TEMPLATES	1
 #endif
+#ifndef PLATFORM_COMPILER_HAS_DEFAULT_FUNCTION_TEMPLATE_ARGUMENTS
+	#define PLATFORM_COMPILER_HAS_DEFAULT_FUNCTION_TEMPLATE_ARGUMENTS	1
+#endif
 #ifndef PLATFORM_COMPILER_HAS_EXPLICIT_OPERATORS
 	#define PLATFORM_COMPILER_HAS_EXPLICIT_OPERATORS	1
 #endif
diff --git a/Engine/Source/Runtime/Core/Public/HTML5/HTML5Platform.h b/Engine/Source/Runtime/Core/Public/HTML5/HTML5Platform.h
--- a/Engine/Source/Runtime/Core/Public/HTML5/HTML5Platform.h
+++ b/Engine/Source/Runtime/Core/Public/HTML5/HTML5Platform.h
@@ -32,6 +32,8 @@ typedef FHTML5Types FPlatformTypes;
 #define PLATFORM_HAS_BSD_TIME						0
 #define PLATFORM_COMPILER_HAS_DEFAULTED_FUNCTIONS	0
 #define PLATFORM_COMPILER_HAS_VARIADIC_TEMPLATES	0
+#define PLATFORM_COMPILER_HAS_EXPLICIT_OPERATORS	0
+#define PLATFORM_COMPILER_HAS_DEFAULT_FUNCTION_TEMPLATE_ARGUMENTS	0
 #define PLATFORM_USES_MICROSOFT_LIBC_FUNCTIONS		1
 #define PLATFORM_MAX_FILEPATH_LENGTH				MAX_PATH
 #define PLATFORM_HAS_BSD_SOCKETS					0
@@ -54,7 +56,6 @@ typedef FHTML5Types FPlatformTypes;
 #if PLATFORM_HTML5_WIN32
 #define FORCEINLINE _forceinline
 #define FORCENOINLINE __declspec(noinline)	/* Force code to NOT be inline */
-#define PLATFORM_COMPILER_HAS_EXPLICIT_OPERATORS 0
 #define CONSTEXPR    
 #else
 #define FORCEINLINE		inline __attribute__((__always_inline__))		/* Force code to be inline */
diff --git a/Engine/Source/Runtime/Core/Public/WinRT/WinRTPlatform.h b/Engine/Source/Runtime/Core/Public/WinRT/WinRTPlatform.h
--- a/Engine/Source/Runtime/Core/Public/WinRT/WinRTPlatform.h
+++ b/Engine/Source/Runtime/Core/Public/WinRT/WinRTPlatform.h
@@ -53,6 +53,7 @@ typedef FWinRTTypes FPlatformTypes;
 #define PLATFORM_COMPILER_HAS_DEFAULTED_FUNCTIONS			0
 #define PLATFORM_COMPILER_HAS_VARIADIC_TEMPLATES			0
 #define PLATFORM_COMPILER_HAS_EXPLICIT_OPERATORS			0
+#define PLATFORM_COMPILER_HAS_DEFAULT_FUNCTION_TEMPLATE_ARGUMENTS	0
 
 [REDACTED_EMAIL]: Fixup once sockets are supported
 #define PLATFORM_SUPPORTS_MESSAGEBUS						0
diff --git a/Engine/Source/Runtime/Core/Public/Windows/WIndowsPlatform.h b/Engine/Source/Runtime/Core/Public/Windows/WIndowsPlatform.h
--- a/Engine/Source/Runtime/Core/Public/Windows/WIndowsPlatform.h
+++ b/Engine/Source/Runtime/Core/Public/Windows/WIndowsPlatform.h
@@ -60,9 +60,11 @@ typedef FWindowsPlatformTypes FPlatformTypes;
 #if _MSC_VER >= 1800 || __clang__
 	#define PLATFORM_COMPILER_HAS_VARIADIC_TEMPLATES		1
 	#define PLATFORM_COMPILER_HAS_EXPLICIT_OPERATORS		1
+	#define PLATFORM_COMPILER_HAS_DEFAULT_FUNCTION_TEMPLATE_ARGUMENTS	1
 #else
 	#define PLATFORM_COMPILER_HAS_VARIADIC_TEMPLATES		0
 	#define PLATFORM_COMPILER_HAS_EXPLICIT_OPERATORS		0
+	#define PLATFORM_COMPILER_HAS_DEFAULT_FUNCTION_TEMPLATE_ARGUMENTS	0
 #endif
 #define PLATFORM_COMPILER_HAS_TCHAR_WMAIN					1
 
diff --git a/Engine/Source/Runtime/CoreUObject/Public/UObject/LazyObjectPtr.h b/Engine/Source/Runtime/CoreUObject/Public/UObject/LazyObjectPtr.h
--- a/Engine/Source/Runtime/CoreUObject/Public/UObject/LazyObjectPtr.h
+++ b/Engine/Source/Runtime/CoreUObject/Public/UObject/LazyObjectPtr.h
@@ -54,7 +54,7 @@ struct COREUOBJECT_API FUniqueObjectGuid
 	 */
 	class UObject *ResolveObject() const;
 
-	/**  
+	/**
 	 * Test if this can ever point to a live UObject
 	 */
 	FORCEINLINE bool IsValid() const
@@ -144,7 +144,7 @@ public:
 	{
 	}
 
-	/**  
+	/**
 	 * Construct from another lazy pointer
 	 * @param Other lazy pointer to copy from
 	 */
@@ -153,7 +153,7 @@ public:
 		(*this)=Other;
 	}
 
-	/**  
+	/**
 	 * Construct from another lazy pointer
 	 * @param Other lazy pointer to copy from
 	 */
@@ -161,7 +161,7 @@ public:
 	{
 		(*this)=Object;
 	}
-	/**  
+	/**
 	 * Copy from an object pointer
 	 * @param Object object to create a weak pointer to
 	 */
@@ -170,7 +170,7 @@ public:
 		TPersistentObjectPtr<FUniqueObjectGuid>::operator=(Object);
 	}
 
-	/**  
+	/**
 	 * Copy from a lazy pointer
 	 */
 	FORCEINLINE void operator=(const FLazyObjectPtr &Other)
@@ -178,7 +178,7 @@ public:
 		TPersistentObjectPtr<FUniqueObjectGuid>::operator=(Other);
 	}
 
-	/**  
+	/**
 	 * Copy from a unique object identifier
 	 * @param ObjectID Object identifier to create a weak pointer to
 	 */
@@ -201,38 +201,72 @@ template<class T=UObject>
 class TLazyObjectPtr : private FLazyObjectPtr
 {
 public:
+#if PLATFORM_COMPILER_HAS_DEFAULTED_FUNCTIONS
+	FORCEINLINE TLazyObjectPtr() = default;
+	FORCEINLINE TLazyObjectPtr(const TLazyObjectPtr<T>& Other) = default;
+	FORCEINLINE TLazyObjectPtr(TLazyObjectPtr<T>&& Other) = default;
+	FORCEINLINE TLazyObjectPtr<T>& operator=(const TLazyObjectPtr<T>& Other) = default;
+	FORCEINLINE TLazyObjectPtr<T>& operator=(TLazyObjectPtr<T>&& Other) = default;
+#else
 	/** NULL constructor **/
 	FORCEINLINE TLazyObjectPtr()
 	{
 	}
 	
-	/**  
+	/**
 	 * Construct from another lazy pointer
 	 * @param Other lazy pointer to copy from
 	 */
-	FORCEINLINE TLazyObjectPtr(const TLazyObjectPtr<T> &Other) :
+	FORCEINLINE TLazyObjectPtr(const TLazyObjectPtr<T>& Other) :
 		FLazyObjectPtr(Other)
 	{
 	}
 
-	/**  
-	 * Construct from an int (assumed to be NULL)
-	 * @param Other int to create a reference to
+	/**
+	* Assign from another lazy pointer
+	* @param Other lazy pointer to copy from
+	*/
+	FORCEINLINE TLazyObjectPtr<T>& operator=(const TLazyObjectPtr<T>& Other)
+	{
+		FLazyObjectPtr::operator=(Other);
+		return *this;
+	}
+#endif
+
+	/**
+	 * Construct from another lazy pointer with implicit upcasting allowed
+	 * @param Other lazy pointer to copy from
 	 */
-	FORCEINLINE TLazyObjectPtr(TYPE_OF_NULL Other)
+#if PLATFORM_COMPILER_HAS_DEFAULT_FUNCTION_TEMPLATE_ARGUMENTS
+	template<typename U, typename = typename TEnableIf<TPointerIsConvertibleFromTo<U, T>::Value>::Type>
+	FORCEINLINE TLazyObjectPtr(const TLazyObjectPtr<U>& Other) :
+#else
+	template<typename U>
+	FORCEINLINE TLazyObjectPtr(const TLazyObjectPtr<U>& Other, typename TEnableIf<TPointerIsConvertibleFromTo<U, T>::Value, T*>::Type = nullptr) :
+#endif
+		FLazyObjectPtr((const FLazyObjectPtr&)Other)
 	{
-		checkSlow(!Other); // we assume this is NULL
+	}
+	
+	/**
+	 * Assign from another lazy pointer with implicit upcasting allowed
+	 * @param Other lazy pointer to copy from
+	 */
+	template<typename U>
+	FORCEINLINE typename TEnableIf<TPointerIsConvertibleFromTo<U, T>::Value, TLazyObjectPtr<T>&>::Type
+		operator=(const TLazyObjectPtr<U>& Other)
+	{
+		FLazyObjectPtr::operator=((const FLazyObjectPtr&)Other);
+		return *this;
 	}
 
-	/**  
+	/**
 	 * Construct from an object pointer
 	 * @param Object Object to create a lazy pointer to
 	 */
-	template<class U>
-	FORCEINLINE TLazyObjectPtr(U* Object)
+	FORCEINLINE TLazyObjectPtr(T* Object)
 	{
-		const T* TempObject = Object;
-		FLazyObjectPtr::operator=(TempObject);
+		FLazyObjectPtr::operator=(Object);
 	}
 
 	/**
@@ -243,74 +277,26 @@ public:
 		FLazyObjectPtr::Reset();
 	}
 
-	/**  
+	/**
 	 * Copy from an object pointer
 	 * @param Object Object to create a lazy pointer to
 	 */
-	template<class U>
-	FORCEINLINE void operator=(const U *Object)
+	FORCEINLINE void operator=(T* Object)
 	{
-		const T* TempObject = Object;
-		FLazyObjectPtr::operator=(TempObject);
+		FLazyObjectPtr::operator=(Object);
 	}
 
-	/**  
+	/**
 	 * Copy from a unique object identifier
-	 * @param ObjectID Object identifier to create a weak pointer to
+	 * WARNING: this doesn't check the type of the object is correct,
+	 * because the object corresponding to this ID may not even be loaded!
+	 * @param ObjectID Object identifier to create a lazy pointer to
 	 */
-	FORCEINLINE void operator=(const FUniqueObjectGuid &ObjectID)
+	FORCEINLINE void operator=(const FUniqueObjectGuid& ObjectID)
 	{
 		FLazyObjectPtr::operator=(ObjectID);
 	}
 
-	/**  
-	 * Copy from a weak pointer
-	 * @param Other Weak pointer to copy from
-	 */
-	FORCEINLINE void operator=(const TWeakObjectPtr<T> &Other)
-	{
-		FLazyObjectPtr::operator=(Other);
-	}
-
-	/**  
-	 * Construct from another lazy pointer
-	 * @param Other lazy pointer to copy from
-	 */
-	FORCEINLINE void operator=(const TLazyObjectPtr<T> &Other)
-	{
-		FLazyObjectPtr::operator=(Other);
-	}
-
-	/**  
-	 * Copy from an int (assumed to be NULL)
-	 * @param Other int to create a reference to
-	 */
-	FORCEINLINE void operator=(TYPE_OF_NULL Other)
-	{
-		checkSlow(!Other); // we assume this is NULL
-		FLazyObjectPtr::Reset();
-	}
-
-	/**  
-	 * Compare lazy pointers for equality
-	 * @param Other lazy pointer to compare to
-	 * Caution: Two lazy pointers might not be equal to each other, but they both might return NULL.
-	 */
-	FORCEINLINE bool operator==(const TLazyObjectPtr<T> &Other) const
-	{
-		return FLazyObjectPtr::operator==(Other);
-	}
-
-	/**  
-	 * Compare lazy pointers for equality
-	 * @param Other lazy pointer to compare to
-	 * Caution: Two lazy pointers might not be equal to each other, but they both might return NULL.
-	 */
-	FORCEINLINE bool operator!=(const TLazyObjectPtr<T> &Other) const
-	{
-		return FLazyObjectPtr::operator!=(Other);
-	}
-
 	/**
 	 * Gets the unique object identifier associated with this lazy pointer. Valid even if pointer is not currently valid
 	 * @return Unique ID for this object, or an invalid FUniqueObjectID if this pointer isn't set to anything
@@ -320,32 +306,35 @@ public:
 		return FLazyObjectPtr::GetUniqueID();
 	}
 
-	/**  
+	/**
 	 * Dereference the lazy pointer.
 	 * @return NULL if this object is gone or the lazy pointer was NULL, otherwise a valid uobject pointer
 	 */
-	FORCEINLINE T *Get() const
+	FORCEINLINE T* Get() const
 	{
+		// there are cases where a TLazyObjectPtr can get an object of the wrong type assigned to it which are difficult to avoid
+		// e.g. operator=(const FUniqueObjectGuid& ObjectID)
+		// "WARNING: this doesn't check the type of the object is correct..."
 		return dynamic_cast<T*>(FLazyObjectPtr::Get());
 	}
 
-	/**  
+	/**
 	 * Dereference the lazy pointer.
 	 */
-	FORCEINLINE T & operator*() const
+	FORCEINLINE T& operator*() const
 	{
 		return *Get();
 	}
 
-	/**  
+	/**
 	 * Dereference the lazy pointer.
 	 */
-	FORCEINLINE T * operator->() const
+	FORCEINLINE T* operator->() const
 	{
 		return Get();
 	}
 
-	/**  
+	/**
 	 * Test if this points to a live UObject
 	 * @return true if Get() would return a valid non-null pointer
 	 */
@@ -354,7 +343,7 @@ public:
 		return FLazyObjectPtr::IsValid();
 	}
 
-	/**  
+	/**
 	 * Slightly different than !IsValid(), returns true if this used to point to a UObject, but doesn't any more and has not been assigned or reset in the mean time.
 	 * @return true if this used to point at a real object but no longer does.
 	**/
@@ -363,7 +352,7 @@ public:
 		return FLazyObjectPtr::IsStale();
 	}
 
-	/**  
+	/**
 	 * Test if this does not point to a live UObject, but may in the future
 	 * @return true if this does not point to a real object, but could possibly
 	**/
@@ -372,7 +361,7 @@ public:
 		return FLazyObjectPtr::IsPending();
 	}
 
-	/**  
+	/**
 	 * Test if this can never point to a live UObject
 	 * @return true if this is explicitly pointing to no object
 	**/
@@ -381,67 +370,95 @@ public:
 		return FLazyObjectPtr::IsNull();
 	}
 
-	/**  
+	/**
 	 * Dereference lazy pointer to see if it points somewhere valid.
 	 */
-	FORCEINLINE operator bool() const
-	{
-		return Get() != NULL;
-	}
-
-	/**  Compare for equality with a raw pointer **/
-	template<class U>
-	FORCEINLINE bool operator==(U* Other) const
-	{
-		return Get() == Other;
-	}
-	/**  Compare for equality with a raw pointer **/
-	template<class U>
-	FORCEINLINE bool operator==(const U* Other) const
-	{
-		return Get() == Other;
-	}
-	/**  Compare to NULL **/
-	FORCEINLINE bool operator==(TYPE_OF_NULL Other) const
-	{
-		checkSlow(!Other); // we assume this is NULL
-		return !IsValid();
-	}
-
-	/**  Compare for inequality with a raw pointer	**/
-	template<class U>
-	FORCEINLINE bool operator!=(U* Other) const
+	FORCEINLINE_EXPLICIT_OPERATOR_BOOL() const
 	{
-		return Get() != Other;
-	}
-	/**  Compare for inequality with a raw pointer	**/
-	template<class U>
-	FORCEINLINE bool operator!=(const U* Other) const
-	{
-		return Get() != Other;
-	}
-	/**  Compare for inequality with NULL **/
-	FORCEINLINE bool operator!=(TYPE_OF_NULL Other) const
-	{
-		checkSlow(!Other); // we assume this is NULL
 		return IsValid();
 	}
 
 	/** Hash function. */
 	FORCEINLINE friend uint32 GetTypeHash(const TLazyObjectPtr<T>& LazyObjectPtr)
 	{
-		return GetTypeHash(static_cast<const TPersistentObjectPtr<FUniqueObjectGuid>&>(LazyObjectPtr));
+		return GetTypeHash(static_cast<const FLazyObjectPtr&>(LazyObjectPtr));
 	}
 
-	friend FArchive& operator<<( FArchive& Ar, TLazyObjectPtr<T>& LazyObjectPtr )
+	friend FArchive& operator<<(FArchive& Ar, TLazyObjectPtr<T>& LazyObjectPtr)
 	{
 		Ar << static_cast<FLazyObjectPtr&>(LazyObjectPtr);
 		return Ar;
 	}
 };
 
-template<class T> struct TIsPODType<TLazyObjectPtr<T> > { enum { Value = TIsPODType<FLazyObjectPtr>::Value }; };
-template<class T> struct TIsWeakPointerType<TLazyObjectPtr<T> > { enum { Value = TIsWeakPointerType<FLazyObjectPtr>::Value }; };
+// The reason these aren't inside the class (above) is because Visual Studio 2012-2013 crashes when compiling them :D
+
+// Compare with another TLazyObjectPtr of related type
+template<typename T, typename U>
+FORCEINLINE bool operator==(const TLazyObjectPtr<T>& Lhs, const TLazyObjectPtr<U>& Rhs)
+{
+	static_assert(TPointerIsConvertibleFromTo<U, T>::Value || TPointerIsConvertibleFromTo<T, U>::Value, "Unable to compare TLazyObjectPtr - types are incompatible");
+	return (const FLazyObjectPtr&)Lhs == (const FLazyObjectPtr&)Rhs;
+}
+template<typename T, typename U>
+FORCEINLINE bool operator!=(const TLazyObjectPtr<T>& Lhs, const TLazyObjectPtr<U>& Rhs)
+{
+	static_assert(TPointerIsConvertibleFromTo<U, T>::Value || TPointerIsConvertibleFromTo<T, U>::Value, "Unable to compare TLazyObjectPtr - types are incompatible");
+	return (const FLazyObjectPtr&)Lhs != (const FLazyObjectPtr&)Rhs;
+}
+
+/** Compare for equality with a raw pointer **/
+template<typename T, typename U>
+FORCEINLINE bool operator==(const TLazyObjectPtr<T>& Lhs, const U* Rhs)
+{
+	static_assert(TPointerIsConvertibleFromTo<U, T>::Value || TPointerIsConvertibleFromTo<T, U>::Value, "Unable to compare TLazyObjectPtr with raw pointer - types are incompatible");
+	return Lhs.Get() == Rhs;
+}
+template<typename T, typename U>
+FORCEINLINE bool operator==(const U* Lhs, const TLazyObjectPtr<T>& Rhs)
+{
+	static_assert(TPointerIsConvertibleFromTo<U, T>::Value || TPointerIsConvertibleFromTo<T, U>::Value, "Unable to compare TLazyObjectPtr with raw pointer - types are incompatible");
+	return Lhs == Rhs.Get();
+}
 
+/** Compare to null **/
+template<typename T>
+FORCEINLINE bool operator==(const TLazyObjectPtr<T>& Lhs, TYPE_OF_NULLPTR)
+{
+	return !Lhs.IsValid();
+}
+template<typename T>
+FORCEINLINE bool operator==(TYPE_OF_NULLPTR, const TLazyObjectPtr<T>& Rhs)
+{
+	return !Rhs.IsValid();
+}
 
+/** Compare for inequality with a raw pointer	**/
+template<typename T, typename U>
+FORCEINLINE bool operator!=(const TLazyObjectPtr<T>& Lhs, const U* Rhs)
+{
+	static_assert(TPointerIsConvertibleFromTo<U, T>::Value || TPointerIsConvertibleFromTo<T, U>::Value, "Unable to compare TLazyObjectPtr - types are incompatible");
+	return Lhs.Get() != Rhs;
+}
+template<typename T, typename U>
+FORCEINLINE bool operator!=(const U* Lhs, const TLazyObjectPtr<T>& Rhs)
+{
+	static_assert(TPointerIsConvertibleFromTo<U, T>::Value || TPointerIsConvertibleFromTo<T, U>::Value, "Unable to compare TLazyObjectPtr - types are incompatible");
+	return Lhs != Rhs.Get();
+}
 
+/** Compare for inequality with null **/
+template<typename T>
+FORCEINLINE bool operator!=(const TLazyObjectPtr<T>& Lhs, TYPE_OF_NULLPTR)
+{
+	return Lhs.IsValid();
+}
+template<typename T>
+FORCEINLINE bool operator!=(TYPE_OF_NULLPTR, const TLazyObjectPtr<T>& Rhs)
+{
+	return Rhs.IsValid();
+}
+
+
+template<class T> struct TIsPODType<TLazyObjectPtr<T> > { enum { Value = TIsPODType<FLazyObjectPtr>::Value }; };
+template<class T> struct TIsWeakPointerType<TLazyObjectPtr<T> > { enum { Value = TIsWeakPointerType<FLazyObjectPtr>::Value }; };
