[UE-9824] Fixing fallout from CL 2471336 (missing include, making shared code truly generic).

diff --git a/Engine/Source/Runtime/CoreUObject/Private/Blueprint/BlueprintSupport.cpp b/Engine/Source/Runtime/CoreUObject/Private/Blueprint/BlueprintSupport.cpp
--- a/Engine/Source/Runtime/CoreUObject/Private/Blueprint/BlueprintSupport.cpp
+++ b/Engine/Source/Runtime/CoreUObject/Private/Blueprint/BlueprintSupport.cpp
@@ -307,6 +307,13 @@ bool ULinkerLoad::RegenerateBlueprintClass(UClass* LoadClass, UObject* ExportObj
 	return bSuccessfulRegeneration;
 }
 
+/** Helper utility function to convert UObject pointers to FLinkerPlaceholderBase* */
+template<class PlaceholderType>
+static FLinkerPlaceholderBase* PlaceholderCast(UObject* PlaceholderObj)
+{
+	return CastChecked<PlaceholderType>(PlaceholderObj, ECastCheckedType::NullAllowed);
+}
+
 bool ULinkerLoad::DeferPotentialCircularImport(const int32 Index)
 {
 #if USE_CIRCULAR_DEPENDENCY_LOAD_DEFERRING
@@ -335,6 +342,7 @@ bool ULinkerLoad::DeferPotentialCircularImport(const int32 Index)
 				UObject* PlaceholderOuter = LinkerRoot;
 				FString  PlaceholderNamePrefix = TEXT("PLACEHOLDER_");
 				UClass*  PlaceholderType  = nullptr;
+				FLinkerPlaceholderBase* (*PlaceholderCastFunc)(UObject* RawUObject) = nullptr;
 
 				bool const bIsBlueprintClass = ImportClass->IsChildOf<UClass>();
 				// @TODO: if we could see if the related package is created 
@@ -345,6 +353,7 @@ bool ULinkerLoad::DeferPotentialCircularImport(const int32 Index)
 				{
 					PlaceholderNamePrefix = TEXT("PLACEHOLDER-CLASS_");
 					PlaceholderType = ULinkerPlaceholderClass::StaticClass();
+					PlaceholderCastFunc = &PlaceholderCast<ULinkerPlaceholderClass>;
 				}
 				else if (ImportClass->IsChildOf<UFunction>())
 				{		
@@ -360,6 +369,7 @@ bool ULinkerLoad::DeferPotentialCircularImport(const int32 Index)
 							PlaceholderNamePrefix = TEXT("PLACEHOLDER-FUNCTION_");
 							PlaceholderOuter = ImportMap[OuterImportIndex].XObject;
 							PlaceholderType  = ULinkerPlaceholderFunction::StaticClass();
+							PlaceholderCastFunc = &PlaceholderCast<ULinkerPlaceholderFunction>;
 						}
 					}
 				}
@@ -369,18 +379,25 @@ bool ULinkerLoad::DeferPotentialCircularImport(const int32 Index)
 					FName PlaceholderName(*FString::Printf(TEXT("%s_%s"), *PlaceholderNamePrefix, *Import.ObjectName.ToString()));
 					PlaceholderName = MakeUniqueObjectName(PlaceholderOuter, PlaceholderType, PlaceholderName);
 
-					ULinkerPlaceholderClass* Placeholder = NewObject<ULinkerPlaceholderClass>(PlaceholderOuter, PlaceholderType, PlaceholderName, RF_Public | RF_Transient);
+					UStruct* PlaceholderObj = NewObject<UStruct>(PlaceholderOuter, PlaceholderType, PlaceholderName, RF_Public | RF_Transient);
 					// store the import index in the placeholder, so we can 
 					// easily look it up in the import map, given the 
 					// placeholder (needed, to find the corresponding import for 
 					// ResolvingDeferredPlaceholder)
-					Placeholder->PackageIndex = FPackageIndex::FromImport(Index);
+					DEFERRED_DEPENDENCY_CHECK(PlaceholderCastFunc != nullptr);
+					if (PlaceholderCastFunc != nullptr)
+					{
+						if (FLinkerPlaceholderBase* AsPlaceholder = PlaceholderCastFunc(PlaceholderObj))
+						{
+							AsPlaceholder->PackageIndex = FPackageIndex::FromImport(Index);
+						}
+					}				
 					// make sure the class is fully formed (has its 
 					// ClassAddReferencedObjects/ClassConstructor members set)
-					Placeholder->Bind();
-					Placeholder->StaticLink(/*bRelinkExistingProperties =*/true);
+					PlaceholderObj->Bind();
+					PlaceholderObj->StaticLink(/*bRelinkExistingProperties =*/true);
 
-					Import.XObject = Placeholder;
+					Import.XObject = PlaceholderObj;
 				}
 			}
 		}
diff --git a/Engine/Source/Runtime/CoreUObject/Private/UObject/LinkerPlaceholderClass.cpp b/Engine/Source/Runtime/CoreUObject/Private/UObject/LinkerPlaceholderClass.cpp
--- a/Engine/Source/Runtime/CoreUObject/Private/UObject/LinkerPlaceholderClass.cpp
+++ b/Engine/Source/Runtime/CoreUObject/Private/UObject/LinkerPlaceholderClass.cpp
@@ -1,7 +1,7 @@
 // Copyright 1998-2015 Epic Games, Inc. All Rights Reserved.
 
 #include "CoreUObjectPrivate.h"
-#include "LinkerPlaceholderFunction.h"
+#include "LinkerPlaceholderClass.h"
 
 //------------------------------------------------------------------------------
 ULinkerPlaceholderClass::ULinkerPlaceholderClass(const FObjectInitializer& ObjectInitializer)
