When calling UField::GetToolTipText, if there is no native tooltip, the UField's name will be returned as a sanitized display string.

diff --git a/Engine/Source/Runtime/CoreUObject/Private/UObject/Class.cpp b/Engine/Source/Runtime/CoreUObject/Private/UObject/Class.cpp
--- a/Engine/Source/Runtime/CoreUObject/Private/UObject/Class.cpp
+++ b/Engine/Source/Runtime/CoreUObject/Private/UObject/Class.cpp
@@ -243,7 +243,7 @@ FText UField::GetToolTipText(bool bShortTooltip) const
 	{
 		if (NativeToolTip.IsEmpty())
 		{
-			NativeToolTip = FDisplayNameHelper::Get(*this);
+			NativeToolTip = FName::NameToDisplayString(FDisplayNameHelper::Get(*this), IsA<UBoolProperty>());
 		}
 		else
 		{
