[UE-9824] NonUnity build fix... forward declaring our TLinkerImportPlaceholder<> specializations.

diff --git a/Engine/Source/Runtime/CoreUObject/Private/UObject/LinkerPlaceholderBase.h b/Engine/Source/Runtime/CoreUObject/Private/UObject/LinkerPlaceholderBase.h
--- a/Engine/Source/Runtime/CoreUObject/Private/UObject/LinkerPlaceholderBase.h
+++ b/Engine/Source/Runtime/CoreUObject/Private/UObject/LinkerPlaceholderBase.h
@@ -192,3 +192,15 @@ private:
 
 // Templatized implementation
 #include "LinkerPlaceholderBase.inl" 
+
+/*******************************************************************************
+ * TLinkerImportPlaceholder<> Specializations
+ ******************************************************************************/
+ 
+/**  */
+template<>
+int32 TLinkerImportPlaceholder<UClass>::ResolvePropertyReferences(UClass* ReplacementClass);
+
+/**  */
+template<>
+int32 TLinkerImportPlaceholder<UFunction>::ResolvePropertyReferences(UFunction* ReplacementFunc);
