[UE-10647] Hook up '-scriptstackonwarnings' to FFrame::KismetExecutionMessage()

diff --git a/Engine/Source/Runtime/CoreUObject/Private/UObject/ScriptCore.cpp b/Engine/Source/Runtime/CoreUObject/Private/UObject/ScriptCore.cpp
--- a/Engine/Source/Runtime/CoreUObject/Private/UObject/ScriptCore.cpp
+++ b/Engine/Source/Runtime/CoreUObject/Private/UObject/ScriptCore.cpp
@@ -168,15 +168,32 @@ void FFrame::KismetExecutionMessage(const TCHAR* Message, ELogVerbosity::Type Ve
 #endif
 	}
 
+#if !(UE_BUILD_SHIPPING || UE_BUILD_TEST)
+	// Walk the script stack, if any
+	FString ScriptStack;
+	if( GScriptStack.Num() > 0 )
+	{
+		ScriptStack = TEXT( "Script call stack:\n" );
+		for( int32 i = GScriptStack.Num() - 1; i >= 0; --i )
+		{
+			ScriptStack += TEXT( "\t" ) + GScriptStack[i].GetStackDescription() + TEXT( "\n" );
+		}
+	}
+#endif
+
 	if (Verbosity == ELogVerbosity::Error)
 	{
+#if !(UE_BUILD_SHIPPING || UE_BUILD_TEST)
+		UE_LOG(LogScriptCore, Fatal,TEXT("%s\n%s"), Message, *ScriptStack);
+#else
 		UE_LOG(LogScriptCore, Fatal,TEXT("%s"), Message);
+#endif
 	}
 	else
 	{
 #if !(UE_BUILD_SHIPPING || UE_BUILD_TEST)
 		static bool GScriptStackForScriptWarning = FParse::Param(FCommandLine::Get(),TEXT("SCRIPTSTACKONWARNINGS"));
-		UE_LOG(LogScript, Warning, TEXT("%s"), Message);
+		UE_LOG(LogScript, Warning, TEXT("%s%s"), Message, GScriptStackForScriptWarning ? *FString::Printf(TEXT("\n%s"), *ScriptStack) : TEXT(""));
 #endif
 	}
 }
