Fix for crash in Deferred Loading when an asset is missing.

diff --git a/Engine/Source/Runtime/CoreUObject/Private/UObject/LinkerPlaceholderBase.inl b/Engine/Source/Runtime/CoreUObject/Private/UObject/LinkerPlaceholderBase.inl
--- a/Engine/Source/Runtime/CoreUObject/Private/UObject/LinkerPlaceholderBase.inl
+++ b/Engine/Source/Runtime/CoreUObject/Private/UObject/LinkerPlaceholderBase.inl
@@ -10,7 +10,7 @@
 template<class PlaceholderType>
 int32 TLinkerImportPlaceholder<PlaceholderType>::ResolveAllPlaceholderReferences(UObject* ReplacementObj)
 {
-	PlaceholderType* TypeCheckedReplacement = CastChecked<PlaceholderType>(ReplacementObj);
+	PlaceholderType* TypeCheckedReplacement = CastChecked<PlaceholderType>(ReplacementObj, ECastCheckedType::NullAllowed);
 
 	int32 ReplacementCount = ResolvePropertyReferences(TypeCheckedReplacement);
 	ReplacementCount += ResolveScriptReferences(TypeCheckedReplacement);
