#ue4 -Serialization code in UStructProperty moved to static function
-UScriptStruct::CompareScriptStruct function is accesible in other modules
-Call Function BP node works better with CustomStructureParam.

diff --git a/Engine/Source/Editor/BlueprintGraph/Private/K2Node_CallFunction.cpp b/Engine/Source/Editor/BlueprintGraph/Private/K2Node_CallFunction.cpp
--- a/Engine/Source/Editor/BlueprintGraph/Private/K2Node_CallFunction.cpp
+++ b/Engine/Source/Editor/BlueprintGraph/Private/K2Node_CallFunction.cpp
@@ -884,6 +884,8 @@ void UK2Node_CallFunction::PostReconstructNode()
 {
 	Super::PostReconstructNode();
 
+	FCustomStructureParamHelper::UpdateCustomStructurePins(GetTargetFunction(), this);
+
 	const UEdGraphSchema_K2* K2Schema = GetDefault<UEdGraphSchema_K2>();
 	// Fixup self node, may have been overridden from old self node
 	UFunction* Function = GetTargetFunction();
diff --git a/Engine/Source/Runtime/CoreUObject/Private/UObject/PropertyStruct.cpp b/Engine/Source/Runtime/CoreUObject/Private/UObject/PropertyStruct.cpp
--- a/Engine/Source/Runtime/CoreUObject/Private/UObject/PropertyStruct.cpp
+++ b/Engine/Source/Runtime/CoreUObject/Private/UObject/PropertyStruct.cpp
@@ -104,15 +104,14 @@ bool UStructProperty::UseBinaryOrNativeSerialization(const FArchive& Ar) const
 	return bUseBinarySerialization || bUseNativeSerialization;
 }
 
-void UStructProperty::SerializeItem( FArchive& Ar, void* Value, void const* Defaults ) const
+void UStructProperty::StaticSerializeItem(FArchive& Ar, void* Value, void const* Defaults, UScriptStruct* Struct, const bool bUseBinarySerialization, const bool bUseNativeSerialization)
 {
-	const bool bUseBinarySerialization = UseBinarySerialization(Ar);
-	const bool bUseNativeSerialization = UseNativeSerialization();
+	check(Struct);
 
 	// Preload struct before serialization tracking to not double count time.
-	if ( bUseBinarySerialization || bUseNativeSerialization)
+	if (bUseBinarySerialization || bUseNativeSerialization)
 	{
-		Ar.Preload( Struct );
+		Ar.Preload(Struct);
 	}
 
 	bool bItemSerialized = false;
@@ -126,21 +125,21 @@ void UStructProperty::SerializeItem( FArchive& Ar, void* Value, void const* Defa
 
 	if (!bItemSerialized)
 	{
-		if( bUseBinarySerialization )
+		if (bUseBinarySerialization)
 		{
 			// Struct is already preloaded above.
-			if ( !Ar.IsPersistent() && Ar.GetPortFlags() != 0 && !Struct->ShouldSerializeAtomically(Ar) )
+			if (!Ar.IsPersistent() && Ar.GetPortFlags() != 0 && !Struct->ShouldSerializeAtomically(Ar))
 			{
-				Struct->SerializeBinEx( Ar, Value, Defaults, Struct );
+				Struct->SerializeBinEx(Ar, Value, Defaults, Struct);
 			}
 			else
 			{
-				Struct->SerializeBin( Ar, Value );
+				Struct->SerializeBin(Ar, Value);
 			}
 		}
 		else
 		{
-			Struct->SerializeTaggedProperties( Ar, (uint8*)Value, Struct, (uint8*)Defaults );
+			Struct->SerializeTaggedProperties(Ar, (uint8*)Value, Struct, (uint8*)Defaults);
 		}
 	}
 
@@ -153,6 +152,14 @@ void UStructProperty::SerializeItem( FArchive& Ar, void* Value, void const* Defa
 	}
 }
 
+void UStructProperty::SerializeItem( FArchive& Ar, void* Value, void const* Defaults ) const
+{
+	const bool bUseBinarySerialization = UseBinarySerialization(Ar);
+	const bool bUseNativeSerialization = UseNativeSerialization();
+
+	StaticSerializeItem(Ar, Value, Defaults, Struct, bUseBinarySerialization, bUseNativeSerialization);
+}
+
 bool UStructProperty::NetSerializeItem( FArchive& Ar, UPackageMap* Map, void* Data, TArray<uint8> * MetaData ) const
 {
 	//------------------------------------------------
diff --git a/Engine/Source/Runtime/CoreUObject/Public/UObject/Class.h b/Engine/Source/Runtime/CoreUObject/Public/UObject/Class.h
--- a/Engine/Source/Runtime/CoreUObject/Public/UObject/Class.h
+++ b/Engine/Source/Runtime/CoreUObject/Public/UObject/Class.h
@@ -1072,7 +1072,7 @@ public:
 	 * @param	PortFlags	Comparison flags
 	 * @return true if the structs are identical
 	 */
-	bool CompareScriptStruct(const void* A, const void* B, uint32 PortFlags) const;
+	COREUOBJECT_API bool CompareScriptStruct(const void* A, const void* B, uint32 PortFlags) const;
 
 	/**
 	 * Copy a struct over an existing struct
diff --git a/Engine/Source/Runtime/CoreUObject/Public/UObject/UnrealType.h b/Engine/Source/Runtime/CoreUObject/Public/UObject/UnrealType.h
--- a/Engine/Source/Runtime/CoreUObject/Public/UObject/UnrealType.h
+++ b/Engine/Source/Runtime/CoreUObject/Public/UObject/UnrealType.h
@@ -2581,6 +2581,8 @@ public:
 	bool UseBinarySerialization(const FArchive& Ar) const;
 	bool UseBinaryOrNativeSerialization(const FArchive& Ar) const;
 
+	static void StaticSerializeItem(FArchive& Ar, void* Value, void const* Defaults, UScriptStruct* Struct, const bool bUseBinarySerialization, const bool bUseNativeSerialization);
+
 public:
 
 #if HACK_HEADER_GENERATOR
