Missing chage for 2481360. Still outcome of previous P4 failure.

diff --git a/Engine/Source/Runtime/CoreUObject/Public/UObject/ObjectBase.h b/Engine/Source/Runtime/CoreUObject/Public/UObject/ObjectBase.h
--- a/Engine/Source/Runtime/CoreUObject/Public/UObject/ObjectBase.h
+++ b/Engine/Source/Runtime/CoreUObject/Public/UObject/ObjectBase.h
@@ -525,7 +525,7 @@ struct COREUOBJECT_API FReferencerInformationList
 #define GENERATED_USTRUCT_BODY(...) GENERATED_BODY()
 #define GENERATED_UCLASS_BODY(...) GENERATED_BODY_LEGACY()
 #define GENERATED_UINTERFACE_BODY(...) GENERATED_BODY_LEGACY()
-#define GENERATED_IINTERFACE_BODY(...) GENERATED_BODY()
+#define GENERATED_IINTERFACE_BODY(...) GENERATED_BODY_LEGACY()
 
 #if UE_BUILD_DOCS
 #define UCLASS(...)
