Fix for crashes when GC is triggered while streaming.

diff --git a/Engine/Source/Runtime/CoreUObject/Private/Serialization/AsyncLoading.cpp b/Engine/Source/Runtime/CoreUObject/Private/Serialization/AsyncLoading.cpp
--- a/Engine/Source/Runtime/CoreUObject/Private/Serialization/AsyncLoading.cpp
+++ b/Engine/Source/Runtime/CoreUObject/Private/Serialization/AsyncLoading.cpp
@@ -725,6 +725,9 @@ EAsyncPackageState::Type FAsyncPackage::CreateImports()
  		UObject* Object	= Linker->CreateImport( ImportIndex++ );
 		LastObjectWorkWasPerformedOn	= Object;
 		LastTypeOfWorkPerformed			= TEXT("creating imports for");
+
+		// Make sure this object is not claimed by GC if it's triggered while streaming.
+		FAsyncObjectsReferencer::Get().AddObject(Object);
 	}
 
 	return ImportIndex == Linker->ImportMap.Num() ? EAsyncPackageState::Complete : EAsyncPackageState::TimeOut;
@@ -1252,6 +1255,7 @@ EAsyncPackageState::Type ProcessAsyncLoading( bool bUseTimeLimit, bool bUseFullT
 	// Not that streaming has finished, we can stop force-referencing all objects that were created during async load.
 	if (CompletionState == EAsyncPackageState::Complete && bWasLoading)
 	{
+		check(GObjAsyncPackages.Num() == 0);
 		FAsyncObjectsReferencer::Get().EmptyReferencedObjects();
 	}
 
