UE-5538: Map doesn't load the asset reference if the asset isn't saved when the map was saved.

Added new package flag to indicate that this package was newly created in the editor. Then linker can treat it similarily as it treated compiled-in packages, just look for objects inside. This flag is cleared when saving the package.

diff --git a/Engine/Source/Runtime/AssetRegistry/Private/AssetRegistry.cpp b/Engine/Source/Runtime/AssetRegistry/Private/AssetRegistry.cpp
--- a/Engine/Source/Runtime/AssetRegistry/Private/AssetRegistry.cpp
+++ b/Engine/Source/Runtime/AssetRegistry/Private/AssetRegistry.cpp
@@ -1140,6 +1140,10 @@ void FAssetRegistry::AssetCreated(UObject* NewAsset)
 		// determined by its long package name.
 		// @todo AssetRegistry We are assuming it will be saved in a single asset package.
 		UPackage* NewPackage = NewAsset->GetOutermost();
+
+		// Mark this package as newly created.
+		NewPackage->PackageFlags |= PKG_NewlyCreated;
+
 		const FString NewPackageName = NewPackage->GetName();
 		const FString Filename = FPackageName::LongPackageNameToFilename(NewPackageName, FPackageName::GetAssetPackageExtension());
 
diff --git a/Engine/Source/Runtime/CoreUObject/Private/UObject/Linker.cpp b/Engine/Source/Runtime/CoreUObject/Private/UObject/Linker.cpp
--- a/Engine/Source/Runtime/CoreUObject/Private/UObject/Linker.cpp
+++ b/Engine/Source/Runtime/CoreUObject/Private/UObject/Linker.cpp
@@ -570,8 +570,8 @@ ULinkerLoad* GetPackageLinker
 
 		if( !FPackageName::DoesPackageExist(InOuter->GetName(), CompatibleGuid, &NewFilename) )
 		{
-			// Compiled in packages have no linker and this is ok
-			if (!(LoadFlags & LOAD_AllowDll) && !(InOuter->PackageFlags & PKG_CompiledIn))
+			// In memory-only packages have no linker and this is ok.
+			if (!(LoadFlags & LOAD_AllowDll) && !(InOuter->PackageFlags & PKG_InMemoryOnly))
 			{
 				FFormatNamedArguments Arguments;
 				Arguments.Add(TEXT("AssetName"), FText::FromString(InOuter->GetName()));
@@ -599,9 +599,9 @@ ULinkerLoad* GetPackageLinker
 
 		if (UPackage* ExistingPackage = FindObject<UPackage>(nullptr, *PackageName))
 		{
-			if (!ExistingPackage->GetOuter() && (ExistingPackage->PackageFlags & PKG_CompiledIn))
+			if (!ExistingPackage->GetOuter() && (ExistingPackage->PackageFlags & PKG_InMemoryOnly))
 			{
-				// this is a compiled in package and so it has no linker and this is ok
+				// This is a memory-only in package and so it has no linker and this is ok.
 				return nullptr;
 			}
 		}
diff --git a/Engine/Source/Runtime/CoreUObject/Private/UObject/LinkerLoad.cpp b/Engine/Source/Runtime/CoreUObject/Private/UObject/LinkerLoad.cpp
--- a/Engine/Source/Runtime/CoreUObject/Private/UObject/LinkerLoad.cpp
+++ b/Engine/Source/Runtime/CoreUObject/Private/UObject/LinkerLoad.cpp
@@ -2305,9 +2305,9 @@ bool ULinkerLoad::VerifyImportInner(const int32 ImportIndex, FString& WarningSuf
 			}
 
 			auto* Package = dynamic_cast<UPackage*>(Top->XObject);
-			if (Package && (Package->PackageFlags & PKG_CompiledIn))
+			if (Package && (Package->PackageFlags & PKG_InMemoryOnly))
 			{
-				// this is an import to a compiled in thing, just search for it in the package
+				// This is an import to a memory-only package, just search for it in the package.
 				TmpPkg = Package;
 			}
 		}
@@ -2464,11 +2464,11 @@ bool ULinkerLoad::VerifyImportInner(const int32 ImportIndex, FString& WarningSuf
 		}
 	}
 
-	bool bCameFromCompiledInPackage = false;
-	if (!Pkg && TmpPkg && (TmpPkg->PackageFlags&PKG_CompiledIn))
+	bool bCameFromMemoryOnlyPackage = false;
+	if (!Pkg && TmpPkg && (TmpPkg->PackageFlags & PKG_InMemoryOnly))
 	{
-		Pkg = TmpPkg; // this is a compiled in package, so that is the package to search regardless of FindIfFail
-		bCameFromCompiledInPackage = true;
+		Pkg = TmpPkg; // this is a package that exists in memory only, so that is the package to search regardless of FindIfFail
+		bCameFromMemoryOnlyPackage = true;
 
 		if (IsCoreUObjectPackage(Import.ClassPackage) && Import.ClassName == NAME_Package && !TmpPkg->GetOuter())
 		{
@@ -2514,10 +2514,10 @@ bool ULinkerLoad::VerifyImportInner(const int32 ImportIndex, FString& WarningSuf
 				}
 
 				UObject* FindObject = StaticFindObject(FindClass, FindOuter, *Import.ObjectName.ToString());
-				// reference to native transient class or CDO of such a class
-				bool IsNativeTransient	= bCameFromCompiledInPackage || (FindObject != NULL && (FindObject->HasAllFlags(EObjectFlags(RF_Public|RF_Native|RF_Transient)) || (FindObject->HasAnyFlags(RF_ClassDefaultObject) && FindObject->GetClass()->HasAllFlags(EObjectFlags(RF_Public|RF_Native|RF_Transient)))));
+				// Reference to in memory-only package's object, native transient class or CDO of such a class.
+				bool bIsInMemoryOnlyOrNativeTransient = bCameFromMemoryOnlyPackage || (FindObject != NULL && (FindObject->HasAllFlags(EObjectFlags(RF_Public | RF_Native | RF_Transient)) || (FindObject->HasAnyFlags(RF_ClassDefaultObject) && FindObject->GetClass()->HasAllFlags(EObjectFlags(RF_Public | RF_Native | RF_Transient)))));
 				// Check for structs which have been moved to another header (within the same class package).
-				if (!FindObject && IsNativeTransient && FindClass == UScriptStruct::StaticClass())
+				if (!FindObject && bIsInMemoryOnlyOrNativeTransient && FindClass == UScriptStruct::StaticClass())
 				{
 					FindObject = StaticFindObject( FindClass, ANY_PACKAGE, *Import.ObjectName.ToString(), true );
 					if (FindObject && FindOuter->GetOutermost() != FindObject->GetOutermost())
@@ -2526,7 +2526,7 @@ bool ULinkerLoad::VerifyImportInner(const int32 ImportIndex, FString& WarningSuf
 						FindObject = NULL;
 					}
 				}
-				if (FindObject != NULL && ((LoadFlags & LOAD_FindIfFail) || IsNativeTransient))
+				if (FindObject != NULL && ((LoadFlags & LOAD_FindIfFail) || bIsInMemoryOnlyOrNativeTransient))
 				{
 					Import.XObject = FindObject;
 					GImportCount++;
diff --git a/Engine/Source/Runtime/CoreUObject/Private/UObject/SavePackage.cpp b/Engine/Source/Runtime/CoreUObject/Private/UObject/SavePackage.cpp
--- a/Engine/Source/Runtime/CoreUObject/Private/UObject/SavePackage.cpp
+++ b/Engine/Source/Runtime/CoreUObject/Private/UObject/SavePackage.cpp
@@ -3727,6 +3727,9 @@ bool UPackage::SavePackage( UPackage* InOuter, UObject* Base, EObjectFlags TopLe
 		}
 		if( Success == true )
 		{
+			// Package has been save, so unmark NewlyCreated flag.
+			InOuter->PackageFlags &= ~PKG_NewlyCreated;
+
 			// send a message that the package was saved
 			UPackage::PackageSavedEvent.Broadcast(Filename, InOuter);
 		}
diff --git a/Engine/Source/Runtime/CoreUObject/Public/UObject/ObjectBase.h b/Engine/Source/Runtime/CoreUObject/Public/UObject/ObjectBase.h
--- a/Engine/Source/Runtime/CoreUObject/Public/UObject/ObjectBase.h
+++ b/Engine/Source/Runtime/CoreUObject/Public/UObject/ObjectBase.h
@@ -67,7 +67,7 @@ enum ESaveFlags
 //
 enum EPackageFlags
 {
-//	PKG_Unused						= 0x00000001,	// 
+	PKG_NewlyCreated				= 0x00000001,	// Newly created package, not saved yet. In editor only.
 	PKG_ClientOptional				= 0x00000002,	// Purely optional for clients.
 	PKG_ServerSideOnly				= 0x00000004,   // Only needed on the server side.
 	PKG_CompiledIn					= 0x00000010,   // This package is from "compiled in" classes.
@@ -89,6 +89,8 @@ enum EPackageFlags
 //	PKG_Unused						= 0x20000000,
 	PKG_ReloadingForCooker			= 0x40000000,   // this package is reloading in the cooker, try to avoid getting data we will never need. We won't save this package.
 	PKG_FilterEditorOnly			= 0x80000000,	// Package has editor-only data filtered
+
+	PKG_InMemoryOnly				= PKG_CompiledIn | PKG_NewlyCreated, // Flag mask that indicates if this package is a package that exists in memory only.
 };
 
 //
