Fixed subobject instancing when creating new objects from a non-CDO template.

Added GetDefaultSubobjects and GetDefaultSubobjectByName helper functions to UObject.

diff --git a/Engine/Source/Runtime/CoreUObject/Private/UObject/Class.cpp b/Engine/Source/Runtime/CoreUObject/Private/UObject/Class.cpp
--- a/Engine/Source/Runtime/CoreUObject/Private/UObject/Class.cpp
+++ b/Engine/Source/Runtime/CoreUObject/Private/UObject/Class.cpp
@@ -2345,29 +2345,25 @@ void UClass::PostInitProperties()
 
 UObject* UClass::GetDefaultSubobjectByName(FName ToFind)
 {
-	TArray<UObject*> SubObjects;
-	GetDefaultObjectSubobjects(SubObjects);
-	for (int32 Index = 0; Index < SubObjects.Num(); ++Index)
+	UObject* DefaultObj = GetDefaultObject();
+	UObject* DefaultSubobject = nullptr;
+	if (DefaultObj)
 	{
-		if (SubObjects[Index]->GetFName() == ToFind)
-		{
-			return SubObjects[Index];
-		}
+		DefaultSubobject = DefaultObj->GetDefaultSubobjectByName(ToFind);
 	}
-	return NULL;
+	return DefaultSubobject;
 }
 
 void UClass::GetDefaultObjectSubobjects(TArray<UObject*>& OutDefaultSubobjects)
 {
-	OutDefaultSubobjects.Empty();
-	GetObjectsWithOuter(GetDefaultObject(), OutDefaultSubobjects, false);
-	for ( int32 SubobjectIndex = 0; SubobjectIndex < OutDefaultSubobjects.Num(); SubobjectIndex++ )
+	UObject* DefaultObj = GetDefaultObject();
+	if (DefaultObj)
 	{
-		UObject* PotentialComponent = OutDefaultSubobjects[SubobjectIndex];
-		if (!PotentialComponent->IsDefaultSubobject())
-		{
-			OutDefaultSubobjects.RemoveAtSwap(SubobjectIndex--);
-		}
+		DefaultObj->GetDefaultSubobjects(OutDefaultSubobjects);
+	}
+	else
+	{
+		OutDefaultSubobjects.Empty();
 	}
 }
 
diff --git a/Engine/Source/Runtime/CoreUObject/Private/UObject/Obj.cpp b/Engine/Source/Runtime/CoreUObject/Private/UObject/Obj.cpp
--- a/Engine/Source/Runtime/CoreUObject/Private/UObject/Obj.cpp
+++ b/Engine/Source/Runtime/CoreUObject/Private/UObject/Obj.cpp
@@ -95,6 +95,34 @@ UObject* UObject::CreateEditorOnlyDefaultSubobjectImpl(FName SubobjectName, UCla
 	return CurrentInitializer->CreateEditorOnlyDefaultSubobject(this, SubobjectName, ReturnType, bTransient);
 }
 
+void UObject::GetDefaultSubobjects(TArray<UObject*>& OutDefaultSubobjects)
+{
+	OutDefaultSubobjects.Empty();
+	GetObjectsWithOuter(this, OutDefaultSubobjects, false);
+	for (int32 SubobjectIndex = 0; SubobjectIndex < OutDefaultSubobjects.Num(); SubobjectIndex++)
+	{
+		UObject* PotentialSubobject = OutDefaultSubobjects[SubobjectIndex];
+		if (!PotentialSubobject->IsDefaultSubobject())
+		{
+			OutDefaultSubobjects.RemoveAtSwap(SubobjectIndex--);
+		}
+	}
+}
+
+UObject* UObject::GetDefaultSubobjectByName(FName ToFind)
+{
+	TArray<UObject*> SubObjects;
+	GetDefaultSubobjects(SubObjects);
+	for (int32 Index = 0; Index < SubObjects.Num(); ++Index)
+	{
+		if (SubObjects[Index]->GetFName() == ToFind)
+		{
+			return SubObjects[Index];
+		}
+	}
+	return nullptr;
+}
+
 bool UObject::Rename( const TCHAR* InName, UObject* NewOuter, ERenameFlags Flags )
 {
 #if WITH_EDITOR
diff --git a/Engine/Source/Runtime/CoreUObject/Private/UObject/UObjectGlobals.cpp b/Engine/Source/Runtime/CoreUObject/Private/UObject/UObjectGlobals.cpp
--- a/Engine/Source/Runtime/CoreUObject/Private/UObject/UObjectGlobals.cpp
+++ b/Engine/Source/Runtime/CoreUObject/Private/UObject/UObjectGlobals.cpp
@@ -2728,11 +2728,22 @@ UObject* FObjectInitializer::CreateDefaultSubobject(UObject* Outer, FName Subobj
 		{
 			UObject* Template = OverrideClass->GetDefaultObject(); // force the CDO to be created if it hasn't already
 			const EObjectFlags SubobjectFlags = Outer->GetMaskedFlags(RF_PropagateToSubObjects);
+			const bool bOwnerArchetypeIsNotNative = !Outer->GetArchetype()->GetClass()->HasAnyClassFlags(CLASS_Native | CLASS_Intrinsic);
+			const bool bOwnerTemplateIsNotCDO = ObjectArchetype != nullptr && ObjectArchetype != Outer->GetClass()->GetDefaultObject(false) && !Outer->HasAnyFlags(RF_ClassDefaultObject);
 			Result = StaticConstructObject_Internal(OverrideClass, Outer, SubobjectFName, SubobjectFlags);
-			if (!bIsTransient && !Outer->GetArchetype()->GetClass()->HasAnyClassFlags(CLASS_Native | CLASS_Intrinsic))
+			if (!bIsTransient && (bOwnerArchetypeIsNotNative || bOwnerTemplateIsNotCDO))
 			{
-				// The archetype of the outer is not native, so we need to copy properties to the subobjects after the C++ constructor chain for the outer has run (because those sets properties on the subobjects)
-				UObject* MaybeTemplate = Outer->GetArchetype()->GetClass()->GetDefaultSubobjectByName(SubobjectFName);
+				UObject* MaybeTemplate = nullptr;
+				if (bOwnerTemplateIsNotCDO)
+				{
+					// Try to get the subobject template from the specified object template
+					MaybeTemplate = ObjectArchetype->GetDefaultSubobjectByName(SubobjectFName);
+				}
+				if (!MaybeTemplate)
+				{
+					// The archetype of the outer is not native, so we need to copy properties to the subobjects after the C++ constructor chain for the outer has run (because those sets properties on the subobjects)
+					MaybeTemplate = Outer->GetArchetype()->GetClass()->GetDefaultSubobjectByName(SubobjectFName);
+				}
 				if (MaybeTemplate && MaybeTemplate->IsA(ReturnType) && Template != MaybeTemplate)
 				{
 					ComponentInits.Add(Result, MaybeTemplate);
diff --git a/Engine/Source/Runtime/CoreUObject/Public/UObject/UObject.h b/Engine/Source/Runtime/CoreUObject/Public/UObject/UObject.h
--- a/Engine/Source/Runtime/CoreUObject/Public/UObject/UObject.h
+++ b/Engine/Source/Runtime/CoreUObject/Public/UObject/UObject.h
@@ -134,6 +134,18 @@ class COREUOBJECT_API UObject : public UObjectBaseUtility
 		return static_cast<TReturnType*>(CreateDefaultSubobject(SubobjectName, ReturnType, ReturnType, /*bIsRequired =*/ true, /*bIsAbstract =*/ true, bTransient));
 	}
 
+	/**
+	* Gets all default subobjects associated with this object instance.
+	* @param	OutDefaultSubobjects	Array containing all default subobjects of this object.
+	*/
+	void GetDefaultSubobjects(TArray<UObject*>& OutDefaultSubobjects);
+
+	/**
+	* Finds a subobject associated with this object instance by its name
+	* @param	OutDefaultSubobjects	Array containing all default subobjects of this object.
+	*/
+	UObject* GetDefaultSubobjectByName(FName ToFind);
+
 	//==========================================
 	// UObject interface.
 	//==========================================	
