#ue4 - Add fast Cast flag for UStaticMeshComponent.

- Touched UBT and UHT to force rebuilds (they don't seem to pick up the change in core dependency).

diff --git a/Engine/Source/Runtime/CoreUObject/Public/Templates/Casts.h b/Engine/Source/Runtime/CoreUObject/Public/Templates/Casts.h
--- a/Engine/Source/Runtime/CoreUObject/Public/Templates/Casts.h
+++ b/Engine/Source/Runtime/CoreUObject/Public/Templates/Casts.h
@@ -226,6 +226,7 @@ template< class T, class U > FORCEINLINE const T* CastChecked( const U      * Sr
 
 // Define a macro that declares all the cast flags.
 // This allows us to reuse these declarations elsewhere to define other properties for these classes.
+// Note: When adding an item to this list, you must also add a CASTCLASS_ flag in ObjectBase.h and rebuild UnrealHeaderTool.
 #define DECLARE_ALL_CAST_FLAGS \
 DECLARE_CAST_BY_FLAG(UField)							\
 DECLARE_CAST_BY_FLAG(UEnum)								\
@@ -264,6 +265,8 @@ DECLARE_CAST_BY_FLAG(UPrimitiveComponent)				\
 DECLARE_CAST_BY_FLAG(USkinnedMeshComponent)				\
 DECLARE_CAST_BY_FLAG(USkeletalMeshComponent)			\
 DECLARE_CAST_BY_FLAG(UBlueprint)						\
+DECLARE_CAST_BY_FLAG(UDelegateFunction)					\
+DECLARE_CAST_BY_FLAG(UStaticMeshComponent)				\
 FINISH_DECLARING_CAST_FLAGS		// This is here to hopefully remind people to include the "\" in all declarations above, especially when copy/pasting the final line.
 
 // Now actually declare the flags
diff --git a/Engine/Source/Runtime/CoreUObject/Public/UObject/ObjectBase.h b/Engine/Source/Runtime/CoreUObject/Public/UObject/ObjectBase.h
--- a/Engine/Source/Runtime/CoreUObject/Public/UObject/ObjectBase.h
+++ b/Engine/Source/Runtime/CoreUObject/Public/UObject/ObjectBase.h
@@ -298,6 +298,7 @@ typedef uint64 EClassCastFlags;
 #define CASTCLASS_USkeletalMeshComponent		DECLARE_UINT64(0x0000040000000000)
 #define CASTCLASS_UBlueprint					DECLARE_UINT64(0x0000080000000000)
 #define CASTCLASS_UDelegateFunction				DECLARE_UINT64(0x0000100000000000)
+#define CASTCLASS_UStaticMeshComponent			DECLARE_UINT64(0x0000200000000000)
 
 #define CASTCLASS_AllFlags						DECLARE_UINT64(0xFFFFFFFFFFFFFFFF)
 
