GitHub 944 : UE-12128: Properly save structs when struct field default value is different

diff --git a/Engine/Source/Runtime/CoreUObject/Private/UObject/Class.cpp b/Engine/Source/Runtime/CoreUObject/Private/UObject/Class.cpp
--- a/Engine/Source/Runtime/CoreUObject/Private/UObject/Class.cpp
+++ b/Engine/Source/Runtime/CoreUObject/Private/UObject/Class.cpp
@@ -2110,13 +2110,18 @@ void UScriptStruct::Link(FArchive& Ar, bool bRelinkExistingProperties)
 bool UScriptStruct::CompareScriptStruct(const void* A, const void* B, uint32 PortFlags) const
 {
 	check(A);
+
+	if (nullptr == B) // if the comparand is NULL, we just call this no-match
+	{
+		return false;
+	}
+
 	if (StructFlags & STRUCT_IdenticalNative)
 	{
 		UScriptStruct::ICppStructOps* TheCppStructOps = GetCppStructOps();
 		check(TheCppStructOps);
 		bool bResult = false;
-		if (!B || // if the comparand is NULL, we just call this no-match
-			TheCppStructOps->Identical(A, B, PortFlags, bResult))
+		if (TheCppStructOps->Identical(A, B, PortFlags, bResult))
 		{
 			return bResult;
 		}
