When cancelling async loading, only kill objects created during async loading if they haven't finished loading yet.

Due to recent changes, some of the objects inside FAsyncObjectsReferencer can come from packages that have previously been loaded.

diff --git a/Engine/Source/Runtime/CoreUObject/Private/Serialization/AsyncLoading.cpp b/Engine/Source/Runtime/CoreUObject/Private/Serialization/AsyncLoading.cpp
--- a/Engine/Source/Runtime/CoreUObject/Private/Serialization/AsyncLoading.cpp
+++ b/Engine/Source/Runtime/CoreUObject/Private/Serialization/AsyncLoading.cpp
@@ -67,11 +67,16 @@ public:
 	/** Removes all referenced objects and markes them for GC */
 	void EmptyReferencedObjectsAndCancelLoading()
 	{
+		const EObjectFlags LoadFlags = RF_AsyncLoading | RF_NeedLoad | RF_NeedPostLoad | RF_NeedPostLoadSubobjects;
+
 		// All of the referenced objects have been created by async loading code and may be in an invalid state so mark them for GC
 		for (auto Object : ReferencedObjects)
 		{
-			Object->ClearFlags(RF_AsyncLoading | RF_NeedLoad | RF_NeedPostLoad | RF_NeedPostLoadSubobjects);
-			Object->SetFlags(RF_PendingKill);
+			if (Object->HasAnyFlags(LoadFlags))
+			{
+				Object->ClearFlags(LoadFlags);
+				Object->SetFlags(RF_PendingKill);
+			}
 		}
 		ReferencedObjects.Empty(ReferencedObjects.Num());
 	}
