#core Added derived->base conversion for TAssetPtr. Tidied up some other dangerous/unnecessary conversions too (operator bool -> TYPE_OF_NULL constructor was the conversion that bit me).

diff --git a/Engine/Source/Runtime/CoreUObject/Public/UObject/AssetPtr.h b/Engine/Source/Runtime/CoreUObject/Public/UObject/AssetPtr.h
--- a/Engine/Source/Runtime/CoreUObject/Public/UObject/AssetPtr.h
+++ b/Engine/Source/Runtime/CoreUObject/Public/UObject/AssetPtr.h
@@ -43,36 +43,22 @@ public:
 	{
 		(*this)=Object;
 	}
-	/**  
-	 * Copy from an object pointer
-	 * @param Object object to create a weak pointer to
-	 */
-	FORCEINLINE void operator=(const UObject *Object)
-	{
-		TPersistentObjectPtr<FStringAssetReference>::operator=(Object);
-	}
-
-	/**  
-	 * Copy from an asset pointer
-	 */
-	FORCEINLINE void operator=(const FAssetPtr &Other)
-	{
-		TPersistentObjectPtr<FStringAssetReference>::operator=(Other);
-	}
 
-	/**  
-	 * Copy from a unique object identifier
-	 * @param ObjectID Object identifier to create a weak pointer to
-	 */
-	FORCEINLINE void operator=(const FStringAssetReference &ObjectID)
-	{
-		TPersistentObjectPtr<FStringAssetReference>::operator=(ObjectID);
-	}
+	using TPersistentObjectPtr<FStringAssetReference>::operator=;
 };
 
 template <> struct TIsPODType<FAssetPtr> { enum { Value = TIsPODType<TPersistentObjectPtr<FStringAssetReference> >::Value }; };
 template <> struct TIsWeakPointerType<FAssetPtr> { enum { Value = TIsWeakPointerType<TPersistentObjectPtr<FStringAssetReference> >::Value }; };
 
+// VS2012 doesn't support defaulted function template arguments, so we need a workaround for that.
+// The workaround technically changes the signature of the constructors, but we assume that the modified
+// signature will not be used.
+#if defined(_MSC_VER) && _MSC_VER < 1800
+	#define ENABLE_TASSETPTR_CONSTRUCTOR_ENABLE_IF_WORKAROUND 1
+#else
+	#define ENABLE_TASSETPTR_CONSTRUCTOR_ENABLE_IF_WORKAROUND 0
+#endif
+
 /***
 * 
 * TAssetPtr is templatized version of the generic FAssetPtr
@@ -80,8 +66,11 @@ template <> struct TIsWeakPointerType<FAssetPtr> { enum { Value = TIsWeakPointer
 **/
 
 template<class T=UObject>
-class TAssetPtr : private FAssetPtr
+class TAssetPtr
 {
+	template <class U>
+	friend class TAssetPtr;
+
 public:
 	/** NULL constructor **/
 	FORCEINLINE TAssetPtr()
@@ -92,26 +81,23 @@ public:
 	 * Construct from another lazy pointer
 	 * @param Other lazy pointer to copy from
 	 */
-	FORCEINLINE TAssetPtr(const TAssetPtr<T> &Other) :
-		FAssetPtr(Other)
+#if ENABLE_TASSETPTR_CONSTRUCTOR_ENABLE_IF_WORKAROUND
+	template <class U>
+	FORCEINLINE TAssetPtr(const TAssetPtr<U>& Other, typename TEnableIf<TPointerIsConvertibleFromTo<U, T>::Value>::Type* = nullptr)
+#else
+	template <class U, class = typename TEnableIf<TPointerIsConvertibleFromTo<U, T>::Value>::Type>
+	FORCEINLINE TAssetPtr(const TAssetPtr<U>& Other)
+#endif
+		: AssetPtr(Other.AssetPtr)
 	{
 	}
 
-	/**  
-	 * Construct from an int (assumed to be NULL)
-	 * @param Other int to create a reference to
-	 */
-	FORCEINLINE TAssetPtr(TYPE_OF_NULL Other)
-	{
-		checkSlow(!Other); // we assume this is NULL
-	}
-
 	/**  
 	 * Construct from an object pointer
 	 * @param Object object to create a reference to
 	**/
-	FORCEINLINE TAssetPtr(const T *Object) :
-		FAssetPtr(Object)
+	FORCEINLINE TAssetPtr(const T* Object)
+		: AssetPtr(Object)
 	{
 	}
 
@@ -120,18 +106,16 @@ public:
 	 */
 	FORCEINLINE void Reset()
 	{
-		FAssetPtr::Reset();
+		AssetPtr.Reset();
 	}
 
 	/**  
 	 * Copy from an object pointer
 	 * @param Object Object to create a lazy pointer to
 	 */
-	template<class U>
-	FORCEINLINE void operator=(const U *Object)
+	FORCEINLINE void operator=(const T* Object)
 	{
-		const T* TempObject = Object;
-		FAssetPtr::operator=(TempObject);
+		AssetPtr = Object;
 	}
 
 	/**  
@@ -140,35 +124,29 @@ public:
 	 */
 	FORCEINLINE void operator=(const FStringAssetReference &ObjectID)
 	{
-		FAssetPtr::operator=(ObjectID);
+		AssetPtr = ObjectID;
 	}
 
 	/**  
 	 * Copy from a weak pointer
 	 * @param Other Weak pointer to copy from
 	 */
-	FORCEINLINE void operator=(const TWeakObjectPtr<T> &Other)
+	template <class U>
+	FORCEINLINE typename TEnableIf<TPointerIsConvertibleFromTo<U, T>::Value>::Type
+		operator=(const TWeakObjectPtr<U>& Other)
 	{
-		FAssetPtr::operator=(Other);
+		AssetPtr = Other;
 	}
 
 	/**  
 	 * Construct from another lazy pointer
 	 * @param Other lazy pointer to copy from
 	 */
-	FORCEINLINE void operator=(const TAssetPtr<T> &Other)
-	{
-		FAssetPtr::operator=(Other);
-	}
-
-	/**  
-	 * Copy from an int (assumed to be NULL)
-	 * @param Other int to create a reference to
-	 */
-	FORCEINLINE void operator=(TYPE_OF_NULL Other)
+	template <class U>
+	FORCEINLINE typename TEnableIf<TPointerIsConvertibleFromTo<U, T>::Value>::Type
+		operator=(const TAssetPtr<U>& Other)
 	{
-		checkSlow(!Other); // we assume this is NULL
-		FAssetPtr::Reset();
+		AssetPtr = Other.AssetPtr;
 	}
 
 	/**  
@@ -176,19 +154,19 @@ public:
 	 * @param Other lazy pointer to compare to
 	 * Caution: Two lazy pointers might not be equal to each other, but they both might return NULL.
 	 */
-	FORCEINLINE bool operator==(const TAssetPtr<T> &Other) const
+	FORCEINLINE friend bool operator==(const TAssetPtr& Lhs, const TAssetPtr& Rhs)
 	{
-		return FAssetPtr::operator==(Other);
+		return Lhs.AssetPtr == Rhs.AssetPtr;
 	}
 
 	/**  
-	 * Compare lazy pointers for equality
+	 * Compare lazy pointers for inequality
 	 * @param Other lazy pointer to compare to
 	 * Caution: Two lazy pointers might not be equal to each other, but they both might return NULL.
 	 */
-	FORCEINLINE bool operator!=(const TAssetPtr<T> &Other) const
+	FORCEINLINE friend bool operator!=(const TAssetPtr& Lhs, const TAssetPtr& Rhs)
 	{
-		return FAssetPtr::operator!=(Other);
+		return Lhs.AssetPtr != Rhs.AssetPtr;
 	}
 
 	/**
@@ -197,22 +175,22 @@ public:
 	 */
 	FStringAssetReference GetUniqueID() const
 	{
-		return FAssetPtr::GetUniqueID();
+		return AssetPtr.GetUniqueID();
 	}
 
 	/**  
 	 * Dereference the lazy pointer.
 	 * @return NULL if this object is gone or the lazy pointer was NULL, otherwise a valid uobject pointer
 	 */
-	FORCEINLINE T *Get() const
+	FORCEINLINE T* Get() const
 	{
-		return dynamic_cast<T*>(FAssetPtr::Get());
+		return dynamic_cast<T*>(AssetPtr.Get());
 	}
 
 	/**  
 	 * Dereference the lazy pointer.
 	 */
-	FORCEINLINE T & operator*() const
+	FORCEINLINE T& operator*() const
 	{
 		return *Get();
 	}
@@ -220,7 +198,7 @@ public:
 	/**  
 	 * Dereference the lazy pointer.
 	 */
-	FORCEINLINE T * operator->() const
+	FORCEINLINE T* operator->() const
 	{
 		return Get();
 	}
@@ -231,7 +209,7 @@ public:
 	 */
 	FORCEINLINE bool IsValid() const
 	{
-		return FAssetPtr::IsValid();
+		return AssetPtr.IsValid();
 	}
 
 	/**  
@@ -240,7 +218,7 @@ public:
 	 */
 	FORCEINLINE bool IsPending() const
 	{
-		return FAssetPtr::IsPending();
+		return AssetPtr.IsPending();
 	}
 
 	/**  
@@ -249,7 +227,7 @@ public:
 	 */
 	FORCEINLINE bool IsNull() const
 	{
-		return FAssetPtr::IsNull();
+		return AssetPtr.IsNull();
 	}
 
 	/**  
@@ -257,69 +235,39 @@ public:
 	 */
 	FORCEINLINE const FStringAssetReference& ToStringReference() const
 	{
-		return FAssetPtr::GetUniqueID();
+		return AssetPtr.GetUniqueID();
 	}
 
 	/**  
 	 * Dereference lazy pointer to see if it points somewhere valid.
 	 */
-	FORCEINLINE operator bool() const
+	FORCEINLINE explicit operator bool() const
 	{
 		return Get() != NULL;
 	}
 
-	/**  Compare for equality with a raw pointer **/
-	FORCEINLINE bool operator==(T* Other) const
-	{
-		return Get() == Other;
-	}
-	/**  Compare for equality with a raw pointer **/
-	FORCEINLINE bool operator==(const T* Other) const
-	{
-		return Get() == Other;
-	}
-	/**  Compare to NULL **/
-	FORCEINLINE bool operator==(TYPE_OF_NULL Other) const
-	{
-		checkSlow(!Other); // we assume this is NULL
-		return !IsValid();
-	}
-
-	/**  Compare for inequality with a raw pointer	**/
-	FORCEINLINE bool operator!=(T* Other) const
-	{
-		return Get() != Other;
-	}
-	/**  Compare for inequality with a raw pointer	**/
-	FORCEINLINE bool operator!=(const T* Other) const
-	{
-		return Get() != Other;
-	}
-	/**  Compare for inequality with NULL **/
-	FORCEINLINE bool operator!=(TYPE_OF_NULL Other) const
-	{
-		checkSlow(!Other); // we assume this is NULL
-		return IsValid();
-	}
-
 	/** Hash function. */
-	FORCEINLINE friend uint32 GetTypeHash(const TAssetPtr<T>& AssetPtr)
+	FORCEINLINE friend uint32 GetTypeHash(const TAssetPtr<T>& Other)
 	{
-		return GetTypeHash(static_cast<const TPersistentObjectPtr<FStringAssetReference>&>(AssetPtr));
+		return GetTypeHash(static_cast<const TPersistentObjectPtr<FStringAssetReference>&>(Other.AssetPtr));
 	}
 
-	friend FArchive& operator<<( FArchive& Ar, TAssetPtr<T>& AssetPtr )
+	friend FArchive& operator<<(FArchive& Ar, TAssetPtr<T>& Other)
 	{
-		Ar << static_cast<FAssetPtr&>(AssetPtr);
+		Ar << Other.AssetPtr;
 		return Ar;
 	}
+
+private:
+	/** Wrapper of a UObject* */
+	FAssetPtr AssetPtr;
 };
 
 template<class T> struct TIsPODType<TAssetPtr<T> > { enum { Value = TIsPODType<FAssetPtr>::Value }; };
 template <class T> struct TIsWeakPointerType<TAssetPtr<T> > { enum { Value = TIsWeakPointerType<FAssetPtr>::Value }; };
 
 template<class TClass>
-class TAssetSubclassOf : private FAssetPtr
+class TAssetSubclassOf
 {
 public:
 	/** NULL constructor **/
@@ -332,8 +280,8 @@ public:
 	 * @param Other lazy pointer to copy from
 	 */
 	template<class TClassA>
-	FORCEINLINE TAssetSubclassOf(const TAssetSubclassOf<TClassA> &Other) :
-		FAssetPtr(Other)
+	FORCEINLINE TAssetSubclassOf(const TAssetSubclassOf<TClassA>& Other)
+		: AssetPtr(Other.AssetPtr)
 	{
 		TClass* Test = (TClassA*)0; // verify that TClassA derives from TClass
 	}
@@ -342,26 +290,17 @@ public:
 	 * Construct from an class pointer
 	 * @param From Class to create a reference to
 	**/
-	FORCEINLINE TAssetSubclassOf(const UClass* From) :
-		FAssetPtr(From)
+	FORCEINLINE TAssetSubclassOf(const UClass* From)
+		: AssetPtr(From)
 	{
 	}
 
-	/**  
-	 * Construct from an int (assumed to be NULL)
-	 * @param Other int to create a reference to
-	 */
-	FORCEINLINE TAssetSubclassOf(TYPE_OF_NULL Other)
-	{
-		checkSlow(!Other); // we assume this is NULL
-	}
-
 	/**
 	 * Reset the lazy pointer back to the NULL state
 	 */
 	FORCEINLINE void Reset()
 	{
-		FAssetPtr::Reset();
+		AssetPtr.Reset();
 	}
 
 	/**  
@@ -370,7 +309,7 @@ public:
 	 */
 	FORCEINLINE void operator=(const UClass* From)
 	{
-		FAssetPtr::operator=(From);
+		AssetPtr = From;
 	}
 
 	/**  
@@ -379,7 +318,7 @@ public:
 	 */
 	FORCEINLINE void operator=(const FStringAssetReference &ObjectID)
 	{
-		FAssetPtr::operator=(ObjectID);
+		AssetPtr = ObjectID;
 	}
 
 	/**  
@@ -390,7 +329,7 @@ public:
 	FORCEINLINE void operator=(const TWeakObjectPtr<TClassA> &Other)
 	{
 		UClass* Test = (TClassA*)0; //verify that TClassA derives from UClass
-		FAssetPtr::operator=(Other);
+		AssetPtr = Other;
 	}
 
 	/**  
@@ -401,17 +340,7 @@ public:
 	FORCEINLINE void operator=(const TAssetSubclassOf<TClassA> &Other)
 	{
 		TClass* Test = (TClassA*)0; // verify that TClassA derives from TClass
-		FAssetPtr::operator=(Other);
-	}
-
-	/**  
-	 * Copy from an int (assumed to be NULL)
-	 * @param Other int to create a reference to
-	 */
-	FORCEINLINE void operator=(TYPE_OF_NULL Other)
-	{
-		checkSlow(!Other); // we assume this is NULL
-		FAssetPtr::Reset();
+		AssetPtr = Other;
 	}
 
 	/**  
@@ -419,19 +348,19 @@ public:
 	 * @param Other lazy pointer to compare to
 	 * Caution: Two lazy pointers might not be equal to each other, but they both might return NULL.
 	 */
-	FORCEINLINE bool operator==(const FAssetPtr &Other) const
+	FORCEINLINE friend bool operator==(const TAssetSubclassOf& Lhs, const TAssetSubclassOf& Rhs)
 	{
-		return FAssetPtr::operator==(Other);
+		return Lhs.AssetPtr == Rhs.AssetPtr;
 	}
 
 	/**  
-	 * Compare lazy pointers for equality
+	 * Compare lazy pointers for inequality
 	 * @param Other lazy pointer to compare to
 	 * Caution: Two lazy pointers might not be equal to each other, but they both might return NULL.
 	 */
-	FORCEINLINE bool operator!=(const FAssetPtr &Other) const
+	FORCEINLINE friend bool operator!=(const TAssetSubclassOf& Lhs, const TAssetSubclassOf& Rhs)
 	{
-		return FAssetPtr::operator!=(Other);
+		return Lhs.AssetPtr != Rhs.AssetPtr;
 	}
 
 	/**
@@ -440,7 +369,7 @@ public:
 	 */
 	FStringAssetReference GetUniqueID() const
 	{
-		return FAssetPtr::GetUniqueID();
+		return AssetPtr.GetUniqueID();
 	}
 
 	/**  
@@ -449,7 +378,7 @@ public:
 	 */
 	FORCEINLINE UClass *Get() const
 	{
-		UClass* Class = dynamic_cast<UClass*>(FAssetPtr::Get());
+		UClass* Class = dynamic_cast<UClass*>(AssetPtr.Get());
 		if (!Class || !Class->IsChildOf(TClass::StaticClass()))
 		{
 			return NULL;
@@ -460,7 +389,7 @@ public:
 	/**  
 	 * Dereference the lazy pointer.
 	 */
-	FORCEINLINE UClass & operator*() const
+	FORCEINLINE UClass& operator*() const
 	{
 		return *Get();
 	}
@@ -468,7 +397,7 @@ public:
 	/**  
 	 * Dereference the lazy pointer.
 	 */
-	FORCEINLINE UClass * operator->() const
+	FORCEINLINE UClass* operator->() const
 	{
 		return Get();
 	}
@@ -488,7 +417,7 @@ public:
 	 */
 	FORCEINLINE bool IsPending() const
 	{
-		return FAssetPtr::IsPending();
+		return AssetPtr.IsPending();
 	}
 
 	/**  
@@ -497,7 +426,7 @@ public:
 	 */
 	FORCEINLINE bool IsNull() const
 	{
-		return FAssetPtr::IsNull();
+		return AssetPtr.IsNull();
 	}
 
 	/**  
@@ -505,63 +434,33 @@ public:
 	 */
 	FORCEINLINE const FStringAssetReference& ToStringReference() const
 	{
-		return FAssetPtr::GetUniqueID();
+		return AssetPtr.GetUniqueID();
 	}
 
 	/**  
 	 * Dereference lazy pointer to see if it points somewhere valid.
 	 */
-	FORCEINLINE operator bool() const
+	FORCEINLINE explicit operator bool() const
 	{
 		return Get() != NULL;
 	}
 
-	/**  Compare for equality with a raw pointer **/
-	FORCEINLINE bool operator==(UClass* Other) const
-	{
-		return Get() == Other;
-	}
-	/**  Compare for equality with a raw pointer **/
-	FORCEINLINE bool operator==(const UClass* Other) const
-	{
-		return Get() == Other;
-	}
-	/**  Compare to NULL **/
-	FORCEINLINE bool operator==(TYPE_OF_NULL Other) const
-	{
-		checkSlow(!Other); // we assume this is NULL
-		return !IsValid();
-	}
-
-	/**  Compare for inequality with a raw pointer	**/
-	FORCEINLINE bool operator!=(UClass* Other) const
-	{
-		return Get() != Other;
-	}
-	/**  Compare for inequality with a raw pointer	**/
-	FORCEINLINE bool operator!=(const UClass* Other) const
-	{
-		return Get() != Other;
-	}
-	/**  Compare for inequality with NULL **/
-	FORCEINLINE bool operator!=(TYPE_OF_NULL Other) const
-	{
-		checkSlow(!Other); // we assume this is NULL
-		return IsValid();
-	}
-
 	/** Hash function. */
-	FORCEINLINE friend uint32 GetTypeHash(const TAssetSubclassOf<TClass>& AssetPtr)
+	FORCEINLINE friend uint32 GetTypeHash(const TAssetSubclassOf<TClass>& Other)
 	{
-		return GetTypeHash(static_cast<const TPersistentObjectPtr<FStringAssetReference>&>(AssetPtr));
+		return GetTypeHash(static_cast<const TPersistentObjectPtr<FStringAssetReference>&>(Other.AssetPtr));
 	}
 
-	friend FArchive& operator<<( FArchive& Ar, TAssetSubclassOf<TClass>& AssetPtr )
+	friend FArchive& operator<<(FArchive& Ar, TAssetSubclassOf<TClass>& Other)
 	{
-		Ar << static_cast<FAssetPtr&>(AssetPtr);
+		Ar << static_cast<FAssetPtr&>(Other.AssetPtr);
 		return Ar;
 	}
+
+private:
+	/** Wrapper of a UObject* */
+	FAssetPtr AssetPtr;
 };
 
 template<class T> struct TIsPODType<TAssetSubclassOf<T> > { enum { Value = TIsPODType<FAssetPtr>::Value }; };
-template <class T> struct TIsWeakPointerType<TAssetSubclassOf<T> > { enum { Value = TIsWeakPointerType<FAssetPtr>::Value }; };
\ No newline at end of file
+template <class T> struct TIsWeakPointerType<TAssetSubclassOf<T> > { enum { Value = TIsWeakPointerType<FAssetPtr>::Value }; };
