Fix operator bool for VS2012.

diff --git a/Engine/Source/Runtime/CoreUObject/Public/UObject/AssetPtr.h b/Engine/Source/Runtime/CoreUObject/Public/UObject/AssetPtr.h
--- a/Engine/Source/Runtime/CoreUObject/Public/UObject/AssetPtr.h
+++ b/Engine/Source/Runtime/CoreUObject/Public/UObject/AssetPtr.h
@@ -241,7 +241,7 @@ public:
 	/**  
 	 * Dereference lazy pointer to see if it points somewhere valid.
 	 */
-	FORCEINLINE explicit operator bool() const
+	FORCEINLINE_EXPLICIT_OPERATOR_BOOL() const
 	{
 		return Get() != NULL;
 	}
@@ -440,7 +440,7 @@ public:
 	/**  
 	 * Dereference lazy pointer to see if it points somewhere valid.
 	 */
-	FORCEINLINE explicit operator bool() const
+	FORCEINLINE_EXPLICIT_OPERATOR_BOOL() const
 	{
 		return Get() != NULL;
 	}
