Fixed reporting of the referencing property displayed when saving a package linked to external private objects.
The bug was in FCollectorTagUsedNonRecursive::HandleObjectReference, I fixed a few other things while I was there (including the signature).

diff --git a/Engine/Source/Runtime/CoreUObject/Private/UObject/GarbageCollection.cpp b/Engine/Source/Runtime/CoreUObject/Private/UObject/GarbageCollection.cpp
--- a/Engine/Source/Runtime/CoreUObject/Private/UObject/GarbageCollection.cpp
+++ b/Engine/Source/Runtime/CoreUObject/Private/UObject/GarbageCollection.cpp
@@ -263,7 +263,7 @@ public:
 	{
 	}
 
-	virtual void HandleObjectReference(UObject*& Object, const UObject* ReferencingObject, const UObject* ReferencingProperty) override
+	virtual void HandleObjectReference(UObject*& Object, const UObject* ReferencingObject, const UProperty* ReferencingProperty) override
 	{
 #if !(UE_BUILD_TEST || UE_BUILD_SHIPPING)
 		if (Object && !Object->IsValidLowLevelFast())
@@ -315,7 +315,7 @@ FReferenceFinder::FReferenceFinder( TArray<UObject*>& InObjectArray, UObject* In
 	}
 }
 
-void FReferenceFinder::FindReferences(UObject* Object, UObject* InReferencingObject, UObject* InReferencingProperty)
+void FReferenceFinder::FindReferences(UObject* Object, UObject* InReferencingObject, UProperty* InReferencingProperty)
 {
 	check(Object != NULL);
 
@@ -328,7 +328,7 @@ void FReferenceFinder::FindReferences(UObject* Object, UObject* InReferencingObj
 	Object->CallAddReferencedObjects(*this);
 }
 
-void FReferenceFinder::HandleObjectReference( UObject*& InObject, const UObject* InReferencingObject /*= NULL*/, const UObject* InReferencingProperty /*= NULL*/ )
+void FReferenceFinder::HandleObjectReference( UObject*& InObject, const UObject* InReferencingObject /*= NULL*/, const UProperty* InReferencingProperty /*= NULL*/ )
 {
 	// Avoid duplicate entries.
 	if ( InObject != NULL )
@@ -348,7 +348,7 @@ void FReferenceFinder::HandleObjectReference( UObject*& InObject, const UObject*
 			if ( bSerializeRecursively == true && !SerializedObjects.Find(Object) )
 			{
 				SerializedObjects.Add(Object);
-				FindReferences(Object, const_cast<UObject*>(InReferencingObject), const_cast<UObject*>(InReferencingProperty));
+				FindReferences(Object, const_cast<UObject*>(InReferencingObject), const_cast<UProperty*>(InReferencingProperty));
 			}
 		}
 	}
diff --git a/Engine/Source/Runtime/CoreUObject/Private/UObject/Obj.cpp b/Engine/Source/Runtime/CoreUObject/Private/UObject/Obj.cpp
--- a/Engine/Source/Runtime/CoreUObject/Private/UObject/Obj.cpp
+++ b/Engine/Source/Runtime/CoreUObject/Private/UObject/Obj.cpp
@@ -1019,7 +1019,7 @@ public:
 	}
 
 	// Begin FReferenceCollector interface.
-	virtual void HandleObjectReference(UObject*& InObject, const UObject* InReferencingObject, const UObject* InReferencingProperty) override
+	virtual void HandleObjectReference(UObject*& InObject, const UObject* InReferencingObject, const UProperty* InReferencingProperty) override
 	{
 		// Only care about unique default subobjects that are outside of the referencing object's outer chain.
 		// Also ignore references to subobjects if they share the same Outer.
diff --git a/Engine/Source/Runtime/CoreUObject/Private/UObject/ReferenceChainSearch.cpp b/Engine/Source/Runtime/CoreUObject/Private/UObject/ReferenceChainSearch.cpp
--- a/Engine/Source/Runtime/CoreUObject/Private/UObject/ReferenceChainSearch.cpp
+++ b/Engine/Source/Runtime/CoreUObject/Private/UObject/ReferenceChainSearch.cpp
@@ -77,7 +77,7 @@ int32 FReferenceChainSearch::FFindReferencerCollector::FindReferencedObjectIndex
 	return Result;
 }
 #endif // !(UE_BUILD_TEST || UE_BUILD_SHIPPING)
-void FReferenceChainSearch::FFindReferencerCollector::HandleObjectReference( UObject*& InObject, const UObject* RefObject, const UObject* ReferencingProperty )
+void FReferenceChainSearch::FFindReferencerCollector::HandleObjectReference(UObject*& InObject, const UObject* RefObject, const UProperty* ReferencingProperty)
 {
 	UObject* RefSrc = RefObject != NULL ? const_cast<UObject*>(RefObject) : ReferencingObject;
 	int32 ReferencedObjectIndex = INDEX_NONE;
diff --git a/Engine/Source/Runtime/CoreUObject/Private/UObject/UObjectGlobals.cpp b/Engine/Source/Runtime/CoreUObject/Private/UObject/UObjectGlobals.cpp
--- a/Engine/Source/Runtime/CoreUObject/Private/UObject/UObjectGlobals.cpp
+++ b/Engine/Source/Runtime/CoreUObject/Private/UObject/UObjectGlobals.cpp
@@ -2551,42 +2551,33 @@ private:
 		ObjectsToSerialize.Add( Object );
 	}
 
-	virtual void HandleObjectReference( UObject*& InObject, const UObject* InReferencingObject, const UObject* InReferencingProperty ) override
+	virtual void HandleObjectReference(UObject*& InObject, const UObject* InReferencingObject, const UProperty* InReferencingProperty) override
 	{
-		checkSlow( !InObject || InObject->IsValidLowLevel() );		
-		if( InObject )
+		checkSlow(!InObject || InObject->IsValidLowLevel());
+		if (InObject)
 		{
-			UObject* Object = const_cast<UObject*>(InObject);
-			if ( Object )
+			if (InObject->HasAnyFlags(ReferenceSearchFlags))
 			{
-				// Remove references to pending kill objects if we're allowed to do so.
-				UObject* ObjectToAdd = Object->HasAnyFlags(RF_Unreachable) ? Object : NULL;
-				if (Object->HasAnyFlags(ReferenceSearchFlags))
+				// Stop recursing, and add to the list of references
+				if (FoundReferencesList)
 				{
-					// Stop recursing, and add to the list of references
-					if (FoundReferencesList)
+					if (!CurrentReferenceInfo)
 					{
-						if (!CurrentReferenceInfo)
-						{
-							CurrentReferenceInfo = new(FoundReferencesList->ExternalReferences) FReferencerInformation(CurrentObject);
-						}
-						if (InReferencingObject != NULL)
-						{
-							CurrentReferenceInfo->ReferencingProperties.AddUnique(dynamic_cast<const UProperty*>(InReferencingObject));
-						}
-						CurrentReferenceInfo->TotalReferences++;
+						CurrentReferenceInfo = new(FoundReferencesList->ExternalReferences) FReferencerInformation(CurrentObject);
 					}
-					if (ObjectToAdd)
+					if (InReferencingProperty)
 					{
-						// Mark it as reachable.
-						Object->ClearFlags(RF_Unreachable);
+						CurrentReferenceInfo->ReferencingProperties.AddUnique(InReferencingProperty);
 					}
+					CurrentReferenceInfo->TotalReferences++;
 				}
-				else if (ObjectToAdd)
-				{
-					// Add encountered object reference to list of to be serialized objects if it hasn't already been added.
-					AddToObjectList(InReferencingObject, InReferencingProperty, ObjectToAdd);
-				}
+				// Mark it as reachable.
+				InObject->ClearFlags(RF_Unreachable);
+			}
+			else if (InObject->HasAnyFlags(RF_Unreachable))
+			{
+				// Add encountered object reference to list of to be serialized objects if it hasn't already been added.
+				AddToObjectList(InReferencingObject, InReferencingProperty, InObject);
 			}
 		}
 	}
diff --git a/Engine/Source/Runtime/CoreUObject/Public/UObject/GarbageCollection.h b/Engine/Source/Runtime/CoreUObject/Public/UObject/GarbageCollection.h
--- a/Engine/Source/Runtime/CoreUObject/Public/UObject/GarbageCollection.h
+++ b/Engine/Source/Runtime/CoreUObject/Public/UObject/GarbageCollection.h
@@ -456,7 +456,7 @@ protected:
 		{
 			auto OldCollectorSerializedProperty = Collector.GetSerializedProperty();
 			Collector.SetSerializedProperty(GetSerializedProperty());
-			Collector.AddReferencedObject(Object, SerializingObject, (UObject*)GetSerializedProperty());
+			Collector.AddReferencedObject(Object, SerializingObject, GetSerializedProperty());
 			Collector.SetSerializedProperty(OldCollectorSerializedProperty);
 		}
 		return *this;
diff --git a/Engine/Source/Runtime/CoreUObject/Public/UObject/ReferenceChainSearch.h b/Engine/Source/Runtime/CoreUObject/Public/UObject/ReferenceChainSearch.h
--- a/Engine/Source/Runtime/CoreUObject/Public/UObject/ReferenceChainSearch.h
+++ b/Engine/Source/Runtime/CoreUObject/Public/UObject/ReferenceChainSearch.h
@@ -93,7 +93,7 @@ public:
 		{
 		}
 
-		virtual void HandleObjectReference( UObject*& InObject, const UObject* ReferencingObject, const UObject* ReferencingProperty ) override;
+		virtual void HandleObjectReference( UObject*& InObject, const UObject* ReferencingObject, const UProperty* ReferencingProperty ) override;
 
 		virtual bool IsIgnoringArchetypeRef() const override	{ return false; }
 		virtual bool IsIgnoringTransient() const override		{ return false; }
diff --git a/Engine/Source/Runtime/CoreUObject/Public/UObject/UObjectGlobals.h b/Engine/Source/Runtime/CoreUObject/Public/UObject/UObjectGlobals.h
--- a/Engine/Source/Runtime/CoreUObject/Public/UObject/UObjectGlobals.h
+++ b/Engine/Source/Runtime/CoreUObject/Public/UObject/UObjectGlobals.h
@@ -1434,7 +1434,7 @@ public:
 	 * @param ReferencingProperty Referencing property (if available).
 	 */
 	template<class UObjectType>
-	void AddReferencedObject(UObjectType*& Object, const UObject* ReferencingObject = NULL, const UObject* ReferencingProperty = NULL)
+	void AddReferencedObject(UObjectType*& Object, const UObject* ReferencingObject = NULL, const UProperty* ReferencingProperty = NULL)
 	{
 		HandleObjectReference(*(UObject**)&Object, ReferencingObject, ReferencingProperty);
 	}
@@ -1447,7 +1447,7 @@ public:
 	* @param ReferencingProperty Referencing property (if available).
 	*/
 	template<class UObjectType>
-	void AddReferencedObjects(TArray<UObjectType>& ObjectArray, const UObject* ReferencingObject = NULL, const UObject* ReferencingProperty = NULL)
+	void AddReferencedObjects(TArray<UObjectType>& ObjectArray, const UObject* ReferencingObject = NULL, const UProperty* ReferencingProperty = NULL)
 	{
 		for (auto& Object : ObjectArray)
 		{
@@ -1463,7 +1463,7 @@ public:
 	* @param ReferencingProperty Referencing property (if available).
 	*/
 	template<class UObjectType>
-	void AddReferencedObjects(TSet<UObjectType>& ObjectSet, const UObject* ReferencingObject = NULL, const UObject* ReferencingProperty = NULL)
+	void AddReferencedObjects(TSet<UObjectType>& ObjectSet, const UObject* ReferencingObject = NULL, const UProperty* ReferencingProperty = NULL)
 	{
 		for (auto& Object : ObjectSet)
 		{
@@ -1490,11 +1490,11 @@ private:
 	* Functions used by AddReferencedObject (TMap version). Adds references to UObjects, ignores value types
 	*/
 	template<class UObjectType>
-	void AddReferencedObjectOrIgnoreValue(UObjectType& Object, const UObject* ReferencingObject, const UObject* ReferencingProperty)
+	void AddReferencedObjectOrIgnoreValue(UObjectType& Object, const UObject* ReferencingObject, const UProperty* ReferencingProperty)
 	{
 	}
 	template<class UObjectType>
-	void AddReferencedObjectOrIgnoreValue(UObjectType*& Object, const UObject* ReferencingObject, const UObject* ReferencingProperty)
+	void AddReferencedObjectOrIgnoreValue(UObjectType*& Object, const UObject* ReferencingObject, const UProperty* ReferencingProperty)
 	{
 		HandleObjectReference(*(UObject**)&Object, ReferencingObject, ReferencingProperty);
 	}
@@ -1509,7 +1509,7 @@ public:
 	* @param ReferencingProperty Referencing property (if available).
 	*/
 	template <typename TKeyType, typename TValueType>
-	void AddReferencedObjects(TMap<TKeyType, TValueType>& Map, const UObject* ReferencingObject = NULL, const UObject* ReferencingProperty = NULL)
+	void AddReferencedObjects(TMap<TKeyType, TValueType>& Map, const UObject* ReferencingObject = NULL, const UProperty* ReferencingProperty = NULL)
 	{
 		static_assert(CanConverFromTo<TKeyType, UObjectBase*>::Result || CanConverFromTo<TValueType, UObjectBase*>::Result, "At least one of TMap template types must be derived from UObject");
 		for (auto& It : Map)
@@ -1553,7 +1553,7 @@ protected:
 	 * @param ReferencingObject Referencing object (if available).
 	 * @param ReferencingProperty Referencing property (if available).
 	 */
-	virtual void HandleObjectReference(UObject*& Object, const UObject* ReferencingObject, const UObject* ReferencingProperty) = 0;
+	virtual void HandleObjectReference(UObject*& InObject, const UObject* InReferencingObject, const UProperty* InReferencingProperty) = 0;
 };
 
 /**
@@ -1585,10 +1585,10 @@ public:
 	 * @param ReferencingObject object that's referencing the current object.
 	 * @param ReferencingProperty property the current object is being referenced through.
 	 */
-	virtual void FindReferences(UObject* Object, UObject* ReferencingObject = nullptr, UObject* ReferencingProperty = nullptr);	
+	virtual void FindReferences(UObject* Object, UObject* ReferencingObject = nullptr, UProperty* ReferencingProperty = nullptr);
 
 	// FReferenceCollector interface.
-	virtual void HandleObjectReference(UObject*& Object, const UObject* ReferencingObject, const UObject* InReferencingProperty) override;
+	virtual void HandleObjectReference(UObject*& Object, const UObject* ReferencingObject, const UProperty* InReferencingProperty) override;
 	virtual bool IsIgnoringArchetypeRef() const override { return bShouldIgnoreArchetype; }
 	virtual bool IsIgnoringTransient() const override { return bShouldIgnoreTransient; }
 	virtual void SetSerializedProperty(class UProperty* Inproperty) override
