Comment the code in SerializeTaggedProperties which automatically converts properties from one type to another on load

diff --git a/Engine/Source/Runtime/CoreUObject/Private/UObject/Class.cpp b/Engine/Source/Runtime/CoreUObject/Private/UObject/Class.cpp
--- a/Engine/Source/Runtime/CoreUObject/Private/UObject/Class.cpp
+++ b/Engine/Source/Runtime/CoreUObject/Private/UObject/Class.cpp
@@ -935,9 +935,14 @@ void UStruct::SerializeTaggedProperties(FArchive& Ar, uint8* Data, UStruct* Defa
 			{
 				UE_LOG(LogClass, Warning, TEXT("Array bounds in %s of %s: %i/%i for package:  %s"), *Tag.Name.ToString(), *GetName(), Tag.ArrayIndex, Property->ArrayDim, *Ar.GetArchiveName() );
 			}
+
+			// Convert properties from old type to new type automatically if types are compatible
+			// If you add an entry to this, you will also need to add an entry to the array case below
+			// For converting to a struct, you can just implement SerializeFromMismatchedTag on the struct
+
 			else if( Tag.Type==NAME_StrProperty && dynamic_cast<UNameProperty*>(Property) )
-			{ 
-				FString str;  
+			{
+				FString str;
 				Ar << str;
 				CastChecked<UNameProperty>(Property)->SetPropertyValue_InContainer(Data, FName(*str), Tag.ArrayIndex);
 				AdvanceProperty = true;
@@ -1031,7 +1036,7 @@ void UStruct::SerializeTaggedProperties(FArchive& Ar, uint8* Data, UStruct* Defa
 				AdvanceProperty = true;
 				continue;
 			}
-			else if(Tag.Type == NAME_IntProperty && Property->GetID() == NAME_BoolProperty )
+			else if (Tag.Type == NAME_IntProperty && Property->GetID() == NAME_BoolProperty)
 			{
 				// Property was saved as an int32, but has been changed to a bool (bitfield)
 				int32 IntValue;
@@ -1085,7 +1090,9 @@ void UStruct::SerializeTaggedProperties(FArchive& Ar, uint8* Data, UStruct* Defa
 				FScriptArrayHelper ScriptArrayHelper(ArrayProperty, ArrayPropertyData);
 				ScriptArrayHelper.EmptyAndAddValues(ElementCount);
 
-				if( Tag.InnerType==NAME_StrProperty && dynamic_cast<UTextProperty*>(ArrayProperty->Inner) ) // Convert serialized string to text.
+				// Convert properties from old type to new type automatically if types are compatible (array case)
+
+				if (Tag.InnerType == NAME_StrProperty && dynamic_cast<UTextProperty*>(ArrayProperty->Inner)) // Convert serialized string to text.
 				{ 
 					for(int32 i = 0; i < ElementCount; ++i)
 					{
@@ -1142,6 +1149,7 @@ void UStruct::SerializeTaggedProperties(FArchive& Ar, uint8* Data, UStruct* Defa
 					}
 					continue;
 				}
+				// TODO: Implement SerializeFromMismatchedTag handling for arrays of structs
 				else
 				{
 					UE_LOG(LogClass, Warning, TEXT("Array Inner Type mismatch in %s of %s - Previous (%s) Current(%s) for package:  %s"), *Tag.Name.ToString(), *GetName(), *Tag.InnerType.ToString(), *CastChecked<UArrayProperty>(Property)->Inner->GetID().ToString(), *Ar.GetArchiveName() );
