CoreUObject: Added missing override keywords

diff --git a/Engine/Source/Runtime/CoreUObject/Private/UObject/SavePackage.cpp b/Engine/Source/Runtime/CoreUObject/Private/UObject/SavePackage.cpp
--- a/Engine/Source/Runtime/CoreUObject/Private/UObject/SavePackage.cpp
+++ b/Engine/Source/Runtime/CoreUObject/Private/UObject/SavePackage.cpp
@@ -1,18 +1,14 @@
 // Copyright 1998-2015 Epic Games, Inc. All Rights Reserved.
 
-/*=============================================================================
-	SavePackage.cpp: File saving support.
-=============================================================================*/
-
 #include "CoreUObjectPrivate.h"
 #include "UObject/UTextProperty.h"
 #include "Interface.h"
 #include "TargetPlatform.h"
 
+
 DEFINE_LOG_CATEGORY_STATIC(LogSavePackage, Log, All);
 
 static const int32 MAX_MERGED_COMPRESSION_CHUNKSIZE = 1024 * 1024;
-
 static const FName WorldClassName = FName("World");
 
 
@@ -558,7 +554,7 @@ public:
 	 *
 	 * This is overridden for the specific Archive Types
 	 **/
-	virtual FString GetArchiveName() const;
+	virtual FString GetArchiveName() const override;
 
 };
 
diff --git a/Engine/Source/Runtime/CoreUObject/Public/Misc/TextBuffer.h b/Engine/Source/Runtime/CoreUObject/Public/Misc/TextBuffer.h
--- a/Engine/Source/Runtime/CoreUObject/Public/Misc/TextBuffer.h
+++ b/Engine/Source/Runtime/CoreUObject/Public/Misc/TextBuffer.h
@@ -1,9 +1,5 @@
 // Copyright 1998-2015 Epic Games, Inc. All Rights Reserved.
 
-/*=============================================================================
-	TextBuffer.h: UObject class for storing text
-=============================================================================*/
-
 #pragma once
 
 #include "ObjectBase.h"
@@ -25,12 +21,11 @@ public:
 	/**
 	 * Creates and initializes a new text buffer.
 	 *
-	 * @param ObjectInitializer - Initialization properties.
-	 * @param InText - The initial text.
+	 * @param ObjectInitializer Initialization properties.
+	 * @param InText The initial text.
 	 */
 	COREUOBJECT_API UTextBuffer (const FObjectInitializer& ObjectInitializer, const TCHAR* InText);
 
-
 public:
 
 	/**
@@ -43,18 +38,15 @@ public:
 		return Text;
 	}
 
-
 public:
 
-	virtual void Serialize (FArchive& Ar);
-
+	virtual void Serialize (FArchive& Ar) override;
 	virtual void Serialize (const TCHAR* Data, ELogVerbosity::Type Verbosity, const class FName& Category ) override;
 
-
 private:
 
 	int32 Pos, Top;
 
-	// Holds the text.
+	/** Holds the text. */
 	FString Text;
-};
\ No newline at end of file
+};
diff --git a/Engine/Source/Runtime/CoreUObject/Public/Serialization/ArchiveUObject.h b/Engine/Source/Runtime/CoreUObject/Public/Serialization/ArchiveUObject.h
--- a/Engine/Source/Runtime/CoreUObject/Public/Serialization/ArchiveUObject.h
+++ b/Engine/Source/Runtime/CoreUObject/Public/Serialization/ArchiveUObject.h
@@ -1,13 +1,11 @@
 // Copyright 1998-2015 Epic Games, Inc. All Rights Reserved.
 
-/*=============================================================================
-	ArchiveUObject.h: Implements the FArchive class and related types for UObjects
-=============================================================================*/
-
 #pragma once
 
+
 struct FObjectInstancingGraph;
 
+
 /**
  * Archive for counting memory usage.
  */
@@ -100,8 +98,8 @@ public:
 	// FArchive interface
 	COREUOBJECT_API virtual FArchive& operator<<( class FName& N ) override;
 	COREUOBJECT_API virtual FArchive& operator<<( class UObject*& Res ) override;
-	COREUOBJECT_API virtual FArchive& operator<<( FLazyObjectPtr& LazyObjectPtr );
-	COREUOBJECT_API virtual FArchive& operator<<( FAssetPtr& AssetPtr );
+	COREUOBJECT_API virtual FArchive& operator<<( FLazyObjectPtr& LazyObjectPtr ) override;
+	COREUOBJECT_API virtual FArchive& operator<<( FAssetPtr& AssetPtr ) override;
 	COREUOBJECT_API virtual FArchive& operator<<(FStringAssetReference& AssetPtr) override;
 	COREUOBJECT_API virtual FString GetArchiveName() const override;
 	// End of FArchive interface
@@ -133,8 +131,8 @@ public:
 	// FArchive interface
 	COREUOBJECT_API virtual FArchive& operator<<( class FName& N ) override;
 	COREUOBJECT_API virtual FArchive& operator<<( class UObject*& Res ) override;
-	COREUOBJECT_API virtual FArchive& operator<<( FLazyObjectPtr& LazyObjectPtr );
-	COREUOBJECT_API virtual FArchive& operator<<( FAssetPtr& AssetPtr );
+	COREUOBJECT_API virtual FArchive& operator<<( FLazyObjectPtr& LazyObjectPtr ) override;
+	COREUOBJECT_API virtual FArchive& operator<<( FAssetPtr& AssetPtr ) override;
 	COREUOBJECT_API virtual FArchive& operator<<(FStringAssetReference& AssetPtr) override;
 	COREUOBJECT_API virtual FString GetArchiveName() const override;
 	// End of FArchive interface
diff --git a/Engine/Source/Runtime/CoreUObject/Public/UObject/CoreNet.h b/Engine/Source/Runtime/CoreUObject/Public/UObject/CoreNet.h
--- a/Engine/Source/Runtime/CoreUObject/Public/UObject/CoreNet.h
+++ b/Engine/Source/Runtime/CoreUObject/Public/UObject/CoreNet.h
@@ -1,15 +1,13 @@
 // Copyright 1998-2015 Epic Games, Inc. All Rights Reserved.
 
-/*=============================================================================
-	CoreNet.h: Core networking support.
-=============================================================================*/
-
 #pragma once
 
 #include "ObjectBase.h"
 
+
 DECLARE_DELEGATE_RetVal_OneParam( bool, FNetObjectIsDynamic, const UObject*);
 
+
 //
 // Information about a field.
 //
@@ -25,6 +23,7 @@ public:
 	{}
 };
 
+
 //
 // Information about a class, cached for network coordination.
 //
@@ -69,6 +68,7 @@ private:
 	TMap< UObject *, FFieldNetCache * > FieldMap;
 };
 
+
 class COREUOBJECT_API FClassNetCacheMgr
 {
 public:
@@ -80,6 +80,7 @@ private:
 	TMap< TWeakObjectPtr< UClass >, FClassNetCache * > ClassFieldIndices;
 };
 
+
 //
 // Maps objects and names to and from indices for network communication.
 //
@@ -129,6 +130,7 @@ protected:
 	FString					DebugContextString;
 };
 
+
 /** Represents a range of PacketIDs, inclusive */
 struct FPacketIdRange
 {
@@ -144,6 +146,7 @@ struct FPacketIdRange
 	}
 };
 
+
 /** Information for tracking retirement and retransmission of a property. */
 struct FPropertyRetirement
 {
@@ -166,6 +169,7 @@ struct FPropertyRetirement
 	{}
 };
 
+
 /** Secondary condition to check before considering the replication of a lifetime property. */
 enum ELifetimeCondition
 {
@@ -181,12 +185,14 @@ enum ELifetimeCondition
 	COND_Max				= 9,
 };
 
+
 enum ELifetimeRepNotifyCondition
 {
 	REPNOTIFY_OnChanged		= 0,		// Only call the property's RepNotify function if it changes from the local value
 	REPNOTIFY_Always		= 1,		// Always Call the property's RepNotify function when it is received from the server
 };
 
+
 /** FLifetimeProperty
  *	This class is used to track a property that is marked to be replicated for the lifetime of the actor channel.
  *  This doesn't mean the property will necessarily always be replicated, it just means:
@@ -219,6 +225,7 @@ public:
 
 template <> struct TIsZeroConstructType<FLifetimeProperty> { enum { Value = true }; };
 
+
 /**
  * FNetBitWriter
  *	A bit writer that serializes FNames and UObject* through
@@ -233,11 +240,12 @@ public:
 
 	class UPackageMap * PackageMap;
 
-	FArchive& operator<<( FName& Name );
-	FArchive& operator<<( UObject*& Object );
-	FArchive& operator<<(FStringAssetReference& Value) override;
+	virtual FArchive& operator<<(FName& Name) override;
+	virtual FArchive& operator<<(UObject*& Object) override;
+	virtual FArchive& operator<<(FStringAssetReference& Value) override;
 };
 
+
 /**
  * FNetBitReader
  *	A bit reader that serializes FNames and UObject* through
@@ -251,9 +259,9 @@ public:
 
 	class UPackageMap * PackageMap;
 
-	FArchive& operator<<( FName& Name );
-	FArchive& operator<<( UObject*& Object );
-	FArchive& operator<<(FStringAssetReference& Value) override;
+	virtual FArchive& operator<<(FName& Name) override;
+	virtual FArchive& operator<<(UObject*& Object) override;
+	virtual FArchive& operator<<(FStringAssetReference& Value) override;
 };
 
 
@@ -272,6 +280,7 @@ public:
 private:
 };
 
+
 class INetSerializeCB
 {
 public:
@@ -280,6 +289,7 @@ public:
 	virtual void NetSerializeStruct( UScriptStruct* Struct, FArchive& Ar, UPackageMap* Map, void* Data, bool& bHasUnmapped ) = 0;
 };
 
+
 class IRepChangedPropertyTracker
 {
 public:
@@ -288,6 +298,7 @@ public:
 	virtual void SetCustomIsActiveOverride( const uint16 RepIndex, const bool bIsActive ) = 0;
 };
 
+
 /**
  * FNetDeltaSerializeInfo
  *  This is the parameter structure for delta serialization. It is kind of a dumping ground for anything custom implementations may need.
@@ -341,6 +352,7 @@ struct FNetDeltaSerializeInfo
 	FString							DebugName;
 };
 
+
 /**
  * Checksum macros for verifying archives stay in sync
  */
@@ -348,6 +360,7 @@ COREUOBJECT_API void SerializeChecksum(FArchive &Ar, uint32 x, bool ErrorOK);
 
 #define NET_ENABLE_CHECKSUMS 0
 
+
 #if !(UE_BUILD_SHIPPING || UE_BUILD_TEST) && NET_ENABLE_CHECKSUMS
 
 #define NET_CHECKSUM_OR_END(Ser) \
@@ -383,6 +396,7 @@ COREUOBJECT_API void SerializeChecksum(FArchive &Ar, uint32 x, bool ErrorOK);
 
 #endif
 
+
 /**
  * Functions to assist in detecting errors during RPC calls
  */
diff --git a/Engine/Source/Runtime/CoreUObject/Public/UObject/CoreObject.h b/Engine/Source/Runtime/CoreUObject/Public/UObject/CoreObject.h
--- a/Engine/Source/Runtime/CoreUObject/Public/UObject/CoreObject.h
+++ b/Engine/Source/Runtime/CoreUObject/Public/UObject/CoreObject.h
@@ -5,6 +5,7 @@
 #include "ObjectBase.h"
 #include "WorldCompositionUtility.h"
 
+
 /**
  * Structure to hold information about an external packages objects used in cross-level references
  */
@@ -98,10 +99,11 @@ public:
 	 * Called after the C++ constructor and after the properties have been initialized, but before the config has been loaded, etc.
 	 * mainly this is to emulate some behavior of when the constructor was called after the properties were intialized.
 	 */
-	virtual void PostInitProperties();
+	virtual void PostInitProperties() override;
 
 	/** Serializer */
-	virtual void Serialize( FArchive& Ar );
+	virtual void Serialize( FArchive& Ar ) override;
+
 	static void AddReferencedObjects(UObject* InThis, FReferenceCollector& Collector);
 
 	/** Packages are never assets */
@@ -182,7 +184,7 @@ public:
 	/**
 	 * Tags the Package's metadata
 	 */
-	virtual void TagSubobjects(EObjectFlags NewFlags);
+	virtual void TagSubobjects(EObjectFlags NewFlags) override;
 
 	/**
 	 * Sets whether exports should be found in memory first or  not.
diff --git a/Engine/Source/Runtime/CoreUObject/Public/UObject/Linker.h b/Engine/Source/Runtime/CoreUObject/Public/UObject/Linker.h
--- a/Engine/Source/Runtime/CoreUObject/Public/UObject/Linker.h
+++ b/Engine/Source/Runtime/CoreUObject/Public/UObject/Linker.h
@@ -1,14 +1,11 @@
 // Copyright 1998-2015 Epic Games, Inc. All Rights Reserved.
 
-/*=============================================================================
-	Linker.h: Unreal object linker.
-=============================================================================*/
-
 #pragma once
 
 #include "ObjectBase.h"
 #include "EngineVersion.h"
 
+
 DECLARE_LOG_CATEGORY_EXTERN(LogLinker, Log, All);
 
 /** Map of packages to their open linkers **/ 
@@ -18,6 +15,7 @@ extern TMap<class UPackage*, class ULinkerLoad*>		GObjPendingLoaders;
 /** List of loaders that have new imports **/
 extern TSet<class ULinkerLoad*>							GObjLoadersWithNewImports;
 
+
 /**
  * Wrapper for index into a ULnker's ImportMap or ExportMap.
  * Values greater than zero indicate that this is an index into the ExportMap.  The
@@ -1692,14 +1690,14 @@ private:
 	 * @param	Owner		UObject owning the bulk data
 	 * @param	BulkData	Bulk data object to associate
 	 */
-	virtual void AttachBulkData( UObject* Owner, FUntypedBulkData* BulkData );
+	virtual void AttachBulkData( UObject* Owner, FUntypedBulkData* BulkData ) override;
 	/**
 	 * Detaches the passed in bulk data object from the linker.
 	 *
 	 * @param	BulkData	Bulk data object to detach
 	 * @param	bEnsureBulkDataIsLoaded	Whether to ensure that the bulk data is loaded before detaching
 	 */
-	virtual void DetachBulkData( FUntypedBulkData* BulkData, bool bEnsureBulkDataIsLoaded );
+	virtual void DetachBulkData( FUntypedBulkData* BulkData, bool bEnsureBulkDataIsLoaded ) override;
 	/**
 	 * Detaches all attached bulk  data objects.
 	 *
diff --git a/Engine/Source/Runtime/CoreUObject/Public/UObject/ObjectRedirector.h b/Engine/Source/Runtime/CoreUObject/Public/UObject/ObjectRedirector.h
--- a/Engine/Source/Runtime/CoreUObject/Public/UObject/ObjectRedirector.h
+++ b/Engine/Source/Runtime/CoreUObject/Public/UObject/ObjectRedirector.h
@@ -19,8 +19,8 @@ class UObjectRedirector : public UObject
 	// Variables.
 	UObject*		DestinationObject;
 	// UObject interface.
-	virtual void PreSave();
-	void Serialize( FArchive& Ar );
+	virtual void PreSave() override;
+	void Serialize( FArchive& Ar ) override;
 	virtual bool NeedsLoadForClient() const override;
 	virtual bool NeedsLoadForServer() const override;
 	virtual bool NeedsLoadForEditorGame() const override;
@@ -38,7 +38,7 @@ class UObjectRedirector : public UObject
 	 *
 	 * @return	return true if property values were added to the map.
 	 */
-	virtual bool GetNativePropertyValues( TMap<FString,FString>& out_PropertyValues, uint32 ExportFlags=0 ) const;
+	virtual bool GetNativePropertyValues( TMap<FString,FString>& out_PropertyValues, uint32 ExportFlags=0 ) const override;
 };
 
 /**
diff --git a/Engine/Source/Runtime/CoreUObject/Public/UObject/UObject.h b/Engine/Source/Runtime/CoreUObject/Public/UObject/UObject.h
--- a/Engine/Source/Runtime/CoreUObject/Public/UObject/UObject.h
+++ b/Engine/Source/Runtime/CoreUObject/Public/UObject/UObject.h
@@ -1,13 +1,11 @@
 // Copyright 1998-2015 Epic Games, Inc. All Rights Reserved.
 
-/*=============================================================================
-	UObject.h: Unreal object base class
-=============================================================================*/
-
 #pragma once
 
+
 DECLARE_LOG_CATEGORY_EXTERN(LogObj, Log, All);
 
+
 namespace ECastCheckedType
 {
 	enum Type
diff --git a/Engine/Source/Runtime/CoreUObject/Public/UObject/UnrealType.h b/Engine/Source/Runtime/CoreUObject/Public/UObject/UnrealType.h
--- a/Engine/Source/Runtime/CoreUObject/Public/UObject/UnrealType.h
+++ b/Engine/Source/Runtime/CoreUObject/Public/UObject/UnrealType.h
@@ -2147,7 +2147,7 @@ public:
 	// UProperty interface
 	virtual void ExportTextItem( FString& ValueStr, const void* PropertyValue, const void* DefaultValue, UObject* Parent, int32 PortFlags, UObject* ExportRootScope ) const override;
 	virtual const TCHAR* ImportText_Internal( const TCHAR* Buffer, void* Data, int32 PortFlags, UObject* OwnerObject, FOutputDevice* ErrorText ) const override;
-	virtual FString GetCPPTypeForwardDeclaration() const
+	virtual FString GetCPPTypeForwardDeclaration() const override
 	{
 		return FString();
 	}
