UE-12106 [CrashReport] Saving after blueprint change causes crash.

A better handling for null CDO in FExportReferenceSorter::ProcessStruct().

diff --git a/Engine/Source/Runtime/CoreUObject/Private/UObject/SavePackage.cpp b/Engine/Source/Runtime/CoreUObject/Private/UObject/SavePackage.cpp
--- a/Engine/Source/Runtime/CoreUObject/Private/UObject/SavePackage.cpp
+++ b/Engine/Source/Runtime/CoreUObject/Private/UObject/SavePackage.cpp
@@ -1883,7 +1883,8 @@ public:
 				if ( ClassObject != NULL )
 				{
 					UObject* CDO = ClassObject->GetDefaultObject();
-					if ( ProcessedObjects.Find(CDO) == INDEX_NONE )
+					ensureMsgf(nullptr != CDO, TEXT("Error: Invalid CDO in class %s"), *GetPathNameSafe(ClassObject));
+					if ((ProcessedObjects.Find(CDO) == INDEX_NONE) && (nullptr != CDO))
 					{
 						ProcessedObjects.Add(CDO);
 
