Don't mark all loaded objects as root set in non-cooked builds. Blueprints don't like it. Instead mark all objects loaded with Constructor Helpers as root set.

diff --git a/Engine/Source/Runtime/CoreUObject/Private/UObject/LinkerLoad.cpp b/Engine/Source/Runtime/CoreUObject/Private/UObject/LinkerLoad.cpp
--- a/Engine/Source/Runtime/CoreUObject/Private/UObject/LinkerLoad.cpp
+++ b/Engine/Source/Runtime/CoreUObject/Private/UObject/LinkerLoad.cpp
@@ -3545,7 +3545,7 @@ UObject* ULinkerLoad::CreateExport( int32 Index )
 			LoadClass,
 			ThisParent,
 			NewName,
-			EObjectFlags(ObjectLoadFlags | (GIsInitialLoad ? RF_RootSet : 0)),
+			EObjectFlags(ObjectLoadFlags | ((FPlatformProperties::RequiresCookedData() && GIsInitialLoad) ? RF_RootSet : 0)),
 			Template
 		);
 		LoadClass = Export.Object->GetClass(); // this may have changed if we are overwriting a CDO component
diff --git a/Engine/Source/Runtime/CoreUObject/Public/UObject/ConstructorHelpers.h b/Engine/Source/Runtime/CoreUObject/Public/UObject/ConstructorHelpers.h
--- a/Engine/Source/Runtime/CoreUObject/Public/UObject/ConstructorHelpers.h
+++ b/Engine/Source/Runtime/CoreUObject/Public/UObject/ConstructorHelpers.h
@@ -29,6 +29,10 @@ namespace ConstructorHelpersInternal
 		UClass* Class = T::StaticClass();
 		Class->GetDefaultObject(); // force the CDO to be created if it hasn't already
 		T* ObjectPtr = LoadObject<T>(NULL, *PathName);
+		if (ObjectPtr)
+		{
+			ObjectPtr->AddToRoot();
+		}
 		return ObjectPtr;
 	}
 
@@ -50,6 +54,10 @@ namespace ConstructorHelpersInternal
 			// If it is not in memory, try to load it.
 			PackagePtr = LoadPackage( nullptr, *PathName, LOAD_None );
 		}
+		if (PackagePtr)
+		{
+			PackagePtr->AddToRoot();
+		}
 		return PackagePtr;
 	}
 
