VerifyImportInner will no longer use a temp FMessageLog object to report errors.

diff --git a/Engine/Source/Runtime/CoreUObject/Private/UObject/LinkerLoad.cpp b/Engine/Source/Runtime/CoreUObject/Private/UObject/LinkerLoad.cpp
--- a/Engine/Source/Runtime/CoreUObject/Private/UObject/LinkerLoad.cpp
+++ b/Engine/Source/Runtime/CoreUObject/Private/UObject/LinkerLoad.cpp
@@ -2124,8 +2124,9 @@ ULinkerLoad::EVerifyResult ULinkerLoad::VerifyImport(int32 ImportIndex)
 				UClass* FindClass = ClassPackage ? FindObject<UClass>( ClassPackage, *OriginalImport.ClassName.ToString() ) : NULL;
 				if( GIsEditor && !IsRunningCommandlet() )
 				{
+					FMessageLog LoadErrors(NAME_LoadErrors);
 					// put something into the load warnings dialog, with any extra information from above (in WarningAppend)
-					TSharedRef<FTokenizedMessage> TokenizedMessage = FMessageLog(NAME_LoadErrors).Error(FText::Format(LOCTEXT("ImportFailure", "Failed import for {ImportClass}"), FText::FromName(GetImportClassName(ImportIndex))));
+					TSharedRef<FTokenizedMessage> TokenizedMessage = LoadErrors.Error(FText::Format(LOCTEXT("ImportFailure", "Failed import for {ImportClass}"), FText::FromName(GetImportClassName(ImportIndex))));
 					TokenizedMessage->AddToken(FAssetNameToken::Create(GetImportPathName(ImportIndex)));
 
 					if (!WarningAppend.IsEmpty())
