Removed "FriendlyName" metadata tag and updated comments and usage of "DisplayName" metadata tag.

diff --git a/Engine/Source/Runtime/CoreUObject/Public/UObject/ObjectBase.h b/Engine/Source/Runtime/CoreUObject/Public/UObject/ObjectBase.h
--- a/Engine/Source/Runtime/CoreUObject/Public/UObject/ObjectBase.h
+++ b/Engine/Source/Runtime/CoreUObject/Public/UObject/ObjectBase.h
@@ -925,7 +925,7 @@ namespace UM
 		/// [ClassMetadata] [FunctionMetadata] Used in conjunction with DeprecatedNode or DeprecatedFunction to customize the warning message displayed to the user.
 		DeprecationMessage,
 
-		/// [ClassMetadata] [PropertyMetadata] The name to display for this class or property instead of auto-generating it from the name.
+		/// [ClassMetadata] [PropertyMetadata] [FunctionMetadata] The name to display for this class, property, or function instead of auto-generating it from the name.
 		DisplayName,
 
 		/// [ClassMetadata] Specifies that this class is an acceptable base class for creating blueprints.
@@ -976,7 +976,7 @@ namespace UM
 		/// [PropertyMetadata] Used for float and integer properties.  Specifies the maximum value that may be entered for the property.
 		ClampMax,
 
-		/// [ClassMetadata] [PropertyMetadata] The name to display for this class or property instead of auto-generating it from the name.
+		/// [ClassMetadata] [PropertyMetadata] [FunctionMetadata] The name to display for this class, property, or function instead of auto-generating it from the name.
 		// DisplayName, (Commented out so as to avoid duplicate name with version in the Class section, but still show in the property section)
 
 		/// [PropertyMetadata] Indicates that the property is an asset type and it should display the thumbnail of the selected asset.
@@ -1127,8 +1127,8 @@ namespace UM
 		/// [FunctionMetadata] For BlueprintCallable functions indicates that an input exec pin should be created for each entry in the enum specified.
 		ExpandEnumAsExecs,
 
-		/// [FunctionMetadata] The name to display for node's of this function in the blueprints.
-		FriendlyName,
+		/// [ClassMetadata] [PropertyMetadata] [FunctionMetadata] The name to display for this class, property, or function instead of auto-generating it from the name.
+		// DisplayName, (Commented out so as to avoid duplicate name with version in the Class section, but still show in the function section)
 
 		/// [FunctionMetadata] For BlueprintCallable functions indicates that the parameter pin should be hidden from the user's view.
 		HidePin,
