UE-11733 Using a function to set button style causes a crash

diff --git a/Engine/Source/Editor/KismetCompiler/Private/KismetCompiler.cpp b/Engine/Source/Editor/KismetCompiler/Private/KismetCompiler.cpp
--- a/Engine/Source/Editor/KismetCompiler/Private/KismetCompiler.cpp
+++ b/Engine/Source/Editor/KismetCompiler/Private/KismetCompiler.cpp
@@ -1586,14 +1586,16 @@ void FKismetCompilerContext::FinishCompilingFunction(FKismetFunctionContext& Con
 				Function->ReturnValueOffset = Property->GetOffset_ForUFunction();
 			}
 		}
-		else
+
+		if (!Property->HasAnyPropertyFlags(CPF_ZeroConstructor) && !Function->FirstPropertyToInit)
 		{
-			if (!Property->HasAnyPropertyFlags(CPF_ZeroConstructor))
-			{
-				Function->FirstPropertyToInit = Property;
-				Function->FunctionFlags |= FUNC_HasDefaults;
-				break;
-			}
+			Function->FirstPropertyToInit = Property;
+			Function->FunctionFlags |= FUNC_HasDefaults;
+		}
+
+		if (!Property->HasAnyPropertyFlags(CPF_Parm) && Function->FirstPropertyToInit)
+		{
+			break;
 		}
 	}
 
diff --git a/Engine/Source/Runtime/CoreUObject/Private/UObject/Class.cpp b/Engine/Source/Runtime/CoreUObject/Private/UObject/Class.cpp
--- a/Engine/Source/Runtime/CoreUObject/Private/UObject/Class.cpp
+++ b/Engine/Source/Runtime/CoreUObject/Private/UObject/Class.cpp
@@ -3747,6 +3747,7 @@ void UFunction::InitializeDerivedMembers()
 	ParmsSize = 0;
 	ReturnValueOffset = MAX_uint16;
 
+	const bool bRequiresInitialization = (FunctionFlags & FUNC_HasDefaults) != 0;
 	for (UProperty* Property = dynamic_cast<UProperty*>(Children); Property; Property = dynamic_cast<UProperty*>(Property->Next))
 	{
 		if (Property->PropertyFlags & CPF_Parm)
@@ -3758,15 +3759,15 @@ void UFunction::InitializeDerivedMembers()
 				ReturnValueOffset = Property->GetOffset_ForUFunction();
 			}
 		}
-		else if ((FunctionFlags & FUNC_HasDefaults) != 0)
+		
+		if (bRequiresInitialization
+			&& !Property->HasAnyPropertyFlags(CPF_ZeroConstructor)
+			&& !FirstPropertyToInit)
 		{
-			if (!Property->HasAnyPropertyFlags(CPF_ZeroConstructor))
-			{
-				FirstPropertyToInit = Property;
-				break;
-			}
+			FirstPropertyToInit = Property;
 		}
-		else
+
+		if (!Property->HasAnyPropertyFlags(CPF_Parm) && (FirstPropertyToInit || !bRequiresInitialization))
 		{
 			break;
 		}
