Changes from PVS-Studio, these could introduce small behavior change.
Consistent null checks, fixes for copy/pasted conditions, copy pasted LOCTEXT identifiers, comparisons against literals of the wrong type, missing return statements, dead code, and some system calls that were ignoring their return value

diff --git a/Engine/Source/Developer/MeshUtilities/Private/SkeletalMeshTools.cpp b/Engine/Source/Developer/MeshUtilities/Private/SkeletalMeshTools.cpp
--- a/Engine/Source/Developer/MeshUtilities/Private/SkeletalMeshTools.cpp
+++ b/Engine/Source/Developer/MeshUtilities/Private/SkeletalMeshTools.cpp
@@ -492,7 +492,7 @@ namespace SkeletalMeshTools
 				{
 					for(int32 j=0; j<MAX_TOTAL_INFLUENCES; j++)
 					{
-						if(SoftVert->InfluenceWeights[j] > 0.01f)
+						if(SoftVert->InfluenceWeights[j] > 0)
 						{
 							int32 BoneIndex = Chunk.BoneMap[SoftVert->InfluenceBones[j]];
 
diff --git a/Engine/Source/Editor/BehaviorTreeEditor/Private/BehaviorTreeDebugger.cpp b/Engine/Source/Editor/BehaviorTreeEditor/Private/BehaviorTreeDebugger.cpp
--- a/Engine/Source/Editor/BehaviorTreeEditor/Private/BehaviorTreeDebugger.cpp
+++ b/Engine/Source/Editor/BehaviorTreeEditor/Private/BehaviorTreeDebugger.cpp
@@ -411,11 +411,11 @@ void FBehaviorTreeDebugger::SetCompositeDecoratorFlags(const struct FBehaviorTre
 				Node->bDebuggerMarkSearchFailedTrigger = SearchStep.bDiscardedTrigger;
 				bTriggerOnly = true;
 			}
-			else if (bMatchesNodeIndex)
-			{
-				SearchPathIdx = i;
-				bTriggerOnly = false;
-			}
+		}
+		else if (bMatchesNodeIndex)
+		{
+			SearchPathIdx = i;
+			bTriggerOnly = false;
 		}
 	}
 
diff --git a/Engine/Source/Editor/GameplayAbilitiesEditor/Private/K2Node_LatentAbilityCall.cpp b/Engine/Source/Editor/GameplayAbilitiesEditor/Private/K2Node_LatentAbilityCall.cpp
--- a/Engine/Source/Editor/GameplayAbilitiesEditor/Private/K2Node_LatentAbilityCall.cpp
+++ b/Engine/Source/Editor/GameplayAbilitiesEditor/Private/K2Node_LatentAbilityCall.cpp
@@ -658,7 +658,7 @@ void UK2Node_LatentAbilityCall::ExpandNode(class FKismetCompilerContext& Compile
 		}
 		else
 		{
-			CompilerContext.MessageLog.Error(*LOCTEXT("MissingBeginSpawningActorArrayFunction", "AbilityTask: Proxy is missing BeginSpawningActorArray native function. @@").ToString(), this);
+			CompilerContext.MessageLog.Error(*LOCTEXT("MissingBeginSpawningActorFunction", "AbilityTask: Proxy is missing BeginSpawningActor native function. @@").ToString(), this);
 		}
 		return;
 	}
diff --git a/Engine/Source/Programs/UnrealHeaderTool/Private/HeaderParser.cpp b/Engine/Source/Programs/UnrealHeaderTool/Private/HeaderParser.cpp
--- a/Engine/Source/Programs/UnrealHeaderTool/Private/HeaderParser.cpp
+++ b/Engine/Source/Programs/UnrealHeaderTool/Private/HeaderParser.cpp
@@ -2834,7 +2834,7 @@ void FHeaderParser::GetVarType
 			// Struct keyword in front of a struct is legal, we 'consume' it
 			bUnconsumedStructKeyword = false;
 		}
-		else if (UScriptStruct* Struct = FindObject<UScriptStruct>( ANY_PACKAGE, *IdentifierStripped ))
+		else if ( FindObject<UScriptStruct>( ANY_PACKAGE, *IdentifierStripped ) != nullptr)
 		{
 			bHandledType = true;
 
diff --git a/Engine/Source/Runtime/Core/Private/Windows/WindowsPlatformCrashContext.cpp b/Engine/Source/Runtime/Core/Private/Windows/WindowsPlatformCrashContext.cpp
--- a/Engine/Source/Runtime/Core/Private/Windows/WindowsPlatformCrashContext.cpp
+++ b/Engine/Source/Runtime/Core/Private/Windows/WindowsPlatformCrashContext.cpp
@@ -163,7 +163,7 @@ void SetReportParameters( HREPORT ReportHandle, EXCEPTION_POINTERS* ExceptionInf
 
 	INT_PTR ExceptionOffset = ( char* )( ExceptionInfo->ExceptionRecord->ExceptionAddress ) - ( char* )FaultModuleHandle;
 	CA_SUPPRESS(6066) // The format specifier should probably be something like %tX, but VS 2013 doesn't support 't'.
-		StringCchPrintf( StringBuffer, MAX_SPRINTF, TEXT( "%p" ), ExceptionOffset );
+		StringCchPrintf( StringBuffer, MAX_SPRINTF, TEXT( "%p" ), (void *)ExceptionOffset );
 	Result = WerReportSetParameter( ReportHandle, WER_P7, TEXT( "Exception Offset" ), StringBuffer );
 
 	// Use LocalBuffer to store the error message.
diff --git a/Engine/Source/Runtime/Core/Public/HAL/ThreadingBase.h b/Engine/Source/Runtime/Core/Public/HAL/ThreadingBase.h
--- a/Engine/Source/Runtime/Core/Public/HAL/ThreadingBase.h
+++ b/Engine/Source/Runtime/Core/Public/HAL/ThreadingBase.h
@@ -881,7 +881,7 @@ private:
 	FScopeLock();
 
 	/** Copy constructor( hidden on purpose). */
-	FScopeLock( FScopeLock* InScopeLock);
+	FScopeLock(const FScopeLock& InScopeLock);
 
 	/** Assignment operator (hidden on purpose). */
 	FScopeLock& operator=( FScopeLock& InScopeLock )
diff --git a/Engine/Source/Runtime/Core/Public/Misc/Variant.h b/Engine/Source/Runtime/Core/Public/Misc/Variant.h
--- a/Engine/Source/Runtime/Core/Public/Misc/Variant.h
+++ b/Engine/Source/Runtime/Core/Public/Misc/Variant.h
@@ -71,6 +71,7 @@ template<typename T> struct TVariantTraits
 	static int32 GetType()
 	{
 		static_assert(!sizeof(T), "Variant trait must be specialized for this type.");
+		return EVariantTypes::Empty;
 	}
 };
 
diff --git a/Engine/Source/Runtime/CoreUObject/Private/Serialization/AsyncLoading.cpp b/Engine/Source/Runtime/CoreUObject/Private/Serialization/AsyncLoading.cpp
--- a/Engine/Source/Runtime/CoreUObject/Private/Serialization/AsyncLoading.cpp
+++ b/Engine/Source/Runtime/CoreUObject/Private/Serialization/AsyncLoading.cpp
@@ -948,7 +948,7 @@ EAsyncPackageState::Type FAsyncPackage::FinishObjects()
 	if (!bLoadHasFailed)
 	{
 		// Mark package as having been fully loaded and update load time.
-		if( Linker->LinkerRoot )
+		if (Linker && Linker->LinkerRoot)
 		{
 			Linker->LinkerRoot->MarkAsFullyLoaded();
 			Linker->LinkerRoot->SetLoadTime( FPlatformTime::Seconds() - LoadStartTime );
@@ -964,8 +964,11 @@ EAsyncPackageState::Type FAsyncPackage::FinishObjects()
 		FIOSystem::Get().HintDoneWithFile(Linker->Filename);
 
 	#if WITH_ENGINE
-		// Cancel all texture allocations that haven't been claimed yet.
-		Linker->Summary.TextureAllocations.CancelRemainingAllocations( true );
+		if (Linker != nullptr)
+		{
+			// Cancel all texture allocations that haven't been claimed yet.
+			Linker->Summary.TextureAllocations.CancelRemainingAllocations(true);
+		}
 	#endif		// WITH_ENGINE
 
 		// Flush cache to free memory
diff --git a/Engine/Source/Runtime/CoreUObject/Private/Serialization/ReloadObjectArc.cpp b/Engine/Source/Runtime/CoreUObject/Private/Serialization/ReloadObjectArc.cpp
--- a/Engine/Source/Runtime/CoreUObject/Private/Serialization/ReloadObjectArc.cpp
+++ b/Engine/Source/Runtime/CoreUObject/Private/Serialization/ReloadObjectArc.cpp
@@ -85,12 +85,15 @@ void FReloadObjectArc::SerializeObject( UObject* Obj )
 			// set this to prevent recursion in serialization
 			if ( IsLoading() )
 			{
-				// InitProperties will call CopyCompleteValue for any instanced object properties, so disable object instancing
-				// because we probably already have an object that will be copied over that value; for any instanced object properties which
-				// did not have a value when we were saving object data, we'll call InstanceSubobjects to instance those.  Also disable component
-				// instancing as this object may not have the correct values for its component properties until its serialized, which result in
-				// possibly creating lots of unnecessary components that will be thrown away anyway
-				InstanceGraph->EnableSubobjectInstancing(false);
+				if (InstanceGraph != NULL)
+				{
+					// InitProperties will call CopyCompleteValue for any instanced object properties, so disable object instancing
+					// because we probably already have an object that will be copied over that value; for any instanced object properties which
+					// did not have a value when we were saving object data, we'll call InstanceSubobjects to instance those.  Also disable component
+					// instancing as this object may not have the correct values for its component properties until its serialized, which result in
+					// possibly creating lots of unnecessary components that will be thrown away anyway
+					InstanceGraph->EnableSubobjectInstancing(false);
+				}
 				if ( Obj->GetClass() != UPackage::StaticClass() )
 				{
 					Obj->ReinitializeProperties(NULL, InstanceGraph);
diff --git a/Engine/Source/Runtime/Slate/Public/Widgets/Views/SListView.h b/Engine/Source/Runtime/Slate/Public/Widgets/Views/SListView.h
--- a/Engine/Source/Runtime/Slate/Public/Widgets/Views/SListView.h
+++ b/Engine/Source/Runtime/Slate/Public/Widgets/Views/SListView.h
@@ -43,17 +43,6 @@ public:
 	typedef typename TSlateDelegates< ItemType >::FOnMouseButtonDoubleClick FOnMouseButtonDoubleClick;
 
 public:
-
-	class FColumnHeaderSlot
-	{
-	public:
-		FColumnHeaderSlot& operator[]( const TSharedRef< SHeaderRow >& InColumnHeaders )
-		{
-			HeaderRow = InColumnHeaders;
-		}			
-		TSharedPtr<SHeaderRow> HeaderRow;
-	};
-
 	SLATE_BEGIN_ARGS( SListView<ItemType> )
 		: _OnGenerateRow()
 		, _ListItemsSource( static_cast<const TArray<ItemType>*>(nullptr) ) //@todo Slate Syntax: Initializing from nullptr without a cast
diff --git a/Engine/Source/Runtime/Windows/XAudio2/Private/XAudio2Device.cpp b/Engine/Source/Runtime/Windows/XAudio2/Private/XAudio2Device.cpp
--- a/Engine/Source/Runtime/Windows/XAudio2/Private/XAudio2Device.cpp
+++ b/Engine/Source/Runtime/Windows/XAudio2/Private/XAudio2Device.cpp
@@ -82,7 +82,18 @@ bool FXAudio2Device::InitializeHardware()
 #if PLATFORM_64BITS
 	// Work around the fact the x64 version of XAudio2_7.dll does not properly ref count
 	// by forcing it to be always loaded
-	LoadLibraryA( "XAudio2_7.dll" );
+
+	// Load the xaudio2 library and keep a handle so we can free it on teardown
+	// Note: windows internally ref-counts the library per call to load library so 
+	// when we call FreeLibrary, it will only free it once the refcount is zero
+	DeviceProperties->XAudio2Dll = LoadLibraryA("XAudio2_7.dll");
+
+	// returning null means we failed to load XAudio2, which means everything will fail
+	if (DeviceProperties->XAudio2Dll == nullptr)
+	{
+		UE_LOG(LogInit, Warning, TEXT("Failed to load XAudio2 dll"));
+		return false;
+	}
 #endif	//PLATFORM_64BITS
 #endif	//PLATFORM_WINDOWS
 
@@ -206,11 +217,21 @@ void FXAudio2Device::TeardownHardware()
 			DeviceProperties->XAudio2->Release();
 			DeviceProperties->XAudio2 = NULL;
 		}
+
+#if PLATFORM_WINDOWS && PLATFORM_64BITS
+		if (DeviceProperties->XAudio2Dll)
+		{
+			if (!FreeLibrary(DeviceProperties->XAudio2Dll))
+			{
+				UE_LOG(LogAudio, Warning, TEXT("Failed to free XAudio2 Dll"));
+			}
+		}
+#endif
+
 		delete DeviceProperties;
 		DeviceProperties = nullptr;
 	}
 
-
 #if PLATFORM_WINDOWS
 	if (bComInitialized)
 	{
diff --git a/Engine/Source/Runtime/Windows/XAudio2/Private/XAudio2Support.h b/Engine/Source/Runtime/Windows/XAudio2/Private/XAudio2Support.h
--- a/Engine/Source/Runtime/Windows/XAudio2/Private/XAudio2Support.h
+++ b/Engine/Source/Runtime/Windows/XAudio2/Private/XAudio2Support.h
@@ -520,6 +520,7 @@ struct FXAudioDeviceProperties
 	// These variables are non-static to support multiple audio device instances
 	struct IXAudio2*					XAudio2;
 	struct IXAudio2MasteringVoice*		MasteringVoice;
+	HMODULE								XAudio2Dll;
 
 	// These variables are static because they are common across all audio device instances
 	static int32						NumSpeakers;
@@ -527,6 +528,13 @@ struct FXAudioDeviceProperties
 #if XAUDIO_SUPPORTS_DEVICE_DETAILS
 	static XAUDIO2_DEVICE_DETAILS		DeviceDetails;
 #endif	//XAUDIO_SUPPORTS_DEVICE_DETAILS
+
+	FXAudioDeviceProperties()
+		: XAudio2(nullptr)
+		, MasteringVoice(nullptr)
+		, XAudio2Dll(nullptr)
+	{
+	}
 };
 
 #if XAUDIO_SUPPORTS_DEVICE_DETAILS
