Core: Corrected comment for MakeEditWidget to replace FRotator with FTransform reference and add that it works on TArrays of supported types too

diff --git a/Engine/Source/Runtime/CoreUObject/Public/UObject/ObjectBase.h b/Engine/Source/Runtime/CoreUObject/Public/UObject/ObjectBase.h
--- a/Engine/Source/Runtime/CoreUObject/Public/UObject/ObjectBase.h
+++ b/Engine/Source/Runtime/CoreUObject/Public/UObject/ObjectBase.h
@@ -1006,7 +1006,7 @@ namespace UM
 		/// [PropertyMetadata] Used for Subclass properties. Indicates whether only placeable classes should be shown in the class picker.
 		OnlyPlaceable,
 
-		/// [PropertyMetadata] Used for Vector/Rotator properties. Indicates that the property should be exposed in the viewport as a movable widget.
+		/// [PropertyMetadata] Used for Transform/Rotator properties (also works on arrays of them). Indicates that the property should be exposed in the viewport as a movable widget.
 		MakeEditWidget,
 
 		/// [PropertyMetadata] For properties in a structure indicates the default value of the property in a blueprint make structure node.
