Stats - Fixed stats metadata issue

diff --git a/Engine/Source/Runtime/CoreUObject/Private/Serialization/BulkData.cpp b/Engine/Source/Runtime/CoreUObject/Private/Serialization/BulkData.cpp
--- a/Engine/Source/Runtime/CoreUObject/Private/Serialization/BulkData.cpp
+++ b/Engine/Source/Runtime/CoreUObject/Private/Serialization/BulkData.cpp
@@ -638,7 +638,7 @@ bool FUntypedBulkData::ShouldFreeOnEmpty() const
 
 void FUntypedBulkData::StartSerializingBulkData(FArchive& Ar, UObject* Owner, int32 Idx, bool bPayloadInline)
 {
-	DECLARE_SCOPE_CYCLE_COUNTER(TEXT("FUntypedBulkData::StartSerializingBulkData"), STAT_UBD_Serialize, STATGROUP_Memory);
+	DECLARE_SCOPE_CYCLE_COUNTER(TEXT("FUntypedBulkData::StartSerializingBulkData"), STAT_UBD_StartSerializingBulkData, STATGROUP_Memory);
 	check(SerializeFuture.IsValid() == false);	
 
 	// Async
