Core - Removed stats from UObjectBaseUtility::IsA

diff --git a/Engine/Source/Runtime/CoreUObject/Private/UObject/UObjectBaseUtility.cpp b/Engine/Source/Runtime/CoreUObject/Private/UObject/UObjectBaseUtility.cpp
--- a/Engine/Source/Runtime/CoreUObject/Private/UObject/UObjectBaseUtility.cpp
+++ b/Engine/Source/Runtime/CoreUObject/Private/UObject/UObjectBaseUtility.cpp
@@ -251,7 +251,6 @@ bool UObjectBaseUtility::RootPackageHasAnyFlags( uint32 CheckFlagMask ) const
  */
 bool UObjectBaseUtility::IsA( const UClass* SomeBase ) const
 {
-	SCOPE_CYCLE_COUNTER(STAT_UObjectIsA);
 	UE_CLOG(!SomeBase, LogObj, Fatal, TEXT("IsA(NULL) cannot yield meaningful results"));
 
 	for ( UClass* TempClass=GetClass(); TempClass; TempClass=TempClass->GetSuperClass() )
