Core - Removed stats from UObjectBaseUtility::IsA

diff --git a/Engine/Source/Runtime/CoreUObject/Private/UObject/UObjectBaseUtility.cpp b/Engine/Source/Runtime/CoreUObject/Private/UObject/UObjectBaseUtility.cpp
--- a/Engine/Source/Runtime/CoreUObject/Private/UObject/UObjectBaseUtility.cpp
+++ b/Engine/Source/Runtime/CoreUObject/Private/UObject/UObjectBaseUtility.cpp
@@ -7,7 +7,6 @@
 #include "CoreUObjectPrivate.h"
 #include "Interface.h"
 
-DEFINE_STAT(STAT_UObjectIsA);
 
 /***********************/
 /******** Names ********/
diff --git a/Engine/Source/Runtime/CoreUObject/Public/UObject/UObjectBaseUtility.h b/Engine/Source/Runtime/CoreUObject/Public/UObject/UObjectBaseUtility.h
--- a/Engine/Source/Runtime/CoreUObject/Public/UObject/UObjectBaseUtility.h
+++ b/Engine/Source/Runtime/CoreUObject/Public/UObject/UObjectBaseUtility.h
@@ -6,8 +6,6 @@
 
 #pragma once
 
-DECLARE_CYCLE_STAT_EXTERN(TEXT("UObjectBaseUtility::IsA"), STAT_UObjectIsA, STATGROUP_Engine, COREUOBJECT_API);
-
 class COREUOBJECT_API UObjectBaseUtility : public UObjectBase
 {
 public:
