[UEBP-7] some minor TFindObjectReferencers optimizations

change summary:
- call Reset() on local collection arrays rather than re-create them each time through the FObjectIterator loop
- removed an unnecessary dynamic_cast operation

diff --git a/Engine/Source/Runtime/CoreUObject/Public/Serialization/ArchiveUObject.h b/Engine/Source/Runtime/CoreUObject/Public/Serialization/ArchiveUObject.h
--- a/Engine/Source/Runtime/CoreUObject/Public/Serialization/ArchiveUObject.h
+++ b/Engine/Source/Runtime/CoreUObject/Public/Serialization/ArchiveUObject.h
@@ -495,23 +495,26 @@ public:
 	TFindObjectReferencers( TArray< T* > TargetObjects, UPackage* PackageToCheck=NULL, bool bIgnoreTemplates = true )
 	: TMultiMap< T*, UObject* >()
 	{
+		TArray<UObject*> ReferencedObjects;
+		TMap<UObject*, int32> ReferenceCounts;
+
 		FFindReferencersArchive FindReferencerAr(nullptr, ( TArray<UObject*>& )TargetObjects);
 
 		// Loop over every object to find any reference that may exist for the target objects
 		for (FObjectIterator It; It; ++It)
 		{
 			UObject* PotentialReferencer = *It;
-			if ( !TargetObjects.Contains(dynamic_cast<T*>( PotentialReferencer ))
+			if ( !TargetObjects.Contains(PotentialReferencer)
 			&&	(PackageToCheck == NULL || PotentialReferencer->IsIn(PackageToCheck))
 			&&	(!bIgnoreTemplates || !PotentialReferencer->IsTemplate()) )
 			{
 				FindReferencerAr.ResetPotentialReferencer(PotentialReferencer);
 
-				TMap<UObject*, int32> ReferenceCounts;
+				ReferenceCounts.Reset();
 				if ( FindReferencerAr.GetReferenceCounts(ReferenceCounts) > 0 )
 				{
 					// here we don't really care about the number of references from PotentialReferencer to the target object...just that he's a referencer
-					TArray<UObject*> ReferencedObjects;
+					ReferencedObjects.Reset();
 					ReferenceCounts.GenerateKeyArray(ReferencedObjects);
 					for ( int32 RefIndex = 0; RefIndex < ReferencedObjects.Num(); RefIndex++ )
 					{
