Fix for reaching backend of some test content with USE_DEFERRED_DEPENDENCY_CHECK_VERIFICATION_TESTS enabled

diff --git a/Engine/Source/Runtime/CoreUObject/Private/Blueprint/BlueprintSupport.cpp b/Engine/Source/Runtime/CoreUObject/Private/Blueprint/BlueprintSupport.cpp
--- a/Engine/Source/Runtime/CoreUObject/Private/Blueprint/BlueprintSupport.cpp
+++ b/Engine/Source/Runtime/CoreUObject/Private/Blueprint/BlueprintSupport.cpp
@@ -341,7 +341,10 @@ bool ULinkerLoad::DeferPotentialCircularImport(const int32 Index)
 			{
 				UObject* PlaceholderOuter = LinkerRoot;
 				FString  PlaceholderNamePrefix = TEXT("PLACEHOLDER_");
-				UClass*  PlaceholderType  = nullptr;
+				UClass*  PlaceholderType = nullptr;
+#if USE_DEFERRED_DEPENDENCY_CHECK_VERIFICATION_TESTS
+				ULinkerPlaceholderClass* Outer = nullptr;
+#endif //USE_DEFERRED_DEPENDENCY_CHECK_VERIFICATION_TESTS
 				FLinkerPlaceholderBase* (*PlaceholderCastFunc)(UObject* RawUObject) = nullptr;
 
 				bool const bIsBlueprintClass = ImportClass->IsChildOf<UClass>();
@@ -370,6 +373,10 @@ bool ULinkerLoad::DeferPotentialCircularImport(const int32 Index)
 							PlaceholderOuter = ImportMap[OuterImportIndex].XObject;
 							PlaceholderType  = ULinkerPlaceholderFunction::StaticClass();
 							PlaceholderCastFunc = &PlaceholderCast<ULinkerPlaceholderFunction>;
+#if USE_DEFERRED_DEPENDENCY_CHECK_VERIFICATION_TESTS
+							Outer = dynamic_cast<ULinkerPlaceholderClass*>(PlaceholderOuter);
+							check(Outer);
+#endif //USE_DEFERRED_DEPENDENCY_CHECK_VERIFICATION_TESTS
 						}
 					}
 				}
@@ -396,6 +403,12 @@ bool ULinkerLoad::DeferPotentialCircularImport(const int32 Index)
 					// ClassAddReferencedObjects/ClassConstructor members set)
 					PlaceholderObj->Bind();
 					PlaceholderObj->StaticLink(/*bRelinkExistingProperties =*/true);
+#if USE_DEFERRED_DEPENDENCY_CHECK_VERIFICATION_TESTS
+					if (Outer)
+					{
+						Outer->AddChildObject(PlaceholderObj);
+					}
+#endif //USE_DEFERRED_DEPENDENCY_CHECK_VERIFICATION_TESTS
 
 					Import.XObject = PlaceholderObj;
 				}
diff --git a/Engine/Source/Runtime/CoreUObject/Private/UObject/LinkerLoad.cpp b/Engine/Source/Runtime/CoreUObject/Private/UObject/LinkerLoad.cpp
--- a/Engine/Source/Runtime/CoreUObject/Private/UObject/LinkerLoad.cpp
+++ b/Engine/Source/Runtime/CoreUObject/Private/UObject/LinkerLoad.cpp
@@ -7,6 +7,7 @@
 #include "EngineVersion.h"
 #include "LinkerPlaceholderClass.h"
 #include "LinkerPlaceholderExportObject.h"
+#include "LinkerPlaceholderFunction.h"
 
 #define LOCTEXT_NAMESPACE "LinkerLoad"
 
@@ -3619,7 +3620,15 @@ UObject* ULinkerLoad::CreateExport( int32 Index )
 			{
 				if ( !Export.SuperIndex.IsNull() )
 				{
-					Struct->SetSuperStruct( (UStruct*)IndexToObject( Export.SuperIndex ) );
+					UStruct* SuperStruct = (UStruct*)IndexToObject(Export.SuperIndex);
+					if (ULinkerPlaceholderFunction* Function = Cast<ULinkerPlaceholderFunction>(SuperStruct))
+					{
+						Function->AddDerivedFunction(Struct);
+					}
+					else
+					{
+						Struct->SetSuperStruct((UStruct*)IndexToObject(Export.SuperIndex));
+					}
 				}
 
 				// If it's a class, bind it to C++.
diff --git a/Engine/Source/Runtime/CoreUObject/Private/UObject/LinkerPlaceholderBase.h b/Engine/Source/Runtime/CoreUObject/Private/UObject/LinkerPlaceholderBase.h
--- a/Engine/Source/Runtime/CoreUObject/Private/UObject/LinkerPlaceholderBase.h
+++ b/Engine/Source/Runtime/CoreUObject/Private/UObject/LinkerPlaceholderBase.h
@@ -164,6 +164,19 @@ public:
 	 */
 	void AddReferencingScriptExpr(PlaceholderType** ExpressionPtr);
 
+	/** 
+	 * Records a child placeholder object. Not needed except for validiation purposes,
+	 * since both the child object and this are placeholders:
+	 */
+#if USE_DEFERRED_DEPENDENCY_CHECK_VERIFICATION_TESTS
+	void AddChildObject(UObject* Child);
+#endif //USE_DEFERRED_DEPENDENCY_CHECK_VERIFICATION_TESTS
+
+	/**
+	 * Records a derived function, which will have a reference back to its (placeholder) parent function.
+	 * We need to update the derived function when the parent finishes loading.
+	 */
+	void AddDerivedFunction(UStruct* DerivedFunctionType);
 private:
 	/**
 	 * Iterates through all known ReferencingProperties and replaces references
@@ -188,6 +201,15 @@ private:
 
 	/** Points directly at UObject* refs that were serialized in as part of script bytecode */
 	TSet<PlaceholderType**> ReferencingScriptExpressions;
+
+#if USE_DEFERRED_DEPENDENCY_CHECK_VERIFICATION_TESTS
+	/** References to us that are equally transient, this is used in the case where we make a placeholder 
+		that requires an outer that is also a placeholder. E.g. when we make a placeholder function it will 
+		have a placeholder outer. */
+	TArray<UObject*> ChildObjects;
+#endif // USE_DEFERRED_DEPENDENCY_CHECK_VERIFICATION_TESTS
+
+	TSet<UStruct*> DerivedFunctions;
 };
 
 // Templatized implementation
diff --git a/Engine/Source/Runtime/CoreUObject/Private/UObject/LinkerPlaceholderBase.inl b/Engine/Source/Runtime/CoreUObject/Private/UObject/LinkerPlaceholderBase.inl
--- a/Engine/Source/Runtime/CoreUObject/Private/UObject/LinkerPlaceholderBase.inl
+++ b/Engine/Source/Runtime/CoreUObject/Private/UObject/LinkerPlaceholderBase.inl
@@ -14,6 +14,19 @@ int32 TLinkerImportPlaceholder<PlaceholderType>::ResolveAllPlaceholderReferences
 
 	int32 ReplacementCount = ResolvePropertyReferences(TypeCheckedReplacement);
 	ReplacementCount += ResolveScriptReferences(TypeCheckedReplacement);
+
+#if USE_DEFERRED_DEPENDENCY_CHECK_VERIFICATION_TESTS
+	ReplacementCount += ChildObjects.Num();
+	ChildObjects.Empty();
+#endif//USE_DEFERRED_DEPENDENCY_CHECK_VERIFICATION_TESTS
+
+	ReplacementCount += DerivedFunctions.Num();
+	for (auto Entry : DerivedFunctions)
+	{
+		Entry->SetSuperStruct(TypeCheckedReplacement);
+	}
+	DerivedFunctions.Empty();
+
 	ReplacementCount += FLinkerPlaceholderBase::ResolveAllPlaceholderReferences(ReplacementObj);
 
 	return ReplacementCount;
@@ -23,7 +36,14 @@ int32 TLinkerImportPlaceholder<PlaceholderType>::ResolveAllPlaceholderReferences
 template<class PlaceholderType>
 bool TLinkerImportPlaceholder<PlaceholderType>::HasKnownReferences() const
 {
-	return FLinkerPlaceholderBase::HasKnownReferences() || (ReferencingProperties.Num() > 0) || (ReferencingScriptExpressions.Num() > 0);
+
+	return FLinkerPlaceholderBase::HasKnownReferences() || 
+		(ReferencingProperties.Num() > 0) || 
+		(ReferencingScriptExpressions.Num() > 0) ||
+#if USE_DEFERRED_DEPENDENCY_CHECK_VERIFICATION_TESTS
+		ChildObjects.Num() > 0 || 
+#endif //USE_DEFERRED_DEPENDENCY_CHECK_VERIFICATION_TESTS
+		DerivedFunctions.Num() > 0;
 }
 
 //------------------------------------------------------------------------------
@@ -56,6 +76,22 @@ void TLinkerImportPlaceholder<PlaceholderType>::AddReferencingProperty(UProperty
 	ReferencingProperties.Add(ReferencingProperty);
 }
 
+#if USE_DEFERRED_DEPENDENCY_CHECK_VERIFICATION_TESTS
+//------------------------------------------------------------------------------
+template<class PlaceholderType>
+void TLinkerImportPlaceholder<PlaceholderType>::AddChildObject(UObject* Child)
+{
+	ChildObjects.Push(Child);
+}
+#endif //USE_DEFERRED_DEPENDENCY_CHECK_VERIFICATION_TESTS
+
+//------------------------------------------------------------------------------
+template<class PlaceholderType>
+void TLinkerImportPlaceholder<PlaceholderType>::AddDerivedFunction(UStruct* DerivedFunctionType)
+{
+	DerivedFunctions.Add(DerivedFunctionType);
+}
+
 //------------------------------------------------------------------------------
 template<class PlaceholderType>
 void TLinkerImportPlaceholder<PlaceholderType>::RemoveReferencingProperty(UProperty* ReferencingProperty)
