Back out changelist 2494280
That change introduced a many issues.

diff --git a/Engine/Source/Editor/KismetCompiler/Private/KismetCompiler.cpp b/Engine/Source/Editor/KismetCompiler/Private/KismetCompiler.cpp
--- a/Engine/Source/Editor/KismetCompiler/Private/KismetCompiler.cpp
+++ b/Engine/Source/Editor/KismetCompiler/Private/KismetCompiler.cpp
@@ -1587,16 +1587,14 @@ void FKismetCompilerContext::FinishCompilingFunction(FKismetFunctionContext& Con
 				Function->ReturnValueOffset = Property->GetOffset_ForUFunction();
 			}
 		}
-
-		if (!Property->HasAnyPropertyFlags(CPF_ZeroConstructor) && !Function->FirstPropertyToInit)
-		{
-			Function->FirstPropertyToInit = Property;
-			Function->FunctionFlags |= FUNC_HasDefaults;
-		}
-
-		if (!Property->HasAnyPropertyFlags(CPF_Parm) && Function->FirstPropertyToInit)
+		else
 		{
-			break;
+			if (!Property->HasAnyPropertyFlags(CPF_ZeroConstructor))
+			{
+				Function->FirstPropertyToInit = Property;
+				Function->FunctionFlags |= FUNC_HasDefaults;
+				break;
+			}
 		}
 	}
 
diff --git a/Engine/Source/Runtime/CoreUObject/Private/UObject/Class.cpp b/Engine/Source/Runtime/CoreUObject/Private/UObject/Class.cpp
--- a/Engine/Source/Runtime/CoreUObject/Private/UObject/Class.cpp
+++ b/Engine/Source/Runtime/CoreUObject/Private/UObject/Class.cpp
@@ -3747,7 +3747,6 @@ void UFunction::InitializeDerivedMembers()
 	ParmsSize = 0;
 	ReturnValueOffset = MAX_uint16;
 
-	const bool bRequiresInitialization = (FunctionFlags & FUNC_HasDefaults) != 0;
 	for (UProperty* Property = dynamic_cast<UProperty*>(Children); Property; Property = dynamic_cast<UProperty*>(Property->Next))
 	{
 		if (Property->PropertyFlags & CPF_Parm)
@@ -3759,15 +3758,15 @@ void UFunction::InitializeDerivedMembers()
 				ReturnValueOffset = Property->GetOffset_ForUFunction();
 			}
 		}
-		
-		if (bRequiresInitialization
-			&& !Property->HasAnyPropertyFlags(CPF_ZeroConstructor)
-			&& !FirstPropertyToInit)
+		else if ((FunctionFlags & FUNC_HasDefaults) != 0)
 		{
-			FirstPropertyToInit = Property;
+			if (!Property->HasAnyPropertyFlags(CPF_ZeroConstructor))
+			{
+				FirstPropertyToInit = Property;
+				break;
+			}
 		}
-
-		if (!Property->HasAnyPropertyFlags(CPF_Parm) && (FirstPropertyToInit || !bRequiresInitialization))
+		else
 		{
 			break;
 		}
