[UE-12550] -  Fixed metadata load order, finishing struct loading before loading dependent bluerpints... This fixes the GUIDs (random numbers and letters) being displayed on set/break struct nodes

diff --git a/Engine/Source/Runtime/CoreUObject/Private/Blueprint/BlueprintSupport.cpp b/Engine/Source/Runtime/CoreUObject/Private/Blueprint/BlueprintSupport.cpp
--- a/Engine/Source/Runtime/CoreUObject/Private/Blueprint/BlueprintSupport.cpp
+++ b/Engine/Source/Runtime/CoreUObject/Private/Blueprint/BlueprintSupport.cpp
@@ -702,6 +702,7 @@ private:
 };
 /** A global set that tracks structs currently being ran through (and unfinished by) ULinkerLoad::ResolveDeferredDependencies() */
 TSet<UObject*> FUnresolvedStructTracker::UnresolvedStructs;
+UPackage* LoadPackageInternal(UPackage* InOuter, const TCHAR* InLongPackageName, uint32 LoadFlags, ULinkerLoad* ImportLinker);
 
 void ULinkerLoad::ResolveDeferredDependencies(UStruct* LoadStruct)
 {
@@ -771,6 +772,13 @@ void ULinkerLoad::ResolveDeferredDependencies(UStruct* LoadStruct)
 				// deferred dependencies in the struct 
 				if (Import.SourceLinker != nullptr)
 				{
+					// Make sure meta data is loaded first, so that we have structure field names available to any blueprints that load:
+					if (!Import.SourceLinker->LinkerRoot->HasAnyFlags(RF_WasLoaded))
+					{
+						uint32 InternalLoadFlags = LoadFlags & (LOAD_NoVerify | LOAD_NoWarn | LOAD_Quiet);
+						LoadPackageInternal(nullptr, *(Import.SourceLinker->Filename), InternalLoadFlags, this);
+					}
+
 					Import.SourceLinker->ResolveDeferredDependencies(StructObj);
 				}
 			}
