Better fix for: UE-11733 Using a function to set button style causes a crash

diff --git a/Engine/Source/Runtime/CoreUObject/Private/UObject/ScriptCore.cpp b/Engine/Source/Runtime/CoreUObject/Private/UObject/ScriptCore.cpp
--- a/Engine/Source/Runtime/CoreUObject/Private/UObject/ScriptCore.cpp
+++ b/Engine/Source/Runtime/CoreUObject/Private/UObject/ScriptCore.cpp
@@ -501,6 +501,8 @@ void UObject::CallFunction( FFrame& Stack, RESULT_DECL, UFunction* Function )
 				uint8* Param = Property->ContainerPtrToValuePtr<uint8>(NewStack.Locals);
 				checkSlow(Param);
 
+				Property->InitializeValue_InContainer(NewStack.Locals);
+
 				Stack.Step(Stack.Object, Param);
 			}
 		}
diff --git a/Engine/Source/Runtime/SlateCore/Public/Styling/SlateTypes.h b/Engine/Source/Runtime/SlateCore/Public/Styling/SlateTypes.h
--- a/Engine/Source/Runtime/SlateCore/Public/Styling/SlateTypes.h
+++ b/Engine/Source/Runtime/SlateCore/Public/Styling/SlateTypes.h
@@ -314,6 +314,7 @@ struct TStructOpsTypeTraits<FButtonStyle> : public TStructOpsTypeTraitsBase
 	enum 
 	{
 		WithPostSerialize = true,
+		WithCopy = true,
 	};
 };
 
