[UE-12953] Fixing it so superclass CDO defaults don't stomp fixed-up self (CDO) references.

diff --git a/Engine/Source/Runtime/CoreUObject/Private/Blueprint/BlueprintSupport.cpp b/Engine/Source/Runtime/CoreUObject/Private/Blueprint/BlueprintSupport.cpp
--- a/Engine/Source/Runtime/CoreUObject/Private/Blueprint/BlueprintSupport.cpp
+++ b/Engine/Source/Runtime/CoreUObject/Private/Blueprint/BlueprintSupport.cpp
@@ -1196,17 +1196,21 @@ void ULinkerLoad::ResolveDeferredExports(UClass* LoadClass)
 		// do. Simply moving the call to SerializeDefaultObject next to the call to create the CDO means
 		// the stream order changes...
 		DEFERRED_DEPENDENCY_CHECK(BlueprintCDO->GetClass() == LoadClass);
-		if (BlueprintCDO->HasAnyFlags(RF_NeedLoad) &&
-			!LoadClass->GetSuperClass()->HasAnyFlags(RF_Native))
-		{
-			for (UProperty* P = LoadClass->PropertyLink; P; P = P->PropertyLinkNext)
-			{
-				if (P->IsInContainer(LoadClass->GetSuperClass()))
-				{
-					P->CopyCompleteValue_InContainer(BlueprintCDO, LoadClass->GetSuperClass()->GetDefaultObject(false));
-				}
-			}
-		}
+		// @TODO: Find a different way of resolving UE-11029, this was causing a
+		//        more subtle bug with UE-12953 (it seems that we need to run 
+		//        all of ~FObjectInitializer() together; we can't pick and 
+		//        choose like we tried here)
+// 		if (BlueprintCDO->HasAnyFlags(RF_NeedLoad) &&
+// 			!LoadClass->GetSuperClass()->HasAnyFlags(RF_Native))
+// 		{
+// 			for (UProperty* P = LoadClass->PropertyLink; P; P = P->PropertyLinkNext)
+// 			{
+// 				if (P->IsInContainer(LoadClass->GetSuperClass()))
+// 				{
+// 					P->CopyCompleteValue_InContainer(BlueprintCDO, LoadClass->GetSuperClass()->GetDefaultObject(false));
+// 				}
+// 			}
+// 		}
 
 		// should load the CDO (ensuring that it has been serialized in by the 
 		// time we get to class regeneration)
