UE-12889 Crashes when tries to send PointerEvent thru Event dispatcher in widget

diff --git a/Engine/Source/Runtime/CoreUObject/Private/UObject/ScriptCore.cpp b/Engine/Source/Runtime/CoreUObject/Private/UObject/ScriptCore.cpp
--- a/Engine/Source/Runtime/CoreUObject/Private/UObject/ScriptCore.cpp
+++ b/Engine/Source/Runtime/CoreUObject/Private/UObject/ScriptCore.cpp
@@ -481,6 +481,7 @@ void UObject::CallFunction( FFrame& Stack, RESULT_DECL, UFunction* Function )
 				// warning: Stack.MostRecentPropertyAddress could be NULL for optional out parameters
 				// if that's the case, we use the extra memory allocated for the out param in the function's locals
 				// so there's always a valid address
+				ensure(Stack.MostRecentPropertyAddress); // possible problem - output param values on local stack are neither initialized nor cleaned.
 				Out->PropAddr = (Stack.MostRecentPropertyAddress != NULL) ? Stack.MostRecentPropertyAddress : Property->ContainerPtrToValuePtr<uint8>(NewStack.Locals);
 				Out->Property = Property;
 
@@ -1616,6 +1617,7 @@ public:
 			{
 				uint8* Param = Property->ContainerPtrToValuePtr<uint8>(Parameters);
 				checkSlow(Param);
+				Property->InitializeValue_InContainer(Parameters);
 				Stack.Step(Stack.Object, Param);
 			}
 		}
