Merging UE4-Streaming to UE4

diff --git a/Engine/Source/Runtime/Core/Public/UObject/NameTypes.h b/Engine/Source/Runtime/Core/Public/UObject/NameTypes.h
--- a/Engine/Source/Runtime/Core/Public/UObject/NameTypes.h
+++ b/Engine/Source/Runtime/Core/Public/UObject/NameTypes.h
@@ -299,6 +299,7 @@ class TStaticIndirectArrayThreadSafeRead
 		}
 		check(ChunkIndex < NumChunks && Chunks[ChunkIndex]); // should have a valid pointer now
 	}
+
 	/**
 	 * Return a pointer to the pointer to a given element
 	 * @param Index The Index of an element we want to retrieve the pointer-to-pointer for
@@ -375,6 +376,28 @@ public:
 	{
 		return Chunks;
 	}
+	/**
+	* Make sure chunks are allocated to hold the specified capacity of items. This is NOT thread safe.
+	**/
+	void Reserve(int32 Capacity)
+	{
+		check(Capacity >= 0 && Capacity <= MaxTotalElements);
+		if (Capacity > NumElements)
+		{			
+			int32 MaxChunks = (Capacity + ElementsPerChunk - 1) / ElementsPerChunk;
+			check(MaxChunks >= NumChunks);
+			for (int32 ChunkIndex = 0; ChunkIndex < MaxChunks; ++ChunkIndex)
+			{
+				if (!Chunks[ChunkIndex])
+				{
+					ElementType** NewChunk = (ElementType**)FMemory::Malloc(sizeof(ElementType*) * ElementsPerChunk);
+					FMemory::Memzero(NewChunk, sizeof(ElementType*) * ElementsPerChunk);
+					Chunks[ChunkIndex] = NewChunk;
+				}
+			}
+			NumChunks = MaxChunks;
+		}
+	}
 };
 
 // Typedef for the threadsafe master name table. 
@@ -748,13 +771,13 @@ public:
 	FName( const TCHAR* Name, int32 InNumber, EFindName FindType=FNAME_Add );
 
 	/**
-	 * Constructor used by ULinkerLoad when loading its name table; Creates an FName with an instance
+	 * Constructor used by FLinkerLoad when loading its name table; Creates an FName with an instance
 	 * number of 0 that does not attempt to split the FName into string and number portions.
 	 */
 	FName( ELinkerNameTableConstructor, const WIDECHAR* Name );
 
 	/**
-	 * Constructor used by ULinkerLoad when loading its name table; Creates an FName with an instance
+	 * Constructor used by FLinkerLoad when loading its name table; Creates an FName with an instance
 	 * number of 0 that does not attempt to split the FName into string and number portions.
 	 */
 	FName( ELinkerNameTableConstructor, const ANSICHAR* Name );
diff --git a/Engine/Source/Runtime/Core/Public/UObject/UObjectHierarchyFwd.h b/Engine/Source/Runtime/Core/Public/UObject/UObjectHierarchyFwd.h
--- a/Engine/Source/Runtime/Core/Public/UObject/UObjectHierarchyFwd.h
+++ b/Engine/Source/Runtime/Core/Public/UObject/UObjectHierarchyFwd.h
@@ -38,9 +38,9 @@ class			UStruct;
 class				UFunction;
 class				UClass;
 class				UScriptStruct;
-class		ULinker;
-class			ULinkerLoad;
-class			ULinkerSave;
+class		FLinker;
+class			FLinkerLoad;
+class			FLinkerSave;
 class		UPackage;
 class		USystem;
 class		UTextBuffer;
diff --git a/Engine/Source/Runtime/CoreUObject/Private/Blueprint/BlueprintSupport.cpp b/Engine/Source/Runtime/CoreUObject/Private/Blueprint/BlueprintSupport.cpp
--- a/Engine/Source/Runtime/CoreUObject/Private/Blueprint/BlueprintSupport.cpp
+++ b/Engine/Source/Runtime/CoreUObject/Private/Blueprint/BlueprintSupport.cpp
@@ -7,6 +7,7 @@
 #include "Linker.h"
 #include "PropertyTag.h"
 #include "UObject/StructScriptLoader.h"
+#include "UObject/UObjectThreadContext.h"
 
 /** 
  * Defined in BlueprintSupport.cpp
@@ -131,12 +132,12 @@ FScopedClassDependencyGather::~FScopedClassDependencyGather()
 			UClass* Dependency = *DependencyIter;
 			if( Dependency->ClassGeneratedBy != BatchMasterClass->ClassGeneratedBy )
 			{
-				Dependency->ConditionalRecompileClass(&GObjLoaded);
+				Dependency->ConditionalRecompileClass(&FUObjectThreadContext::Get().ObjLoaded);
 			}
 		}
 
 		// Finally, recompile the master class to make sure it gets updated too
-		BatchMasterClass->ConditionalRecompileClass(&GObjLoaded);
+		BatchMasterClass->ConditionalRecompileClass(&FUObjectThreadContext::Get().ObjLoaded);
 		
 		BatchMasterClass = NULL;
 	}
@@ -151,7 +152,7 @@ TArray<UClass*> const& FScopedClassDependencyGather::GetCachedDependencies()
 
 
 /*******************************************************************************
- * ULinkerLoad
+ * FLinkerLoad
  ******************************************************************************/
 
 // rather than littering the code with USE_DEFERRED_DEPENDENCY_CHECK_VERIFICATION_TESTS
@@ -192,7 +193,7 @@ struct FPreloadMembersHelper
 		if (InObject && !InObject->HasAnyFlags(RF_LoadCompleted))
 		{
 			InObject->SetFlags(RF_NeedLoad);
-			if (ULinkerLoad* Linker = InObject->GetLinker())
+			if (FLinkerLoad* Linker = InObject->GetLinker())
 			{
 				Linker->Preload(InObject);
 			}
@@ -207,7 +208,7 @@ struct FPreloadMembersHelper
  * @param	ExportObject	Current object being exported
  * @return	Returns true if regeneration was successful, otherwise false
  */
-bool ULinkerLoad::RegenerateBlueprintClass(UClass* LoadClass, UObject* ExportObject)
+bool FLinkerLoad::RegenerateBlueprintClass(UClass* LoadClass, UObject* ExportObject)
 {
 	// determine if somewhere further down the callstack, we're already in this
 	// function for this class
@@ -286,7 +287,7 @@ bool ULinkerLoad::RegenerateBlueprintClass(UClass* LoadClass, UObject* ExportObj
 			// Preload the blueprint to make sure it has all the data the class needs for regeneration
 			FPreloadMembersHelper::PreloadObject(BlueprintObject);
 
-			UClass* RegeneratedClass = BlueprintObject->RegenerateClass(LoadClass, CurrentCDO, GObjLoaded);
+			UClass* RegeneratedClass = BlueprintObject->RegenerateClass(LoadClass, CurrentCDO, FUObjectThreadContext::Get().ObjLoaded);
 			if (RegeneratedClass)
 			{
 				BlueprintObject->ClearFlags(RF_BeingRegenerated);
@@ -314,7 +315,7 @@ static FLinkerPlaceholderBase* PlaceholderCast(UObject* PlaceholderObj)
 	return CastChecked<PlaceholderType>(PlaceholderObj, ECastCheckedType::NullAllowed);
 }
 
-bool ULinkerLoad::DeferPotentialCircularImport(const int32 Index)
+bool FLinkerLoad::DeferPotentialCircularImport(const int32 Index)
 {
 #if USE_CIRCULAR_DEPENDENCY_LOAD_DEFERRING
 	if (!FBlueprintSupport::UseDeferredDependencyLoading())
@@ -437,13 +438,13 @@ struct FResolvingExportTracker : TThreadSingleton<FResolvingExportTracker>
 {
 public:
 	/**  */
-	void FlagLinkerExportAsResolving(ULinkerLoad* Linker, int32 ExportIndex)
+	void FlagLinkerExportAsResolving(FLinkerLoad* Linker, int32 ExportIndex)
 	{
 		ResolvingExports.FindOrAdd(Linker).Add(ExportIndex);
 	}
 
 	/**  */
-	bool IsLinkerExportBeingResolved(ULinkerLoad* Linker, int32 ExportIndex) const
+	bool IsLinkerExportBeingResolved(FLinkerLoad* Linker, int32 ExportIndex) const
 	{
 		if (auto* ExportIndices = ResolvingExports.Find(Linker))
 		{
@@ -453,7 +454,7 @@ public:
 	}
 
 	/**  */
-	void FlagExportClassAsFullyResolved(ULinkerLoad* Linker, int32 ExportIndex)
+	void FlagExportClassAsFullyResolved(FLinkerLoad* Linker, int32 ExportIndex)
 	{
 		if (auto* ExportIndices = ResolvingExports.Find(Linker))
 		{
@@ -466,18 +467,18 @@ public:
 	}
 
 #if USE_DEFERRED_DEPENDENCY_CHECK_VERIFICATION_TESTS
-	void FlagFullExportResolvePassComplete(ULinkerLoad* Linker)
+	void FlagFullExportResolvePassComplete(FLinkerLoad* Linker)
 	{
 		FullyResolvedLinkers.Add(Linker);
 	}
 
-	bool HasPerformedFullExportResolvePass(ULinkerLoad* Linker)
+	bool HasPerformedFullExportResolvePass(FLinkerLoad* Linker)
 	{
 		return FullyResolvedLinkers.Contains(Linker);
 	}
 #endif // USE_DEFERRED_DEPENDENCY_CHECK_VERIFICATION_TESTS
 
-	void Reset(ULinkerLoad* Linker)
+	void Reset(FLinkerLoad* Linker)
 	{
 		ResolvingExports.Remove(Linker);
 #if USE_DEFERRED_DEPENDENCY_CHECK_VERIFICATION_TESTS
@@ -487,10 +488,10 @@ public:
 
 private:
 	/**  */
-	TMap< ULinkerLoad*, TSet<int32> > ResolvingExports;
+	TMap< FLinkerLoad*, TSet<int32> > ResolvingExports;
 
 #if USE_DEFERRED_DEPENDENCY_CHECK_VERIFICATION_TESTS
-	TSet<ULinkerLoad*> FullyResolvedLinkers;
+	TSet<FLinkerLoad*> FullyResolvedLinkers;
 #endif // USE_DEFERRED_DEPENDENCY_CHECK_VERIFICATION_TESTS
 };
 
@@ -501,7 +502,7 @@ private:
 struct FScopedResolvingExportTracker
 {
 public: 
-	FScopedResolvingExportTracker(ULinkerLoad* Linker, int32 ExportIndex)
+	FScopedResolvingExportTracker(FLinkerLoad* Linker, int32 ExportIndex)
 		: TrackedLinker(Linker), TrackedExport(ExportIndex)
 	{
 		FResolvingExportTracker::Get().FlagLinkerExportAsResolving(Linker, ExportIndex);
@@ -513,11 +514,11 @@ public:
 	}
 
 private:
-	ULinkerLoad* TrackedLinker;
+	FLinkerLoad* TrackedLinker;
 	int32        TrackedExport;
 };
 
-bool ULinkerLoad::DeferExportCreation(const int32 Index)
+bool FLinkerLoad::DeferExportCreation(const int32 Index)
 {
 	FObjectExport& Export = ExportMap[Index];
 
@@ -541,7 +542,7 @@ bool ULinkerLoad::DeferExportCreation(const int32 Index)
 	ULinkerPlaceholderClass* AsPlaceholderClass = Cast<ULinkerPlaceholderClass>(LoadClass);
 	bool const bIsPlaceholderClass = (AsPlaceholderClass != nullptr);
 
-	ULinkerLoad* ClassLinker = LoadClass->GetLinker();
+	FLinkerLoad* ClassLinker = LoadClass->GetLinker();
 	if ( !bIsPlaceholderClass && ((ClassLinker == nullptr) || !ClassLinker->IsBlueprintFinalizationPending()) )
 	{
 		return false;
@@ -568,7 +569,7 @@ bool ULinkerLoad::DeferExportCreation(const int32 Index)
 	}
 	// we haven't come across a scenario where this happens, but if this hits 
 	// then we're deferring exports that will NEVER be resolved (possibly 
-	// stemming from CDO serialization in ULinkerLoad::ResolveDeferredExports)
+	// stemming from CDO serialization in FLinkerLoad::ResolveDeferredExports)
 	DEFERRED_DEPENDENCY_CHECK(!FResolvingExportTracker::Get().HasPerformedFullExportResolvePass(this));
 	
 	UPackage* PlaceholderOuter = LinkerRoot;
@@ -588,7 +589,7 @@ bool ULinkerLoad::DeferExportCreation(const int32 Index)
 	return true;
 }
 
-int32 ULinkerLoad::FindCDOExportIndex(UClass* LoadClass)
+int32 FLinkerLoad::FindCDOExportIndex(UClass* LoadClass)
 {
 	DEFERRED_DEPENDENCY_CHECK(LoadClass->GetLinker() == this);
 	int32 const ClassExportIndex = LoadClass->GetLinkerIndex();
@@ -658,7 +659,7 @@ public:
 	 * @param  Linker	The linker you want to check.
 	 * @return True if the specified linker is in the midst of an unfinished ResolveDeferredDependencies() call (otherwise false).
 	 */
-	static bool IsAssociatedStructUnresolved(const ULinkerLoad* Linker)
+	static bool IsAssociatedStructUnresolved(const FLinkerLoad* Linker)
 	{
 		for (UObject* UnresolvedObj : UnresolvedStructs)
 		{
@@ -673,7 +674,7 @@ public:
 	}
 
 	/**  */
-	static void Reset(const ULinkerLoad* Linker)
+	static void Reset(const FLinkerLoad* Linker)
 	{
 		TArray<UObject*> ToRemove;
 		for (UObject* UnresolvedObj : UnresolvedStructs)
@@ -700,11 +701,11 @@ private:
 	 */
 	static TSet<UObject*> UnresolvedStructs;
 };
-/** A global set that tracks structs currently being ran through (and unfinished by) ULinkerLoad::ResolveDeferredDependencies() */
+/** A global set that tracks structs currently being ran through (and unfinished by) FLinkerLoad::ResolveDeferredDependencies() */
 TSet<UObject*> FUnresolvedStructTracker::UnresolvedStructs;
-UPackage* LoadPackageInternal(UPackage* InOuter, const TCHAR* InLongPackageName, uint32 LoadFlags, ULinkerLoad* ImportLinker);
+UPackage* LoadPackageInternal(UPackage* InOuter, const TCHAR* InLongPackageName, uint32 LoadFlags, FLinkerLoad* ImportLinker);
 
-void ULinkerLoad::ResolveDeferredDependencies(UStruct* LoadStruct)
+void FLinkerLoad::ResolveDeferredDependencies(UStruct* LoadStruct)
 {
 #if USE_CIRCULAR_DEPENDENCY_LOAD_DEFERRING
 	//--------------------------------------
@@ -786,7 +787,7 @@ void ULinkerLoad::ResolveDeferredDependencies(UStruct* LoadStruct)
 
 		if (UStruct* SuperStruct = LoadStruct->GetSuperStruct())
 		{
-			ULinkerLoad* SuperLinker = SuperStruct->GetLinker();
+			FLinkerLoad* SuperLinker = SuperStruct->GetLinker();
 			// NOTE: there is no harm in calling this when the super is not 
 			//       "actively resolving deferred dependencies"... this condition  
 			//       just saves on wasted ops, looping over the super's ImportMap
@@ -828,12 +829,12 @@ void ULinkerLoad::ResolveDeferredDependencies(UStruct* LoadStruct)
 #endif // USE_CIRCULAR_DEPENDENCY_LOAD_DEFERRING
 }
 
-bool ULinkerLoad::HasUnresolvedDependencies() const
+bool FLinkerLoad::HasUnresolvedDependencies() const
 {
 #if USE_CIRCULAR_DEPENDENCY_LOAD_DEFERRING
 	// checking (ResolvingDeferredPlaceholder != nullptr) is not sufficient, 
 	// because the linker could be in the midst of a nested resolve (for a 
-	// struct, or super... see ULinkerLoad::ResolveDeferredDependencies)
+	// struct, or super... see FLinkerLoad::ResolveDeferredDependencies)
 	bool bIsClassExportUnresolved = FUnresolvedStructTracker::IsAssociatedStructUnresolved(this);
 
 	// (ResolvingDeferredPlaceholder != nullptr) should imply 
@@ -847,7 +848,7 @@ bool ULinkerLoad::HasUnresolvedDependencies() const
 #endif // USE_CIRCULAR_DEPENDENCY_LOAD_DEFERRING
 }
 
-int32 ULinkerLoad::ResolveDependencyPlaceholder(FLinkerPlaceholderBase* PlaceholderIn, UClass* ReferencingClass)
+int32 FLinkerLoad::ResolveDependencyPlaceholder(FLinkerPlaceholderBase* PlaceholderIn, UClass* ReferencingClass)
 {
 #if USE_CIRCULAR_DEPENDENCY_LOAD_DEFERRING
 #if USE_DEFERRED_DEPENDENCY_CHECK_VERIFICATION_TESTS
@@ -942,7 +943,7 @@ int32 ULinkerLoad::ResolveDependencyPlaceholder(FLinkerPlaceholderBase* Placehol
 #endif // USE_CIRCULAR_DEPENDENCY_LOAD_DEFERRING
 }
 
-void ULinkerLoad::FinalizeBlueprint(UClass* LoadClass)
+void FLinkerLoad::FinalizeBlueprint(UClass* LoadClass)
 {
 #if USE_CIRCULAR_DEPENDENCY_LOAD_DEFERRING
 	if (!FBlueprintSupport::UseDeferredDependencyLoading())
@@ -963,7 +964,7 @@ void ULinkerLoad::FinalizeBlueprint(UClass* LoadClass)
 	// (so we can have a properly formed sub-class)
 	if (UClass* SuperClass = LoadClass->GetSuperClass())
 	{
-		ULinkerLoad* SuperLinker = SuperClass->GetLinker();
+		FLinkerLoad* SuperLinker = SuperClass->GetLinker();
 		if ((SuperLinker != nullptr) && SuperLinker->IsBlueprintFinalizationPending())
 		{
 			SuperLinker->FinalizeBlueprint(SuperClass);
@@ -1087,7 +1088,7 @@ void ULinkerLoad::FinalizeBlueprint(UClass* LoadClass)
 #endif // USE_CIRCULAR_DEPENDENCY_LOAD_DEFERRING
 }
 
-void ULinkerLoad::ResolveDeferredExports(UClass* LoadClass)
+void FLinkerLoad::ResolveDeferredExports(UClass* LoadClass)
 {
 #if USE_CIRCULAR_DEPENDENCY_LOAD_DEFERRING
 
@@ -1098,7 +1099,7 @@ void ULinkerLoad::ResolveDeferredExports(UClass* LoadClass)
 	DEFERRED_DEPENDENCY_CHECK(BlueprintCDO != nullptr);
 
 #if USE_DEFERRED_DEPENDENCY_CHECK_VERIFICATION_TESTS
-	static TSet<ULinkerLoad*> ResolvingExportsTracker;
+	static TSet<FLinkerLoad*> ResolvingExportsTracker;
 	// the first half of this function CANNOT be re-entrant (specifically for 
 	// this linker)
 	//
@@ -1125,7 +1126,7 @@ void ULinkerLoad::ResolveDeferredExports(UClass* LoadClass)
 				DEFERRED_DEPENDENCY_CHECK(Export.ClassIndex.IsImport());
 				UClass* ExportClass = GetExportLoadClass(ExportIndex);
 				DEFERRED_DEPENDENCY_CHECK((ExportClass != nullptr) && !ExportClass->HasAnyClassFlags(CLASS_Intrinsic) && ExportClass->HasAnyClassFlags(CLASS_CompiledFromBlueprint));
-				ULinkerLoad* ClassLinker = ExportClass->GetLinker();
+				FLinkerLoad* ClassLinker = ExportClass->GetLinker();
 				DEFERRED_DEPENDENCY_CHECK((ClassLinker != nullptr) && (ClassLinker != this));
 
 				{
@@ -1229,7 +1230,7 @@ void ULinkerLoad::ResolveDeferredExports(UClass* LoadClass)
 #endif // USE_CIRCULAR_DEPENDENCY_LOAD_DEFERRING
 }
 
-bool ULinkerLoad::IsBlueprintFinalizationPending() const
+bool FLinkerLoad::IsBlueprintFinalizationPending() const
 {
 #if USE_CIRCULAR_DEPENDENCY_LOAD_DEFERRING
 	return (DeferredCDOIndex != INDEX_NONE);
@@ -1238,10 +1239,10 @@ bool ULinkerLoad::IsBlueprintFinalizationPending() const
 #endif // USE_CIRCULAR_DEPENDENCY_LOAD_DEFERRING
 }
 
-bool ULinkerLoad::ForceRegenerateClass(UClass* ImportClass)
+bool FLinkerLoad::ForceRegenerateClass(UClass* ImportClass)
 {
 #if USE_CIRCULAR_DEPENDENCY_LOAD_DEFERRING
-	if (ULinkerLoad* ClassLinker = ImportClass->GetLinker())
+	if (FLinkerLoad* ClassLinker = ImportClass->GetLinker())
 	{
 		//
 		// BE VERY CAREFUL with this! if these following statements are called 
@@ -1264,7 +1265,7 @@ bool ULinkerLoad::ForceRegenerateClass(UClass* ImportClass)
 	return false;
 }
 
-bool ULinkerLoad::IsExportBeingResolved(int32 ExportIndex)
+bool FLinkerLoad::IsExportBeingResolved(int32 ExportIndex)
 {
 	FObjectExport& Export = ExportMap[ExportIndex];
 	bool bIsExportClassBeingForceRegened = FResolvingExportTracker::Get().IsLinkerExportBeingResolved(this, ExportIndex);
@@ -1292,7 +1293,7 @@ bool ULinkerLoad::IsExportBeingResolved(int32 ExportIndex)
 	return bIsExportClassBeingForceRegened;
 }
 
-bool ULinkerLoad::HasPerformedFullExportResolvePass()
+bool FLinkerLoad::HasPerformedFullExportResolvePass()
 {
 #if USE_DEFERRED_DEPENDENCY_CHECK_VERIFICATION_TESTS
 	return FResolvingExportTracker::Get().HasPerformedFullExportResolvePass(this);
@@ -1302,7 +1303,7 @@ bool ULinkerLoad::HasPerformedFullExportResolvePass()
 	
 }
 
-void ULinkerLoad::ResetDeferredLoadingState()
+void FLinkerLoad::ResetDeferredLoadingState()
 {
 #if USE_CIRCULAR_DEPENDENCY_LOAD_DEFERRING
 	DeferredCDOIndex = INDEX_NONE;
diff --git a/Engine/Source/Runtime/CoreUObject/Private/Misc/RedirectCollector.cpp b/Engine/Source/Runtime/CoreUObject/Private/Misc/RedirectCollector.cpp
--- a/Engine/Source/Runtime/CoreUObject/Private/Misc/RedirectCollector.cpp
+++ b/Engine/Source/Runtime/CoreUObject/Private/Misc/RedirectCollector.cpp
@@ -3,6 +3,7 @@
 
 // Core includes.
 #include "CoreUObjectPrivate.h"
+#include "UObject/UObjectThreadContext.h"
 
 DEFINE_LOG_CATEGORY_STATIC(LogRedirectors, Log, All);
 
@@ -44,11 +45,11 @@ void FRedirectCollector::OnRedirectorFollowed(const FString& InString, UObject*
 
 void FRedirectCollector::OnStringAssetReferenceLoaded(const FString& InString)
 {
-	extern UObject* GSerializedObject;
+	FUObjectThreadContext& ThreadContext = FUObjectThreadContext::Get();
 	FString ContainingPackage;
-	if (GSerializedObject)
+	if (ThreadContext.SerializedObject)
 	{
-		ULinker* Linker = GSerializedObject->GetLinker();
+		auto Linker = ThreadContext.SerializedObject->GetLinker();
 		if (Linker)
 		{
 			ContainingPackage = Linker->Filename;
diff --git a/Engine/Source/Runtime/CoreUObject/Private/Misc/StartupPackages.cpp b/Engine/Source/Runtime/CoreUObject/Private/Misc/StartupPackages.cpp
--- a/Engine/Source/Runtime/CoreUObject/Private/Misc/StartupPackages.cpp
+++ b/Engine/Source/Runtime/CoreUObject/Private/Misc/StartupPackages.cpp
@@ -47,8 +47,8 @@ static void AsyncPreloadPackageList(const TArray<FString>& PackageNames)
 	// Iterate over all native script packages and preload them.
 	for (int32 PackageIndex=0; PackageIndex<PackageNames.Num(); PackageIndex++)
 	{
-		// let ULinkerLoad class manage preloading
-		ULinkerLoad::AsyncPreloadPackage(*PackageNames[PackageIndex]);
+		// let FLinkerLoad class manage preloading
+		FLinkerLoad::AsyncPreloadPackage(*PackageNames[PackageIndex]);
 	}
 }
 
diff --git a/Engine/Source/Runtime/CoreUObject/Private/Serialization/ArchiveObjectCrc32.cpp b/Engine/Source/Runtime/CoreUObject/Private/Serialization/ArchiveObjectCrc32.cpp
--- a/Engine/Source/Runtime/CoreUObject/Private/Serialization/ArchiveObjectCrc32.cpp
+++ b/Engine/Source/Runtime/CoreUObject/Private/Serialization/ArchiveObjectCrc32.cpp
@@ -40,12 +40,6 @@ FArchive& FArchiveObjectCrc32::operator<<(class UObject*& Object)
 {
 	FArchive& Ar = *this;
 
-	// No linker should be a part of the object's state
-	if (Object && Object->IsA<ULinker>())
-	{
-		Object = NULL;
-	}
-
 	if (!Object || !Object->IsIn(RootObject))
 	{
 		auto UniqueName = GetPathNameSafe(Object);
diff --git a/Engine/Source/Runtime/CoreUObject/Private/Serialization/ArchiveObjectGraph.cpp b/Engine/Source/Runtime/CoreUObject/Private/Serialization/ArchiveObjectGraph.cpp
--- a/Engine/Source/Runtime/CoreUObject/Private/Serialization/ArchiveObjectGraph.cpp
+++ b/Engine/Source/Runtime/CoreUObject/Private/Serialization/ArchiveObjectGraph.cpp
@@ -17,8 +17,8 @@ FArchiveObjectGraph::FArchiveObjectGraph(bool IncludeTransients, EObjectFlags	Ke
 	TArray<UObject*> RootObjects;
 
 	// allocate enough memory for all objects
-	ObjectGraph.Empty(GUObjectArray.GetObjectArrayNum());
-	RootObjects.Empty(GUObjectArray.GetObjectArrayNum());
+	ObjectGraph.Empty(GetUObjectArray().GetObjectArrayNum());
+	RootObjects.Empty(GetUObjectArray().GetObjectArrayNum());
 
 	// search for objects that have the right flags and add them to the list of objects that we're going to start with
 	// all other objects need to be tagged so that we can tell whether they've been serialized or not.
diff --git a/Engine/Source/Runtime/CoreUObject/Private/Serialization/ArchiveShowReferences.cpp b/Engine/Source/Runtime/CoreUObject/Private/Serialization/ArchiveShowReferences.cpp
--- a/Engine/Source/Runtime/CoreUObject/Private/Serialization/ArchiveShowReferences.cpp
+++ b/Engine/Source/Runtime/CoreUObject/Private/Serialization/ArchiveShowReferences.cpp
@@ -33,9 +33,6 @@ FArchiveShowReferences::FArchiveShowReferences( FOutputDevice& inOutputAr, UObje
 	// every object we serialize obviously references our package
 	Found.AddUnique(SourceOuter);
 
-	// every object we serialize obviously references its linker
-	Found.AddUnique(SourceObject->GetLinker());
-
 	// every object we serialize obviously references its class and its class's parent classes
 	for ( UClass* ObjectClass = SourceObject->GetClass(); ObjectClass; ObjectClass = ObjectClass->GetSuperClass() )
 	{
diff --git a/Engine/Source/Runtime/CoreUObject/Private/Serialization/ArchiveTraceRoute.cpp b/Engine/Source/Runtime/CoreUObject/Private/Serialization/ArchiveTraceRoute.cpp
--- a/Engine/Source/Runtime/CoreUObject/Private/Serialization/ArchiveTraceRoute.cpp
+++ b/Engine/Source/Runtime/CoreUObject/Private/Serialization/ArchiveTraceRoute.cpp
@@ -140,8 +140,8 @@ FArchiveTraceRoute::FArchiveTraceRoute( UObject* TargetObject, TMap<UObject*,FTr
 	TSparseArray<UObject*> RootObjects;
 
 	// allocate enough memory for all objects
-	ObjectGraph.Empty(GUObjectArray.GetObjectArrayNum());
-	RootObjects.Empty(GUObjectArray.GetObjectArrayNum() / 2);
+	ObjectGraph.Empty(GetUObjectArray().GetObjectArrayNum());
+	RootObjects.Empty(GetUObjectArray().GetObjectArrayNum() / 2);
 
 	// search for objects that have the right flags and add them to the list of objects that we're going to start with
 	// all other objects need to be tagged so that we can tell whether they've been serialized or not.
diff --git a/Engine/Source/Runtime/CoreUObject/Private/Serialization/AsyncLoading.cpp b/Engine/Source/Runtime/CoreUObject/Private/Serialization/AsyncLoading.cpp
--- a/Engine/Source/Runtime/CoreUObject/Private/Serialization/AsyncLoading.cpp
+++ b/Engine/Source/Runtime/CoreUObject/Private/Serialization/AsyncLoading.cpp
@@ -6,6 +6,10 @@
 
 #include "CoreUObjectPrivate.h"
 #include "Serialization/AsyncLoading.h"
+#include "Serialization/DeferredMessageLog.h"
+#include "Serialization/AsyncPackage.h"
+#include "UObject/UObjectThreadContext.h"
+#include "Serialization/AsyncLoadingThread.h"
 
 /*-----------------------------------------------------------------------------
 	Async loading stats.
@@ -24,17 +28,19 @@ DECLARE_CYCLE_STAT(TEXT("CreateImports AsyncPackage"),STAT_FAsyncPackage_CreateI
 DECLARE_CYCLE_STAT(TEXT("FinishTextureAllocations AsyncPackage"),STAT_FAsyncPackage_FinishTextureAllocations,STATGROUP_AsyncLoad);
 DECLARE_CYCLE_STAT(TEXT("CreateExports AsyncPackage"),STAT_FAsyncPackage_CreateExports,STATGROUP_AsyncLoad);
 DECLARE_CYCLE_STAT(TEXT("PreLoadObjects AsyncPackage"),STAT_FAsyncPackage_PreLoadObjects,STATGROUP_AsyncLoad);
-DECLARE_FLOAT_ACCUMULATOR_STAT(TEXT("PostLoadObjects AsyncPackage Time"), STAT_FAsyncPackage_PostLoadObjectsTime, STATGROUP_AsyncLoad);
 DECLARE_CYCLE_STAT(TEXT("PostLoadObjects AsyncPackage"),STAT_FAsyncPackage_PostLoadObjects,STATGROUP_AsyncLoad);
+DECLARE_FLOAT_ACCUMULATOR_STAT(TEXT("PostLoadObjects AsyncPackage Time"), STAT_FAsyncPackage_PostLoadObjectsTime, STATGROUP_AsyncLoad);
+DECLARE_CYCLE_STAT(TEXT("PostLoadObjects GameThread"), STAT_FAsyncPackage_PostLoadObjectsGameThread, STATGROUP_AsyncLoad);
+DECLARE_FLOAT_ACCUMULATOR_STAT(TEXT("PostLoadObjects GameThread Time"), STAT_FAsyncPackage_PostLoadObjectsGameThreadTime, STATGROUP_AsyncLoad);
 DECLARE_CYCLE_STAT(TEXT("FinishObjects AsyncPackage"),STAT_FAsyncPackage_FinishObjects,STATGROUP_AsyncLoad);
+DECLARE_FLOAT_ACCUMULATOR_STAT(TEXT("CreateAsyncPackagesFromQueue Time"), STAT_FAsyncPackage_CreateAsyncPackagesFromQueueTime, STATGROUP_AsyncLoad);
+DECLARE_FLOAT_ACCUMULATOR_STAT(TEXT("AsyncProcess Time"), STAT_FAsyncPackage_AsyncProcessTime, STATGROUP_AsyncLoad);
 DECLARE_FLOAT_ACCUMULATOR_STAT(TEXT("Flush Async Loading Time"), STAT_FAsyncPackage_FlushAsyncLoadingTime, STATGROUP_AsyncLoad);
 
 DECLARE_CYCLE_STAT(TEXT("Async Loading Time"),STAT_AsyncLoadingTime,STATGROUP_AsyncLoad);
+DECLARE_CYCLE_STAT(TEXT("Async Loading Time Detailed"), STAT_AsyncLoadingTimeDetailed, STATGROUP_AsyncLoad);
 
-
-
-/** Objects that have been constructed during async loading phase.						*/
-static TArray<UObject*>			GObjConstructedDuringAsyncLoading;
+DECLARE_FLOAT_ACCUMULATOR_STAT(TEXT("FindAsyncPackage"), STAT_FAsyncPackage_FindAsyncPackage, STATGROUP_AsyncLoad);
 
 /** Keeps a reference to all objects created during async load until streaming has finished */
 class FAsyncObjectsReferencer : FGCObject
@@ -44,6 +50,9 @@ class FAsyncObjectsReferencer : FGCObject
 
 	/** List of referenced objects */
 	TArray<UObject*> ReferencedObjects;
+	/** Critical section for referenced objects list */
+	FCriticalSection ReferencedObjectsCritical;
+
 public:
 	/** Returns the one and only instance of this object */
 	static FAsyncObjectsReferencer& Get();
@@ -51,17 +60,31 @@ public:
 	virtual void AddReferencedObjects(FReferenceCollector& Collector) override
 	{
 		Collector.AllowEliminatingReferences(false);
-		Collector.AddReferencedObjects(ReferencedObjects);
+		{
+			FScopeLock ReferencedObjectsLock(&ReferencedObjectsCritical);
+			Collector.AddReferencedObjects(ReferencedObjects);
+		}
 		Collector.AllowEliminatingReferences(true);
 	}
 	/** Adds an object to be referenced */
 	FORCEINLINE void AddObject(UObject* InObject)
 	{
-		ReferencedObjects.Add(InObject);
+		if (InObject)
+		{
+			FScopeLock ReferencedObjectsLock(&ReferencedObjectsCritical);
+			ReferencedObjects.Add(InObject);
+			InObject->ThisThreadAtomicallyClearedRFUnreachable();
+		}
 	}
 	/** Removes all objects from the list */
 	FORCEINLINE void EmptyReferencedObjects()
 	{
+		FScopeLock ReferencedObjectsLock(&ReferencedObjectsCritical);
+		for (auto Obj : ReferencedObjects)
+		{
+			check(Obj);
+			Obj->ClearFlags(RF_Async | RF_AsyncLoading);
+		}
 		ReferencedObjects.Empty(ReferencedObjects.Num());
 	}
 	/** Removes all referenced objects and markes them for GC */
@@ -87,6 +110,444 @@ FAsyncObjectsReferencer& FAsyncObjectsReferencer::Get()
 	return Singleton;
 }
 
+static FORCEINLINE bool IsTimeLimitExceeded(double InTickStartTime, bool bUseTimeLimit, float InTimeLimit, const TCHAR* InLastTypeOfWorkPerformed = nullptr, UObject* InLastObjectWorkWasPerformedOn = nullptr)
+{
+	bool bTimeLimitExceeded = false;
+	if (bUseTimeLimit)
+	{
+		double CurrentTime = FPlatformTime::Seconds();
+		bTimeLimitExceeded = CurrentTime - InTickStartTime > InTimeLimit;
+
+		// Log single operations that take longer than timelimit.
+		if ((CurrentTime - InTickStartTime) > (2.5 * InTimeLimit))
+		{
+			UE_LOG(LogStreaming, Log, TEXT("FAsyncPackage: %s %s took (less than) %5.2f ms"),
+				InLastTypeOfWorkPerformed ? InLastTypeOfWorkPerformed : TEXT("unknown"),
+				InLastObjectWorkWasPerformedOn ? *InLastObjectWorkWasPerformedOn->GetFullName() : TEXT("nullptr"),
+				(CurrentTime - InTickStartTime) * 1000);
+		}
+	}
+	return bTimeLimitExceeded;
+}
+
+int32 FAsyncLoadingThread::AsyncLoadingThreadID = -1;
+
+FAsyncLoadingThread& FAsyncLoadingThread::Get()
+{
+	static FAsyncLoadingThread GAsyncLoader;
+	return GAsyncLoader;
+}
+
+void FAsyncLoadingThread::InitializeAsyncThread()
+{
+	AsyncThreadReady.Increment();
+}
+
+void FAsyncLoadingThread::CancelAsyncLoadingInternal()
+{
+	{
+		// Packages we haven't yet started processing.
+		FScopeLock QueueLock(&QueueCritical);
+		QueuedPackages.Empty();
+	}
+
+	{
+		// Packages we started processing, need to be canceled.
+		// Accessed only in async thread, no need to protect region.
+		for (auto AsyncPackage : AsyncPackages)
+		{
+			AsyncPackage->Cancel();
+		}
+
+		AsyncPackages.Empty();
+	}
+
+	{
+		// Packages that are already loaded. May be halfway through PostLoad
+		FScopeLock LoadedLock(&LoadedPackagesCritical);
+		for (auto LoadedPackage : LoadedPackages)
+		{
+			LoadedPackage->Cancel();
+		}
+		LoadedPackages.Empty();
+
+		for (auto LoadedPackage : LoadedPackagesToProcess)
+		{
+			LoadedPackage->Cancel();
+		}
+		LoadedPackagesToProcess.Empty();
+	}
+
+	AsyncLoadingCounter.Reset();
+	AsyncPackagesCounter.Reset();
+	QueuedPackagesCounter.Reset();
+
+	FUObjectThreadContext::Get().ObjLoaded.Empty();
+	FAsyncObjectsReferencer::Get().EmptyReferencedObjectsAndCancelLoading();
+
+	// Notify everyone streaming is canceled.
+	CancelLoadingEvent->Trigger();
+}
+
+void FAsyncLoadingThread::QueuePackage(const FAsyncPackageDesc& Package)
+{
+	{
+		FScopeLock QueueLock(&QueueCritical);
+		QueuedPackagesCounter.Increment();
+		QueuedPackages.Add(new FAsyncPackageDesc(Package));
+	}
+	QueuedRequestsEvent->Trigger();
+}
+
+int32 FAsyncLoadingThread::CreateAsyncPackagesFromQueue()
+{
+	int32 NumCreated = 0;
+
+	STAT(double ThisTime = 0);
+	{
+		SCOPE_SECONDS_COUNTER(ThisTime);
+
+		checkSlow(IsInAsyncLoadThread());
+
+		TArray<FAsyncPackageDesc*> QueueCopy;
+		{
+			FScopeLock QueueLock(&QueueCritical);
+			QueueCopy = QueuedPackages;
+			QueuedPackages.Empty();
+		}
+		for (auto PackageRequest : QueueCopy)
+		{
+			FAsyncPackage* Package = nullptr;
+			int ExistingPackageIndex = FindAsyncPackage(PackageRequest->Name);
+			if (ExistingPackageIndex != INDEX_NONE)
+			{
+				// This is a package that is currently being loaded.
+				Package = AsyncPackages[ExistingPackageIndex];
+			}
+			else
+			{
+				// [BLOCKING] LoadedPackages are accessed on the main thread too, so lock to be able to add a completion callback
+				FScopeLock LoadedLock(&LoadedPackagesCritical);
+				ExistingPackageIndex = FindLoadedPackage(PackageRequest->Name);
+				if (ExistingPackageIndex != INDEX_NONE && PackageRequest->PackageLoadedDelegate.IsBound())
+				{
+					const bool bInternalCallback = false;
+					LoadedPackages[ExistingPackageIndex]->AddCompletionCallback(PackageRequest->PackageLoadedDelegate, bInternalCallback);
+					continue;
+				}
+			}
+			if (!Package)
+			{
+				// New package that needs to be loaded or a package has already been loaded long time ago
+				Package = new FAsyncPackage(*PackageRequest);
+
+				// Add to queue according to priority.
+				InsertPackage(Package);
+			}
+			if (PackageRequest->PackageLoadedDelegate.IsBound())
+			{
+				const bool bInternalCallback = false;
+				Package->AddCompletionCallback(PackageRequest->PackageLoadedDelegate, bInternalCallback);
+			}
+			QueuedPackagesCounter.Decrement();
+			check(QueuedPackagesCounter.GetValue() >= 0);
+		}
+		NumCreated = QueueCopy.Num();
+	}
+	INC_FLOAT_STAT_BY(STAT_FAsyncPackage_CreateAsyncPackagesFromQueueTime, (float)ThisTime);
+
+	return NumCreated;
+}
+
+/**
+* [ASYNC THREAD] Loads all packages
+*/
+EAsyncPackageState::Type FAsyncLoadingThread::ProcessAsyncLoading(int32& OutPackagesProcessed, bool bUseTimeLimit /*= false*/, bool bUseFullTimeLimit /*= false*/, float TimeLimit /*= 0.0f*/, FName ExcludeType /*= NAME_None*/)
+{
+	EAsyncPackageState::Type LoadingState = EAsyncPackageState::Complete;
+	OutPackagesProcessed = 0;
+	STAT(double ThisTime = 0);
+	{
+		SCOPE_SECONDS_COUNTER(ThisTime);
+
+		// We need to loop as the function has to handle finish loading everything given no time limit
+		// like e.g. when called from FlushAsyncLoading.
+		for (int32 PackageIndex = 0; LoadingState != EAsyncPackageState::TimeOut && PackageIndex < AsyncPackages.Num(); ++PackageIndex)
+		{
+			OutPackagesProcessed++;
+
+			// Package to be loaded.
+			FAsyncPackage* Package = AsyncPackages[PackageIndex];
+
+			if (ExcludeType != NAME_None && ExcludeType == Package->GetPackageType())
+			{
+				// We should skip packages of this type
+				continue;
+			}
+
+			if (Package->HasFinishedLoading() == false)
+			{
+				// Package tick returns EAsyncPackageState::Complete on completion.
+				// We only tick packages that have not yet been loaded.
+				LoadingState = Package->Tick(bUseTimeLimit, bUseFullTimeLimit, TimeLimit);
+			}
+			else
+			{
+				// This package has finished loading but some other package is still holding
+				// a reference to it because it has this package in its dependency list.
+				LoadingState = EAsyncPackageState::Complete;
+			}
+			bool bPackageFullyLoaded = false;
+			if (LoadingState == EAsyncPackageState::Complete)
+			{
+				// Only remove packages that are not being referenced by anything (we still need their linkers
+				// when the referencing packages create their imports and exports).
+				if (Package->GetDependencyRefCount() == 0)
+				{
+					// We're done so we can remove the package now.
+					AddToLoadedPackages(Package);
+					{
+						FScopeLock LockAsyncPackages(&AsyncPackagesCritical);
+						AsyncPackages.RemoveAt(PackageIndex);
+					}
+					
+					// Need to process this index again as we just removed an item
+					PackageIndex--;
+					bPackageFullyLoaded = true;
+				}
+				else
+				{
+					// This package is being referenced by another package so we can't say streaming is complete yet.
+					LoadingState = EAsyncPackageState::PendingImports;
+				}
+			}
+			else if (!bUseTimeLimit && !FPlatformProcess::SupportsMultithreading())
+			{
+				// Tick async loading when multithreading is disabled.
+				FIOSystem::Get().TickSingleThreaded();
+			}
+
+			// Check if there's any new packages in the queue.
+			CreateAsyncPackagesFromQueue();
+
+			if (bPackageFullyLoaded)
+			{
+				AsyncPackagesCounter.Decrement();
+			}
+		}
+	}
+	INC_FLOAT_STAT_BY(STAT_FAsyncPackage_AsyncProcessTime, (float)ThisTime);
+
+	return LoadingState;
+}
+
+
+EAsyncPackageState::Type FAsyncLoadingThread::ProcessLoadedPackages(bool bUseTimeLimit, bool bUseFullTimeLimit, float TimeLimit, FName ExcludeType)
+{
+	EAsyncPackageState::Type Result = EAsyncPackageState::Complete;
+
+	double TickStartTime = FPlatformTime::Seconds();
+	{
+		FScopeLock LoadedLock(&LoadedPackagesCritical);
+		LoadedPackagesToProcess.Append(LoadedPackages);
+		LoadedPackages.Empty();
+	}
+		
+	int32 PackageIndex = 0;
+	while (PackageIndex < LoadedPackagesToProcess.Num() && !IsTimeLimitExceeded(TickStartTime, bUseTimeLimit, TimeLimit))
+	{
+		auto Package = LoadedPackagesToProcess[PackageIndex];
+		Result = Package->PostLoadDeferredObjects(TickStartTime, bUseTimeLimit, TimeLimit);
+		if (Result == EAsyncPackageState::Complete)
+		{
+			if (FPlatformProperties::RequiresCookedData())
+			{
+				// Emulates ResetLoaders on the package linker's linkerroot.
+				Package->ResetLoader();
+			}
+
+			const bool bInternalCallbacks = false;
+			const EAsyncLoadingResult::Type LoadingResult = Package->HasLoadFailed() ? EAsyncLoadingResult::Failed : EAsyncLoadingResult::Succeeded;
+			Package->CallCompletionCallbacks(bInternalCallbacks, LoadingResult);
+			PackageIndex++;
+
+			// Incremented on the Async Thread, decremented on the Game Thread				
+			AsyncLoadingCounter.Decrement();
+			check(AsyncLoadingCounter.GetValue() >= 0);
+		}
+		else
+		{
+			break;
+		}
+	}
+	if (PackageIndex > 0)
+	{
+		LoadedPackagesToProcess.RemoveAt(0, PackageIndex);
+	}
+
+	return Result;
+}
+
+
+
+EAsyncPackageState::Type FAsyncLoadingThread::TickAsyncLoading(bool bUseTimeLimit, bool bUseFullTimeLimit, float TimeLimit, FName ExcludeType)
+{	
+	TGuardValue<bool> InAsyncLoadingTickGuard(bIsInAsyncLoadingTick, true);
+
+	static bool bWasAsyncLoading = false;
+	EAsyncPackageState::Type Result = EAsyncPackageState::Complete;	
+	double TickStartTime = FPlatformTime::Seconds();
+	double TimeLimitUsedForPreloading = 0;
+
+	if (!FAsyncLoadingThread::IsMultithreaded())
+	{		
+		TickAsyncThread(bUseTimeLimit, bUseFullTimeLimit, TimeLimit);
+		TimeLimitUsedForPreloading = FPlatformTime::Seconds() - TickStartTime;
+	}
+
+	if (!IsTimeLimitExceeded(TickStartTime, bUseTimeLimit, TimeLimit))
+	{
+		double RemainingTimeLimit = FMath::Max(0.0, TimeLimit - TimeLimitUsedForPreloading);
+		Result = ProcessLoadedPackages(bUseTimeLimit, bUseFullTimeLimit, RemainingTimeLimit, ExcludeType);
+	}
+
+	if (!IsTimeLimitExceeded(TickStartTime, bUseTimeLimit, TimeLimit))
+	{
+		FScopeLock QueueLock(&QueueCritical);
+		FScopeLock LoadedLock(&LoadedPackagesCritical);
+		if (AsyncPackagesCounter.GetValue() == 0 && LoadedPackagesToProcess.Num() == 0)
+		{
+			FDeferredMessageLog::Flush();
+			FAsyncObjectsReferencer::Get().EmptyReferencedObjects();
+		}
+	}
+
+	return Result;
+}
+
+FAsyncLoadingThread::FAsyncLoadingThread()
+{
+	QueuedRequestsEvent = FPlatformProcess::CreateSynchEvent();
+	CancelLoadingEvent = FPlatformProcess::CreateSynchEvent();
+	if (FAsyncLoadingThread::IsMultithreaded())
+	{
+		Thread = FRunnableThread::Create(this, TEXT("FAsyncLoadingThread"), 0, TPri_Normal);
+	}
+	else
+	{
+		Thread = nullptr;
+		Init();
+	}
+	bIsInAsyncLoadingTick = false;
+}
+
+FAsyncLoadingThread::~FAsyncLoadingThread()
+{
+	delete Thread;
+	Thread = nullptr;
+	delete QueuedRequestsEvent;
+	QueuedRequestsEvent = nullptr;
+	delete CancelLoadingEvent;
+	CancelLoadingEvent = nullptr;
+}
+
+bool FAsyncLoadingThread::Init()
+{
+	return true;
+}
+
+uint32 FAsyncLoadingThread::Run()
+{
+	AsyncLoadingThreadID = FPlatformTLS::GetCurrentThreadId();
+
+	while (StopTaskCounter.GetValue() == 0)
+	{
+		TickAsyncThread(false, true, 0.0f);
+	}
+	return 0;
+}
+
+void FAsyncLoadingThread::TickAsyncThread(bool bUseTimeLimit, bool bUseFullTimeLimit, float TimeLimit)
+{
+	if (!bShouldCancelLoading)
+	{
+		int32 ProcessedRequests = 0;
+		if (AsyncThreadReady.GetValue())
+		{
+			CreateAsyncPackagesFromQueue();
+			ProcessAsyncLoading(ProcessedRequests, bUseTimeLimit, bUseFullTimeLimit, TimeLimit, NAME_None);
+		}
+		if (ProcessedRequests == 0 && IsMultithreaded())
+		{
+			const bool bIgnoreThreadIdleStats = true;
+			QueuedRequestsEvent->Wait(30, bIgnoreThreadIdleStats);
+		}
+	}
+	else
+	{
+		// Blocks main thread
+		CancelAsyncLoadingInternal();
+		bShouldCancelLoading = false;
+	}
+}
+
+void FAsyncLoadingThread::Stop()
+{
+	StopTaskCounter.Increment();
+}
+
+void FAsyncLoadingThread::CancelAsyncLoading()
+{
+	checkSlow(IsInGameThread());
+
+	bShouldCancelLoading = true;
+	if (IsMultithreaded())
+	{
+		CancelLoadingEvent->Wait();
+	}
+	else
+	{
+		// This will immediately cancel async loading without waiting for packages to finish loading
+		FlushAsyncLoading();
+		// It's possible we haven't been async loading at all in which case the above call would not reset bShouldCancelLoading
+		bShouldCancelLoading = false;
+	}
+}
+
+float FAsyncLoadingThread::GetAsyncLoadPercentage(const FName& PackageName)
+{
+	float LoadPercentage = -1.0f;
+	{
+		FScopeLock LockAsyncPackages(&AsyncPackagesCritical);
+		const int32 PackageIndex = FindPackageByName(AsyncPackages, PackageName);
+		if (PackageIndex != INDEX_NONE)
+		{
+			LoadPercentage = AsyncPackages[PackageIndex]->GetLoadPercentage();
+		}
+	}
+	if (LoadPercentage < 0.0f)
+	{
+		FScopeLock LockLoadedPackages(&LoadedPackagesCritical);
+		const int32 PackageIndex = FindPackageByName(LoadedPackages, PackageName);
+		if (PackageIndex != INDEX_NONE)
+		{
+			LoadPercentage = LoadedPackages[PackageIndex]->GetLoadPercentage();
+		}
+	}
+	if (LoadPercentage < 0.0f)
+	{
+		checkSlow(IsInGameThread());
+		// No lock required as we're in the game thread and LoadedPackagesToProcess are only managed on the game thread
+		const int32 PackageIndex = FindPackageByName(LoadedPackagesToProcess, PackageName);
+		if (PackageIndex != INDEX_NONE)
+		{
+			LoadPercentage = LoadedPackagesToProcess[PackageIndex]->GetLoadPercentage();
+		}
+	}
+
+	return LoadPercentage;
+}
+
 /**
  * Call back into the async loading code to inform of the creation of a new object
  * @param Object		Object created
@@ -100,17 +561,34 @@ void NotifyConstructedDuringAsyncLoading(UObject* Object, bool bSubObject)
 	if (!bSubObject)
 	{
 		Object->SetFlags(RF_AsyncLoading);
-		GObjConstructedDuringAsyncLoading.Add(Object);
 	}
 	FAsyncObjectsReferencer::Get().AddObject(Object);
 }
 
 
+class FAsyncLoadingThreadManager : private FSelfRegisteringExec
+{
+public:
+	// FSelfRegisteringExec interface
+	virtual bool Exec(class UWorld* InWorld, const TCHAR* Cmd, FOutputDevice& Ar) override
+	{
+		if (FParse::Command(&Cmd, TEXT("ASYNCCREATE")))
+		{
+			return true;
+		}
+		return false;
+	}
+
+private:
+
+};
+static FAsyncLoadingThreadManager GAsyncLoaderManager;
+
+
 /*-----------------------------------------------------------------------------
 	FAsyncPackage implementation.
 -----------------------------------------------------------------------------*/
-/** Array of packages that are being preloaded							*/
-static TIndirectArray<struct FAsyncPackage>	GObjAsyncPackages;
+
 
 int32 FAsyncPackage::PreLoadIndex = 0;
 int32 FAsyncPackage::PostLoadIndex = 0;
@@ -118,16 +596,14 @@ int32 FAsyncPackage::PostLoadIndex = 0;
 /**
 * Constructor
 */
-FAsyncPackage::FAsyncPackage(const FName& InPackageName, const FGuid* InPackageGuid, FName InPackageType, const FName& InPackageNameToLoad)
-: PackageName(InPackageName)
-, PackageNameToLoad(InPackageNameToLoad)
-, PackageGuid(InPackageGuid != NULL ? *InPackageGuid : FGuid(0, 0, 0, 0))
-, PackageType(InPackageType)
+FAsyncPackage::FAsyncPackage(const FAsyncPackageDesc& InDesc)
+: Desc(InDesc)
 , Linker(NULL)
 , DependencyRefCount(0)
 , LoadImportIndex(0)
 , ImportIndex(0)
 , ExportIndex(0)
+, DeferredPostLoadIndex(0)
 , TimeLimit(FLT_MAX)
 , bUseTimeLimit(false)
 , bUseFullTimeLimit(false)
@@ -139,10 +615,6 @@ FAsyncPackage::FAsyncPackage(const FName& InPackageName, const FGuid* InPackageG
 , LastTypeOfWorkPerformed(NULL)
 , LoadStartTime(0.0)
 , LoadPercentage(0)
-, PackageFlags(0)
-#if WITH_EDITOR
-, PIEInstanceID(INDEX_NONE)
-#endif	
 #if PERF_TRACK_DETAILED_ASYNC_STATS
 , TickCount(0)
 , TickLoopCount(0)
@@ -165,26 +637,6 @@ FAsyncPackage::FAsyncPackage(const FName& InPackageName, const FGuid* InPackageG
 {
 }
 
-/**
- * Returns an estimated load completion percentage.
- */
-float FAsyncPackage::GetLoadPercentage() const
-{
-	float Accum = LoadPercentage;
-
-	for (auto Import : PendingImportedPackages)
-	{
-		Accum += Import->GetLoadPercentage();
-	}
-
-	for (auto Import : ReferencedImports)
-	{
-		Accum += Import->GetLoadPercentage();
-	}
-
-	return Accum / (float)(1 + PendingImportedPackages.Num() + ReferencedImports.Num());
-}
-
 /**
  * @return Time load begun. This is NOT the time the load was requested in the case of other pending requests.
  */
@@ -199,9 +651,9 @@ double FAsyncPackage::GetLoadStartTime() const
 void FAsyncPackage::ResetLoader()
 {
 	// Reset loader.
-	if( Linker )
+	if (Linker)
 	{
-		Linker->Detach(false);
+		delete Linker;
 		Linker = nullptr;
 	}
 }
@@ -213,23 +665,7 @@ void FAsyncPackage::ResetLoader()
  */
 bool FAsyncPackage::IsTimeLimitExceeded()
 {
-	if( !bTimeLimitExceeded && bUseTimeLimit )
-	{
-		double CurrentTime = FPlatformTime::Seconds();
-		bTimeLimitExceeded = CurrentTime - TickStartTime > TimeLimit;
-//		if( GUseSeekFreeLoading )
-		{
-			// Log single operations that take longer than timelimit.
-			if( (CurrentTime - TickStartTime) > (2.5 * TimeLimit) )
-			{
- 				UE_LOG(LogStreaming, Log, TEXT("FAsyncPackage: %s %s took (less than) %5.2f ms"), 
- 					LastTypeOfWorkPerformed, 
- 					LastObjectWorkWasPerformedOn ? *LastObjectWorkWasPerformedOn->GetFullName() : TEXT(""), 
- 					(CurrentTime - TickStartTime) * 1000);
-			}
-		}
-	}
-	return bTimeLimitExceeded;
+	return ::IsTimeLimitExceeded(TickStartTime, bUseTimeLimit, TimeLimit, LastTypeOfWorkPerformed, LastObjectWorkWasPerformedOn);
 }
 
 /**
@@ -253,9 +689,6 @@ bool FAsyncPackage::GiveUpTimeSlice()
  */
 void FAsyncPackage::BeginAsyncLoad()
 {
-	// All objects created from now on should be flagged as RF_AsyncLoading so StaticFindObject doesn't return them.
-	GIsAsyncLoading = true;
-
 	// this won't do much during async loading except increase the load count which causes IsLoading to return true
 	BeginLoad();
 }
@@ -266,14 +699,11 @@ void FAsyncPackage::BeginAsyncLoad()
  */
 void FAsyncPackage::EndAsyncLoad()
 {
-	check(GIsAsyncLoading);
+	check(IsAsyncLoading());
 
 	// this won't do much during async loading except decrease the load count which causes IsLoading to return false
 	EndLoad();
 
-	// We're no longer async loading.
-	GIsAsyncLoading = false;
-
 	if (!bLoadHasFailed)
 	{
 		// Mark the package as loaded, if we succeeded
@@ -291,81 +721,81 @@ void FAsyncPackage::EndAsyncLoad()
  * @return	true if package has finished loading, false otherwise
  */
 
-EAsyncPackageState::Type FAsyncPackage::Tick( bool InbUseTimeLimit, bool InbUseFullTimeLimit, float& InOutTimeLimit )
+EAsyncPackageState::Type FAsyncPackage::Tick(bool InbUseTimeLimit, bool InbUseFullTimeLimit, float& InOutTimeLimit)
 {
 	SCOPE_CYCLE_COUNTER(STAT_FAsyncPackage_Tick);
 
- 	check(LastObjectWorkWasPerformedOn==nullptr);
-	check(LastTypeOfWorkPerformed==nullptr);
-
-	// Set up tick relevant variables.
-	bUseTimeLimit			= InbUseTimeLimit;
-	bUseFullTimeLimit		= InbUseFullTimeLimit;
-	bTimeLimitExceeded		= false;
-	TimeLimit				= InOutTimeLimit;
-	TickStartTime			= FPlatformTime::Seconds();
-
-	// Keep track of time when we start loading.
-	if( LoadStartTime == 0.0 )
-	{
-		LoadStartTime = TickStartTime;
-	}
-
 	// Whether we should execute the next step.
-	EAsyncPackageState::Type LoadingState	= EAsyncPackageState::Complete;
+	EAsyncPackageState::Type LoadingState = EAsyncPackageState::Complete;
 
 	STAT(double ThisTime = 0);
 	{
 		SCOPE_SECONDS_COUNTER(ThisTime);
 
+		check(LastObjectWorkWasPerformedOn == nullptr);
+		check(LastTypeOfWorkPerformed == nullptr);
+
+		// Set up tick relevant variables.
+		bUseTimeLimit = InbUseTimeLimit;
+		bUseFullTimeLimit = InbUseFullTimeLimit;
+		bTimeLimitExceeded = false;
+		TimeLimit = InOutTimeLimit;
+		TickStartTime = FPlatformTime::Seconds();
+
+		// Keep track of time when we start loading.
+		if (LoadStartTime == 0.0)
+		{
+			LoadStartTime = TickStartTime;
+		}
+
 		// Make sure we finish our work if there's no time limit. The loop is required as PostLoad
 		// might cause more objects to be loaded in which case we need to Preload them again.
 		do
 		{
 			// Reset value to true at beginning of loop.
-			LoadingState	= EAsyncPackageState::Complete;
+			LoadingState = EAsyncPackageState::Complete;
 
-			// Begin async loading, simulates BeginLoad and sets GIsAsyncLoading to true.
+			// Begin async loading, simulates BeginLoad
 			BeginAsyncLoad();
 
 			// Create raw linker. Needs to be async created via ticking before it can be used.
-			if( LoadingState == EAsyncPackageState::Complete )
+			if (LoadingState == EAsyncPackageState::Complete)
 			{
 				LoadingState = CreateLinker();
 			}
 
 			// Async create linker.
-			if( LoadingState == EAsyncPackageState::Complete )
+			if (LoadingState == EAsyncPackageState::Complete)
 			{
 				LoadingState = FinishLinker();
 			}
 
 			// Load imports from linker import table asynchronously.
-			if( LoadingState == EAsyncPackageState::Complete )
+			if (LoadingState == EAsyncPackageState::Complete)
 			{
 				LoadingState = LoadImports();
 			}
 
 			// Create imports from linker import table.
-			if( LoadingState == EAsyncPackageState::Complete )
+			if (LoadingState == EAsyncPackageState::Complete)
 			{
 				LoadingState = CreateImports();
 			}
 
 			// Finish all async texture allocations.
-			if( LoadingState == EAsyncPackageState::Complete )
+			if (LoadingState == EAsyncPackageState::Complete)
 			{
 				LoadingState = FinishTextureAllocations();
 			}
 
 			// Create exports from linker export table and also preload them.
-			if( LoadingState == EAsyncPackageState::Complete )
+			if (LoadingState == EAsyncPackageState::Complete)
 			{
 				LoadingState = CreateExports();
 			}
 
 			// Call Preload on the linker for all loaded objects which causes actual serialization.
-			if( LoadingState == EAsyncPackageState::Complete )
+			if (LoadingState == EAsyncPackageState::Complete)
 			{
 				LoadingState = PreLoadObjects();
 			}
@@ -377,30 +807,28 @@ EAsyncPackageState::Type FAsyncPackage::Tick( bool InbUseTimeLimit, bool InbUseF
 				LoadingState = PostLoadObjects();
 			}
 
-			// End async loading, simulates EndLoad and sets GIsAsyncLoading to false.
+			// End async loading, simulates EndLoad
 			EndAsyncLoad();
 
 			// Finish objects (removing RF_AsyncLoading, dissociate imports and forced exports, 
 			// call completion callback, ...
 			// If the load has failed, perform completion callbacks and then quit
-			if( LoadingState == EAsyncPackageState::Complete || bLoadHasFailed )
+			if (LoadingState == EAsyncPackageState::Complete || bLoadHasFailed)
 			{
 				LoadingState = FinishObjects();
 			}
-		// Only exit loop if we're either done or the time limit has been exceeded.
-		} while( !IsTimeLimitExceeded() && LoadingState == EAsyncPackageState::TimeOut );	
+		} while (!IsTimeLimitExceeded() && LoadingState == EAsyncPackageState::TimeOut);
 
-		check( bUseTimeLimit || LoadingState != EAsyncPackageState::TimeOut );
+		check(bUseTimeLimit || LoadingState != EAsyncPackageState::TimeOut);
 
 		// We can't have a reference to a UObject.
-		LastObjectWorkWasPerformedOn	= nullptr;
+		LastObjectWorkWasPerformedOn = nullptr;
 		// Reset type of work performed.
-		LastTypeOfWorkPerformed			= nullptr;
+		LastTypeOfWorkPerformed = nullptr;
 		// Mark this package as loaded if everything completed.
 		bLoadHasFinished = (LoadingState == EAsyncPackageState::Complete);
 		// Subtract the time it took to load this package from the global limit.
 		InOutTimeLimit = (float)FMath::Max(0.0, InOutTimeLimit - (FPlatformTime::Seconds() - TickStartTime));
-	
 	}
 	INC_FLOAT_STAT_BY(STAT_FAsyncPackage_TickTime, (float)ThisTime);
 
@@ -423,37 +851,37 @@ EAsyncPackageState::Type FAsyncPackage::CreateLinker()
 		LastTypeOfWorkPerformed			= TEXT("creating Linker");
 
 		// Try to find existing package or create it if not already present.
-		UPackage* Package = CreatePackage(nullptr, *PackageName.ToString());
+		UPackage* Package = CreatePackage(nullptr, *Desc.Name.ToString());
 		FScopeCycleCounterUObject ConstructorScope(Package, GET_STATID(STAT_FAsyncPackage_CreateLinker));
 
 		// Set package specific data 
-		Package->PackageFlags |= PackageFlags;
+		Package->PackageFlags |= Desc.Flags;
 #if WITH_EDITOR
-		Package->PIEInstanceID = PIEInstanceID;
+		Package->PIEInstanceID = Desc.PIEInstanceID;
 #endif
 
 		// Always store package filename we loading from
-		Package->FileName = PackageNameToLoad;
+		Package->FileName = Desc.NameToLoad;
 
 		// if the linker already exists, we don't need to lookup the file (it may have been pre-created with
 		// a different filename)
-		Linker = ULinkerLoad::FindExistingLinkerForPackage(Package);
+		Linker = FLinkerLoad::FindExistingLinkerForPackage(Package);
 
 		if (!Linker)
 		{
 			FString PackageFileName;
-			if (!FPackageName::DoesPackageExist(PackageNameToLoad.ToString(), PackageGuid.IsValid() ? &PackageGuid : nullptr, &PackageFileName))
+			if (!FPackageName::DoesPackageExist(Desc.NameToLoad.ToString(), Desc.Guid.IsValid() ? &Desc.Guid : nullptr, &PackageFileName))
 			{
-				UE_LOG(LogStreaming, Error, TEXT("Couldn't find file for package %s requested by async loading code."), *PackageName.ToString());
+				UE_LOG(LogStreaming, Error, TEXT("Couldn't find file for package %s requested by async loading code."), *Desc.Name.ToString());
 				bLoadHasFailed = true;
 				return EAsyncPackageState::TimeOut;
 			}
 
 			// Create raw async linker, requiring to be ticked till finished creating.
-			Linker = ULinkerLoad::CreateLinkerAsync( Package, *PackageFileName, (FApp::IsGame() && !GIsEditor) ? (LOAD_SeekFree | LOAD_NoVerify) : LOAD_None  );
+			Linker = FLinkerLoad::CreateLinkerAsync( Package, *PackageFileName, (FApp::IsGame() && !GIsEditor) ? (LOAD_SeekFree | LOAD_NoVerify) : LOAD_None  );
 		}
 
-		UE_LOG(LogStreaming, Verbose, TEXT("FAsyncPackage::CreateLinker for %s finished."), *PackageNameToLoad.ToString());
+		UE_LOG(LogStreaming, Verbose, TEXT("FAsyncPackage::CreateLinker for %s finished."), *Desc.NameToLoad.ToString());
 	}
 	return EAsyncPackageState::Complete;
 }
@@ -475,13 +903,13 @@ EAsyncPackageState::Type FAsyncPackage::FinishLinker()
 		const float RemainingTimeLimit = TimeLimit - (float)(FPlatformTime::Seconds() - TickStartTime);
 
 		// Operation still pending if Tick returns false
-		ULinkerLoad::ELinkerStatus LinkerResult = Linker->Tick(RemainingTimeLimit, bUseTimeLimit, bUseFullTimeLimit);
-		if (LinkerResult != ULinkerLoad::LINKER_Loaded)
+		FLinkerLoad::ELinkerStatus LinkerResult = Linker->Tick(RemainingTimeLimit, bUseTimeLimit, bUseFullTimeLimit);
+		if (LinkerResult != FLinkerLoad::LINKER_Loaded)
 		{
 			// Give up remainder of timeslice if there is one to give up.
 			GiveUpTimeSlice();
 			Result = EAsyncPackageState::TimeOut;
-			if (LinkerResult == ULinkerLoad::LINKER_Failed)
+			if (LinkerResult == FLinkerLoad::LINKER_Failed)
 			{
 				// If linker failed we exit with EAsyncPackageState::TimeOut to skip all the remaining steps.
 				// The error will be handled as bLoadHasFailed will be true.
@@ -512,24 +940,7 @@ FORCEINLINE int32 ContainsDependencyPackage(TArray<FAsyncPackage*>& Dependencies
 	return INDEX_NONE;
 }
 
-/**
- * Finds an existing async package in the GObjAsyncPackages by its name.
- *
- * @param PackageName async package name.
- * @return Index of the async package in GObjAsyncPackages array or INDEX_NONE if not found.
- */
-FORCEINLINE int32 FindAsyncPackage(const FName& PackageName)
-{
-	for (int32 PackageIndex = 0; PackageIndex < GObjAsyncPackages.Num(); ++PackageIndex)
-	{
-		FAsyncPackage& PendingPackage = GObjAsyncPackages[PackageIndex];
-		if (PendingPackage.GetPackageName() == PackageName)
-		{
-			return PackageIndex;
-		}
-	}
-	return INDEX_NONE;
-}
+
 /**
 	* Adds a package to the list of pending import packages.
 	*
@@ -538,21 +949,29 @@ FORCEINLINE int32 FindAsyncPackage(const FName& PackageName)
 void FAsyncPackage::AddImportDependency(int32 CurrentPackageIndex, const FName& PendingImport)
 {
 	FAsyncPackage* PackageToStream = nullptr;
-	int32 ExistingAsyncPackageIndex = FindAsyncPackage(PendingImport);
+	int32 ExistingAsyncPackageIndex = FAsyncLoadingThread::Get().FindAsyncPackage(PendingImport);
 	if (ExistingAsyncPackageIndex == INDEX_NONE)
 	{
-		PackageToStream = new FAsyncPackage(PendingImport, nullptr, NAME_None, PendingImport);
-		GObjAsyncPackages.Insert(PackageToStream, CurrentPackageIndex);
+		const FAsyncPackageDesc Info(PendingImport, FGuid(), PendingImport);
+		PackageToStream = new FAsyncPackage(Info);
+
+		// If priority of the dependency is not set, inherit from parent.
+		if (PackageToStream->Desc.Priority == 0)
+		{
+			PackageToStream->Desc.Priority = Desc.Priority;
+		}
+		FAsyncLoadingThread::Get().InsertPackage(PackageToStream);
 	}
 	else
 	{
-		PackageToStream = &GObjAsyncPackages[ExistingAsyncPackageIndex];
+		PackageToStream = FAsyncLoadingThread::Get().GetPackage(ExistingAsyncPackageIndex);
 	}
 	
 	if (!PackageToStream->HasFinishedLoading() && 
 		!PackageToStream->bLoadHasFailed)
 	{
-		PackageToStream->AddCompletionCallback(FLoadPackageAsyncDelegate::CreateRaw(this, &FAsyncPackage::ImportFullyLoadedCallback));
+		const bool bInternalCallback = true;
+		PackageToStream->AddCompletionCallback(FLoadPackageAsyncDelegate::CreateRaw(this, &FAsyncPackage::ImportFullyLoadedCallback), bInternalCallback);
 		PackageToStream->DependencyRefCount++;
 		PendingImportedPackages.Add(PackageToStream);
 	}
@@ -619,9 +1038,12 @@ EAsyncPackageState::Type FAsyncPackage::LoadImports()
 	SCOPE_CYCLE_COUNTER(STAT_FAsyncPackage_LoadImports);
 	LastObjectWorkWasPerformedOn	= Linker->LinkerRoot;
 	LastTypeOfWorkPerformed			= TEXT("loading imports");
-
+	
 	// Index of this package in the async queue.
-	const int32 AsyncQueueIndex = FindAsyncPackage(GetPackageName());
+	const int32 AsyncQueueIndex = FAsyncLoadingThread::Get().FindAsyncPackage(GetPackageName());
+
+	// GC can't run in here
+	FGCScopeGuard GCGuard;
 
 	// Create imports.
 	while (LoadImportIndex < Linker->ImportMap.Num() && !IsTimeLimitExceeded())
@@ -651,10 +1073,10 @@ EAsyncPackageState::Type FAsyncPackage::LoadImports()
 			// we add all dependencies that don't yet have linkers created otherwise we risk that if the current package
 			// doesn't depend on any other packages that have not yet started streaming, creating imports is going
 			// to load packages blocking the main thread.
-			int32 PendingAsyncPackageIndex = FindAsyncPackage(ImportPackageFName);
+			int32 PendingAsyncPackageIndex = FAsyncLoadingThread::Get().FindAsyncPackage(ImportPackageFName);
 			if (PendingAsyncPackageIndex != INDEX_NONE)
 			{
-				FAsyncPackage& PendingPackage = GObjAsyncPackages[PendingAsyncPackageIndex];
+				FAsyncPackage& PendingPackage = *FAsyncLoadingThread::Get().GetPackage(PendingAsyncPackageIndex);
 				if (PendingPackage.Linker == nullptr || !PendingPackage.Linker->HasFinishedInitialization())
 				{
 					// Add this import to the dependency list.
@@ -662,7 +1084,7 @@ EAsyncPackageState::Type FAsyncPackage::LoadImports()
 				}
 				else
 				{
-					UE_LOG(LogStreaming, Verbose, TEXT("FAsyncPackage::LoadImports for %s: Linker exists for %s"), *PackageNameToLoad.ToString(), *ImportPackageFName.ToString());
+					UE_LOG(LogStreaming, Verbose, TEXT("FAsyncPackage::LoadImports for %s: Linker exists for %s"), *Desc.NameToLoad.ToString(), *ImportPackageFName.ToString());
 					// Only keep a reference to this package so that its linker doesn't go away too soon
 					PendingPackage.DependencyRefCount++;
 					ReferencedImports.Add(&PendingPackage);
@@ -679,16 +1101,17 @@ EAsyncPackageState::Type FAsyncPackage::LoadImports()
 			// The package doesn't exist and this import is not in the dependency list so add it now.
 			if (!FPackageName::IsShortPackageName(ImportPackageName))
 			{
-				UE_LOG(LogStreaming, Verbose, TEXT("FAsyncPackage::LoadImports for %s: Loading %s"), *PackageNameToLoad.ToString(), *ImportPackageName);
-
+				UE_LOG(LogStreaming, Verbose, TEXT("FAsyncPackage::LoadImports for %s: Loading %s"), *Desc.NameToLoad.ToString(), *ImportPackageName);
 				AddImportDependency(AsyncQueueIndex, ImportPackageFName);
 			}
 			else
 			{
 				// This usually means there's a reference to a script package from another project
-				UE_LOG(LogStreaming, Warning, TEXT("FAsyncPackage::LoadImports for %s: Short package name in imports list: %s"), *PackageNameToLoad.ToString(), *ImportPackageName);
+				UE_LOG(LogStreaming, Warning, TEXT("FAsyncPackage::LoadImports for %s: Short package name in imports list: %s"), *Desc.NameToLoad.ToString(), *ImportPackageName);
 			}
 		}
+
+		UpdateLoadPercentage();
 	}
 
 	if (PendingImportedPackages.Num())
@@ -706,7 +1129,7 @@ void FAsyncPackage::ImportFullyLoadedCallback(const FName& InPackageName, UPacka
 {
 	if (Result != EAsyncLoadingResult::Canceled)
 	{
-		UE_LOG(LogStreaming, Verbose, TEXT("FAsyncPackage::LoadImports for %s: Loaded %s"), *PackageNameToLoad.ToString(), *InPackageName.ToString());
+		UE_LOG(LogStreaming, Verbose, TEXT("FAsyncPackage::LoadImports for %s: Loaded %s"), *Desc.NameToLoad.ToString(), *InPackageName.ToString());
 		int32 Index = ContainsDependencyPackage(PendingImportedPackages, InPackageName);
 		check(Index != INDEX_NONE);
 		// Keep a reference to this package so that its linker doesn't go away too soon
@@ -724,6 +1147,9 @@ EAsyncPackageState::Type FAsyncPackage::CreateImports()
 {
 	SCOPE_CYCLE_COUNTER(STAT_FAsyncPackage_CreateImports);
 
+	// GC can't run in here
+	FGCScopeGuard GCGuard;
+
 	// Create imports.
 	while( ImportIndex < Linker->ImportMap.Num() && !IsTimeLimitExceeded() )
 	{
@@ -746,7 +1172,6 @@ EAsyncPackageState::Type FAsyncPackage::CreateImports()
 EAsyncPackageState::Type FAsyncPackage::FinishTextureAllocations()
 {
 	//@TODO: Cancel allocations if they take too long.
-
 #if WITH_ENGINE
 	bool bHasCompleted = Linker->Summary.TextureAllocations.HasCompleted();
 	if ( !bHasCompleted )
@@ -780,6 +1205,9 @@ EAsyncPackageState::Type FAsyncPackage::CreateExports()
 {
 	SCOPE_CYCLE_COUNTER(STAT_FAsyncPackage_CreateExports);
 
+	// GC can't run in here
+	FGCScopeGuard GCGuard;
+
 	// Create exports.
 	while( ExportIndex < Linker->ExportMap.Num() && !IsTimeLimitExceeded() )
 	{
@@ -804,8 +1232,8 @@ EAsyncPackageState::Type FAsyncPackage::CreateExports()
 
 			LastObjectWorkWasPerformedOn	= Object;
 			LastTypeOfWorkPerformed			= TEXT("creating exports for");
-			// This assumes that only CreateExports is taking a significant amount of time.
-			LoadPercentage = 100.f * ExportIndex / Linker->ExportMap.Num();
+			
+			UpdateLoadPercentage();
 		}
 		// Data isn't ready yet. Give up remainder of time slice if we're not using a time limit.
 		else if( GiveUpTimeSlice() )
@@ -830,7 +1258,7 @@ void FAsyncPackage::FreeReferencedImports()
 	{
 		FAsyncPackage& Ref = *ReferencedImports[ReferenceIndex];
 		Ref.DependencyRefCount--;
-		UE_LOG(LogStreaming, Verbose, TEXT("FAsyncPackage::FreeReferencedImports for %s: Releasing %s (%d)"), *PackageNameToLoad.ToString(), *Ref.GetPackageName().ToString(), Ref.DependencyRefCount);
+		UE_LOG(LogStreaming, Verbose, TEXT("FAsyncPackage::FreeReferencedImports for %s: Releasing %s (%d)"), *Desc.NameToLoad.ToString(), *Ref.GetPackageName().ToString(), Ref.DependencyRefCount);
 		check(Ref.DependencyRefCount >= 0);
 	}
 	ReferencedImports.Empty();
@@ -845,20 +1273,24 @@ EAsyncPackageState::Type FAsyncPackage::PreLoadObjects()
 {
 	SCOPE_CYCLE_COUNTER(STAT_FAsyncPackage_PreLoadObjects);
 
+	// GC can't run in here
+	FGCScopeGuard GCGuard;
+
+	auto& ObjLoaded = FUObjectThreadContext::Get().ObjLoaded;
 	// Preload (aka serialize) the objects.
-	while( PreLoadIndex < GObjLoaded.Num() && !IsTimeLimitExceeded() )
+	while (PreLoadIndex < ObjLoaded.Num() && !IsTimeLimitExceeded())
 	{
 		//@todo async: make this part async as well.
-		UObject* Object = GObjLoaded[ PreLoadIndex++ ];
+		UObject* Object = ObjLoaded[ PreLoadIndex++ ];
 		if( Object && Object->GetLinker() )
 		{
 			Object->GetLinker()->Preload( Object );
-			LastObjectWorkWasPerformedOn	= Object;
-			LastTypeOfWorkPerformed			= TEXT("preloading");
+			LastObjectWorkWasPerformedOn = Object;
+			LastTypeOfWorkPerformed = TEXT("preloading");
 		}
 	}
 
-	return PreLoadIndex == GObjLoaded.Num() ? EAsyncPackageState::Complete : EAsyncPackageState::TimeOut;
+	return PreLoadIndex == ObjLoaded.Num() ? EAsyncPackageState::Complete : EAsyncPackageState::TimeOut;
 }
 
 /**
@@ -870,33 +1302,94 @@ EAsyncPackageState::Type FAsyncPackage::PostLoadObjects()
 {
 	SCOPE_CYCLE_COUNTER(STAT_FAsyncPackage_PostLoadObjects);
 	
+	// GC can't run in here
+	FGCScopeGuard GCGuard;
+
 	EAsyncPackageState::Type Result = EAsyncPackageState::Complete;
 	STAT(double ThisTime = 0);
 	{
 		SCOPE_SECONDS_COUNTER(ThisTime);
-		TGuardValue<bool> GuardIsRoutingPostLoad(GIsRoutingPostLoad, true);
+		TGuardValue<bool> GuardIsRoutingPostLoad(FUObjectThreadContext::Get().IsRoutingPostLoad, true);
 
+		auto& ObjLoaded = FUObjectThreadContext::Get().ObjLoaded;
 		// PostLoad objects.
-		while( PostLoadIndex < GObjLoaded.Num() && PostLoadIndex < PreLoadIndex && !IsTimeLimitExceeded() )
+		while (PostLoadIndex < ObjLoaded.Num() && PostLoadIndex < PreLoadIndex && !IsTimeLimitExceeded())
 		{
-			UObject* Object	= GObjLoaded[ PostLoadIndex++ ];
+			UObject* Object = ObjLoaded[PostLoadIndex++];
 			check(Object);
-			FScopeCycleCounterUObject ConstructorScope(Object, GET_STATID(STAT_FAsyncPackage_PostLoadObjects));
+			if (!FAsyncLoadingThread::Get().IsMultithreaded() || Object->IsPostLoadThreadSafe())
+			{
+				FScopeCycleCounterUObject ConstructorScope(Object, GET_STATID(STAT_FAsyncPackage_PostLoadObjects));
 
-			Object->ConditionalPostLoad();
+				Object->ConditionalPostLoad();
 
-			LastObjectWorkWasPerformedOn	= Object;
-			LastTypeOfWorkPerformed			= TEXT("postloading");
+				LastObjectWorkWasPerformedOn = Object;
+				LastTypeOfWorkPerformed = TEXT("postloading_async");
+			}
+			else
+			{
+				DeferredPostLoadObjects.Add(Object);
+			}
 		}
 
 		// New objects might have been loaded during PostLoad.
-		Result = (PreLoadIndex == GObjLoaded.Num()) && (PostLoadIndex == GObjLoaded.Num()) ? EAsyncPackageState::Complete : EAsyncPackageState::TimeOut;
+		Result = (PreLoadIndex == ObjLoaded.Num()) && (PostLoadIndex == ObjLoaded.Num()) ? EAsyncPackageState::Complete : EAsyncPackageState::TimeOut;
+		UE_CLOG(PreLoadIndex < ObjLoaded.Num(), LogStreaming, Fatal, TEXT("Objects have been loading in PostLoad. This is no longer allowed."));
 	}
 	INC_FLOAT_STAT_BY(STAT_FAsyncPackage_PostLoadObjectsTime, (float)ThisTime);
 
 	return Result;
 }
 
+EAsyncPackageState::Type FAsyncPackage::PostLoadDeferredObjects(double InTickStartTime, bool bUseTimeLimit, float& InOutTimeLimit)
+{
+	SCOPE_CYCLE_COUNTER(STAT_FAsyncPackage_PostLoadObjectsGameThread);
+
+	EAsyncPackageState::Type Result = EAsyncPackageState::Complete;
+	STAT(double ThisTime = 0);
+	{
+		SCOPE_SECONDS_COUNTER(ThisTime);
+
+		TGuardValue<bool> GuardIsRoutingPostLoad(FUObjectThreadContext::Get().IsRoutingPostLoad, true);
+
+		while (DeferredPostLoadIndex < DeferredPostLoadObjects.Num() && !::IsTimeLimitExceeded(InTickStartTime, bUseTimeLimit, InOutTimeLimit))
+		{
+			UObject* Object = DeferredPostLoadObjects[DeferredPostLoadIndex++];
+			check(Object);
+
+			FScopeCycleCounterUObject ConstructorScope(Object, GET_STATID(STAT_FAsyncPackage_PostLoadObjectsGameThread));
+
+			Object->ConditionalPostLoad();
+
+			LastObjectWorkWasPerformedOn = Object;
+			LastTypeOfWorkPerformed = TEXT("postloading_gamethread");
+
+			UpdateLoadPercentage();
+		}
+
+		// New objects might have been loaded during PostLoad.
+		Result = (DeferredPostLoadIndex == DeferredPostLoadObjects.Num()) ? EAsyncPackageState::Complete : EAsyncPackageState::TimeOut;
+		if (Result == EAsyncPackageState::Complete)
+		{
+			// Mark package as having been fully loaded and update load time.
+			if (Linker && !bLoadHasFailed)
+			{
+				if (Linker->LinkerRoot)
+				{
+					Linker->LinkerRoot->MarkAsFullyLoaded();
+					Linker->LinkerRoot->SetLoadTime(FPlatformTime::Seconds() - LoadStartTime);
+				}
+
+				// give a hint to the IO system that we are done with this file for now
+				FIOSystem::Get().HintDoneWithFile(Linker->Filename);
+			}
+		}
+	}
+	INC_FLOAT_STAT_BY(STAT_FAsyncPackage_PostLoadObjectsGameThreadTime, (float)ThisTime);
+
+	return Result;
+}
+
 /**
  * Finish up objects and state, which means clearing the RF_AsyncLoading flag on newly created ones
  *
@@ -907,102 +1400,63 @@ EAsyncPackageState::Type FAsyncPackage::FinishObjects()
 	SCOPE_CYCLE_COUNTER(STAT_FAsyncPackage_FinishObjects);
 	LastObjectWorkWasPerformedOn	= nullptr;
 	LastTypeOfWorkPerformed			= TEXT("finishing all objects");
-
-	// All loaded objects are now finished loading so we can clear the RF_AsyncLoading flag so
-	// StaticFindObject returns them by default.
-	for( int32 ObjectIndex=0; ObjectIndex<GObjConstructedDuringAsyncLoading.Num(); ObjectIndex++ )
-	{
-		UObject* Object = GObjConstructedDuringAsyncLoading[ObjectIndex];
-		Object->ClearFlags( RF_AsyncLoading );
-	}
-	GObjConstructedDuringAsyncLoading.Empty();
 		
+	auto& LoadingGlobals = FUObjectThreadContext::Get();
+
 	EAsyncLoadingResult::Type LoadingResult;
 	if (!bLoadHasFailed)
 	{
-		GObjLoaded.Empty();
-		LoadingResult = EAsyncLoadingResult::Succeeded;		
+		LoadingGlobals.ObjLoaded.Empty();
+		LoadingResult = EAsyncLoadingResult::Succeeded;
 	}
 	else
 	{
 		// Cleanup objects from this package only
-		for (int ObjectIndex = GObjLoaded.Num() - 1; ObjectIndex >= 0; --ObjectIndex)
+		for (int32 ObjectIndex = LoadingGlobals.ObjLoaded.Num() - 1; ObjectIndex >= 0; --ObjectIndex)
 		{
-			UObject* Object = GObjLoaded[ObjectIndex];
-			if (Object->GetOutermost()->GetFName() == PackageName)
+			UObject* Object = LoadingGlobals.ObjLoaded[ObjectIndex];
+			if (Object->GetOutermost()->GetFName() == Desc.Name)
 			{
 				Object->ClearFlags(RF_NeedPostLoad | RF_NeedLoad | RF_NeedPostLoadSubobjects);
 				Object->MarkPendingKill();
-				GObjLoaded.RemoveAt(ObjectIndex);
-			}			
+				LoadingGlobals.ObjLoaded.RemoveAt(ObjectIndex);
+			}
 		}
 		LoadingResult = EAsyncLoadingResult::Failed;
 	}
 
-	// Simulate what EndLoad does.	
+	// Simulate what EndLoad does.
 	DissociateImportsAndForcedExports(); //@todo: this should be avoidable
 	PreLoadIndex = 0;
 	PostLoadIndex = 0;
-
+	
 	// If we successfully loaded
-	if (!bLoadHasFailed)
+	if (!bLoadHasFailed && Linker)
 	{
-		// Mark package as having been fully loaded and update load time.
-		if (Linker && Linker->LinkerRoot)
-		{
-			Linker->LinkerRoot->MarkAsFullyLoaded();
-			Linker->LinkerRoot->SetLoadTime( FPlatformTime::Seconds() - LoadStartTime );
-		}
-
-		if (Linker)
-		{
-			// Call any completion callbacks specified.
-			for (int32 i = 0; i < CompletionCallbacks.Num(); i++)
-			{
-				CompletionCallbacks[i].ExecuteIfBound(PackageName, Linker->LinkerRoot, LoadingResult);
-			}
-
-			// give a hint to the IO system that we are done with this file for now
-			FIOSystem::Get().HintDoneWithFile(Linker->Filename);
-		}
-
 	#if WITH_ENGINE
-		if (Linker != nullptr)
-		{
-			// Cancel all texture allocations that haven't been claimed yet.
-			Linker->Summary.TextureAllocations.CancelRemainingAllocations(true);
-		}
+		// Cancel all texture allocations that haven't been claimed yet.
+		Linker->Summary.TextureAllocations.CancelRemainingAllocations( true );
 	#endif		// WITH_ENGINE
-
-		// Flush cache to free memory
-		if (Linker != nullptr)
-		{
-			Linker->FlushCache();
-		}
 	}
-	else
-	{
-		// On failure, just call completion callbacks
 
-		// Call any completion callbacks specified.
-		for (int32 i = 0; i < CompletionCallbacks.Num(); i++)
-		{
-			CompletionCallbacks[i].ExecuteIfBound(PackageName, nullptr, LoadingResult);
-		}
+	{
+		const bool bInternalCallbacks = true;
+		CallCompletionCallbacks(bInternalCallbacks, LoadingResult);
 	}
 
-	// Flush cache of linkers that were requirested to be closed to free memory
-	extern TArray<ULinkerLoad*>	GDelayedLinkerClosePackages;
-	for (int32 Index = 0; Index < GDelayedLinkerClosePackages.Num(); Index++)
+	return EAsyncPackageState::Complete;
+}
+
+void FAsyncPackage::CallCompletionCallbacks(bool bInternal, EAsyncLoadingResult::Type LoadingResult)
+{
+	UPackage* LinkerRoot = (!bLoadHasFailed && Linker) ? Linker->LinkerRoot : nullptr;
+	for (auto& CompletionCallback : CompletionCallbacks)
 	{
-		ULinkerLoad* DelayedLinker = GDelayedLinkerClosePackages[Index];
-		if (DelayedLinker)
+		if (CompletionCallback.bIsInternal == bInternal)
 		{
-			DelayedLinker->FlushCache();
+			CompletionCallback.Callback.ExecuteIfBound(Desc.Name, LinkerRoot, LoadingResult);
 		}
 	}
-
-	return EAsyncPackageState::Complete;
 }
 
 void FAsyncPackage::Cancel()
@@ -1011,7 +1465,7 @@ void FAsyncPackage::Cancel()
 	const EAsyncLoadingResult::Type Result = EAsyncLoadingResult::Canceled;
 	for (int32 CallbackIndex = 0; CallbackIndex < CompletionCallbacks.Num(); CallbackIndex++)
 	{
-		CompletionCallbacks[CallbackIndex].ExecuteIfBound(PackageName, nullptr, Result);
+		CompletionCallbacks[CallbackIndex].Callback.ExecuteIfBound(Desc.Name, nullptr, Result);
 	}
 	if (Linker)
 	{
@@ -1022,118 +1476,104 @@ void FAsyncPackage::Cancel()
 		{
 			Linker->LinkerRoot->ClearFlags(RF_WasLoaded);
 			Linker->LinkerRoot->bHasBeenFullyLoaded = false;
+			Linker->LinkerRoot->Rename(*MakeUniqueObjectName(GetTransientPackage(), UPackage::StaticClass()).ToString(), nullptr, REN_DontCreateRedirectors | REN_DoNotDirty | REN_ForceNoResetLoaders | REN_NonTransactional);
 		}
 	}
+
+	PreLoadIndex = 0;
+}
+
+void FAsyncPackage::AddCompletionCallback(const FLoadPackageAsyncDelegate& Callback, bool bInternal)
+{
+	// This is to ensure that there is no one trying to subscribe to a already loaded package
+	//check(!bLoadHasFinished && !bLoadHasFailed);
+	CompletionCallbacks.Add(FCompletionCallback(bInternal, Callback));
+}
+
+void FAsyncPackage::UpdateLoadPercentage()
+{
+	// PostLoadCount is just an estimate to prevent packages to go to 100% too quickly
+	// We may never reach 100% this way, but it's better than spending most of the load package time at 100%
+	float NewLoadPercentage = 0.0f;
+	if (Linker)
+	{
+		const int32 PostLoadCount = FMath::Max(DeferredPostLoadObjects.Num(), Linker->ImportMap.Num());
+		NewLoadPercentage = 100.f * (LoadImportIndex + ExportIndex + DeferredPostLoadIndex) / (Linker->ExportMap.Num() + Linker->ImportMap.Num() + PostLoadCount);		
+	}
+	else if (DeferredPostLoadObjects.Num() > 0)
+	{
+		NewLoadPercentage = DeferredPostLoadIndex / DeferredPostLoadObjects.Num();
+	}
+	// It's also possible that we got so many objects to PostLoad that LoadPercantage will actually drop
+	LoadPercentage = FMath::Max(NewLoadPercentage, LoadPercentage);
 }
 
 /*-----------------------------------------------------------------------------
 	UObject async (pre)loading.
 -----------------------------------------------------------------------------*/
-FAsyncPackage& LoadPackageAsync( const FString& InPackageName, const FGuid* PackageGuid, FName PackageType, const TCHAR* InPackageToLoadFrom)
+void LoadPackageAsync(const FString& InName, const FGuid* InGuid /*= nullptr*/, FName InType /*= NAME_None*/, const TCHAR* InPackageToLoadFrom /*= nullptr*/, FLoadPackageAsyncDelegate InCompletionDelegate /*= FLoadPackageAsyncDelegate()*/, uint32 InFlags /*= 0*/, int32 InPIEInstanceID /*= INDEX_NONE*/, uint32 InPackagePriority /*= 0*/)
 {
 	// The comments clearly state that it should be a package name but we also handle it being a filename as this function is not perf critical
 	// and LoadPackage handles having a filename being passed in as well.
 	FString PackageName;
-	
-	if ( FPackageName::IsValidLongPackageName(InPackageName, /*bIncludeReadOnlyRoots*/true) )
+	if (FPackageName::IsValidLongPackageName(InName, /*bIncludeReadOnlyRoots*/true))
 	{
-		PackageName = InPackageName;
+		PackageName = InName;
 	}
-	else if ( FPackageName::IsPackageFilename(InPackageName) && FPackageName::TryConvertFilenameToLongPackageName(InPackageName, PackageName) )
-	{
-		// PackageName got populated by the conditional function
-	}
-	else
+	// PackageName got populated by the conditional function
+	else if (!(FPackageName::IsPackageFilename(InName) && FPackageName::TryConvertFilenameToLongPackageName(PackageName, PackageName)))
 	{
 		// PackageName will get populated by the conditional function
 		FString ClassName;
-		if (!FPackageName::ParseExportTextPath(InPackageName, &ClassName, &PackageName))
+		if (!FPackageName::ParseExportTextPath(PackageName, &ClassName, &PackageName))
 		{
-			UE_LOG(LogStreaming, Fatal, TEXT("LoadPackageAsync failed to begin to load a level because the supplied package name was neither a valid long package name nor a filename of a map within a content folder: '%s'"), *InPackageName);
+			UE_LOG(LogStreaming, Fatal, TEXT("LoadPackageAsync failed to begin to load a level because the supplied package name was neither a valid long package name nor a filename of a map within a content folder: '%s'"), *PackageName);
 		}
 	}
 
-	// When constructing the FName from filename, we need to make sure that if the filename ends with _number, we're not splitting the FName
-	FName PackageFName(/*ENAME_LinkerConstructor,*/ *PackageName);
-
-	// Check whether the file is already in the queue.	
-	for( int32 PackageIndex=0; PackageIndex<GObjAsyncPackages.Num(); PackageIndex++ )
+	FString PackageNameToLoad(InPackageToLoadFrom);
+	if (PackageNameToLoad.IsEmpty())
 	{
-		FAsyncPackage& PendingPackage = GObjAsyncPackages[PackageIndex];
-		if (PendingPackage.GetPackageName() == PackageFName)
-		{
-			// Early out as package is already being preloaded.
-			return PendingPackage;
-		}
-	}
-	FString PackageToLoadFrom = PackageName;
-	if (InPackageToLoadFrom)
-	{
-		PackageToLoadFrom = FString(InPackageToLoadFrom);
+		PackageNameToLoad = PackageName;
 	}
 	// Make sure long package name is passed to FAsyncPackage so that it doesn't attempt to 
 	// create a package with short name.
-	if (FPackageName::IsShortPackageName(PackageToLoadFrom))
+	if (FPackageName::IsShortPackageName(PackageNameToLoad))
 	{
-		UE_LOG(LogStreaming, Fatal, TEXT("Async loading code requires long package names (%s)."), *InPackageName);
+		UE_LOG(LogStreaming, Fatal, TEXT("Async loading code requires long package names (%s)."), *PackageNameToLoad);
 	}
-	// Add to (FIFO) queue.
-	FAsyncPackage *Package = new(GObjAsyncPackages)FAsyncPackage(PackageFName, PackageGuid, PackageType, FName(/*ENAME_LinkerConstructor,*/ *PackageToLoadFrom));
-	return *Package;
+
+	FAsyncPackageDesc PackageDesc(*PackageName, InGuid ? *InGuid : FGuid(), InType, *PackageNameToLoad, InCompletionDelegate, InFlags, InPIEInstanceID, InPackagePriority);
+	FAsyncLoadingThread::Get().QueuePackage(PackageDesc);
 }
 
-FAsyncPackage& LoadPackageAsync( const FString& InPackageName, FLoadPackageAsyncDelegate CompletionDelegate, const FGuid* PackageGuid, FName PackageType, const TCHAR* InPackageToLoadFrom)
+void LoadPackageAsync(const FString& PackageName, FLoadPackageAsyncDelegate CompletionDelegate, uint32 InPackagePriority /*= 0*/)
 {
-	FAsyncPackage& AsyncPackage = LoadPackageAsync(InPackageName, PackageGuid, PackageType, InPackageToLoadFrom);
-	AsyncPackage.AddCompletionCallback(CompletionDelegate);
-	return AsyncPackage;
+	const FGuid* Guid = nullptr;
+	const TCHAR* PackageToLoadFrom = nullptr;
+	LoadPackageAsync(PackageName, Guid, NAME_None, PackageToLoadFrom, CompletionDelegate, InPackagePriority);
 }
 
 void CancelAsyncLoading()
 {
-	// Can't cancel from within async loading code
-	check(!GIsAsyncLoading);
-
-	GObjLoaded.Empty();
-
-	for (int32 PackageIndex = 0; PackageIndex < GObjAsyncPackages.Num(); PackageIndex++)
-	{
-		FAsyncPackage& PendingPackage = GObjAsyncPackages[PackageIndex];
-		PendingPackage.Cancel();
-	}
-	GObjAsyncPackages.Empty();
-	for (int32 ObjectIndex = 0; ObjectIndex < GObjConstructedDuringAsyncLoading.Num(); ObjectIndex++)
-	{
-		UObject* Object = GObjConstructedDuringAsyncLoading[ObjectIndex];
-		Object->ClearFlags(RF_AsyncLoading | RF_NeedLoad | RF_NeedPostLoad | RF_NeedPostLoadSubobjects);
-	}
-	GObjConstructedDuringAsyncLoading.Empty();
+	FAsyncLoadingThread::Get().CancelAsyncLoading();
+	CollectGarbage(GARBAGE_COLLECTION_KEEPFLAGS, false);
+}
 
-	// Not that streaming has finished, we can stop force-referencing all objects that were created during async load.
-	FAsyncObjectsReferencer::Get().EmptyReferencedObjectsAndCancelLoading();
+float GetAsyncLoadPercentage(const FName& PackageName)
+{
+	return FAsyncLoadingThread::Get().GetAsyncLoadPercentage(PackageName);
 }
 
-/**
- * Returns the async load percentage for a package in flight with the passed in name or -1 if there isn't one.
- *
- * @param	PackageName			Name of package to query load percentage for
- * @return	Async load percentage if package is currently being loaded, -1 otherwise
- */
-float GetAsyncLoadPercentage( const FName& PackageName )
+void InitAsyncThread()
 {
-	float LoadPercentage = -1.f;
-	for( int32 PackageIndex=0; PackageIndex<GObjAsyncPackages.Num(); PackageIndex++ )
-	{
-		const FAsyncPackage& PendingPackage = GObjAsyncPackages[PackageIndex];
-		if( PendingPackage.GetPackageName() == PackageName )
-		{
-			LoadPercentage = PendingPackage.GetLoadPercentage();
-			break;
-		}
-	}
-	return LoadPercentage;
+	FAsyncLoadingThread::Get().InitializeAsyncThread();
 }
 
+bool IsInAsyncLoadingThreadCoreUObjectInternal()
+{
+	return FAsyncLoadingThread::Get().IsInAsyncLoadThread();
+}
 
 /**
  * Blocks till all pending package/ linker requests are fulfilled.
@@ -1142,7 +1582,7 @@ float GetAsyncLoadPercentage( const FName& PackageName )
  */
 void FlushAsyncLoading(FName ExcludeType/*=NAME_None*/)
 {
-	if( GObjAsyncPackages.Num() )
+	if (IsAsyncLoading())
 	{
 		// Disallow low priority requests like texture streaming while we are flushing streaming
 		// in order to avoid excessive seeking.
@@ -1153,7 +1593,11 @@ void FlushAsyncLoading(FName ExcludeType/*=NAME_None*/)
 		STAT(double ThisTime = 0);
 		{
 			SCOPE_SECONDS_COUNTER(ThisTime);
-			while (ProcessAsyncLoading(false, false, 0, ExcludeType) == EAsyncPackageState::PendingImports) {}
+
+			while (IsAsyncLoading())
+			{
+				FAsyncLoadingThread::Get().TickAsyncLoading(false, false, 0, ExcludeType);
+			}
 		}
 		INC_FLOAT_STAT_BY(STAT_FAsyncPackage_FlushAsyncLoadingTime, (float)ThisTime);
 		UE_LOG(LogStreaming, Log,  TEXT("Flushed async loaders.") );
@@ -1169,22 +1613,12 @@ void FlushAsyncLoading(FName ExcludeType/*=NAME_None*/)
 	}
 }
 
-/**
- * Returns whether we are currently async loading a package.
- *
- * @return true if we are async loading a package, false otherwise
- */
-bool IsAsyncLoading()
-{
-	return GObjAsyncPackages.Num() > 0 ? true : false;
-}
-
 /**
  * @return number of active async load package requests
  */
 int32 GetNumAsyncPackages()
 {
-	return GObjAsyncPackages.Num() ;
+	return FAsyncLoadingThread::Get().GetAsyncPackagesCount();
 }
 
 /**
@@ -1196,84 +1630,20 @@ int32 GetNumAsyncPackages()
  * @param	TimeLimit		Soft limit of time this function is allowed to consume
  * @param	ExcludeType		Do not process packages associated with this specific type name
  */
-EAsyncPackageState::Type ProcessAsyncLoading( bool bUseTimeLimit, bool bUseFullTimeLimit, float TimeLimit, FName ExcludeType )
+EAsyncPackageState::Type ProcessAsyncLoading(bool bUseTimeLimit, bool bUseFullTimeLimit, float TimeLimit, FName ExcludeType)
 {
 	SCOPE_CYCLE_COUNTER(STAT_AsyncLoadingTime);
-	// Whether to continue execution.
-	EAsyncPackageState::Type LoadingState = EAsyncPackageState::Complete;
-	EAsyncPackageState::Type CompletionState = EAsyncPackageState::Complete;
-
-	bool bWasLoading = GObjAsyncPackages.Num() > 0;
-	// We need to loop as the function has to handle finish loading everything given no time limit
-	// like e.g. when called from FlushAsyncLoading.
-	for (int32 i = 0; LoadingState != EAsyncPackageState::TimeOut && i < GObjAsyncPackages.Num(); i++)
-	{
-		// Package to be loaded.
-		FAsyncPackage& Package = GObjAsyncPackages[i];
-
-		if (ExcludeType != NAME_None && ExcludeType == Package.GetPackageType())
-		{
-			// We should skip packages of this type
-			continue;
-		}
-
-		
-		if (Package.HasFinishedLoading() == false)
-		{
-			// Package tick returns EAsyncPackageState::Complete on completion.
-			// We only tick packages that have not yet been loaded.
-			LoadingState = Package.Tick( bUseTimeLimit, bUseFullTimeLimit, TimeLimit );
-		}
-		else
-		{
-			// This package has finished loading but some other package is still holding
-			// a reference to it because it has this package in its dependency list.
-			LoadingState = EAsyncPackageState::Complete;
-		}
-		if( LoadingState == EAsyncPackageState::Complete )
-		{
-			// Only remove packages that are not being referenced by anything (we still need their linkers
-			// when the referencing packages create their imports and exports).
-			if (Package.GetDependencyRefCount() == 0)
-			{
-				if( FPlatformProperties::RequiresCookedData() )
-				{
-					// Emulates ResetLoaders on the package linker's linkerroot.
-					Package.ResetLoader();
-				}
 
-				// We're done so we can remove the package now. @warning invalidates local Package variable!.
-				GObjAsyncPackages.RemoveAt( i );
+	FAsyncLoadingThread::Get().TickAsyncLoading(bUseTimeLimit, bUseFullTimeLimit, TimeLimit, ExcludeType);
 
-				// Need to process this index again as we just removed an item
-				i--;
-			}
-			else
-			{
-				// This package is being referenced by another package so we can't say streaming is complete yet.
-				LoadingState = EAsyncPackageState::PendingImports;
-			}
-		}
-		else if (!bUseTimeLimit && !FPlatformProcess::SupportsMultithreading())
-		{
-			// Tick async loading when multithreading is disabled.
-			FIOSystem::Get().TickSingleThreaded();
-		}
-		CompletionState = FMath::Min(CompletionState, LoadingState);
-		// We cannot access Package anymore!
-	}
-
-	// Not that streaming has finished, we can stop force-referencing all objects that were created during async load.
-	if (CompletionState == EAsyncPackageState::Complete && bWasLoading)
-	{
-		check(GObjAsyncPackages.Num() == 0);
-		FAsyncObjectsReferencer::Get().EmptyReferencedObjects();
-	}
-
-	return CompletionState;
+	return IsAsyncLoading() ? EAsyncPackageState::TimeOut : EAsyncPackageState::Complete;
 }
 
-
+bool IsAsyncLoadingCoreUObjectInternal()
+{
+	// GIsInitialLoad guards the async loading thread from being created too early
+	return !GIsInitialLoad && FAsyncLoadingThread::Get().IsAsyncLoadingPackages();
+}
 /*----------------------------------------------------------------------------
 	FArchiveAsync.
 ----------------------------------------------------------------------------*/
@@ -1753,3 +2123,4 @@ void FArchiveAsync::Seek( int64 InPos )
 /*----------------------------------------------------------------------------
 	End.
 ----------------------------------------------------------------------------*/
+
diff --git a/Engine/Source/Runtime/CoreUObject/Private/Serialization/AsyncLoadingThread.h b/Engine/Source/Runtime/CoreUObject/Private/Serialization/AsyncLoadingThread.h
new file mode 100644
--- /dev/null
+++ b/Engine/Source/Runtime/CoreUObject/Private/Serialization/AsyncLoadingThread.h
@@ -0,0 +1,285 @@
+// Copyright 1998-2014 Epic Games, Inc. All Rights Reserved.
+
+/*=============================================================================
+	AsyncLoadingThread.h: Unreal async loading code.
+=============================================================================*/
+#pragma once
+
+/**
+ * Async loading thread. Preloads/serializes packages on async loading thread. Postloads objects on the game thread.
+ */
+class FAsyncLoadingThread : public FRunnable
+{
+	/** Thread to run the worker FRunnable on */
+	FRunnableThread* Thread;
+	/** Stops this thread */
+	FThreadSafeCounter StopTaskCounter;
+
+	/** [ASYNC/GAME THREAD] Event used to signal there's queued packages to stream */
+	FEvent* QueuedRequestsEvent;
+	/** [ASYNC/GAME THREAD] Event used to signal loading should be cancelled */
+	FEvent* CancelLoadingEvent;
+	/** [ASYNC/GAME THREAD] List of queued packages to stream */
+	TArray<FAsyncPackageDesc*> QueuedPackages;
+	/** [ASYNC/GAME THREAD] Package queue critical section */
+	FCriticalSection QueueCritical;
+	/** [ASYNC/GAME THREAD] True if the async loading thread received a request to cancel async loading **/
+	FThreadSafeBool bShouldCancelLoading;
+	/** [ASYNC/GAME THREAD] Event used to signal there's queued packages to stream */
+	TArray<FAsyncPackage*> LoadedPackages;	
+	/** [ASYNC/GAME THREAD] Critical section for LoadedPackages list */
+	FCriticalSection LoadedPackagesCritical;
+	/** [GAME THREAD] Event used to signal there's queued packages to stream */
+	TArray<FAsyncPackage*> LoadedPackagesToProcess;
+	
+	/** [ASYNC THREAD] Array of packages that are being preloaded */
+	TArray<FAsyncPackage*> AsyncPackages;
+	/** We only lock AsyncPackages array to make GetAsyncLoadPercentage thread safe, so we only care about locking Add/Remove operations on the async thread */
+	FCriticalSection AsyncPackagesCritical;
+
+	/** [ASYNC/GAME THREAD] Number of package load requests in the async loading queue */
+	FThreadSafeCounter QueuedPackagesCounter;
+	/** [ASYNC/GAME THREAD] Number of packages being loaded on the async thread and post loaded on the game thread */
+	FThreadSafeCounter AsyncLoadingCounter;
+	/** [ASYNC/GAME THREAD] Number of packages being loaded on the async thread */
+	FThreadSafeCounter AsyncPackagesCounter;
+
+	FThreadSafeCounter AsyncThreadReady;
+
+	/** Async loading thread ID */
+	static int32 AsyncLoadingThreadID;
+
+	FAsyncLoadingThread();
+	virtual ~FAsyncLoadingThread();
+
+	/**
+	* [ASYNC/GAME THREAD] Finds an existing async package by its name.
+	*/
+	FORCEINLINE int32 FindPackageByName(TArray<FAsyncPackage*>& InPackages, const FName& PackageName)
+	{
+		for (int32 PackageIndex = 0; PackageIndex < InPackages.Num(); ++PackageIndex)
+		{
+			FAsyncPackage* Package = InPackages[PackageIndex];
+			if (Package->GetPackageName() == PackageName)
+			{
+				return PackageIndex;
+			}
+		}
+		return INDEX_NONE;
+	}
+
+public:
+
+	// Begin FRunnable interface.
+	virtual bool Init();
+	virtual uint32 Run();
+	virtual void Stop();
+	// End FRunnable interface
+
+	/** Returns the async loading thread singleton */
+	static FAsyncLoadingThread& Get();
+
+	/** True if multithreaded async loading should be used. */
+	static FORCEINLINE bool IsMultithreaded()
+	{
+		static struct FAsyncLoadingThreadEnabled
+		{
+			bool Value;
+			FAsyncLoadingThreadEnabled()
+			{
+				if (FPlatformProperties::RequiresCookedData())
+				{
+					check(GConfig);
+					bool bConfigValue = true;
+					GConfig->GetBool(TEXT("Core.System"), TEXT("AsyncLoadingThreadEnabled"), bConfigValue, GEngineIni);
+					Value = bConfigValue && FPlatformProcess::SupportsMultithreading();
+				}
+				else
+				{
+					Value = false;
+				}
+			}
+		} AsyncLoadingThreadEnabled;
+		return AsyncLoadingThreadEnabled.Value;
+	}
+
+	/** Returns true if packages are currently being loaded on the async thread */
+	FORCEINLINE bool IsAsyncLoadingPackages()
+	{
+		FPlatformMisc::MemoryBarrier();
+		return QueuedPackagesCounter.GetValue() != 0 || AsyncLoadingCounter.GetValue() != 0;
+	}
+
+	/** Returns true this codes runs on the async loading thread */
+	static FORCEINLINE bool IsInAsyncLoadThread()
+	{
+		bool bResult = false;
+		if (IsMultithreaded())
+		{
+			bResult = FPlatformTLS::GetCurrentThreadId() == AsyncLoadingThreadID;
+		}
+		else
+		{
+			bResult = IsInGameThread() && FAsyncLoadingThread::Get().bIsInAsyncLoadingTick;
+		}
+		return bResult;
+	}
+
+	/** Returns the number of async packages that are currently being processed */
+	FORCEINLINE int32 GetAsyncPackagesCount()
+	{
+		FPlatformMisc::MemoryBarrier();
+		return AsyncLoadingCounter.GetValue();
+	}
+
+	/**
+	* [ASYNC THREAD] Finds an existing async package in the AsyncPackages by its name.
+	*
+	* @param PackageName async package name.
+	* @return Index of the async package in AsyncPackages array or INDEX_NONE if not found.
+	*/
+	FORCEINLINE int32 FindAsyncPackage(const FName& PackageName)
+	{
+		checkSlow(IsInAsyncLoadThread());
+		return FindPackageByName(AsyncPackages, PackageName);
+	}
+
+	/**
+	* [ASYNC THREAD] Inserts package to queue according to priority.
+	*
+	* @param PackageName async package name.
+	*/
+	void InsertPackage(FAsyncPackage* Package)
+	{
+		checkSlow(IsInAsyncLoadThread());
+
+		// Incremented on the Async Thread, decremented on the game thread
+		AsyncLoadingCounter.Increment();
+
+		// Incemented and decremented on the AsyncThread
+		AsyncPackagesCounter.Increment();
+
+		// Insert new package keeping descending priority order in AsyncPackages
+		auto InsertIndex = AsyncPackages.IndexOfByPredicate([Package](const FAsyncPackage* Element)
+		{
+			return Element->GetPriority() <= Package->GetPriority();
+		});
+
+		InsertIndex = InsertIndex == -1 ? AsyncPackages.Num() : InsertIndex;
+		{
+			FScopeLock LockAsyncPackages(&AsyncPackagesCritical);
+			AsyncPackages.InsertUninitialized(InsertIndex);
+			AsyncPackages[InsertIndex] = Package;
+		}
+	}
+
+	/**
+	* [ASYNC THREAD] Finds an existing async package in the LoadedPackages by its name.
+	*
+	* @param PackageName async package name.
+	* @return Index of the async package in LoadedPackages array or INDEX_NONE if not found.
+	*/
+	FORCEINLINE int32 FindLoadedPackage(const FName& PackageName)
+	{
+		checkSlow(IsInAsyncLoadThread());
+		return FindPackageByName(LoadedPackages, PackageName);
+	}
+
+	/**
+	* [ASYNC/GAME THREAD] Queues a package for streaming.
+	*
+	* @param Package package descriptor.
+	*/
+	void QueuePackage(const FAsyncPackageDesc& Package);
+
+	/**
+	* [GAME THREAD] Cancels streaming
+	*/
+	void CancelAsyncLoading();
+
+	/**
+	* [ASYNC/GAME THREAD] Queues a package for streaming.
+	*
+	* @param Package package descriptor.
+	*/
+	FORCEINLINE FAsyncPackage* GetPackage(int32 PackageIndex)
+	{
+		checkSlow(IsInAsyncLoadThread());
+		return AsyncPackages[PackageIndex];
+	}
+
+	/**
+	* [ASYNC* THREAD] Loads all packages
+	*
+	* @param OutPackagesProcessed Number of packages processed in this call.
+	* @param bUseTimeLimit True if time limit should be used [time-slicing].
+	* @param bUseFullTimeLimit True if full time limit should be used [time-slicing].
+	* @param TimeLimit Maximum amount of time that can be spent in this call [time-slicing].
+	* @param ExcludeType Packages to exclude.
+	* @return The current state of async loading
+	*/
+	EAsyncPackageState::Type ProcessAsyncLoading(int32& OutPackagesProcessed, bool bUseTimeLimit = false, bool bUseFullTimeLimit = false, float TimeLimit = 0.0f, FName ExcludeType = NAME_None);
+
+	/**
+	* [GAME THREAD] Ticks game thread side of async loading.
+	*
+	* @param bUseTimeLimit True if time limit should be used [time-slicing].
+	* @param bUseFullTimeLimit True if full time limit should be used [time-slicing].
+	* @param TimeLimit Maximum amount of time that can be spent in this call [time-slicing].
+	* @param ExcludeType Packages to exclude.
+	* @return The current state of async loading
+	*/
+	EAsyncPackageState::Type TickAsyncLoading(bool bUseTimeLimit, bool bUseFullTimeLimit, float TimeLimit, FName ExcludeType);
+
+	/**
+	* [ASYNC THREAD] Main thread loop
+	*
+	* @param bUseTimeLimit True if time limit should be used [time-slicing].
+	* @param bUseFullTimeLimit True if full time limit should be used [time-slicing].
+	* @param TimeLimit Maximum amount of time that can be spent in this call [time-slicing].
+	*/
+	void TickAsyncThread(bool bUseTimeLimit, bool bUseFullTimeLimit, float TimeLimit);
+
+	/** Initializes async loading thread */
+	void InitializeAsyncThread();
+
+	/** 
+	 * [GAME THREAD] Gets the load percentage of the specified package
+	 * @param PackageName Name of the package to return async load percentage for
+	 * @return Percentage (0-100) of the async package load or -1 of package has not been found
+	 */
+	float GetAsyncLoadPercentage(const FName& PackageName);
+
+private:
+
+	/**
+	* [GAME THREAD] Performs game-thread specific operations on loaded packages (not-thread-safe PostLoad, callbacks)
+	*
+	* @param bUseTimeLimit True if time limit should be used [time-slicing].
+	* @param bUseFullTimeLimit True if full time limit should be used [time-slicing].
+	* @param TimeLimit Maximum amount of time that can be spent in this call [time-slicing].
+	* @param ExcludeType Packages to exclude.
+	* @return The current state of async loading
+	*/
+	EAsyncPackageState::Type ProcessLoadedPackages(bool bUseTimeLimit, bool bUseFullTimeLimit, float TimeLimit, FName ExcludeType);
+
+	/**
+	* [ASYNC THREAD] Creates async packages from the queued requests
+	*/
+	int32 CreateAsyncPackagesFromQueue();
+
+	/**
+	* [ASYNC THREAD] Adds a package to a list of packages that have finished loading on the async thread
+	*/
+	FORCEINLINE void AddToLoadedPackages(FAsyncPackage* Package)
+	{
+		FScopeLock LoadedLock(&LoadedPackagesCritical);
+		LoadedPackages.Add(Package);
+	}
+
+	/** Cancels async loading internally */
+	void CancelAsyncLoadingInternal();
+
+	/** True if async loading is currently being ticked, mostly used by singlethreaded ticking */
+	bool bIsInAsyncLoadingTick;
+};
