Merging UE4-Streaming to UE4

diff --git a/Engine/Source/Runtime/CoreUObject/Private/Serialization/BulkData.cpp b/Engine/Source/Runtime/CoreUObject/Private/Serialization/BulkData.cpp
--- a/Engine/Source/Runtime/CoreUObject/Private/Serialization/BulkData.cpp
+++ b/Engine/Source/Runtime/CoreUObject/Private/Serialization/BulkData.cpp
@@ -318,7 +318,7 @@ bool FUntypedBulkData::CanLoadFromDisk() const
 #if WITH_EDITOR
 	return AttachedAr != NULL;
 #else
-	return Filename != TEXT("") || Linker.IsValid();
+	return Filename != TEXT("") || (Package.IsValid() && Package->LinkerLoad);
 #endif // WITH_EDITOR
 }
 
@@ -684,7 +684,7 @@ bool FUntypedBulkData::ShouldStreamBulkData()
 	const bool bForceStream = !!(BulkDataFlags & BULKDATA_ForceStreamPayload);
 
 	return (FPlatformProperties::RequiresCookedData() && !Filename.IsEmpty() &&
-		FPlatformProcess::SupportsMultithreading() &&
+		FPlatformProcess::SupportsMultithreading() && IsInGameThread() &&
 		(bForceStream || GetBulkDataSize() > MinBulkDataSizeForAsyncLoading.Value) &&
 		MinBulkDataSizeForAsyncLoading.Value >= 0);
 }
@@ -790,23 +790,19 @@ void FUntypedBulkData::Serialize( FArchive& Ar, UObject* Owner, int32 Idx )
 
 			// We're allowing defered serialization.
 			if( Ar.IsAllowingLazyLoading() && Owner != NULL)
-			{
-				Linker = Owner->GetLinker();
-#if WITH_EDITOR
-				if (Linker)
-#else
-				if (Linker.IsValid())
-#endif
-				{
-					Filename = Linker->Filename;
-				}
+			{				
 #if WITH_EDITOR
+				Linker = Owner->GetLinker();
 				check(Linker);
 				Ar.AttachBulkData( Owner, this );
 				AttachedAr = &Ar;
-				
+				Filename = Linker->Filename;
 #else
-				check(Linker.IsValid());				
+				Package = Owner->GetOutermost();
+				check(Package.IsValid());
+				auto Linker = FLinkerLoad::FindExistingLinkerForPackage(Package.Get());
+				check(Linker);
+				Filename = Linker->Filename;		
 #endif // WITH_EDITOR
 				if (bPayloadInline)
 				{
@@ -909,10 +905,10 @@ void FUntypedBulkData::Serialize( FArchive& Ar, UObject* Owner, int32 Idx )
 			Ar << BulkDataOffsetInFile;
 
 				// try to get the linkersave object
-			ULinkerSave* LinkerSave = dynamic_cast<ULinkerSave*>(Ar.GetLinker());
+			FLinkerSave* LinkerSave = Cast<FLinkerSave>(Ar.GetLinker());
 
 			// determine whether we are going to store the payload inline or not.
-			bool bStoreInline = !!(BulkDataFlags&BULKDATA_ForceInlinePayload) || LinkerSave == NULL;
+			bool bStoreInline = !!(BulkDataFlags&BULKDATA_ForceInlinePayload) || !LinkerSave;
 
 			if (!bStoreInline)
 			{
@@ -924,7 +920,7 @@ void FUntypedBulkData::Serialize( FArchive& Ar, UObject* Owner, int32 Idx )
 				
 				// add the bulkdata storage info object to the linkersave
 				int32 Index = LinkerSave->BulkDataToAppend.AddZeroed(1);
-				ULinkerSave::FBulkDataStorageInfo& BulkStore = LinkerSave->BulkDataToAppend[Index];
+				FLinkerSave::FBulkDataStorageInfo& BulkStore = LinkerSave->BulkDataToAppend[Index];
 
 				BulkStore.BulkDataOffsetInFilePos = SavedBulkDataOffsetInFilePos;
 				BulkStore.BulkDataSizeOnDiskPos = SavedBulkDataSizeOnDiskPos;
@@ -990,7 +986,7 @@ bool FUntypedBulkData::RequiresSingleElementSerialization( FArchive& Ar )
 }
 
 /*-----------------------------------------------------------------------------
-	Accessors for friend classes ULinkerLoad and content cookers.
+	Accessors for friend classes FLinkerLoad and content cookers.
 -----------------------------------------------------------------------------*/
 
 #if WITH_EDITOR
@@ -1083,10 +1079,12 @@ void FUntypedBulkData::InitializeMemberVariables()
 	BulkData = nullptr;
 	BulkDataAsync = nullptr;
 	LockStatus = LOCKSTATUS_Unlocked;
-	bShouldFreeOnEmpty = true;
-	Linker = nullptr;
+	bShouldFreeOnEmpty = true;	
 #if WITH_EDITOR
+	Linker = nullptr;
 	AttachedAr = nullptr;
+#else
+	Package = nullptr;
 #endif
 }
 
@@ -1279,10 +1277,10 @@ void FUntypedBulkData::LoadDataIntoMemory( void* Dest )
 	AttachedAr->Seek( PushedPos );
 #else
 	bool bWasLoadedSuccessfully = false;
-	if (IsInGameThread() && Linker.IsValid())
+	if ((IsInGameThread() || IsInAsyncLoadingThread()) && Package.IsValid() && Package->LinkerLoad && Package->LinkerLoad->GetOwnerThreadId() == FPlatformTLS::GetCurrentThreadId())
 	{
-		ULinkerLoad* LinkerLoad = Linker.Get();
-		if ( LinkerLoad && LinkerLoad->Loader && !LinkerLoad->IsCompressed() )
+		FLinkerLoad* LinkerLoad = Package->LinkerLoad;
+		if (LinkerLoad && LinkerLoad->Loader && !LinkerLoad->IsCompressed())
 		{
 			FArchive* Ar = LinkerLoad;
 			// keep track of current position in this archive
diff --git a/Engine/Source/Runtime/CoreUObject/Private/Serialization/DuplicateDataReader.cpp b/Engine/Source/Runtime/CoreUObject/Private/Serialization/DuplicateDataReader.cpp
--- a/Engine/Source/Runtime/CoreUObject/Private/Serialization/DuplicateDataReader.cpp
+++ b/Engine/Source/Runtime/CoreUObject/Private/Serialization/DuplicateDataReader.cpp
@@ -1,6 +1,7 @@
 // Copyright 1998-2015 Epic Games, Inc. All Rights Reserved.
 
 #include "CoreUObjectPrivate.h"
+#include "UObject/UObjectThreadContext.h"
 
 /*----------------------------------------------------------------------------
 	FDuplicateDataReader.
@@ -23,8 +24,8 @@ FDuplicateDataReader::FDuplicateDataReader( class FUObjectAnnotationSparse<FDupl
 
 void FDuplicateDataReader::SerializeFail()
 {
-	extern UObject* GSerializedObject;
-	UE_LOG(LogObj, Fatal, TEXT("FDuplicateDataReader Overread. SerializedObject = %s SerializedProperty = %s"), *GetFullNameSafe(GSerializedObject), *GetFullNameSafe(GetSerializedProperty()));
+	FUObjectThreadContext& ThreadContext = FUObjectThreadContext::Get();
+	UE_LOG(LogObj, Fatal, TEXT("FDuplicateDataReader Overread. SerializedObject = %s SerializedProperty = %s"), *GetFullNameSafe(ThreadContext.SerializedObject), *GetFullNameSafe(GetSerializedProperty()));
 }
 
 FArchive& FDuplicateDataReader::operator<<(FName& N)
diff --git a/Engine/Source/Runtime/CoreUObject/Private/UObject/Class.cpp b/Engine/Source/Runtime/CoreUObject/Private/UObject/Class.cpp
--- a/Engine/Source/Runtime/CoreUObject/Private/UObject/Class.cpp
+++ b/Engine/Source/Runtime/CoreUObject/Private/UObject/Class.cpp
@@ -907,7 +907,7 @@ void UStruct::SerializeTaggedProperties(FArchive& Ar, uint8* Data, UStruct* Defa
 			// Check if this is a struct property and we have a redirector
 			if (Tag.Type==NAME_StructProperty && Property != NULL && Tag.Type == Property->GetID())
 			{
-				FName* NewName = ULinkerLoad::StructNameRedirects.Find(Tag.StructName);
+				FName* NewName = FLinkerLoad::StructNameRedirects.Find(Tag.StructName);
 				FName StructName = CastChecked<UStructProperty>(Property)->Struct->GetFName();
 				if (NewName != NULL && *NewName == StructName)
 				{
@@ -1346,7 +1346,7 @@ void UStruct::Serialize( FArchive& Ar )
 			if (Ar.IsPersistent() && Ar.GetLinker())
 			{
 				// make sure this is a ULinkerSave
-				ULinkerSave* LinkerSave = CastChecked<ULinkerSave>(Ar.GetLinker());
+				FLinkerSave* LinkerSave = CastChecked<FLinkerSave>(Ar.GetLinker());
 
 				// remember how we were saving
 				FArchive* SavedSaver = LinkerSave->Saver;
@@ -2514,8 +2514,8 @@ UObject* UClass::CreateDefaultObject()
 			// If this is a class that can be regenerated, it is potentially not completely loaded.  Preload and Link here to ensure we properly zero memory and read in properties for the CDO
 			if( HasAnyClassFlags(CLASS_CompiledFromBlueprint) && (PropertyLink == NULL) && !GIsDuplicatingClassForReinstancing)
 			{
-				ULinkerLoad* ClassLinker = GetLinker();
-				if( ClassLinker )
+				auto ClassLinker = GetLinker();
+				if (ClassLinker)
 				{
 					UField* FieldIt = Children;
 					while(FieldIt && (FieldIt->GetOuter() == this))
diff --git a/Engine/Source/Runtime/CoreUObject/Private/UObject/GarbageCollection.cpp b/Engine/Source/Runtime/CoreUObject/Private/UObject/GarbageCollection.cpp
--- a/Engine/Source/Runtime/CoreUObject/Private/UObject/GarbageCollection.cpp
+++ b/Engine/Source/Runtime/CoreUObject/Private/UObject/GarbageCollection.cpp
@@ -51,6 +51,122 @@ static volatile bool GIsRunningParallelReachability = false;
 /** Whether we are currently purging an object in the GC purge pass. */
 static bool GIsPurgingObject = false;
 
+/** Locks all UObject hash tables when performing GC */
+class FGCScopeLock
+{
+	/** Previous value of the GetGarbageCollectingFlag() */
+	bool bPreviousGabageCollectingFlagValue;
+public:
+
+	static FThreadSafeBool& GetGarbageCollectingFlag();
+
+	/** 
+	 * We're storing the value of GetGarbageCollectingFlag in the constructor, it's safe as only 
+	 * one thread is ever going to be setting it and calling this code - the game thread.
+	 **/
+	FORCEINLINE FGCScopeLock()
+		: bPreviousGabageCollectingFlagValue(GetGarbageCollectingFlag())
+	{		
+		void LockUObjectHashTablesForGC();
+		LockUObjectHashTablesForGC();		
+		GetGarbageCollectingFlag() = true;
+	}
+	FORCEINLINE ~FGCScopeLock()
+	{		
+		GetGarbageCollectingFlag() = bPreviousGabageCollectingFlagValue;
+		void UnlockUObjectHashTablesForGC();
+		UnlockUObjectHashTablesForGC();		
+	}
+};
+FThreadSafeBool& FGCScopeLock::GetGarbageCollectingFlag()
+{
+	static FThreadSafeBool IsGarbageCollecting(false);
+	return IsGarbageCollecting;
+}
+
+/**
+ * Garbage Collection synchronization objects
+ * Will not lock other threads if GC is not running.
+ * Has the ability to only lock for GC if no other locks are present.
+ */
+class FGCCSyncObject
+{
+	FThreadSafeCounter AsyncCounter;
+	FThreadSafeCounter GCCounter;
+	FCriticalSection Critical;
+public:
+	/** Lock on non-game thread. Will block if GC is running. */
+	void LockAsync()
+	{
+		if (!IsInGameThread())
+		{
+			FScopeLock CriticalLock(&Critical);
+			
+			// Wait until GC is done if it's currently running
+			FPlatformProcess::ConditionalSleep([&]()
+			{
+				return GCCounter.GetValue() == 0;
+			});
+
+			AsyncCounter.Increment();
+		}
+	}
+	/** Release lock from non-game thread */
+	void UnlockAsync()
+	{
+		if (!IsInGameThread())
+		{
+			AsyncCounter.Decrement();
+		}
+	}
+	/** Lock for GC. Will block if any other thread has locked. */
+	void GCLock()
+	{
+		FScopeLock CriticalLock(&Critical);
+
+		// Wait until all other threads are done if they're currently holding the lock
+		FPlatformProcess::ConditionalSleep([&]()
+		{
+			return AsyncCounter.GetValue() == 0;
+		});
+
+		GCCounter.Increment();
+	}
+	/** Lock for GC. Will not block and return false if any other thread has already locked. */
+	bool TryGCLock()
+	{		
+		bool bSuccess = false;
+		FScopeLock CriticalLock(&Critical);
+		// If any other thread is currently locking we just exit
+		if (AsyncCounter.GetValue() == 0)
+		{
+			GCCounter.Increment();
+			bSuccess = true;
+		}
+		return bSuccess;
+	}
+	/** Unlock GC */
+	void GCUnlock()
+	{
+		GCCounter.Decrement();
+	}
+};
+
+static FGCCSyncObject GGarbageCollectionGuardCritical;
+FGCScopeGuard::FGCScopeGuard()
+{
+	GGarbageCollectionGuardCritical.LockAsync();
+}
+FGCScopeGuard::~FGCScopeGuard()
+{
+	GGarbageCollectionGuardCritical.UnlockAsync();
+}
+
+bool IsGarbageCollecting()
+{
+	return FGCScopeLock::GetGarbageCollectingFlag();
+}
+
 /**
  * If set and VERIFY_DISREGARD_GC_ASSUMPTIONS is true, we verify GC assumptions about "Disregard For GC" objects. We also
  * verify that no unreachable actors/ components are referenced if VERIFY_NO_UNREACHABLE_OBJECTS_ARE_REFERENCED
@@ -436,9 +552,9 @@ public:
 		GObjectCountDuringLastMarkPhase = 0;
 
 		// Presize array and add a bit of extra slack for prefetching.
-		ObjectsToSerialize.Empty( GUObjectArray.GetObjectArrayNumMinusPermanent() + 3 );
+		ObjectsToSerialize.Empty( GetUObjectArray().GetObjectArrayNumMinusPermanent() + 3 );
 		// Make sure GC referencer object is checked for references to other objects even if it resides in permanent object pool
-		if (FPlatformProperties::RequiresCookedData() && FGCObject::GGCObjectReferencer && GUObjectArray.IsDisregardForGC(FGCObject::GGCObjectReferencer))
+		if (FPlatformProperties::RequiresCookedData() && FGCObject::GGCObjectReferencer && GetUObjectArray().IsDisregardForGC(FGCObject::GGCObjectReferencer))
 		{
 			ObjectsToSerialize.Add(FGCObject::GGCObjectReferencer);
 		}
@@ -790,7 +906,7 @@ void IncrementalPurgeGarbage( bool bUseTimeLimit, float TimeLimit )
 	if (GExitPurge)
 	{
 		GObjPurgeIsRequired = true;
-		GUObjectArray.DisableDisregardForGC();
+		GetUObjectArray().DisableDisregardForGC();
 		GObjCurrentPurgeObjectIndexNeedsReset = true;
 		GObjCurrentPurgeObjectIndexResetPastPermanent = false;
 	}
@@ -801,7 +917,7 @@ void IncrementalPurgeGarbage( bool bUseTimeLimit, float TimeLimit )
 	}
 
 	// Set 'I'm garbage collecting' flag - might be checked inside UObject::Destroy etc.
-	TGuardValue<bool> GuardIsGarbageCollecting(GIsGarbageCollecting, true);
+	FGCScopeLock GCLock;
 
 	// Incremental purge is now in progress.
 	GObjIncrementalPurgeIsInProgress						= true;
@@ -1052,9 +1168,11 @@ static const auto CVarAllowParallelGC =
  * @param	KeepFlags			objects with those flags will be kept regardless of being referenced or not
  * @param	bPerformFullPurge	if true, perform a full purge after the mark pass
  */
-
-void CollectGarbage( EObjectFlags KeepFlags, bool bPerformFullPurge )
+void CollectGarbageInternal(EObjectFlags KeepFlags, bool bPerformFullPurge)
 {
+	// We can't collect garbage while there's a load in progress. E.g. one potential issue is Import.XObject
+	check(!IsLoading());
+
 	// Helper class to register FlushAsyncLoadingCallback on first GC run.
 	struct FAddFlushAsyncLoadingCallback
 	{
@@ -1076,15 +1194,12 @@ void CollectGarbage( EObjectFlags KeepFlags, bool bPerformFullPurge )
 	// Add FlushAsyncLoadingCallback the first time CollectGarbage is called if requested by ini settings
 	static FAddFlushAsyncLoadingCallback MaybeAddFlushAsyncLoadingCallback;
 
-	// We can't collect garbage while there's a load in progress. E.g. one potential issue is Import.XObject
-	check( !IsLoading() );
-
 	// Route callbacks so we can ensure that we are e.g. not in the middle of loading something by flushing
 	// the async loading, etc...
 	FCoreUObjectDelegates::PreGarbageCollect.Broadcast();
-	
+
 	// Set 'I'm garbage collecting' flag - might be checked inside various functions.
-	GIsGarbageCollecting = true; 
+	FGCScopeLock GCLock;
 
 	UE_LOG(LogGarbage, Log, TEXT("Collecting garbage") );
 
@@ -1099,7 +1214,7 @@ void CollectGarbage( EObjectFlags KeepFlags, bool bPerformFullPurge )
 
 #if VERIFY_DISREGARD_GC_ASSUMPTIONS
 	// Only verify assumptions if option is enabled. This avoids false positives in the Editor or commandlets.
-	if( GUObjectArray.DisregardForGCEnabled() && GShouldVerifyGCAssumptions )
+	if( GetUObjectArray().DisregardForGCEnabled() && GShouldVerifyGCAssumptions )
 	{
 		// Verify that objects marked to be disregarded for GC are not referencing objects that are not part of the root set.
 		for( FObjectIterator It; It; ++It )
@@ -1109,7 +1224,7 @@ void CollectGarbage( EObjectFlags KeepFlags, bool bPerformFullPurge )
 			// Don't require UGCObjectReferencer's references to adhere to the assumptions.
 			// Although we want the referencer itself to sit in the disregard for gc set, most of the objects
 			// it's referencing will not be in the root set.
-			if (GUObjectArray.IsDisregardForGC(Object) && !Object->IsA(UGCObjectReferencer::StaticClass()))
+			if (GetUObjectArray().IsDisregardForGC(Object) && !Object->IsA(UGCObjectReferencer::StaticClass()))
 			{
 				// Serialize object with reference collector.
 				TArray<UObject*> CollectedReferences;
@@ -1120,15 +1235,12 @@ void CollectGarbage( EObjectFlags KeepFlags, bool bPerformFullPurge )
 				for( int32 ReferenceIndex=0; ReferenceIndex<CollectedReferences.Num(); ReferenceIndex++ )
 				{
 					UObject* ReferencedObject = CollectedReferences[ReferenceIndex];
-					if( ReferencedObject && !(ReferencedObject->HasAnyFlags(RF_RootSet) || GUObjectArray.IsDisregardForGC(ReferencedObject)))
+					if( ReferencedObject && !(ReferencedObject->HasAnyFlags(RF_RootSet) || GetUObjectArray().IsDisregardForGC(ReferencedObject)))
 					{
-						if( !ReferencedObject->IsA(ULinkerLoad::StaticClass()) || ReferencedObject->HasAnyFlags(RF_ClassDefaultObject) )
-						{
-							UE_LOG(LogGarbage, Warning, TEXT("Disregard for GC object %s referencing %s which is not part of root set"),
-								*Object->GetFullName(),
-								*ReferencedObject->GetFullName());
-							bShouldAssert = true;
-						}
+						UE_LOG(LogGarbage, Warning, TEXT("Disregard for GC object %s referencing %s which is not part of root set"),
+							*Object->GetFullName(),
+							*ReferencedObject->GetFullName());
+						bShouldAssert = true;
 					}
 				}
 			}
@@ -1192,13 +1304,38 @@ void CollectGarbage( EObjectFlags KeepFlags, bool bPerformFullPurge )
 		IncrementalPurgeGarbage( false );	
 	}
 
-	// We're done collecting garbage. Note that IncrementalPurgeGarbage above might already clear it internally.
-	GIsGarbageCollecting = false;
-
 	// Route callbacks to verify GC assumptions
 	FCoreUObjectDelegates::PostGarbageCollect.Broadcast();
 }
 
+void CollectGarbage(EObjectFlags KeepFlags, bool bPerformFullPurge)
+{
+	// No other thread may be performing UOBject operations while we're running
+	GGarbageCollectionGuardCritical.GCLock();
+
+	// Perform actual garbage collection
+	CollectGarbageInternal(KeepFlags, bPerformFullPurge);
+
+	// Other threads are free to use UObjects
+	GGarbageCollectionGuardCritical.GCUnlock();
+}
+
+bool TryCollectGarbage(EObjectFlags KeepFlags, bool bPerformFullPurge)
+{
+	// No other thread may be performing UOBject operations while we're running
+	bool bCanRunGC = GGarbageCollectionGuardCritical.TryGCLock();
+	if (bCanRunGC)
+	{
+		// Perform actual garbage collection
+		CollectGarbageInternal(KeepFlags, bPerformFullPurge);
+
+		// Other threads are free to use UObjects
+		GGarbageCollectionGuardCritical.GCUnlock();
+	}
+	return bCanRunGC;
+}
+
+
 /**
  * Helper function to add referenced objects via serialization
  *
@@ -1224,10 +1361,10 @@ void UObject::AddReferencedObjects(UObject* This, FReferenceCollector& Collector
 {
 #if WITH_EDITOR
 	//@todo UE4 - This seems to be required and it should not be. Seems to be related to the texture streamer.
-	ULinkerLoad* LinkerLoad = This->GetLinker();	
+	FLinkerLoad* LinkerLoad = This->GetLinker();	
 	if (LinkerLoad)
 	{
-		Collector.AddReferencedObject( LinkerLoad, This );
+		LinkerLoad->AddReferencedObjects(Collector);
 	}
 	// Required by the unified GC when running in the editor
 	if (GIsEditor)
diff --git a/Engine/Source/Runtime/CoreUObject/Private/UObject/Linker.cpp b/Engine/Source/Runtime/CoreUObject/Private/UObject/Linker.cpp
--- a/Engine/Source/Runtime/CoreUObject/Private/UObject/Linker.cpp
+++ b/Engine/Source/Runtime/CoreUObject/Private/UObject/Linker.cpp
@@ -8,24 +8,13 @@
 #include "SecureHash.h"
 #include "MessageLog.h"
 #include "UObjectToken.h"
+#include "LinkerManager.h"
+#include "UObject/UObjectThreadContext.h"
 
 DEFINE_LOG_CATEGORY(LogLinker);
 
 #define LOCTEXT_NAMESPACE "Linker"
 
-TMap<UPackage*, ULinkerLoad*>		GObjLoaders;
-TMap<UPackage*, ULinkerLoad*>		GObjPendingLoaders;
-TSet<ULinkerLoad*>					GObjLoadersWithNewImports;
-
-/** Points to the main PackageLinker currently being serialized */
-ULinkerLoad* GSerializedPackageLinker = NULL;
-/** The main Import Index currently being used for serialization by CreateImports()		*/
-int32 GSerializedImportIndex = INDEX_NONE;
-/** Points to the main Linker currently being used for serialization by CreateImports()	*/
-ULinkerLoad* GSerializedImportLinker = NULL;
-/** Points to the main UObject currently being serialized								*/
-UObject* GSerializedObject = NULL;
-
 
 /*-----------------------------------------------------------------------------
 	Helper functions.
@@ -101,18 +90,17 @@ FName FLinkerTables::GetExportClassName( int32 i )
 }
 
 /*----------------------------------------------------------------------------
-	ULinker.
+	FLinker.
 ----------------------------------------------------------------------------*/
-ULinker::ULinker(const FObjectInitializer& ObjectInitializer, UPackage* InRoot, const TCHAR* InFilename )
-:	UObject(ObjectInitializer)
-,	LinkerRoot( InRoot )
-,	Summary()
-,	Filename( InFilename )
-,	FilterClientButNotServer(false)
-,	FilterServerButNotClient(false)
+FLinker::FLinker(ELinkerType::Type InType, UPackage* InRoot, const TCHAR* InFilename)
+: LinkerType(InType)
+, bDestroyed(false)
+, LinkerRoot( InRoot )
+, Filename( InFilename )
+, FilterClientButNotServer(false)
+, FilterServerButNotClient(false)
+, ScriptSHA(nullptr)
 {
-	check(!HasAnyFlags(RF_ClassDefaultObject));
-
 	check(LinkerRoot);
 	check(InFilename);
 
@@ -126,10 +114,8 @@ ULinker::ULinker(const FObjectInitializer& ObjectInitializer, UPackage* InRoot,
 	}
 }
 
-void ULinker::Serialize( FArchive& Ar )
+void FLinker::Serialize( FArchive& Ar )
 {
-	Super::Serialize( Ar );
-
 	if( Ar.IsCountingMemory() )
 	{
 		// Can't use CountBytes as ExportMap is array of structs of arrays.
@@ -146,35 +132,36 @@ void ULinker::Serialize( FArchive& Ar )
 	// Prevent garbage collecting of linker's names and package.
 	Ar << NameMap << LinkerRoot;
 	{
-		for( int32 i=0; i<ExportMap.Num(); i++ )
+		for (auto& E : ExportMap)
 		{
-			FObjectExport& E = ExportMap[i];
 			Ar << E.ObjectName;
 		}
 	}
 	{
-		for( int32 i=0; i<ImportMap.Num(); i++ )
+		for (auto& I : ImportMap)
 		{
-			FObjectImport& I = ImportMap[i];
-			Ar << (UObject*&)I.SourceLinker;
+			UObject* LegacyLinkerObject = nullptr;
+			Ar << LegacyLinkerObject;
 			Ar << I.ClassPackage << I.ClassName;
 		}
 	}
 }
 
-void ULinker::AddReferencedObjects(UObject* InThis, FReferenceCollector& Collector)
+void FLinker::AddReferencedObjects(FReferenceCollector& Collector)
 {
-	ULinker* This = CastChecked<ULinker>(InThis);
 #if WITH_EDITOR
 	if (GIsEditor)
 	{
-		Collector.AddReferencedObject( *(UObject**)&This->LinkerRoot, This );
+		Collector.AddReferencedObject(*(UObject**)&LinkerRoot);
 	}
 #endif
-	Super::AddReferencedObjects( This, Collector );
+	for (auto& Import : ImportMap)
+	{
+		Collector.AddReferencedObject(Import.XObject);
+	}
 }
 
-// ULinker interface.
+// FLinker interface.
 /**
  * Return the path name of the UObject represented by the specified import. 
  * (can be used with StaticFindObject)
@@ -183,18 +170,16 @@ void ULinker::AddReferencedObjects(UObject* InThis, FReferenceCollector& Collect
  *
  * @return	the path name of the UObject represented by the resource at ImportIndex
  */
-FString ULinker::GetImportPathName(int32 ImportIndex)
+FString FLinker::GetImportPathName(int32 ImportIndex)
 {
-	ULinkerLoad* Loader = dynamic_cast<ULinkerLoad*>(this);
-
 	FString Result;
 	for (FPackageIndex LinkerIndex = FPackageIndex::FromImport(ImportIndex); !LinkerIndex.IsNull();)
 	{
 		FObjectResource Resource = ImpExp(LinkerIndex);
 		bool bSubobjectDelimiter=false;
 
-		if (Result.Len() > 0 && Loader != NULL && Loader->GetClassName(LinkerIndex) != NAME_Package
-			&& (Resource.OuterIndex.IsNull() || Loader->GetClassName(Resource.OuterIndex) == NAME_Package) )
+		if (Result.Len() > 0 && GetClassName(LinkerIndex) != NAME_Package
+			&& (Resource.OuterIndex.IsNull() || GetClassName(Resource.OuterIndex) == NAME_Package) )
 		{
 			bSubobjectDelimiter = true;
 		}
@@ -229,10 +214,9 @@ FString ULinker::GetImportPathName(int32 ImportIndex)
  *
  * @return	the path name of the UObject represented by the resource at ExportIndex
  */
-FString ULinker::GetExportPathName(int32 ExportIndex, const TCHAR* FakeRoot,bool bResolveForcedExports/*=false*/)
+FString FLinker::GetExportPathName(int32 ExportIndex, const TCHAR* FakeRoot,bool bResolveForcedExports/*=false*/)
 {
 	FString Result;
-	ULinkerLoad* Loader = dynamic_cast<ULinkerLoad*>(this);
 
 	bool bForcedExport = false;
 	for ( FPackageIndex LinkerIndex = FPackageIndex::FromExport(ExportIndex); !LinkerIndex.IsNull(); LinkerIndex = Exp(LinkerIndex).OuterIndex )
@@ -243,10 +227,8 @@ FString ULinker::GetExportPathName(int32 ExportIndex, const TCHAR* FakeRoot,bool
 		if ( Result.Len() > 0 )
 		{
 			// if this export is not a UPackage but this export's Outer is a UPackage, we need to use subobject notation
-			if (Loader != NULL
-			&&	(	Export.OuterIndex.IsNull()
-				||	Loader->GetExportClassName(Export.OuterIndex) == NAME_Package)
-			&&	Loader->GetExportClassName(LinkerIndex) != NAME_Package)
+			if ((Export.OuterIndex.IsNull() || GetExportClassName(Export.OuterIndex) == NAME_Package)
+			  && GetExportClassName(LinkerIndex) != NAME_Package)
 			{
 				Result = FString(SUBOBJECT_DELIMITER) + Result;
 			}
@@ -268,12 +250,12 @@ FString ULinker::GetExportPathName(int32 ExportIndex, const TCHAR* FakeRoot,bool
 	return (FakeRoot ? FakeRoot : LinkerRoot->GetPathName()) + TEXT(".") + Result;
 }
 
-FString ULinker::GetImportFullName(int32 ImportIndex)
+FString FLinker::GetImportFullName(int32 ImportIndex)
 {
 	return ImportMap[ImportIndex].ClassName.ToString() + TEXT(" ") + GetImportPathName(ImportIndex);
 }
 
-FString ULinker::GetExportFullName(int32 ExportIndex, const TCHAR* FakeRoot,bool bResolveForcedExports/*=false*/)
+FString FLinker::GetExportFullName(int32 ExportIndex, const TCHAR* FakeRoot,bool bResolveForcedExports/*=false*/)
 {
 	FPackageIndex ClassIndex = ExportMap[ExportIndex].ClassIndex;
 	FName ClassName = ClassIndex.IsNull() ? FName(NAME_Class) : ImpExp(ClassIndex).ObjectName;
@@ -284,7 +266,7 @@ FString ULinker::GetExportFullName(int32 ExportIndex, const TCHAR* FakeRoot,bool
 /**
  * Tell this linker to start SHA calculations
  */
-void ULinker::StartScriptSHAGeneration()
+void FLinker::StartScriptSHAGeneration()
 {
 	// create it if needed
 	if (ScriptSHA == NULL)
@@ -301,7 +283,7 @@ void ULinker::StartScriptSHAGeneration()
  *
  * @param ScriptCode Code to SHAify
  */
-void ULinker::UpdateScriptSHAKey(const TArray<uint8>& ScriptCode)
+void FLinker::UpdateScriptSHAKey(const TArray<uint8>& ScriptCode)
 {
 	// if we are doing SHA, update it
 	if (ScriptSHA && ScriptCode.Num())
@@ -315,7 +297,7 @@ void ULinker::UpdateScriptSHAKey(const TArray<uint8>& ScriptCode)
  *
  * @param OutKey Storage for the key bytes (20 bytes)
  */
-void ULinker::GetScriptSHAKey(uint8* OutKey)
+void FLinker::GetScriptSHAKey(uint8* OutKey)
 {
 	check(ScriptSHA);
 
@@ -324,9 +306,14 @@ void ULinker::GetScriptSHAKey(uint8* OutKey)
 	ScriptSHA->GetHash(OutKey);
 }
 
-void ULinker::BeginDestroy()
+FLinker::~FLinker()
 {
-	Super::BeginDestroy();
+	if (bDestroyed)
+	{
+		static volatile int32 xx = 0;
+		xx++;
+	}
+	bDestroyed = true;
 
 	// free any SHA memory
 	delete ScriptSHA;
@@ -353,21 +340,20 @@ void ResetLoaders( UObject* InPkg )
 	if( TopLevelPackage )
 	{
 		// Linker to reset/ detach.
-		ULinkerLoad* LinkerToReset = ULinkerLoad::FindExistingLinkerForPackage(CastChecked<UPackage>(TopLevelPackage));
-		if ( LinkerToReset )
+		auto LinkerToReset = FLinkerLoad::FindExistingLinkerForPackage(CastChecked<UPackage>(TopLevelPackage));
+		if (LinkerToReset)
 		{
-			for (TMap<UPackage*, ULinkerLoad*>::TIterator It(GObjLoaders); It; ++It)
+			for (auto Linker : FLinkerManager::Get().GetLoaders())
 			{
-				ULinkerLoad* Linker = It.Value();
 				// Detach LinkerToReset from other linker's import table.
 				if( Linker->LinkerRoot != TopLevelPackage )
 				{
-					for( int32 j=0; j<Linker->ImportMap.Num(); j++ )
+					for (auto& Import : Linker->ImportMap)
 					{
-						if( Linker->ImportMap[j].SourceLinker == LinkerToReset )
+						if (Import.SourceLinker == LinkerToReset)
 						{
-							Linker->ImportMap[j].SourceLinker	= NULL;
-							Linker->ImportMap[j].SourceIndex	= INDEX_NONE;
+							Import.SourceLinker = NULL;
+							Import.SourceIndex = INDEX_NONE;
 						}
 					}
 				}
@@ -377,19 +363,19 @@ void ResetLoaders( UObject* InPkg )
 				}
 			}
 			// Detach linker, also removes from array and sets LinkerRoot to NULL.
-			LinkerToReset->Detach(true);
+			LinkerToReset->LoadAndDetachAllBulkData();
+			delete LinkerToReset;
+			LinkerToReset = nullptr;
 		}
 	}
 	else
 	{
-		TArray<ULinkerLoad *> LinkersToDetach;
-		GObjLoaders.GenerateValueArray(LinkersToDetach);
-
-		for (int32 Index = 0; Index < LinkersToDetach.Num(); Index++)
+		auto LinkersToDetach = FLinkerManager::Get().GetLoaders();
+		for (auto Linker : LinkersToDetach)
 		{
-			ULinkerLoad* Linker = LinkersToDetach[Index];
 			// Detach linker, also removes from array and sets LinkerRoot to NULL.
-			Linker->Detach(true);
+			Linker->LoadAndDetachAllBulkData();
+			delete Linker;
 		}
 	}
 
@@ -403,11 +389,11 @@ void ResetLoaders( UObject* InPkg )
  */
 void DissociateImportsAndForcedExports()
 {
-	if( GImportCount && GObjLoadersWithNewImports.Num())
+	int32& ImportCount = FUObjectThreadContext::Get().ImportCount;
+	if (ImportCount && FLinkerManager::Get().GetLoadersWithNewImports().Num())
 	{
-		for (TSet< ULinkerLoad*>::TIterator It(GObjLoadersWithNewImports); It; ++It)
+		for (auto Linker : FLinkerManager::Get().GetLoadersWithNewImports())
 		{
-			ULinkerLoad* Linker = *It;
 			for( int32 ImportIndex=0; ImportIndex<Linker->ImportMap.Num(); ImportIndex++ )
 			{
 				FObjectImport& Import = Linker->ImportMap[ImportIndex];
@@ -422,28 +408,27 @@ void DissociateImportsAndForcedExports()
 			}
 		}
 	}
-	GImportCount = 0;
-	GObjLoadersWithNewImports.Empty();
+	ImportCount = 0;
+	FLinkerManager::Get().GetLoadersWithNewImports().Empty();
 
-	if( GForcedExportCount )
+	int32& ForcedExportCount = FUObjectThreadContext::Get().ForcedExportCount;
+	if (ForcedExportCount)
 	{
-		for (TMap<UPackage*, ULinkerLoad*>::TIterator It(GObjLoaders); It; ++It)
+		for (auto Linker : FLinkerManager::Get().GetLoaders())
 		{
-			ULinkerLoad* Linker = It.Value();
 			//@todo optimization: only dissociate exports for loaders that had forced exports created
 			//@todo optimization: since the last time this function was called.
-			for( int32 ExportIndex=0; ExportIndex<Linker->ExportMap.Num(); ExportIndex++ )
+			for (auto& Export : Linker->ExportMap)
 			{
-				FObjectExport& Export = Linker->ExportMap[ExportIndex];
-				if( Export.Object && Export.bForcedExport )
+				if (Export.Object && Export.bForcedExport)
 				{
-					Export.Object->SetLinker( NULL, INDEX_NONE );
+					Export.Object->SetLinker(NULL, INDEX_NONE);
 					Export.Object = NULL;
 				}
 			}
 		}
 	}
-	GForcedExportCount = 0;
+	ForcedExportCount = 0;
 }
 
 static void LogGetPackageLinkerError(FArchiveUObject* LinkerArchive, const TCHAR* InFilename, const FText& InFullErrorMessage, const FText& InSummaryErrorMessage, UObject* InOuter, uint32 LoadFlags)
@@ -454,15 +439,16 @@ static void LogGetPackageLinkerError(FArchiveUObject* LinkerArchive, const TCHAR
 		/** Helper function to output more detailed error info if available */
 		static void OutputErrorDetail(FArchiveUObject* InLinkerArchive, const FName& LogName)
 		{
-			if ( GSerializedObject && GSerializedImportLinker )
+			FUObjectThreadContext& ThreadContext = FUObjectThreadContext::Get();
+			if (ThreadContext.SerializedObject && ThreadContext.SerializedImportLinker)
 			{
 				FMessageLog LoadErrors(LogName);
 
 				TSharedRef<FTokenizedMessage> Message = LoadErrors.Info();
 				Message->AddToken(FTextToken::Create(LOCTEXT("FailedLoad_Message", "Failed to load")));
-				Message->AddToken(FAssetNameToken::Create(GSerializedImportLinker->GetImportPathName(GSerializedImportIndex)));
+				Message->AddToken(FAssetNameToken::Create(ThreadContext.SerializedImportLinker->GetImportPathName(ThreadContext.SerializedImportIndex)));
 				Message->AddToken(FTextToken::Create(LOCTEXT("FailedLoad_Referenced", "Referenced by")));
-				Message->AddToken(FUObjectToken::Create(GSerializedObject));
+				Message->AddToken(FUObjectToken::Create(ThreadContext.SerializedObject));
 				auto SerializedProperty = InLinkerArchive ? InLinkerArchive->GetSerializedProperty() : nullptr;
 				if (SerializedProperty != nullptr)
 				{
@@ -537,7 +523,7 @@ static void LogGetPackageLinkerError(FArchiveUObject* LinkerArchive, const TCHAR
 //
 // Find or create the linker for a package.
 //
-ULinkerLoad* GetPackageLinker
+FLinkerLoad* GetPackageLinker
 (
 	UPackage*		InOuter,
 	const TCHAR*	InLongPackageName,
@@ -547,11 +533,11 @@ ULinkerLoad* GetPackageLinker
 )
 {
 	// See if there is already a linker for this package.
-	ULinkerLoad* Result = ULinkerLoad::FindExistingLinkerForPackage(InOuter);
+	auto Result = FLinkerLoad::FindExistingLinkerForPackage(InOuter);
 
 	// Try to load the linker.
 	// See if the linker is already loaded.
-	if( Result )
+	if (Result)
 	{
 		return Result;
 	}
@@ -573,10 +559,11 @@ ULinkerLoad* GetPackageLinker
 			// In memory-only packages have no linker and this is ok.
 			if (!(LoadFlags & LOAD_AllowDll) && !(InOuter->PackageFlags & PKG_InMemoryOnly))
 			{
+				FUObjectThreadContext& ThreadContext = FUObjectThreadContext::Get();
 				FFormatNamedArguments Arguments;
 				Arguments.Add(TEXT("AssetName"), FText::FromString(InOuter->GetName()));
-				Arguments.Add(TEXT("PackageName"), FText::FromString(GSerializedPackageLinker ? *(GSerializedPackageLinker->Filename) : TEXT("NULL")));
-				LogGetPackageLinkerError(Result, GSerializedPackageLinker ? *GSerializedPackageLinker->Filename : nullptr,
+				Arguments.Add(TEXT("PackageName"), FText::FromString(ThreadContext.SerializedPackageLinker ? *(ThreadContext.SerializedPackageLinker->Filename) : TEXT("NULL")));
+				LogGetPackageLinkerError(Result, ThreadContext.SerializedPackageLinker ? *ThreadContext.SerializedPackageLinker->Filename : nullptr,
 											FText::Format(LOCTEXT("PackageNotFound", "Can't find file for asset '{AssetName}' while loading {PackageName}."), Arguments),
 											LOCTEXT("PackageNotFoundShort", "Can't find file for asset."),
 											InOuter,
@@ -635,7 +622,7 @@ ULinkerLoad* GetPackageLinker
 				return nullptr;
 			}
 			InOuter = FilenamePkg;
-			Result = ULinkerLoad::FindExistingLinkerForPackage(InOuter);
+			Result = FLinkerLoad::FindExistingLinkerForPackage(InOuter);
 		}
 		else if (InOuter != FilenamePkg) //!!should be tested and validated in new UnrealEd
 		{
@@ -665,7 +652,7 @@ ULinkerLoad* GetPackageLinker
 		// we will already have found the filename above
 		check(NewFilename.Len() > 0);
 
-		Result = ULinkerLoad::CreateLinker( InOuter, *NewFilename, LoadFlags );
+		Result = FLinkerLoad::CreateLinker( InOuter, *NewFilename, LoadFlags );
 	}
 
 	// Verify compatibility.
@@ -693,7 +680,7 @@ void ResetLoadersForSave(UObject* InOuter, const TCHAR *Filename)
 {
 	UPackage* Package = dynamic_cast<UPackage*>(InOuter);
 	// If we have a loader for the package, unload it to prevent conflicts if we are resaving to the same filename
-	ULinkerLoad* Loader = ULinkerLoad::FindExistingLinkerForPackage(Package);
+	FLinkerLoad* Loader = FLinkerLoad::FindExistingLinkerForPackage(Package);
 	// This is the loader corresponding to the package we're saving.
 	if( Loader )
 	{
@@ -710,14 +697,4 @@ void ResetLoadersForSave(UObject* InOuter, const TCHAR *Filename)
 	}
 }
 
-IMPLEMENT_CORE_INTRINSIC_CLASS(ULinker, UObject,
-	{
-		Class->ClassAddReferencedObjects = &ULinker::AddReferencedObjects;
-		Class->EmitObjectReference(STRUCT_OFFSET(ULinker, LinkerRoot), TEXT("LinkerRoot"));
-		const uint32 SkipIndexIndex = Class->EmitStructArrayBegin(STRUCT_OFFSET(ULinker, ImportMap), TEXT("ImportMap"), sizeof(FObjectImport));
-		Class->EmitObjectReference(STRUCT_OFFSET(FObjectImport, SourceLinker), TEXT("SourceLinker"));
-		Class->EmitStructArrayEnd( SkipIndexIndex );
-	}
-);
-
 #undef LOCTEXT_NAMESPACE
