Merging UE4-Streaming to UE4

diff --git a/Engine/Source/Runtime/CoreUObject/Private/UObject/LinkerLoad.cpp b/Engine/Source/Runtime/CoreUObject/Private/UObject/LinkerLoad.cpp
--- a/Engine/Source/Runtime/CoreUObject/Private/UObject/LinkerLoad.cpp
+++ b/Engine/Source/Runtime/CoreUObject/Private/UObject/LinkerLoad.cpp
@@ -8,10 +8,13 @@
 #include "LinkerPlaceholderClass.h"
 #include "LinkerPlaceholderExportObject.h"
 #include "LinkerPlaceholderFunction.h"
+#include "LinkerManager.h"
+#include "Serialization/DeferredMessageLog.h"
+#include "UObject/UObjectThreadContext.h"
 
 #define LOCTEXT_NAMESPACE "LinkerLoad"
 
-DECLARE_STATS_GROUP_VERBOSE(TEXT("Linker Load"), STATGROUP_LinkerLoad, STATCAT_Advanced);
+DECLARE_STATS_GROUP(TEXT("Linker Load"), STATGROUP_LinkerLoad, STATCAT_Advanced);
 
 DECLARE_CYCLE_STAT(TEXT("Linker Preload"),STAT_LinkerPreload,STATGROUP_LinkerLoad);
 DECLARE_CYCLE_STAT(TEXT("Linker Precache"),STAT_LinkerPrecache,STATGROUP_LinkerLoad);
@@ -22,26 +25,12 @@ DECLARE_STATS_GROUP( TEXT( "Linker Count" ), STATGROUP_LinkerCount, STATCAT_Adva
 DECLARE_DWORD_ACCUMULATOR_STAT(TEXT("Linker Count"), STAT_LinkerCount, STATGROUP_LinkerCount);
 DECLARE_DWORD_ACCUMULATOR_STAT(TEXT("Live Linker Count"), STAT_LiveLinkerCount, STATGROUP_LinkerCount);
 
-/** Points to the main PackageLinker currently being serialized (Defined in Linker.cpp) */
-extern ULinkerLoad* GSerializedPackageLinker;
-/** The main Import Index currently being used for serialization by CreateImports() (Defined in Linker.cpp) */
-extern int32 GSerializedImportIndex;
-/** Points to the main Linker currently being used for serialization by CreateImports() (Defined in Linker.cpp) */
-extern ULinkerLoad* GSerializedImportLinker;
-/** Points to the main UObject currently being serialized (Defined in Linker.cpp) */
-extern UObject* GSerializedObject;
-
-/** The most recently used export Index for serialization by CreateExport() */
-static int32 GSerializedExportIndex = INDEX_NONE;
-/** Points to the most recently used Linker for serialization by CreateExport() */
-static ULinkerLoad*	GSerializedExportLinker = NULL;
-
 /** Map that keeps track of any precached full package reads															*/
-TMap<FString, ULinkerLoad::FPackagePrecacheInfo> ULinkerLoad::PackagePrecacheMap;
+TMap<FString, FLinkerLoad::FPackagePrecacheInfo> FLinkerLoad::PackagePrecacheMap;
 
-UClass* ULinkerLoad::UTexture2DStaticClass = NULL;
+UClass* FLinkerLoad::UTexture2DStaticClass = NULL;
 
-FName ULinkerLoad::NAME_LoadErrors("LoadErrors");
+FName FLinkerLoad::NAME_LoadErrors("LoadErrors");
 
 /**
  * Here is the format for the ClassRedirection:
@@ -57,15 +46,15 @@ FName ULinkerLoad::NAME_LoadErrors("LoadErrors");
  */
 
 
-TMap<FName, FName> ULinkerLoad::ObjectNameRedirects;			    // OldClassName to NewClassName for ImportMap
-TMap<FName, FName> ULinkerLoad::ObjectNameRedirectsInstanceOnly;	// OldClassName to NewClassName for ExportMap
-TMap<FName, FName> ULinkerLoad::ObjectNameRedirectsObjectOnly;		// Object name to NewClassName for export map
-TMap<FName, FName> ULinkerLoad::GameNameRedirects;					// Game package name to new game package name
-TMap<FName, FName> ULinkerLoad::StructNameRedirects;				// Old struct name to new struct name mapping
-TMap<FString, FString> ULinkerLoad::PluginNameRedirects;			// Old plugin name to new plugin name mapping
-TMap<FName, ULinkerLoad::FSubobjectRedirect> ULinkerLoad::SubobjectNameRedirects;	
+TMap<FName, FName> FLinkerLoad::ObjectNameRedirects;			    // OldClassName to NewClassName for ImportMap
+TMap<FName, FName> FLinkerLoad::ObjectNameRedirectsInstanceOnly;	// OldClassName to NewClassName for ExportMap
+TMap<FName, FName> FLinkerLoad::ObjectNameRedirectsObjectOnly;		// Object name to NewClassName for export map
+TMap<FName, FName> FLinkerLoad::GameNameRedirects;					// Game package name to new game package name
+TMap<FName, FName> FLinkerLoad::StructNameRedirects;				// Old struct name to new struct name mapping
+TMap<FString, FString> FLinkerLoad::PluginNameRedirects;			// Old plugin name to new plugin name mapping
+TMap<FName, FLinkerLoad::FSubobjectRedirect> FLinkerLoad::SubobjectNameRedirects;	
 
-void ULinkerLoad::AddGameNameRedirect(const FName OldName, const FName NewName)
+void FLinkerLoad::AddGameNameRedirect(const FName OldName, const FName NewName)
 {
 	GameNameRedirects.Add(OldName, NewName);
 }
@@ -75,9 +64,9 @@ void ULinkerLoad::AddGameNameRedirect(const FName OldName, const FName NewName)
 ----------------------------------------------------------------------------*/
 
 /**
- * Add redirects to ULinkerLoad static map
+ * Add redirects to FLinkerLoad static map
  */
-void ULinkerLoad::CreateActiveRedirectsMap(const FString& GEngineIniName)
+void FLinkerLoad::CreateActiveRedirectsMap(const FString& GEngineIniName)
 {		
 	static bool bAlreadyInitialized_CreateActiveRedirectsMap = false;
 	if (bAlreadyInitialized_CreateActiveRedirectsMap)
@@ -190,14 +179,15 @@ struct FScopedCreateImportCounter
 	 *	@param Linker	- Current Linker
 	 *	@param Index	- Index of the current Import
 	 */
-	FScopedCreateImportCounter( ULinkerLoad* Linker, int32 Index )
+	FScopedCreateImportCounter( FLinkerLoad* Linker, int32 Index )
 	{
 		// First time CreateImport() is called for this callstack?
 		if ( Counter++ == 0 )
 		{
+			FUObjectThreadContext& ThreadContext = FUObjectThreadContext::Get();
 			// Remember the current linker and index.
-			GSerializedImportLinker = Linker;
-			GSerializedImportIndex = Index;
+			ThreadContext.SerializedImportLinker = Linker;
+			ThreadContext.SerializedImportIndex = Index;
 		}
 	}
 
@@ -207,8 +197,9 @@ struct FScopedCreateImportCounter
 		// Last time CreateImport() exits for this callstack?
 		if ( --Counter == 0 )
 		{
-			GSerializedImportLinker = NULL;
-			GSerializedImportIndex = INDEX_NONE;
+			FUObjectThreadContext& ThreadContext = FUObjectThreadContext::Get();
+			ThreadContext.SerializedImportLinker = nullptr;
+			ThreadContext.SerializedImportIndex = INDEX_NONE;
 		}
 	}
 
@@ -227,39 +218,41 @@ struct FScopedCreateExportCounter
 	 *	@param Linker	- Current Linker
 	 *	@param Index	- Index of the current Import
 	 */
-	FScopedCreateExportCounter( ULinkerLoad* Linker, int32 Index )
+	FScopedCreateExportCounter( FLinkerLoad* Linker, int32 Index )
 	{
-		GSerializedExportLinker = Linker;
-		GSerializedExportIndex = Index;
+		FUObjectThreadContext& ThreadContext = FUObjectThreadContext::Get();
+		ThreadContext.SerializedExportLinker = Linker;
+		ThreadContext.SerializedExportIndex = Index;
 	}
 
 	/** Destructor. Called upon CreateImport() exit. */
 	~FScopedCreateExportCounter()
 	{
-		GSerializedExportLinker = NULL;
-		GSerializedExportIndex = INDEX_NONE;
+		FUObjectThreadContext& ThreadContext = FUObjectThreadContext::Get();
+		ThreadContext.SerializedExportLinker = nullptr;
+		ThreadContext.SerializedExportIndex = INDEX_NONE;
 	}
 };
 
 /**
- * Exception save guard to ensure GSerializedPackageLinker is reset after this
+ * Exception save guard to ensure SerializedPackageLinker is reset after this
  * class goes out of scope.
  */
 class FSerializedPackageLinkerGuard
 {
 	/** Pointer to restore to after going out of scope. */
-	ULinkerLoad* PrevSerializedPackageLinker;
+	FLinkerLoad* PrevSerializedPackageLinker;
 public:
 	FSerializedPackageLinkerGuard() 
-	:	PrevSerializedPackageLinker(GSerializedPackageLinker) 
+		: PrevSerializedPackageLinker(FUObjectThreadContext::Get().SerializedPackageLinker)
 	{}
 	~FSerializedPackageLinkerGuard() 
 	{ 
-		GSerializedPackageLinker = PrevSerializedPackageLinker; 
+		FUObjectThreadContext::Get().SerializedPackageLinker = PrevSerializedPackageLinker;
 	}
 };
 
-namespace ULinkerDefs
+namespace FLinkerDefs
 {
 	/** Number of progress steps for reporting status to a GUI while loading packages */
 	const int32 TotalProgressSteps = 5;
@@ -315,7 +308,7 @@ static FORCEINLINE bool IsCoreUObjectPackage(const FName& PackageName)
 }
 
 /*----------------------------------------------------------------------------
-	ULinkerLoad.
+	FLinkerLoad.
 ----------------------------------------------------------------------------*/
 
 /**
@@ -323,37 +316,30 @@ static FORCEINLINE bool IsCoreUObjectPackage(const FName& PackageName)
  *
  * @param TArray<FString> to be populated
  */
-void ULinkerLoad::GetListOfPackagesInPackagePrecacheMap( TArray<FString>& ListOfPackages )
+void FLinkerLoad::GetListOfPackagesInPackagePrecacheMap( TArray<FString>& ListOfPackages )
 {
-	for ( TMap<FString, ULinkerLoad::FPackagePrecacheInfo>::TIterator It(PackagePrecacheMap); It; ++It )
+	for ( TMap<FString, FLinkerLoad::FPackagePrecacheInfo>::TIterator It(PackagePrecacheMap); It; ++It )
 	{
 		ListOfPackages.Add( It.Key() );
 	}
 }
 
-void ULinkerLoad::StaticInit(UClass* InUTexture2DStaticClass)
+void FLinkerLoad::StaticInit(UClass* InUTexture2DStaticClass)
 {
 	UTexture2DStaticClass = InUTexture2DStaticClass;
 }
 
 /**
- * Creates and returns a ULinkerLoad object.
+ * Creates and returns a FLinkerLoad object.
  *
  * @param	Parent		Parent object to load into, can be NULL (most likely case)
  * @param	Filename	Name of file on disk to load
  * @param	LoadFlags	Load flags determining behavior
  *
- * @return	new ULinkerLoad object for Parent/ Filename
+ * @return	new FLinkerLoad object for Parent/ Filename
  */
-ULinkerLoad* ULinkerLoad::CreateLinker( UPackage* Parent, const TCHAR* Filename, uint32 LoadFlags )
+FLinkerLoad* FLinkerLoad::CreateLinker( UPackage* Parent, const TCHAR* Filename, uint32 LoadFlags )
 {
-	// This should not happen during async load, otherwise we're blocking async streaming to load a package
-	// in the main thread.
-	if (FPlatformProperties::RequiresCookedData() && GIsAsyncLoading)
-	{
-		UE_LOG(LogLinker, Warning, TEXT("ULinkerLoad::CreateLinker(%s) blocking async loading!"), Filename);
-	}
-
 #if USE_CIRCULAR_DEPENDENCY_LOAD_DEFERRING
 	// we don't want the linker permanently created with the 
 	// DeferDependencyLoads flag (we also want to be able to determine if the 
@@ -368,7 +354,7 @@ ULinkerLoad* ULinkerLoad::CreateLinker( UPackage* Parent, const TCHAR* Filename,
 	LoadFlags &= ~LOAD_DeferDependencyLoads;
 #endif // USE_CIRCULAR_DEPENDENCY_LOAD_DEFERRING
 
-	ULinkerLoad* Linker = CreateLinkerAsync( Parent, Filename, LoadFlags );
+	FLinkerLoad* Linker = CreateLinkerAsync( Parent, Filename, LoadFlags );
 	{
 #if USE_CIRCULAR_DEPENDENCY_LOAD_DEFERRING
 		// the linker could already have the DeferDependencyLoads flag present 
@@ -381,7 +367,7 @@ ULinkerLoad* ULinkerLoad::CreateLinker( UPackage* Parent, const TCHAR* Filename,
 #endif // USE_CIRCULAR_DEPENDENCY_LOAD_DEFERRING
 
 		FSerializedPackageLinkerGuard Guard;	
-		GSerializedPackageLinker = Linker;
+		FUObjectThreadContext::Get().SerializedPackageLinker = Linker;
 		if (Linker->Tick( 0.f, false, false ) == LINKER_Failed)
 		{
 			return NULL;
@@ -394,12 +380,17 @@ ULinkerLoad* ULinkerLoad::CreateLinker( UPackage* Parent, const TCHAR* Filename,
 /**
  * Looks for an existing linker for the given package, without trying to make one if it doesn't exist
  */
-ULinkerLoad* ULinkerLoad::FindExistingLinkerForPackage(UPackage* Package)
+FLinkerLoad* FLinkerLoad::FindExistingLinkerForPackage(UPackage* Package)
 {
-	return GObjLoaders.FindRef(Package);
+	FLinkerLoad* Linker = nullptr;
+	if (Package)
+	{
+		Linker = Package->LinkerLoad;
+	}
+	return Linker;
 }
 
-ULinkerLoad* ULinkerLoad::FindExistingLinkerForImport(int32 Index) const
+FLinkerLoad* FLinkerLoad::FindExistingLinkerForImport(int32 Index) const
 {
 	const FObjectImport& Import = ImportMap[Index];
 	if (Import.SourceLinker != nullptr)
@@ -408,19 +399,19 @@ ULinkerLoad* ULinkerLoad::FindExistingLinkerForImport(int32 Index) const
 	}
 	else if (Import.XObject != nullptr)
 	{
-		if (ULinkerLoad* ObjLinker = Import.XObject->GetLinker())
+		if (FLinkerLoad* ObjLinker = Import.XObject->GetLinker())
 		{
 			return ObjLinker;
 		}
 	}
 
-	ULinkerLoad* FoundLinker = nullptr;
+	FLinkerLoad* FoundLinker = nullptr;
 	if (Import.OuterIndex.IsNull() && (Import.ClassName == NAME_Package))
 	{
 		FString PackageName = Import.ObjectName.ToString();
 		if (UPackage* FoundPackage = FindObject<UPackage>(/*Outer =*/nullptr, *PackageName))
 		{
-			FoundLinker = ULinkerLoad::FindExistingLinkerForPackage(FoundPackage);
+			FoundLinker = FLinkerLoad::FindExistingLinkerForPackage(FoundPackage);
 		}
 	}
 	else if (Import.OuterIndex.IsImport())
@@ -443,10 +434,10 @@ ULinkerLoad* ULinkerLoad::FindExistingLinkerForImport(int32 Index) const
  * WARNING!!!	This function is potentially very dangerous!  It should only be used at very specific times, and in very specific cases.
  *				If you're unsure, DON'T TRY TO USE IT!!!
  */
-void ULinkerLoad::PRIVATE_PatchNewObjectIntoExport(UObject* OldObject, UObject* NewObject)
+void FLinkerLoad::PRIVATE_PatchNewObjectIntoExport(UObject* OldObject, UObject* NewObject)
 {
 	// Cache off the old object's linker and export index.  We'll slide the new object in here.
-	ULinkerLoad* OldObjectLinker = OldObject->GetLinker();
+	FLinkerLoad* OldObjectLinker = OldObject->GetLinker();
 	// if this thing doesn't have a linker, then it wasn't loaded off disk and all of this is moot
 	if (OldObjectLinker)
 	{
@@ -461,18 +452,19 @@ void ULinkerLoad::PRIVATE_PatchNewObjectIntoExport(UObject* OldObject, UObject*
 		NewObject->SetLinker(OldObjectLinker, CachedLinkerIndex);
 		ObjExport.Object = NewObject;
 
-		// If the object was in the GObjLoaded queue (exported, but not yet serialized), swap out for our new object
-		int32 ObjLoadedIdx = GObjLoaded.Find(OldObject);
+		auto& ObjLoaded = FUObjectThreadContext::Get().ObjLoaded;
+		// If the object was in the ObjLoaded queue (exported, but not yet serialized), swap out for our new object
+		int32 ObjLoadedIdx = ObjLoaded.Find(OldObject);
 		if( ObjLoadedIdx != INDEX_NONE )
 		{
-			GObjLoaded[ObjLoadedIdx] = NewObject;
+			ObjLoaded[ObjLoadedIdx] = NewObject;
 		}
 	}
 }
 
-void ULinkerLoad::InvalidateExport(UObject* OldObject)
+void FLinkerLoad::InvalidateExport(UObject* OldObject)
 {
-	ULinkerLoad* OldObjectLinker = OldObject->GetLinker();
+	FLinkerLoad* OldObjectLinker = OldObject->GetLinker();
 	const int32 CachedLinkerIndex = OldObject->GetLinkerIndex();
 
 	if (OldObjectLinker && OldObjectLinker->ExportMap.IsValidIndex(CachedLinkerIndex))
@@ -482,7 +474,7 @@ void ULinkerLoad::InvalidateExport(UObject* OldObject)
 	}
 }
 
-FName ULinkerLoad::FindSubobjectRedirectName(const FName& Name)
+FName FLinkerLoad::FindSubobjectRedirectName(const FName& Name)
 {
 	FSubobjectRedirect *Redirect = SubobjectNameRedirects.Find(Name);
 	if (Redirect)
@@ -494,28 +486,24 @@ FName ULinkerLoad::FindSubobjectRedirectName(const FName& Name)
 }
 
 /**
- * Creates a ULinkerLoad object for async creation. Tick has to be called manually till it returns
+ * Creates a FLinkerLoad object for async creation. Tick has to be called manually till it returns
  * true in which case the returned linker object has finished the async creation process.
  *
  * @param	Parent		Parent object to load into, can be NULL (most likely case)
  * @param	Filename	Name of file on disk to load
  * @param	LoadFlags	Load flags determining behavior
  *
- * @return	new ULinkerLoad object for Parent/ Filename
+ * @return	new FLinkerLoad object for Parent/ Filename
  */
-ULinkerLoad* ULinkerLoad::CreateLinkerAsync( UPackage* Parent, const TCHAR* Filename, uint32 LoadFlags )
+FLinkerLoad* FLinkerLoad::CreateLinkerAsync( UPackage* Parent, const TCHAR* Filename, uint32 LoadFlags )
 {
+	check(Parent);
+
 	// See whether there already is a linker for this parent/ linker root.
-	ULinkerLoad* Linker = FindExistingLinkerForPackage(Parent);
-	if (!Linker)
-	{
-		// In case we're Async loading and something from the main thread is trying to load the same package,
-		// also look for linkers that haven't been finalized yet.
-		Linker = GObjPendingLoaders.FindRef(Parent);
-	}
+	FLinkerLoad* Linker = FindExistingLinkerForPackage(Parent);
 	if (Linker)
 	{
-		UE_LOG(LogStreaming, Log, TEXT("ULinkerLoad::CreateLinkerAsync: Found existing linker for '%s'"), *Parent->GetName());
+		UE_LOG(LogStreaming, Log, TEXT("FLinkerLoad::CreateLinkerAsync: Found existing linker for '%s'"), *Parent->GetName());
 	}
 
 	// Create a new linker if there isn't an existing one.
@@ -525,10 +513,12 @@ ULinkerLoad* ULinkerLoad::CreateLinkerAsync( UPackage* Parent, const TCHAR* File
 		{
 			LoadFlags |= LOAD_SeekFree;
 		}
-		Linker = new(EC_InternalUseOnlyConstructor) ULinkerLoad(FObjectInitializer(), Parent, Filename, LoadFlags);
-		// Add to the list of linkers that haven't been finalized yet
-		GObjPendingLoaders.Add(Parent, Linker);
+		Linker = new FLinkerLoad(Parent, Filename, LoadFlags );
+		Parent->LinkerLoad = Linker;
 	}
+	
+	check(Parent->LinkerLoad == Linker);
+
 	return Linker;
 }
 
@@ -542,7 +532,7 @@ ULinkerLoad* ULinkerLoad::CreateLinkerAsync( UPackage* Parent, const TCHAR* File
  * 
  * @return	true if linker has finished creation, false if it is still in flight
  */
-ULinkerLoad::ELinkerStatus ULinkerLoad::Tick( float InTimeLimit, bool bInUseTimeLimit, bool bInUseFullTimeLimit )
+FLinkerLoad::ELinkerStatus FLinkerLoad::Tick( float InTimeLimit, bool bInUseTimeLimit, bool bInUseFullTimeLimit )
 {
 	ELinkerStatus Status = LINKER_Loaded;
 
@@ -645,7 +635,7 @@ ULinkerLoad::ELinkerStatus ULinkerLoad::Tick( float InTimeLimit, bool bInUseTime
 
 	if (Status == LINKER_Failed)
 	{
-		GObjPendingLoaders.Remove(LinkerRoot);
+		LinkerRoot->LinkerLoad = nullptr;
 #if WITH_EDITOR
 
 		delete LoadProgressScope;
@@ -664,10 +654,28 @@ ULinkerLoad::ELinkerStatus ULinkerLoad::Tick( float InTimeLimit, bool bInUseTime
  * @param	Filename	Name of file on disk to load
  * @param	LoadFlags	Load flags determining behavior
  */
-ULinkerLoad::ULinkerLoad( const FObjectInitializer& ObjectInitializer, UPackage* InParent, const TCHAR* InFilename, uint32 InLoadFlags )
-:	ULinker( ObjectInitializer, InParent, InFilename )
-,	LoadFlags( InLoadFlags )
-,	bHaveImportsBeenVerified( false )
+FLinkerLoad::FLinkerLoad(UPackage* InParent, const TCHAR* InFilename, uint32 InLoadFlags )
+: FLinker(ELinkerType::Load, InParent, InFilename)
+, LoadFlags(InLoadFlags)
+, bHaveImportsBeenVerified(false)
+, Loader(nullptr)
+, NameMapIndex(0)
+, ImportMapIndex(0)
+, ExportMapIndex(0)
+, DependsMapIndex(0)
+, ExportHashIndex(0)
+, bHasSerializedPackageFileSummary(false)
+, bHasFixedUpImportMap(false)
+, bHasFoundExistingExports(false)
+, bHasFinishedInitialization(false)
+, bIsGatheringDependencies(false)
+, bTimeLimitExceeded(false)
+, bUseTimeLimit(false)
+, bUseFullTimeLimit(false)
+, IsTimeLimitExceededCallCount(0)
+, TimeLimit(0.0f)
+, TickStartTime(0.0)
+, bFixupExportMapDone(false)
 #if WITH_EDITOR
 ,	LoadProgressScope( nullptr )
 #endif // WITH_EDITOR
@@ -676,14 +684,27 @@ ULinkerLoad::ULinkerLoad( const FObjectInitializer& ObjectInitializer, UPackage*
 ,	ResolvingDeferredPlaceholder(nullptr)
 #endif // USE_CIRCULAR_DEPENDENCY_LOAD_DEFERRING
 {
-	check(!HasAnyFlags(RF_ClassDefaultObject)); 
-
+	FMemory::Memset(ExportHash, 0, sizeof(ExportHash));
 	INC_DWORD_STAT(STAT_LinkerCount);
 	INC_DWORD_STAT(STAT_LiveLinkerCount);
+#if UE_BUILD_DEBUG
+	FLinkerManager::Get().GetLiveLinkers().Add(this);
+#endif
+
+	OwnerThread = FPlatformTLS::GetCurrentThreadId();
 }
 
-ULinkerLoad::~ULinkerLoad()
+FLinkerLoad::~FLinkerLoad()
 {
+#if UE_BUILD_DEBUG
+	FLinkerManager::Get().GetLiveLinkers().Add(this);
+#endif
+
+	// Detaches linker.
+	Detach();
+
+	DEC_DWORD_STAT(STAT_LiveLinkerCount);
+
 #if WITH_EDITOR
 	// Make sure this is deleted if it's still allocated
 	delete LoadProgressScope;
@@ -698,7 +719,7 @@ ULinkerLoad::~ULinkerLoad()
  *
  * @return true if time limit has been exceeded (and is enabled), false otherwise (including if time limit is disabled)
  */
-bool ULinkerLoad::IsTimeLimitExceeded( const TCHAR* CurrentTask, int32 Granularity )
+bool FLinkerLoad::IsTimeLimitExceeded( const TCHAR* CurrentTask, int32 Granularity )
 {
 	IsTimeLimitExceededCallCount++;
 	if( !bTimeLimitExceeded 
@@ -712,7 +733,7 @@ bool ULinkerLoad::IsTimeLimitExceeded( const TCHAR* CurrentTask, int32 Granulari
 			// Log single operations that take longer than timelimit.
 			if( (CurrentTime - TickStartTime) > (2.5 * TimeLimit) )
 			{
- 				UE_LOG(LogStreaming, Log, TEXT("ULinkerLoad: %s took (less than) %5.2f ms"), 
+ 				UE_LOG(LogStreaming, Log, TEXT("FLinkerLoad: %s took (less than) %5.2f ms"), 
  					CurrentTask, 
  					(CurrentTime - TickStartTime) * 1000);
 			}
@@ -724,15 +745,15 @@ bool ULinkerLoad::IsTimeLimitExceeded( const TCHAR* CurrentTask, int32 Granulari
 /**
  * Creates loader used to serialize content.
  */
-ULinkerLoad::ELinkerStatus ULinkerLoad::CreateLoader()
+FLinkerLoad::ELinkerStatus FLinkerLoad::CreateLoader()
 {
-	DECLARE_SCOPE_CYCLE_COUNTER( TEXT( "ULinkerLoad::CreateLoader" ), STAT_LinkerLoad_CreateLoader, STATGROUP_LinkerLoad );
+	DECLARE_SCOPE_CYCLE_COUNTER( TEXT( "FLinkerLoad::CreateLoader" ), STAT_LinkerLoad_CreateLoader, STATGROUP_LinkerLoad );
 
 #if WITH_EDITOR
 
 	if (!LoadProgressScope)
 	{
-		LoadProgressScope = new FScopedSlowTask(ULinkerDefs::TotalProgressSteps, NSLOCTEXT("Core", "GenericLoading", "Loading..."), ShouldReportProgress());
+		LoadProgressScope = new FScopedSlowTask(FLinkerDefs::TotalProgressSteps, NSLOCTEXT("Core", "GenericLoading", "Loading..."), ShouldReportProgress());
 	}
 
 #endif
@@ -832,11 +853,11 @@ ULinkerLoad::ELinkerStatus ULinkerLoad::CreateLoader()
 		check( Loader );
 		check( !Loader->IsError() );
 
-		if( ULinkerLoad::FindExistingLinkerForPackage(LinkerRoot) )
-		{
-			UE_LOG(LogLinker, Warning, TEXT("Linker for '%s' already exists"), *LinkerRoot->GetName() );
-			return LINKER_Failed;
-		}
+		//if( FLinkerLoad::FindExistingLinkerForPackage(LinkerRoot) )
+		//{
+		//	UE_LOG(LogLinker, Warning, TEXT("Linker for '%s' already exists"), *LinkerRoot->GetName() );
+		//	return LINKER_Failed;
+		//}
 
 		// Set status info.
 		ArUE4Ver = GPackageFileUE4Version;
@@ -870,9 +891,9 @@ ULinkerLoad::ELinkerStatus ULinkerLoad::CreateLoader()
 /**
  * Serializes the package file summary.
  */
-ULinkerLoad::ELinkerStatus ULinkerLoad::SerializePackageFileSummary()
+FLinkerLoad::ELinkerStatus FLinkerLoad::SerializePackageFileSummary()
 {
-	DECLARE_SCOPE_CYCLE_COUNTER( TEXT( "ULinkerLoad::SerializePackageFileSummary" ), STAT_LinkerLoad_SerializePackageFileSummary, STATGROUP_LinkerLoad );
+	DECLARE_SCOPE_CYCLE_COUNTER( TEXT( "FLinkerLoad::SerializePackageFileSummary" ), STAT_LinkerLoad_SerializePackageFileSummary, STATGROUP_LinkerLoad );
 
 	if (bHasSerializedPackageFileSummary == false)
 	{
@@ -1051,6 +1072,11 @@ ULinkerLoad::ELinkerStatus ULinkerLoad::SerializePackageFileSummary()
 
 			// Propagate package Guid
 			LinkerRootPackage->SetGuid( Summary.Guid );
+
+			// Remember the linker versions
+			LinkerRootPackage->LinkerPackageVersion = ArUE4Ver;
+			LinkerRootPackage->LinkerLicenseeVersion = ArLicenseeUE4Ver;
+			LinkerRootPackage->LinkerCustomVersion = SummaryVersions;
 		}
 		
 		// Propagate fact that package cannot use lazy loading to archive (aka this).
@@ -1079,9 +1105,9 @@ ULinkerLoad::ELinkerStatus ULinkerLoad::SerializePackageFileSummary()
 /**
  * Serializes the name table.
  */
-ULinkerLoad::ELinkerStatus ULinkerLoad::SerializeNameMap()
+FLinkerLoad::ELinkerStatus FLinkerLoad::SerializeNameMap()
 {
-	DECLARE_SCOPE_CYCLE_COUNTER( TEXT( "ULinkerLoad::SerializeNameMap" ), STAT_LinkerLoad_SerializeNameMap, STATGROUP_LinkerLoad );
+	DECLARE_SCOPE_CYCLE_COUNTER( TEXT( "FLinkerLoad::SerializeNameMap" ), STAT_LinkerLoad_SerializeNameMap, STATGROUP_LinkerLoad );
 
 	// The name map is the first item serialized. We wait till all the header information is read
 	// before any serialization. @todo async, @todo seamless: this could be spread out across name,
@@ -1131,9 +1157,9 @@ ULinkerLoad::ELinkerStatus ULinkerLoad::SerializeNameMap()
 /**
  * Serializes the import map.
  */
-ULinkerLoad::ELinkerStatus ULinkerLoad::SerializeImportMap()
+FLinkerLoad::ELinkerStatus FLinkerLoad::SerializeImportMap()
 {
-	DECLARE_SCOPE_CYCLE_COUNTER( TEXT( "ULinkerLoad::SerializeImportMap" ), STAT_LinkerLoad_SerializeImportMap, STATGROUP_LinkerLoad );
+	DECLARE_SCOPE_CYCLE_COUNTER( TEXT( "FLinkerLoad::SerializeImportMap" ), STAT_LinkerLoad_SerializeImportMap, STATGROUP_LinkerLoad );
 
 	if( ImportMapIndex == 0 && Summary.ImportCount > 0 )
 	{
@@ -1154,9 +1180,9 @@ ULinkerLoad::ELinkerStatus ULinkerLoad::SerializeImportMap()
 /**
  * Fixes up the import map, performing remapping for backward compatibility and such.
  */
-ULinkerLoad::ELinkerStatus ULinkerLoad::FixupImportMap()
+FLinkerLoad::ELinkerStatus FLinkerLoad::FixupImportMap()
 {
-	DECLARE_SCOPE_CYCLE_COUNTER( TEXT( "ULinkerLoad::FixupImportMap" ), STAT_LinkerLoad_FixupImportMap, STATGROUP_LinkerLoad );
+	DECLARE_SCOPE_CYCLE_COUNTER( TEXT( "FLinkerLoad::FixupImportMap" ), STAT_LinkerLoad_FixupImportMap, STATGROUP_LinkerLoad );
 
 	if( bHasFixedUpImportMap == false )
 	{
@@ -1188,11 +1214,11 @@ ULinkerLoad::ELinkerStatus ULinkerLoad::FixupImportMap()
 									if (Import.ObjectName != NAME_None)
 									{
 										FString Now = GetImportFullName(i);
-										UE_LOG(LogLinker, Verbose, TEXT("ULinkerLoad::FixupImportMap() - Renamed component from %s   to   %s"), *Was, *Now);
+										UE_LOG(LogLinker, Verbose, TEXT("FLinkerLoad::FixupImportMap() - Renamed component from %s   to   %s"), *Was, *Now);
 									}
 									else
 									{
-										UE_LOG(LogLinker, Verbose, TEXT("ULinkerLoad::FixupImportMap() - Removed component %s"), *Was);
+										UE_LOG(LogLinker, Verbose, TEXT("FLinkerLoad::FixupImportMap() - Removed component %s"), *Was);
 									}
 									
 									bDone = false;
@@ -1336,7 +1362,7 @@ ULinkerLoad::ELinkerStatus ULinkerLoad::FixupImportMap()
 						// Log the object redirection to the console for review
 						if ( OldImport.ObjectName != Import.ObjectName || OldImport.ClassName != Import.ClassName || OldImport.ClassPackage != Import.ClassPackage || OldImport.OuterIndex != Import.OuterIndex )
 						{
-							UE_LOG(LogLinker, Verbose, TEXT("ULinkerLoad::FixupImportMap() - Pkg<%s> [Obj<%s> Cls<%s> Pkg<%s> Out<%s>] -> [Obj<%s> Cls<%s> Pkg<%s> Out<%s>]"), *LinkerRoot->GetName(),
+							UE_LOG(LogLinker, Verbose, TEXT("FLinkerLoad::FixupImportMap() - Pkg<%s> [Obj<%s> Cls<%s> Pkg<%s> Out<%s>] -> [Obj<%s> Cls<%s> Pkg<%s> Out<%s>]"), *LinkerRoot->GetName(),
 								*OldImport.ObjectName.ToString(), *OldImport.ClassName.ToString(), *OldImport.ClassPackage.ToString(), OldImport.OuterIndex.IsImport() ? *Imp(OldImport.OuterIndex).ObjectName.ToString() : TEXT("None"),
 								*Import.ObjectName.ToString(), *Import.ClassName.ToString(), *Import.ClassPackage.ToString(),	Import.OuterIndex.IsImport() ? *Imp(Import.OuterIndex).ObjectName.ToString() : TEXT("None"));
 						}
@@ -1353,9 +1379,9 @@ ULinkerLoad::ELinkerStatus ULinkerLoad::FixupImportMap()
 /**
  * Serializes the export map.
  */
-ULinkerLoad::ELinkerStatus ULinkerLoad::SerializeExportMap()
+FLinkerLoad::ELinkerStatus FLinkerLoad::SerializeExportMap()
 {
-	DECLARE_SCOPE_CYCLE_COUNTER( TEXT( "ULinkerLoad::SerializeExportMap" ), STAT_LinkerLoad_SerializeExportMap, STATGROUP_LinkerLoad );
+	DECLARE_SCOPE_CYCLE_COUNTER( TEXT( "FLinkerLoad::SerializeExportMap" ), STAT_LinkerLoad_SerializeExportMap, STATGROUP_LinkerLoad );
 
 	if( ExportMapIndex == 0 && Summary.ExportCount > 0 )
 	{
@@ -1373,9 +1399,9 @@ ULinkerLoad::ELinkerStatus ULinkerLoad::SerializeExportMap()
 	return ((ExportMapIndex == Summary.ExportCount) && !IsTimeLimitExceeded( TEXT("serializing export map") )) ? LINKER_Loaded : LINKER_TimedOut;
 }
 
-ULinkerLoad::ELinkerStatus ULinkerLoad::RemapImports()
+FLinkerLoad::ELinkerStatus FLinkerLoad::RemapImports()
 {
-	DECLARE_SCOPE_CYCLE_COUNTER( TEXT( "ULinkerLoad::RemapImports" ), STAT_LinkerLoad_RemapImports, STATGROUP_LinkerLoad );
+	DECLARE_SCOPE_CYCLE_COUNTER( TEXT( "FLinkerLoad::RemapImports" ), STAT_LinkerLoad_RemapImports, STATGROUP_LinkerLoad );
 
 #if !(UE_BUILD_SHIPPING || UE_BUILD_TEST)
 
@@ -1416,9 +1442,9 @@ ULinkerLoad::ELinkerStatus ULinkerLoad::RemapImports()
 /**
  * Kicks off async memory allocations for all textures that will be loaded from this package.
  */
-ULinkerLoad::ELinkerStatus ULinkerLoad::StartTextureAllocation()
+FLinkerLoad::ELinkerStatus FLinkerLoad::StartTextureAllocation()
 {
-	DECLARE_SCOPE_CYCLE_COUNTER( TEXT( "ULinkerLoad::StartTextureAllocation" ), STAT_LinkerLoad_StartTextureAllocation, STATGROUP_LinkerLoad );
+	DECLARE_SCOPE_CYCLE_COUNTER( TEXT( "FLinkerLoad::StartTextureAllocation" ), STAT_LinkerLoad_StartTextureAllocation, STATGROUP_LinkerLoad );
 
 	double StartTime = FPlatformTime::Seconds();
 	int32 NumAllocationsStarted = 0;
@@ -1486,9 +1512,9 @@ ULinkerLoad::ELinkerStatus ULinkerLoad::StartTextureAllocation()
 /**
  * Serializes the depends map.
  */
-ULinkerLoad::ELinkerStatus ULinkerLoad::SerializeDependsMap()
+FLinkerLoad::ELinkerStatus FLinkerLoad::SerializeDependsMap()
 {
-	DECLARE_SCOPE_CYCLE_COUNTER( TEXT( "ULinkerLoad::SerializeDependsMap" ), STAT_LinkerLoad_SerializeDependsMap, STATGROUP_LinkerLoad );
+	DECLARE_SCOPE_CYCLE_COUNTER( TEXT( "FLinkerLoad::SerializeDependsMap" ), STAT_LinkerLoad_SerializeDependsMap, STATGROUP_LinkerLoad );
 
 	// Skip serializing depends map if we are using seekfree loading
 	if( FPlatformProperties::RequiresCookedData() 
@@ -1521,7 +1547,7 @@ ULinkerLoad::ELinkerStatus ULinkerLoad::SerializeDependsMap()
 /**
  * Serializes thumbnails
  */
-ULinkerLoad::ELinkerStatus ULinkerLoad::SerializeThumbnails( bool bForceEnableInGame/*=false*/ )
+FLinkerLoad::ELinkerStatus FLinkerLoad::SerializeThumbnails( bool bForceEnableInGame/*=false*/ )
 {
 #if WITH_EDITORONLY_DATA
 	// Skip serializing thumbnails if we are using seekfree loading
@@ -1611,9 +1637,9 @@ ULinkerLoad::ELinkerStatus ULinkerLoad::SerializeThumbnails( bool bForceEnableIn
 /** 
  * Creates the export hash. This relies on the import and export maps having already been serialized.
  */
-ULinkerLoad::ELinkerStatus ULinkerLoad::CreateExportHash()
+FLinkerLoad::ELinkerStatus FLinkerLoad::CreateExportHash()
 {
-	DECLARE_SCOPE_CYCLE_COUNTER( TEXT( "ULinkerLoad::CreateExportHash" ), STAT_LinkerLoad_CreateExportHash, STATGROUP_LinkerLoad );
+	DECLARE_SCOPE_CYCLE_COUNTER( TEXT( "FLinkerLoad::CreateExportHash" ), STAT_LinkerLoad_CreateExportHash, STATGROUP_LinkerLoad );
 
 	// Zero initialize hash on first iteration.
 	if( ExportHashIndex == 0 )
@@ -1645,9 +1671,9 @@ ULinkerLoad::ELinkerStatus ULinkerLoad::CreateExportHash()
  * and also for script compilation as saving a package will reset its linker and loading will reload/ replace existing
  * objects without a linker.
  */
-ULinkerLoad::ELinkerStatus ULinkerLoad::FindExistingExports()
+FLinkerLoad::ELinkerStatus FLinkerLoad::FindExistingExports()
 {
-	DECLARE_SCOPE_CYCLE_COUNTER( TEXT( "ULinkerLoad::FindExistingExports" ), STAT_LinkerLoad_FindExistingExports, STATGROUP_LinkerLoad );
+	DECLARE_SCOPE_CYCLE_COUNTER( TEXT( "FLinkerLoad::FindExistingExports" ), STAT_LinkerLoad_FindExistingExports, STATGROUP_LinkerLoad );
 
 	if( bHasFoundExistingExports == false )
 	{
@@ -1681,9 +1707,9 @@ ULinkerLoad::ELinkerStatus ULinkerLoad::FindExistingExports()
 /**
  * Finalizes linker creation, adding linker to loaders array and potentially verifying imports.
  */
-ULinkerLoad::ELinkerStatus ULinkerLoad::FinalizeCreation()
+FLinkerLoad::ELinkerStatus FLinkerLoad::FinalizeCreation()
 {
-	DECLARE_SCOPE_CYCLE_COUNTER( TEXT( "ULinkerLoad::FinalizeCreation" ), STAT_LinkerLoad_FinalizeCreation, STATGROUP_LinkerLoad );
+	DECLARE_SCOPE_CYCLE_COUNTER( TEXT( "FLinkerLoad::FinalizeCreation" ), STAT_LinkerLoad_FinalizeCreation, STATGROUP_LinkerLoad );
 
 	if( bHasFinishedInitialization == false )
 	{
@@ -1692,9 +1718,7 @@ ULinkerLoad::ELinkerStatus ULinkerLoad::FinalizeCreation()
 #endif
 
 		// Add this linker to the object manager's linker array.
-		GObjLoaders.Add(LinkerRoot, this);
-		// And remove it from the pending loaders list.
-		GObjPendingLoaders.Remove(LinkerRoot);
+		FLinkerManager::Get().GetLoaders().Add(this);
 
 		// check if the package source matches the package filename's CRC (if it doens't match, a user saved this package)
 		if (Summary.PackageSource != FCrc::StrCrc_DEPRECATED(*FPaths::GetBaseFilename(Filename).ToUpper()))
@@ -1707,9 +1731,6 @@ ULinkerLoad::ELinkerStatus ULinkerLoad::FinalizeCreation()
 			Verify();
 		}
 
-		// This means that _Linker references are not NULL'd when using FArchiveReplaceObjectRef
-		SetFlags(RF_Public);
-
 		// Avoid duplicate work in the case of async linker creation.
 		bHasFinishedInitialization = true;
 
@@ -1732,7 +1753,7 @@ ULinkerLoad::ELinkerStatus ULinkerLoad::FinalizeCreation()
  * @param ExportIndex	The index of the export to hunt down
  * @return The object that was found, or NULL if it wasn't found
  */
-UObject* ULinkerLoad::FindExistingExport(int32 ExportIndex)
+UObject* FLinkerLoad::FindExistingExport(int32 ExportIndex)
 {
 	check(ExportMap.IsValidIndex(ExportIndex));
 	FObjectExport& Export = ExportMap[ExportIndex];
@@ -1757,8 +1778,7 @@ UObject* ULinkerLoad::FindExistingExport(int32 ExportIndex)
 	}
 
 	// if we found one, keep going. if we didn't find one, then this package has never been loaded before
-	// things inside a class however should not be touched, as they are in .u files and shouldn't have SetLinker called on them
-	if (OuterObject && !GetOuter()->IsInA(UClass::StaticClass()))
+	if (OuterObject)
 	{
 		// find the class of this object
 		UClass* TheClass;
@@ -1792,7 +1812,7 @@ UObject* ULinkerLoad::FindExistingExport(int32 ExportIndex)
 	return Export.Object;
 }
 
-void ULinkerLoad::Verify()
+void FLinkerLoad::Verify()
 {
 	if(!FApp::IsGame() || GIsEditor || IsRunningCommandlet())
 	{
@@ -1816,7 +1836,7 @@ void ULinkerLoad::Verify()
 	bHaveImportsBeenVerified = true;
 }
 
-FName ULinkerLoad::GetExportClassPackage( int32 i )
+FName FLinkerLoad::GetExportClassPackage( int32 i )
 {
 	FObjectExport& Export = ExportMap[ i ];
 	if( Export.ClassIndex.IsImport() )
@@ -1835,7 +1855,7 @@ FName ULinkerLoad::GetExportClassPackage( int32 i )
 	}
 }
 
-FString ULinkerLoad::GetArchiveName() const
+FString FLinkerLoad::GetArchiveName() const
 {
 	return *Filename;
 }
@@ -1849,7 +1869,7 @@ FString ULinkerLoad::GetArchiveName() const
  * @param Dependencies Array of all dependencies needed
  * @param bSkipLoadedObjects Whether to skip already loaded objects when gathering dependencies
  */
-void ULinkerLoad::GatherExportDependencies(int32 ExportIndex, TSet<FDependencyRef>& Dependencies, bool bSkipLoadedObjects)
+void FLinkerLoad::GatherExportDependencies(int32 ExportIndex, TSet<FDependencyRef>& Dependencies, bool bSkipLoadedObjects)
 {
 	// make sure we have dependencies
 	// @todo: remove this check after all packages have been saved up to VER_ADDED_LINKER_DEPENDENCIES
@@ -1909,7 +1929,7 @@ void ULinkerLoad::GatherExportDependencies(int32 ExportIndex, TSet<FDependencyRe
  * @param Dependencies Set of all dependencies needed
  * @param bSkipLoadedObjects Whether to skip already loaded objects when gathering dependencies
  */
-void ULinkerLoad::GatherImportDependencies(int32 ImportIndex, TSet<FDependencyRef>& Dependencies, bool bSkipLoadedObjects)
+void FLinkerLoad::GatherImportDependencies(int32 ImportIndex, TSet<FDependencyRef>& Dependencies, bool bSkipLoadedObjects)
 {
 	// get the import
 	FObjectImport& Import = ImportMap[ImportIndex];
@@ -1933,7 +1953,7 @@ void ULinkerLoad::GatherImportDependencies(int32 ImportIndex, TSet<FDependencyRe
 	if (Import.SourceLinker == NULL || Import.SourceIndex == INDEX_NONE)
 	{
 #if DO_CHECK
-		int32 NumObjectsBefore = GUObjectArray.GetObjectArrayNum();
+		int32 NumObjectsBefore = GetUObjectArray().GetObjectArrayNum();
 #endif
 
 		// temp storage we can ignore
@@ -1959,20 +1979,20 @@ void ULinkerLoad::GatherImportDependencies(int32 ImportIndex, TSet<FDependencyRe
 			// don't print out for intrinsic native classes
 			if (!Import.XObject || !(Import.XObject->GetClass()->HasAnyClassFlags(CLASS_Intrinsic)))
 			{
-				UE_LOG(LogLinker, Warning, TEXT("VerifyImportInner failed [(%x, %d), (%x, %d)] for %s with linker: %s %s"), 
+				UE_LOG(LogLinker, Warning, TEXT("VerifyImportInner failed [(%x, %d), (%x, %d)] for %s with linker: %s"), 
 					Import.XObject, Import.XObject ? (Import.XObject->HasAnyFlags(RF_Native) ? 1 : 0) : 0, 
 					Import.SourceLinker, Import.SourceIndex, 
-					*GetImportFullName(ImportIndex), *this->GetFullName(), *this->Filename );
+					*GetImportFullName(ImportIndex), *this->Filename );
 			}
 			EndLoad();
 			return;
 		}
 
 #if DO_CHECK && !NO_LOGGING
-		// only object we should create are one ULinkerLoad for source linker
-		if (GUObjectArray.GetObjectArrayNum() - NumObjectsBefore > 2)
+		// only object we should create are one FLinkerLoad for source linker
+		if (GetUObjectArray().GetObjectArrayNum() - NumObjectsBefore > 2)
 		{
-			UE_LOG(LogLinker, Warning, TEXT("Created %d objects checking %s"), GUObjectArray.GetObjectArrayNum() - NumObjectsBefore, *GetImportFullName(ImportIndex));
+			UE_LOG(LogLinker, Warning, TEXT("Created %d objects checking %s"), GetUObjectArray().GetObjectArrayNum() - NumObjectsBefore, *GetImportFullName(ImportIndex));
 		}
 #endif
 	}
@@ -2003,7 +2023,7 @@ void ULinkerLoad::GatherImportDependencies(int32 ImportIndex, TSet<FDependencyRe
 }
 
 
-ULinkerLoad::EVerifyResult ULinkerLoad::VerifyImport(int32 ImportIndex)
+FLinkerLoad::EVerifyResult FLinkerLoad::VerifyImport(int32 ImportIndex)
 {
 	FObjectImport& Import = ImportMap[ImportIndex];
 
@@ -2088,8 +2108,8 @@ ULinkerLoad::EVerifyResult ULinkerLoad::VerifyImport(int32 ImportIndex)
 
 					// now, fake our Import to be what the redirector pointed to
 					Import.XObject = Redir->DestinationObject;
-					GImportCount++;
-					GObjLoadersWithNewImports.Add(this);
+					FUObjectThreadContext::Get().ImportCount++;
+					FLinkerManager::Get().GetLoadersWithNewImports().Add(this);
 				}
 			}
 		}
@@ -2123,9 +2143,9 @@ ULinkerLoad::EVerifyResult ULinkerLoad::VerifyImport(int32 ImportIndex)
 				// try to get a pointer to the class of the original object so that we can display the class name of the missing resource
 				UObject* ClassPackage = FindObject<UPackage>( NULL, *Import.ClassPackage.ToString() );
 				UClass* FindClass = ClassPackage ? FindObject<UClass>( ClassPackage, *OriginalImport.ClassName.ToString() ) : NULL;
-				if( GIsEditor && !IsRunningCommandlet() )
+				if( GIsEditor && !IsRunningCommandlet())
 				{
-					FMessageLog LoadErrors(NAME_LoadErrors);
+					FDeferredMessageLog LoadErrors(NAME_LoadErrors);
 					// put something into the load warnings dialog, with any extra information from above (in WarningAppend)
 					TSharedRef<FTokenizedMessage> TokenizedMessage = LoadErrors.Error(FText::Format(LOCTEXT("ImportFailure", "Failed import for {ImportClass}"), FText::FromName(GetImportClassName(ImportIndex))));
 					TokenizedMessage->AddToken(FAssetNameToken::Create(GetImportPathName(ImportIndex)));
@@ -2142,13 +2162,14 @@ ULinkerLoad::EVerifyResult ULinkerLoad::VerifyImport(int32 ImportIndex)
 #if UE_BUILD_DEBUG
 				if( !IgnoreMissingReferencedClass( Import.ObjectName ) )
 				{
+					FUObjectThreadContext& ThreadContext = FUObjectThreadContext::Get();
 					// failure to load a class, most likely deleted instead of deprecated
 					if ( (!GIsEditor || IsRunningCommandlet()) && (FindClass->IsChildOf(UClass::StaticClass())) )
 					{
 						UE_LOG(LogLinker, Warning, TEXT("Missing Class '%s' referenced by package '%s' ('%s').  Classes should not be removed if referenced by content; mark the class 'deprecated' instead."),
 							*GetImportFullName(ImportIndex),
 							*LinkerRoot->GetName(),
-							GSerializedExportLinker ? *GSerializedExportLinker->GetExportPathName(GSerializedExportIndex) : TEXT("Unknown") );
+							ThreadContext.SerializedExportLinker ? *ThreadContext.SerializedExportLinker->GetExportPathName(ThreadContext.SerializedExportIndex) : TEXT("Unknown") );
 					}
 					// ignore warnings for missing imports if the object's class has been deprecated.
 					else if ( FindClass == NULL || !FindClass->HasAnyClassFlags(CLASS_Deprecated) )
@@ -2156,7 +2177,7 @@ ULinkerLoad::EVerifyResult ULinkerLoad::VerifyImport(int32 ImportIndex)
 						UE_LOG(LogLinker, Warning, TEXT("Missing Class '%s' referenced by package '%s' ('%s')."),
 							*GetImportFullName(ImportIndex),
 							*LinkerRoot->GetName(),
-							GSerializedExportLinker ? *GSerializedExportLinker->GetExportPathName(GSerializedExportIndex) : TEXT("Unknown") );
+							ThreadContext.SerializedExportLinker ? *ThreadContext.SerializedExportLinker->GetExportPathName(ThreadContext.SerializedExportIndex) : TEXT("Unknown") );
 					}
 				}
 #endif
@@ -2168,7 +2189,7 @@ ULinkerLoad::EVerifyResult ULinkerLoad::VerifyImport(int32 ImportIndex)
 }
 
 // Internal Load package call so that we can pass the linker that requested this package as an import dependency
-UPackage* LoadPackageInternal(UPackage* InOuter, const TCHAR* InLongPackageName, uint32 LoadFlags, ULinkerLoad* ImportLinker);
+UPackage* LoadPackageInternal(UPackage* InOuter, const TCHAR* InLongPackageName, uint32 LoadFlags, FLinkerLoad* ImportLinker);
 
 /**
  * Safely verify that an import in the ImportMap points to a good object. This decides whether or not
@@ -2178,7 +2199,7 @@ UPackage* LoadPackageInternal(UPackage* InOuter, const TCHAR* InLongPackageName,
  *
  * @return true if the wrapper should crash if it can't find a good object redirector to load
  */
-bool ULinkerLoad::VerifyImportInner(const int32 ImportIndex, FString& WarningSuffix)
+bool FLinkerLoad::VerifyImportInner(const int32 ImportIndex, FString& WarningSuffix)
 {
 	check(IsLoading());
 
@@ -2478,8 +2499,8 @@ bool ULinkerLoad::VerifyImportInner(const int32 ImportIndex, FString& WarningSuf
 			{
 				// except if we are looking for _the_ package...in which case we are looking for TmpPkg, so we are done
 				Import.XObject = TmpPkg;
-				GImportCount++;
-				GObjLoadersWithNewImports.Add(this);
+				FUObjectThreadContext::Get().ImportCount++;
+				FLinkerManager::Get().GetLoadersWithNewImports().Add(this);
 				return false;
 			}
 		}
@@ -2531,8 +2552,8 @@ bool ULinkerLoad::VerifyImportInner(const int32 ImportIndex, FString& WarningSuf
 				if (FindObject != NULL && ((LoadFlags & LOAD_FindIfFail) || bIsInMemoryOnlyOrNativeTransient))
 				{
 					Import.XObject = FindObject;
-					GImportCount++;
-					GObjLoadersWithNewImports.Add(this);
+					FUObjectThreadContext::Get().ImportCount++;
+					FLinkerManager::Get().GetLoadersWithNewImports().Add(this);
 				}
 				else
 				{
@@ -2553,7 +2574,7 @@ bool ULinkerLoad::VerifyImportInner(const int32 ImportIndex, FString& WarningSuf
 	return false;
 }
 
-UObject* ULinkerLoad::CreateExportAndPreload(int32 ExportIndex, bool bForcePreload /* = false */)
+UObject* FLinkerLoad::CreateExportAndPreload(int32 ExportIndex, bool bForcePreload /* = false */)
 {
 	UObject *Object = CreateExport(ExportIndex);
 	if (Object && (bForcePreload || dynamic_cast<UClass*>(Object) || Object->IsTemplate() || dynamic_cast<UObjectRedirector*>(Object)))
@@ -2564,7 +2585,7 @@ UObject* ULinkerLoad::CreateExportAndPreload(int32 ExportIndex, bool bForcePrelo
 	return Object;
 }
 
-UClass* ULinkerLoad::GetExportLoadClass(int32 Index)
+UClass* FLinkerLoad::GetExportLoadClass(int32 Index)
 {
 	FObjectExport& Export = ExportMap[Index];
 
@@ -2592,7 +2613,7 @@ UClass* ULinkerLoad::GetExportLoadClass(int32 Index)
 	return ExportClass;
 }
 
-int32 ULinkerLoad::LoadMetaDataFromExportMap(bool bForcePreload/* = false */)
+int32 FLinkerLoad::LoadMetaDataFromExportMap(bool bForcePreload/* = false */)
 {
 	int32 MetaDataIndex = INDEX_NONE;
 
@@ -2632,7 +2653,7 @@ int32 ULinkerLoad::LoadMetaDataFromExportMap(bool bForcePreload/* = false */)
  * @param bForcePreload	Whether to explicitly call Preload (serialize) right away instead of being
  *						called from EndLoad()
  */
-void ULinkerLoad::LoadAllObjects( bool bForcePreload )
+void FLinkerLoad::LoadAllObjects( bool bForcePreload )
 {
 #if WITH_EDITOR
 	FScopedSlowTask SlowTask(ExportMap.Num(), NSLOCTEXT("Core", "LinkerLoad_LoadingObjects", "Loading Objects"), ShouldReportProgress());
@@ -2708,7 +2729,7 @@ void ULinkerLoad::LoadAllObjects( bool bForcePreload )
  *
  * @return	ObjectName for the FObjectResource at ResourceIndex, or NAME_None if not found
  */
-FName ULinkerLoad::ResolveResourceName( FPackageIndex ResourceIndex )
+FName FLinkerLoad::ResolveResourceName( FPackageIndex ResourceIndex )
 {
 	if (ResourceIndex.IsNull())
 	{
@@ -2718,7 +2739,7 @@ FName ULinkerLoad::ResolveResourceName( FPackageIndex ResourceIndex )
 }
 
 // Find the index of a specified object without regard to specific package.
-int32 ULinkerLoad::FindExportIndex( FName ClassName, FName ClassPackage, FName ObjectName, FPackageIndex ExportOuterIndex )
+int32 FLinkerLoad::FindExportIndex( FName ClassName, FName ClassPackage, FName ObjectName, FPackageIndex ExportOuterIndex )
 {
 	int32 iHash = HashNames( ObjectName, ClassName, ClassPackage ) & (ARRAY_COUNT(ExportHash)-1);
 
@@ -2769,7 +2790,7 @@ int32 ULinkerLoad::FindExportIndex( FName ClassName, FName ClassPackage, FName O
  * @param Checked		Whether or not a failure will throw an error
  * @return The created object, or (UObject*)-1 if this is just verifying
  */
-UObject* ULinkerLoad::Create( UClass* ObjectClass, FName ObjectName, UObject* Outer, uint32 InLoadFlags, bool Checked )
+UObject* FLinkerLoad::Create( UClass* ObjectClass, FName ObjectName, UObject* Outer, uint32 InLoadFlags, bool Checked )
 {
 	// We no longer handle a NULL outer, which used to mean look in any outer, but we need fully qualified names now
 	// The other case where this was NULL is if you are calling StaticLoadObject on the top-level package, but
@@ -2872,7 +2893,7 @@ UObject* ULinkerLoad::Create( UClass* ObjectClass, FName ObjectName, UObject* Ou
  * When this function exits, Object is guaranteed to contain the data stored that was stored on disk.
  *
  * @param	Object	The object to load data for.  If the data for this object isn't stored in this
- *					ULinkerLoad, routes the call to the appropriate linker.  Data serialization is 
+ *					FLinkerLoad, routes the call to the appropriate linker.  Data serialization is 
  *					skipped if the object has already been loaded (as indicated by the RF_NeedLoad flag
  *					not set for the object), so safe to call on objects that have already been loaded.
  *					Note that this function assumes that Object has already been initialized against
@@ -2880,9 +2901,9 @@ UObject* ULinkerLoad::Create( UClass* ObjectClass, FName ObjectName, UObject* Ou
  *					If Object is a UClass and the class default object has already been created, calls
  *					Preload for the class default object as well.
  */
-void ULinkerLoad::Preload( UObject* Object )
+void FLinkerLoad::Preload( UObject* Object )
 {
-	check(IsValidLowLevel());
+	//check(IsValidLowLevel());
 	check(Object);
 
 #if USE_CIRCULAR_DEPENDENCY_LOAD_DEFERRING
@@ -2914,9 +2935,9 @@ void ULinkerLoad::Preload( UObject* Object )
 #endif // USE_CIRCULAR_DEPENDENCY_LOAD_DEFERRING
 	
 	// Preload the object if necessary.
-	if( Object->HasAnyFlags(RF_NeedLoad) )
+	if (Object->HasAnyFlags(RF_NeedLoad))
 	{
-		if( Object->GetLinker()==this )
+		if (Object->GetLinker() == this)
 		{
 			SCOPE_CYCLE_COUNTER(STAT_LinkerPreload);
 			FScopeCycleCounterUObject PreloadScope(Object, GET_STATID(STAT_LinkerPreload));
@@ -3003,12 +3024,13 @@ void ULinkerLoad::Preload( UObject* Object )
 					}
 					else
 					{
-						// Maintain the current GSerializedObjects.
-						UObject* PrevSerializedObject = GSerializedObject;
-						GSerializedObject = Object;
+						FUObjectThreadContext& ThreadContext = FUObjectThreadContext::Get();
+						// Maintain the current SerializedObjects.
+						UObject* PrevSerializedObject = ThreadContext.SerializedObject;
+						ThreadContext.SerializedObject = Object;
 						Object->Serialize( *this );
 						Object->SetFlags(RF_LoadCompleted);
-						GSerializedObject = PrevSerializedObject;
+						ThreadContext.SerializedObject = PrevSerializedObject;
 					}
 				}
 
@@ -3125,7 +3147,7 @@ void ULinkerLoad::Preload( UObject* Object )
 				if (Object->HasAnyFlags(RF_ClassDefaultObject) && Object->GetClass()->HasAnyClassFlags(CLASS_CompiledFromBlueprint))
 				{
 					Object->SetFlags(RF_NeedPostLoad|RF_WasLoaded);
-					GObjLoaded.Add( Object );
+					FUObjectThreadContext::Get().ObjLoaded.Add(Object);
 				}
 			}
 		}
@@ -3143,7 +3165,7 @@ void ULinkerLoad::Preload( UObject* Object )
  * @param OutPathName		[out] Will contain the full path for the resource
  * @param ResourceIndex		Index of a resource in the export table
  */
-void ULinkerLoad::BuildPathName( FString& OutPathName, FPackageIndex ResourceIndex ) const
+void FLinkerLoad::BuildPathName( FString& OutPathName, FPackageIndex ResourceIndex ) const
 {
 	if ( ResourceIndex.IsNull() )
 	{
@@ -3165,7 +3187,7 @@ void ULinkerLoad::BuildPathName( FString& OutPathName, FPackageIndex ResourceInd
  * @param ExportIndex	Index of the export to check
  * @return				true of the export should be loaded
  */
-bool ULinkerLoad::WillTextureBeLoaded( UClass* Class, int32 ExportIndex )
+bool FLinkerLoad::WillTextureBeLoaded( UClass* Class, int32 ExportIndex )
 {
 	const FObjectExport& Export = ExportMap[ ExportIndex ];
 
@@ -3191,10 +3213,10 @@ bool ULinkerLoad::WillTextureBeLoaded( UClass* Class, int32 ExportIndex )
 	}
 }
 
-UObject* ULinkerLoad::CreateExport( int32 Index )
+UObject* FLinkerLoad::CreateExport( int32 Index )
 {
 	FScopedCreateExportCounter ScopedCounter( this, Index );
-	FMessageLog LoadErrors(NAME_LoadErrors);
+	FDeferredMessageLog LoadErrors(NAME_LoadErrors);
 
 	// Map the object into our table.
 	FObjectExport& Export = ExportMap[ Index ];
@@ -3377,7 +3399,7 @@ UObject* ULinkerLoad::CreateExport( int32 Index )
 			// will find and return an existing object if one exists and only create a new one if there doesn't.
 			Export.Object = CreatePackage( NULL, *Export.ObjectName.ToString() );
 			check(Export.Object);
-			GForcedExportCount++;
+			FUObjectThreadContext::Get().ForcedExportCount++;
 		}
 		else
 		{
@@ -3412,7 +3434,7 @@ UObject* ULinkerLoad::CreateExport( int32 Index )
 		}
 
 		// Find the Archetype object for the one we are loading.
-		UObject* Template = GetArchetypeFromRequiredInfo(LoadClass, ThisParent, Export.ObjectName, Export.ObjectFlags);
+		UObject* Template = UObject::GetArchetypeFromRequiredInfo(LoadClass, ThisParent, Export.ObjectName, Export.ObjectFlags);
 
 		check(Template);
 		checkfSlow(((Export.ObjectFlags&RF_ClassDefaultObject)!=0 || Template->IsA(LoadClass)), TEXT("Mismatch between template %s and load class %s.  If this is a legacy blueprint or map, it may need to be resaved with bRecompileOnLoad turned off."), *Template->GetPathName(), *LoadClass->GetPathName());
@@ -3427,9 +3449,9 @@ UObject* ULinkerLoad::CreateExport( int32 Index )
 		
 		// if we require cooked data, attempt to find exports in memory first
 		if(	FPlatformProperties::RequiresCookedData() 
-		||	GIsAsyncLoading 
-		||	Export.bForcedExport 
-		||	LinkerRoot->ShouldFindExportsInMemoryFirst()
+			|| IsAsyncLoading()
+			|| Export.bForcedExport
+			|| LinkerRoot->ShouldFindExportsInMemoryFirst()
 			)
 		{
 			// Find object after making sure it isn't already set. This would be bad as the code below NULLs it in a certain
@@ -3446,7 +3468,7 @@ UObject* ULinkerLoad::CreateExport( int32 Index )
 				// Mark that we need to dissociate forced exports later on if we are a forced export.
 				if( Export.bForcedExport )
 				{
-					GForcedExportCount++;
+					FUObjectThreadContext::Get().ForcedExportCount++;
 				}
 				// Associate linker with object to avoid detachment mismatches.
 				else
@@ -3455,7 +3477,7 @@ UObject* ULinkerLoad::CreateExport( int32 Index )
 
 					// If this object was allocated but never loaded (components created by a constructor) make sure it gets loaded
 					// Do this for all subobjects created in the native constructor.
-					GObjLoaded.AddUnique( Export.Object );
+					FUObjectThreadContext::Get().ObjLoaded.AddUnique(Export.Object);
 					if( Export.Object->HasAnyFlags(RF_DefaultSubObject) 
 						|| (ThisParent && ThisParent->HasAnyFlags(RF_ClassDefaultObject)))
 					{
@@ -3478,8 +3500,8 @@ UObject* ULinkerLoad::CreateExport( int32 Index )
 				LoadClass = UObjectRedirector::StaticClass();
 				// Create new import for UObjectRedirector class
 				FObjectImport* RedirectorImport = new( ImportMap )FObjectImport( UObjectRedirector::StaticClass() );
-				GObjLoadersWithNewImports.Add(this);
-				GImportCount++;
+				FLinkerManager::Get().GetLoadersWithNewImports().Add(this);
+				FUObjectThreadContext::Get().ImportCount++;
 				Export.ClassIndex = FPackageIndex::FromImport(ImportMap.Num() - 1);
 				Export.Object = Redirector;
 				Export.Object->SetLinker( this, Index );
@@ -3605,12 +3627,12 @@ UObject* ULinkerLoad::CreateExport( int32 Index )
 				// so add the object to the list of objects that need to be loaded, which will be processed
 				// in EndLoad()
 				Export.Object->SetLinker( this, Index );
-				GObjLoaded.Add( Export.Object );
+				FUObjectThreadContext::Get().ObjLoaded.Add(Export.Object);
 			}
 		}
 		else
 		{
-			UE_LOG(LogLinker, Warning, TEXT("ULinker::CreatedExport failed to construct object %s %s"), *LoadClass->GetName(), *Export.ObjectName.ToString() );
+			UE_LOG(LogLinker, Warning, TEXT("FLinker::CreatedExport failed to construct object %s %s"), *LoadClass->GetName(), *Export.ObjectName.ToString() );
 		}
 
 		if ( Export.Object != NULL )
@@ -3662,14 +3684,14 @@ UObject* ULinkerLoad::CreateExport( int32 Index )
 			// Mark that we need to dissociate forced exports later on.
 			if( Export.bForcedExport )
 			{
-				GForcedExportCount++;
+				FUObjectThreadContext::Get().ForcedExportCount++;
 			}
 		}
 	}
 	return Export.Object;
 }
 
-bool ULinkerLoad::IsImportNative(const int32 Index) const
+bool FLinkerLoad::IsImportNative(const int32 Index) const
 {
 	const FObjectImport& Import = ImportMap[Index];
 
@@ -3694,7 +3716,7 @@ bool ULinkerLoad::IsImportNative(const int32 Index) const
 }
 
 // Return the loaded object corresponding to an import index; any errors are fatal.
-UObject* ULinkerLoad::CreateImport( int32 Index )
+UObject* FLinkerLoad::CreateImport( int32 Index )
 {
 	FScopedCreateImportCounter ScopedCounter( this, Index );
 	FObjectImport& Import = ImportMap[ Index ];
@@ -3776,8 +3798,8 @@ UObject* ULinkerLoad::CreateImport( int32 Index )
 					{		
 						// Associate import and indicate that we associated an import for later cleanup.
 						Import.XObject = FindObject;
-						GImportCount++;
-						GObjLoadersWithNewImports.Add(this);
+						FUObjectThreadContext::Get().ImportCount++;
+						FLinkerManager::Get().GetLoadersWithNewImports().Add(this);
 					}
 				}
 			}
@@ -3813,8 +3835,8 @@ UObject* ULinkerLoad::CreateImport( int32 Index )
 					}
 				}
 #endif
-				GImportCount++;
-				GObjLoadersWithNewImports.Add(this);
+				FUObjectThreadContext::Get().ImportCount++;
+				FLinkerManager::Get().GetLoadersWithNewImports().Add(this);
 			}
 		}
 	}
@@ -3822,7 +3844,7 @@ UObject* ULinkerLoad::CreateImport( int32 Index )
 }
 
 // Map an import/export index to an object; all errors here are fatal.
-UObject* ULinkerLoad::IndexToObject( FPackageIndex Index )
+UObject* FLinkerLoad::IndexToObject( FPackageIndex Index )
 {
 	if( Index.IsExport() )
 	{
@@ -3843,7 +3865,7 @@ UObject* ULinkerLoad::IndexToObject( FPackageIndex Index )
 
 
 // Detach an export from this linker.
-void ULinkerLoad::DetachExport( int32 i )
+void FLinkerLoad::DetachExport( int32 i )
 {
 	FObjectExport& E = ExportMap[ i ];
 	check(E.Object);
@@ -3855,9 +3877,9 @@ void ULinkerLoad::DetachExport( int32 i )
 	{
 		UObject* Object = E.Object;
 		UE_LOG(LogLinker, Log, TEXT("Object            : %s"), *Object->GetFullName() );
-		UE_LOG(LogLinker, Log, TEXT("Object Linker     : %s"), *Object->GetLinker()->GetFullName() );
+		//UE_LOG(LogLinker, Log, TEXT("Object Linker     : %s"), *Object->GetLinker()->GetFullName() );
 		UE_LOG(LogLinker, Log, TEXT("Linker LinkerRoot : %s"), Object->GetLinker() ? *Object->GetLinker()->LinkerRoot->GetFullName() : TEXT("None") );
-		UE_LOG(LogLinker, Log, TEXT("Detach Linker     : %s"), *GetFullName() );
+		//UE_LOG(LogLinker, Log, TEXT("Detach Linker     : %s"), *GetFullName() );
 		UE_LOG(LogLinker, Log, TEXT("Detach LinkerRoot : %s"), *LinkerRoot->GetFullName() );
 		UE_LOG(LogLinker, Fatal, TEXT("Linker object %s %s.%s mislinked!"), *GetExportClassName(i).ToString(), *LinkerRoot->GetName(), *E.ObjectName.ToString() );
 	}
@@ -3865,40 +3887,45 @@ void ULinkerLoad::DetachExport( int32 i )
 	ExportMap[i].Object->SetLinker( NULL, INDEX_NONE );
 }
 
-//@todo. Remove extern of global during next linker cleanup
-extern TArray<ULinkerLoad*> GDelayedLinkerClosePackages;
+void FLinkerLoad::LoadAndDetachAllBulkData()
+{
+#if WITH_EDITOR
+	// Detach all lazy loaders.
+	const bool bEnsureAllBulkDataIsLoaded = true;
+	DetachAllBulkData(bEnsureAllBulkDataIsLoaded);
+#endif
+}
 
-/**
- * Detaches linker from exports and removes itself from array of loaders.
- */
-void ULinkerLoad::Detach( bool bEnsureAllBulkDataIsLoaded )
+void FLinkerLoad::Detach()
 {
 #if WITH_EDITOR
 	// Detach all lazy loaders.
-	DetachAllBulkData( bEnsureAllBulkDataIsLoaded );
+	const bool bEnsureAllBulkDataIsLoaded = false;
+	DetachAllBulkData(bEnsureAllBulkDataIsLoaded);
 #endif
+
 	// Detach all objects linked with this linker.
-	for( int32 i=0; i<ExportMap.Num(); i++ )
+	for (int32 ExportIndex = 0; ExportIndex < ExportMap.Num(); ++ExportIndex)
 	{	
-		if( ExportMap[i].Object )
+		if (ExportMap[ExportIndex].Object)
 		{
-			DetachExport( i );
+			DetachExport(ExportIndex);
 		}
 	}
 
 	// Remove from object manager, if it has been added.
-	GObjLoaders.Remove(LinkerRoot);
-	GObjPendingLoaders.Remove(LinkerRoot);
-	GObjLoadersWithNewImports.Remove(this);
+	FLinkerManager::Get().GetLoaders().Remove(this);
+	FLinkerManager::Get().GetLoadersWithNewImports().Remove(this);
 	if (!FPlatformProperties::HasEditorOnlyData())
 	{
-		GDelayedLinkerClosePackages.Remove(this);
+		FUObjectThreadContext::Get().DelayedLinkerClosePackages.Remove(this);
 	}
-	if( Loader )
+
+	if (Loader)
 	{
 		delete Loader;
-	}
-	Loader = NULL;
+		Loader = nullptr;
+	}	
 
 	// Empty out no longer used arrays.
 	NameMap.Empty();
@@ -3910,26 +3937,21 @@ void ULinkerLoad::Detach( bool bEnsureAllBulkDataIsLoaded )
 #endif // USE_CIRCULAR_DEPENDENCY_LOAD_DEFERRING
 
 	// Make sure we're never associated with LinkerRoot again.
-	LinkerRoot = NULL;
-}
-
-void ULinkerLoad::BeginDestroy()
-{
-	DEC_DWORD_STAT(STAT_LiveLinkerCount);
-	// Detaches linker.
-	Detach( false );
-	Super::BeginDestroy();
+	if (LinkerRoot)
+	{
+		LinkerRoot->LinkerLoad = nullptr;
+		LinkerRoot = nullptr;
+	}
 }
 
 #if WITH_EDITOR
-
 /**
  * Attaches/ associates the passed in bulk data object with the linker.
  *
  * @param	Owner		UObject owning the bulk data
  * @param	BulkData	Bulk data object to associate
  */
-void ULinkerLoad::AttachBulkData( UObject* Owner, FUntypedBulkData* BulkData )
+void FLinkerLoad::AttachBulkData( UObject* Owner, FUntypedBulkData* BulkData )
 {
 	check( BulkDataLoaders.Find(BulkData)==INDEX_NONE );
 	BulkDataLoaders.Add( BulkData );
@@ -3941,7 +3963,7 @@ void ULinkerLoad::AttachBulkData( UObject* Owner, FUntypedBulkData* BulkData )
  * @param	BulkData	Bulk data object to detach
  * @param	bEnsureBulkDataIsLoaded	Whether to ensure that the bulk data is loaded before detaching
  */
-void ULinkerLoad::DetachBulkData( FUntypedBulkData* BulkData, bool bEnsureBulkDataIsLoaded )
+void FLinkerLoad::DetachBulkData( FUntypedBulkData* BulkData, bool bEnsureBulkDataIsLoaded )
 {
 	int32 RemovedCount = BulkDataLoaders.Remove( BulkData );
 	if( RemovedCount!=1 )
@@ -3956,13 +3978,13 @@ void ULinkerLoad::DetachBulkData( FUntypedBulkData* BulkData, bool bEnsureBulkDa
  *
  * @param	bEnsureBulkDataIsLoaded	Whether to ensure that the bulk data is loaded before detaching
  */
-void ULinkerLoad::DetachAllBulkData( bool bEnsureAllBulkDataIsLoaded )
+void FLinkerLoad::DetachAllBulkData(bool bEnsureAllBulkDataIsLoaded)
 {
-	for( int32 BulkDataIndex=0; BulkDataIndex<BulkDataLoaders.Num(); BulkDataIndex++ )
+	auto BulkDataToDetach = BulkDataLoaders;
+	for (auto BulkData : BulkDataToDetach)
 	{
-		FUntypedBulkData* BulkData = BulkDataLoaders[BulkDataIndex];
 		check( BulkData );
-		BulkData->DetachFromArchive( this, bEnsureAllBulkDataIsLoaded );
+		BulkData->DetachFromArchive(this, bEnsureAllBulkDataIsLoaded);
 	}
 	BulkDataLoaders.Empty();
 }
@@ -3984,27 +4006,27 @@ void ULinkerLoad::DetachAllBulkData( bool bEnsureAllBulkDataIsLoaded )
  * @param	PrecacheSize	Number of bytes to precache
  * @return	false if precache operation is still pending, true otherwise
  */
-bool ULinkerLoad::Precache( int64 PrecacheOffset, int64 PrecacheSize )
+bool FLinkerLoad::Precache( int64 PrecacheOffset, int64 PrecacheSize )
 {
 	return Loader->Precache( PrecacheOffset, PrecacheSize );
 }
 
-void ULinkerLoad::Seek( int64 InPos )
+void FLinkerLoad::Seek( int64 InPos )
 {
 	Loader->Seek( InPos );
 }
 
-int64 ULinkerLoad::Tell()
+int64 FLinkerLoad::Tell()
 {
 	return Loader->Tell();
 }
 
-int64 ULinkerLoad::TotalSize()
+int64 FLinkerLoad::TotalSize()
 {
 	return Loader->TotalSize();
 }
 
-FArchive& ULinkerLoad::operator<<( UObject*& Object )
+FArchive& FLinkerLoad::operator<<( UObject*& Object )
 {
 	FPackageIndex Index;
 	FArchive& Ar = *this;
@@ -4017,7 +4039,7 @@ FArchive& ULinkerLoad::operator<<( UObject*& Object )
 	return *this;
 }
 
-FArchive& ULinkerLoad::operator<<( FLazyObjectPtr& LazyObjectPtr)
+FArchive& FLinkerLoad::operator<<( FLazyObjectPtr& LazyObjectPtr)
 {
 	FArchive& Ar = *this;
 	FUniqueObjectGuid ID;
@@ -4026,7 +4048,7 @@ FArchive& ULinkerLoad::operator<<( FLazyObjectPtr& LazyObjectPtr)
 	return Ar;
 }
 
-FArchive& ULinkerLoad::operator<<( FAssetPtr& AssetPtr)
+FArchive& FLinkerLoad::operator<<( FAssetPtr& AssetPtr)
 {
 	FArchive& Ar = *this;
 	FStringAssetReference ID;
@@ -4035,7 +4057,7 @@ FArchive& ULinkerLoad::operator<<( FAssetPtr& AssetPtr)
 	return Ar;
 }
 
-FArchive& ULinkerLoad::operator<<( FName& Name )
+FArchive& FLinkerLoad::operator<<( FName& Name )
 {
 	NAME_INDEX NameIndex;
 	FArchive& Ar = *this;
@@ -4071,8 +4093,9 @@ FArchive& ULinkerLoad::operator<<( FName& Name )
 	return *this;
 }
 
-void ULinkerLoad::Serialize( void* V, int64 Length )
+void FLinkerLoad::Serialize( void* V, int64 Length )
 {
+	checkSlow(FPlatformTLS::GetCurrentThreadId() == OwnerThread);
 	Loader->Serialize( V, Length );
 }
 
@@ -4081,7 +4104,7 @@ void ULinkerLoad::Serialize( void* V, int64 Length )
 * 
 * @param PackageName Name of package to read in. Must be the same name as passed into LoadPackage
 */
-void ULinkerLoad::AsyncPreloadPackage(const TCHAR* PackageName)
+void FLinkerLoad::AsyncPreloadPackage(const TCHAR* PackageName)
 {
 	// get package filename
 	FString PackageFilename;
@@ -4153,7 +4176,7 @@ void ULinkerLoad::AsyncPreloadPackage(const TCHAR* PackageName)
 /**
  * Called when an object begins serializing property data using script serialization.
  */
-void ULinkerLoad::MarkScriptSerializationStart( const UObject* Obj )
+void FLinkerLoad::MarkScriptSerializationStart( const UObject* Obj )
 {
 	if ( Obj != NULL && Obj->GetLinker() == this && ExportMap.IsValidIndex(Obj->GetLinkerIndex()) )
 	{
@@ -4165,9 +4188,9 @@ void ULinkerLoad::MarkScriptSerializationStart( const UObject* Obj )
 /**
  * Called when an object stops serializing property data using script serialization.
  */
-void ULinkerLoad::MarkScriptSerializationEnd( const UObject* Obj )
+void FLinkerLoad::MarkScriptSerializationEnd( const UObject* Obj )
 {
-	if ( Obj != NULL && Obj->GetLinker() == this && ExportMap.IsValidIndex(Obj->GetLinkerIndex()) )
+	if (Obj != NULL && Obj->GetLinker() == this && ExportMap.IsValidIndex(Obj->GetLinkerIndex()))
 	{
 		FObjectExport& Export = ExportMap[Obj->GetLinkerIndex()];
 		Export.ScriptSerializationEndOffset = Tell();
@@ -4177,7 +4200,7 @@ void ULinkerLoad::MarkScriptSerializationEnd( const UObject* Obj )
 /**
  * Locates the class adjusted index and its package adjusted index for a given class name in the import map
  */
-bool ULinkerLoad::FindImportClassAndPackage( FName ClassName, FPackageIndex &ClassIdx, FPackageIndex &PackageIdx )
+bool FLinkerLoad::FindImportClassAndPackage( FName ClassName, FPackageIndex &ClassIdx, FPackageIndex &PackageIdx )
 {
 	for ( int32 ImportMapIdx = 0; ImportMapIdx < ImportMap.Num(); ImportMapIdx++ )
 	{
@@ -4195,7 +4218,7 @@ bool ULinkerLoad::FindImportClassAndPackage( FName ClassName, FPackageIndex &Cla
 /**
 * Attempts to find the index for the given class object in the import list and adds it + its package if it does not exist
 */
-bool ULinkerLoad::CreateImportClassAndPackage( FName ClassName, FName PackageName, FPackageIndex &ClassIdx, FPackageIndex &PackageIdx )
+bool FLinkerLoad::CreateImportClassAndPackage( FName ClassName, FName PackageName, FPackageIndex &ClassIdx, FPackageIndex &PackageIdx )
 {
 	//look for an existing import first
 	//might as well look for the package at the same time ...
@@ -4246,7 +4269,7 @@ bool ULinkerLoad::CreateImportClassAndPackage( FName ClassName, FName PackageNam
 	return true;
 }
 
-TArray<FName> ULinkerLoad::FindPreviousNamesForClass(FString CurrentClassPath, bool bIsInstance)
+TArray<FName> FLinkerLoad::FindPreviousNamesForClass(FString CurrentClassPath, bool bIsInstance)
 {
 	TArray<FName> OldNames;
 	for (auto It = ObjectNameRedirects.CreateConstIterator(); It; ++It)
@@ -4271,7 +4294,7 @@ TArray<FName> ULinkerLoad::FindPreviousNamesForClass(FString CurrentClassPath, b
 	return OldNames;
 }
 
-FName ULinkerLoad::FindNewNameForClass(FName OldClassName, bool bIsInstance)
+FName FLinkerLoad::FindNewNameForClass(FName OldClassName, bool bIsInstance)
 {
 	FName *RedirectName = ObjectNameRedirects.Find(OldClassName);
 
@@ -4297,9 +4320,9 @@ FName ULinkerLoad::FindNewNameForClass(FName OldClassName, bool bIsInstance)
 /**
 * Allows object instances to be converted to other classes upon loading a package
 */
-ULinkerLoad::ELinkerStatus ULinkerLoad::FixupExportMap()
+FLinkerLoad::ELinkerStatus FLinkerLoad::FixupExportMap()
 {
-	DECLARE_SCOPE_CYCLE_COUNTER( TEXT( "ULinkerLoad::FixupExportMap" ), STAT_LinkerLoad_FixupExportMap, STATGROUP_LinkerLoad );
+	DECLARE_SCOPE_CYCLE_COUNTER( TEXT( "FLinkerLoad::FixupExportMap" ), STAT_LinkerLoad_FixupExportMap, STATGROUP_LinkerLoad );
 
 	// No need to fixup exports if everything is cooked.
 	if (!FPlatformProperties::RequiresCookedData())
@@ -4357,12 +4380,12 @@ ULinkerLoad::ELinkerStatus ULinkerLoad::FixupExportMap()
 							if (Export.ObjectName != NAME_None)
 							{
 								FString Now = GetExportFullName(ExportMapIdx);
-								UE_LOG(LogLinker, Log, TEXT("ULinkerLoad::FixupExportMap() - Renamed component from %s   to   %s"), *Was, *Now);
+								UE_LOG(LogLinker, Log, TEXT("FLinkerLoad::FixupExportMap() - Renamed component from %s   to   %s"), *Was, *Now);
 							}
 							else
 							{
 								Export.bExportLoadFailed = true;
-								UE_LOG(LogLinker, Log, TEXT("ULinkerLoad::FixupExportMap() - Removed component %s"), *Was);
+								UE_LOG(LogLinker, Log, TEXT("FLinkerLoad::FixupExportMap() - Removed component %s"), *Was);
 							}
 							continue;
 						}
@@ -4399,7 +4422,7 @@ ULinkerLoad::ELinkerStatus ULinkerLoad::FixupExportMap()
 					FPackageIndex NewPackageIndex;
 					if ( ResultClass == TEXT("None") )
 					{
-						UE_LOG(LogLinker, Log, TEXT("ULinkerLoad::FixupExportMap() - Pkg<%s> [Obj<%s> Cls<%s> ClsPkg<%s>] -> removed"), *LinkerRoot->GetName(),
+						UE_LOG(LogLinker, Log, TEXT("FLinkerLoad::FixupExportMap() - Pkg<%s> [Obj<%s> Cls<%s> ClsPkg<%s>] -> removed"), *LinkerRoot->GetName(),
 							*Export.ObjectName.ToString(), *NameClass.ToString(), *NamePackage.ToString());
 
 						Export.ClassIndex = NewClassIndex;
@@ -4417,13 +4440,13 @@ ULinkerLoad::ELinkerStatus ULinkerLoad::FixupExportMap()
 #endif
 						//Export.OuterIndex = newPackageIndex;
 
-						UE_LOG(LogLinker, Log, TEXT("ULinkerLoad::FixupExportMap() - Pkg<%s> [Obj<%s> Cls<%s> ClsPkg<%s>] -> [Obj<%s> Cls<%s> ClsPkg<%s>]"), *LinkerRoot->GetName(),
+						UE_LOG(LogLinker, Log, TEXT("FLinkerLoad::FixupExportMap() - Pkg<%s> [Obj<%s> Cls<%s> ClsPkg<%s>] -> [Obj<%s> Cls<%s> ClsPkg<%s>]"), *LinkerRoot->GetName(),
 							*Export.ObjectName.ToString(), *NameClass.ToString(), *NamePackage.ToString(),
 							*Export.ObjectName.ToString(), *ResultClass, *ResultPackage);
 					}
 					else
 					{
-						UE_LOG(LogLinker, Log, TEXT("ULinkerLoad::FixupExportMap() - object redirection failed at %s"), *Export.ObjectName.ToString());
+						UE_LOG(LogLinker, Log, TEXT("FLinkerLoad::FixupExportMap() - object redirection failed at %s"), *Export.ObjectName.ToString());
 					}
 				}
 			}
@@ -4462,13 +4485,13 @@ ULinkerLoad::ELinkerStatus ULinkerLoad::FixupExportMap()
 #if WITH_EDITOR
 							Export.OldClassName = NameClass;
 #endif
-							UE_LOG(LogLinker, Log, TEXT("ULinkerLoad::FixupExportMap() - Pkg<%s> [Obj<%s> Cls<%s> ClsPkg<%s>] -> [Obj<%s> Cls<%s> ClsPkg<%s>]"), *LinkerRoot->GetName(),
+							UE_LOG(LogLinker, Log, TEXT("FLinkerLoad::FixupExportMap() - Pkg<%s> [Obj<%s> Cls<%s> ClsPkg<%s>] -> [Obj<%s> Cls<%s> ClsPkg<%s>]"), *LinkerRoot->GetName(),
 								*Export.ObjectName.ToString(), *NameClass.ToString(), *NamePackage.ToString(),
 								*Export.ObjectName.ToString(), *ResultClass, *ResultPackage);
 						}
 						else
 						{
-							UE_LOG(LogLinker, Log, TEXT("ULinkerLoad::FixupExportMap() - object redirection failed at %s"), *Export.ObjectName.ToString());
+							UE_LOG(LogLinker, Log, TEXT("FLinkerLoad::FixupExportMap() - object redirection failed at %s"), *Export.ObjectName.ToString());
 						}
 					}
 				}
@@ -4484,9 +4507,8 @@ ULinkerLoad::ELinkerStatus ULinkerLoad::FixupExportMap()
 	}
 }
 
-IMPLEMENT_CORE_INTRINSIC_CLASS(ULinkerLoad, ULinker,
-	{
-	}
-);
+void FLinkerLoad::InitLinkers()
+{
+}
 
 #undef LOCTEXT_NAMESPACE
diff --git a/Engine/Source/Runtime/CoreUObject/Private/UObject/LinkerManager.cpp b/Engine/Source/Runtime/CoreUObject/Private/UObject/LinkerManager.cpp
new file mode 100644
--- /dev/null
+++ b/Engine/Source/Runtime/CoreUObject/Private/UObject/LinkerManager.cpp
@@ -0,0 +1,71 @@
+// Copyright 1998-2014 Epic Games, Inc. All Rights Reserved.
+
+/*=============================================================================
+	LinkerManager.h: Unreal object linker manager
+=============================================================================*/
+
+#pragma once
+
+#include "CoreUObjectPrivate.h"
+#include "LinkerManager.h"
+
+FLinkerManager& FLinkerManager::Get()
+{
+	static TAutoPtr<FLinkerManager> Singleton(new FLinkerManager());
+	return *Singleton.GetOwnedPointer();
+}
+
+bool FLinkerManager::Exec(class UWorld* InWorld, const TCHAR* Cmd, FOutputDevice& Ar)
+{
+	if (FParse::Command(&Cmd, TEXT("LinkerLoadList")))
+	{
+		UE_LOG(LogLinker, Display, TEXT("ObjectLoaders: %d"), ObjectLoaders.Num());
+		UE_LOG(LogLinker, Display, TEXT("LoadersWithNewImports: %d"), LoadersWithNewImports.Num());
+
+#if UE_BUILD_DEBUG
+		UE_LOG(LogLinker, Display, TEXT("LiveLinkers: %d"), LiveLinkers.Num());
+		for (auto Linker : LiveLinkers)
+		{
+			UE_LOG(LogLinker, Display, TEXT("%s"), *Linker->Filename);
+		}
+#endif
+		return true;
+	}
+	else if (FParse::Command(&Cmd, TEXT("LINKERS")))
+	{
+		Ar.Logf(TEXT("Linkers:"));
+		for (auto Linker : ObjectLoaders)
+		{
+			int32 NameSize = 0;
+			for (int32 j = 0; j < Linker->NameMap.Num(); j++)
+			{
+				if (Linker->NameMap[j] != NAME_None)
+				{
+					NameSize += FNameEntry::GetSize(*Linker->NameMap[j].ToString());
+				}
+			}
+			Ar.Logf
+				(
+				TEXT("%s (%s): Names=%i (%iK/%iK) Imports=%i (%iK) Exports=%i (%iK) Gen=%i Bulk=%i"),
+				*Linker->Filename,
+				*Linker->LinkerRoot->GetFullName(),
+				Linker->NameMap.Num(),
+				Linker->NameMap.Num() * sizeof(FName) / 1024,
+				NameSize / 1024,
+				Linker->ImportMap.Num(),
+				Linker->ImportMap.Num() * sizeof(FObjectImport) / 1024,
+				Linker->ExportMap.Num(),
+				Linker->ExportMap.Num() * sizeof(FObjectExport) / 1024,
+				Linker->Summary.Generations.Num(),
+#if WITH_EDITOR
+				Linker->BulkDataLoaders.Num()
+#else
+				0
+#endif // WITH_EDITOR
+				);
+		}
+
+		return true;
+	}
+	return false;
+}
\ No newline at end of file
diff --git a/Engine/Source/Runtime/CoreUObject/Private/UObject/LinkerManager.h b/Engine/Source/Runtime/CoreUObject/Private/UObject/LinkerManager.h
new file mode 100644
--- /dev/null
+++ b/Engine/Source/Runtime/CoreUObject/Private/UObject/LinkerManager.h
@@ -0,0 +1,47 @@
+// Copyright 1998-2014 Epic Games, Inc. All Rights Reserved.
+
+/*=============================================================================
+	LinkerManager.h: Unreal object linker manager
+=============================================================================*/
+
+#pragma once
+
+#include "Linker.h"
+
+class FLinkerManager : private FSelfRegisteringExec
+{
+public:
+
+	static FLinkerManager& Get();
+
+	FORCEINLINE TSet<FLinkerLoad*>& GetLoaders()
+	{
+		return ObjectLoaders;
+	}
+
+	FORCEINLINE TSet<FLinkerLoad*>& GetLoadersWithNewImports()
+	{
+		return LoadersWithNewImports;
+	}
+
+#if UE_BUILD_DEBUG
+	FORCEINLINE TArray<FLinkerLoad*>& GetLiveLinkers()
+	{
+		return LiveLinkers;
+	}
+#endif
+
+	// FSelfRegisteringExec interface
+	virtual bool Exec(class UWorld* InWorld, const TCHAR* Cmd, FOutputDevice& Ar) override;
+
+private:
+
+	/** Map of packages to their open linkers **/
+	TSet<FLinkerLoad*> ObjectLoaders;
+	/** List of loaders that have new imports **/
+	TSet<FLinkerLoad*> LoadersWithNewImports;
+#if UE_BUILD_DEBUG
+	/** List of all the existing linker loaders **/
+	TArray<FLinkerLoad*> LiveLinkers;
+#endif
+};
diff --git a/Engine/Source/Runtime/CoreUObject/Private/UObject/LinkerSave.cpp b/Engine/Source/Runtime/CoreUObject/Private/UObject/LinkerSave.cpp
--- a/Engine/Source/Runtime/CoreUObject/Private/UObject/LinkerSave.cpp
+++ b/Engine/Source/Runtime/CoreUObject/Private/UObject/LinkerSave.cpp
@@ -3,18 +3,16 @@
 #include "CoreUObjectPrivate.h"
 
 /*----------------------------------------------------------------------------
-	ULinkerSave.
+	FLinkerSave.
 ----------------------------------------------------------------------------*/
 
 /** A mapping of package name to generated script SHA keys */
-TMap<FString, TArray<uint8> > ULinkerSave::PackagesToScriptSHAMap;
+TMap<FString, TArray<uint8> > FLinkerSave::PackagesToScriptSHAMap;
 
-ULinkerSave::ULinkerSave( const FObjectInitializer& ObjectInitializer, UPackage* InParent, const TCHAR* InFilename, bool bForceByteSwapping, bool bInSaveUnversioned )
-:	ULinker(ObjectInitializer, InParent, InFilename )
-,	Saver( NULL )
+FLinkerSave::FLinkerSave(UPackage* InParent, const TCHAR* InFilename, bool bForceByteSwapping, bool bInSaveUnversioned)
+:	FLinker(ELinkerType::Save, InParent, InFilename )
+,	Saver(nullptr)
 {
-	check(!HasAnyFlags(RF_ClassDefaultObject));
-
 	if (FPlatformProperties::HasEditorOnlyData())
 	{
 		// Create file saver.
@@ -46,12 +44,10 @@ ULinkerSave::ULinkerSave( const FObjectInitializer& ObjectInitializer, UPackage*
 	}
 }
 
-ULinkerSave::ULinkerSave(const FObjectInitializer& ObjectInitializer, UPackage* InParent, bool bForceByteSwapping, bool bInSaveUnversioned )
-:	ULinker( ObjectInitializer, InParent,TEXT("$$Memory$$") )
-,	Saver( NULL )
+FLinkerSave::FLinkerSave(UPackage* InParent, bool bForceByteSwapping, bool bInSaveUnversioned )
+:	FLinker(ELinkerType::Save, InParent,TEXT("$$Memory$$"))
+,	Saver(nullptr)
 {
-	check(!HasAnyFlags(RF_ClassDefaultObject));
-
 	if (FPlatformProperties::HasEditorOnlyData())
 	{
 		// Create file saver.
@@ -82,23 +78,22 @@ ULinkerSave::ULinkerSave(const FObjectInitializer& ObjectInitializer, UPackage*
 /**
  * Detaches file saver and hence file handle.
  */
-void ULinkerSave::Detach()
+void FLinkerSave::Detach()
 {
-	if( Saver )
+	if (Saver)
 	{
 		delete Saver;
 	}
-	Saver = NULL;
+	Saver = nullptr;
 }
 
-void ULinkerSave::BeginDestroy()
+FLinkerSave::~FLinkerSave()
 {
 	// Detach file saver/ handle.
 	Detach();
-	Super::BeginDestroy();
 }
 
-int32 ULinkerSave::MapName(const FName& Name) const
+int32 FLinkerSave::MapName(const FName& Name) const
 {
 	const int32* IndexPtr = NameIndices.Find(Name);
 
@@ -110,7 +105,7 @@ int32 ULinkerSave::MapName(const FName& Name) const
 	return INDEX_NONE;
 }
 
-FPackageIndex ULinkerSave::MapObject( const UObject* Object ) const
+FPackageIndex FLinkerSave::MapObject( const UObject* Object ) const
 {
 	if (Object)
 	{
@@ -123,22 +118,22 @@ FPackageIndex ULinkerSave::MapObject( const UObject* Object ) const
 	return FPackageIndex();
 }
 
-void ULinkerSave::Seek( int64 InPos )
+void FLinkerSave::Seek( int64 InPos )
 {
 	Saver->Seek( InPos );
 }
 
-int64 ULinkerSave::Tell()
+int64 FLinkerSave::Tell()
 {
 	return Saver->Tell();
 }
 
-void ULinkerSave::Serialize( void* V, int64 Length )
+void FLinkerSave::Serialize( void* V, int64 Length )
 {
 	Saver->Serialize( V, Length );
 }
 	
-FArchive& ULinkerSave::operator<<( FName& InName )
+FArchive& FLinkerSave::operator<<( FName& InName )
 {
 	int32 Save = MapName(InName);
 	int32 Number = InName.GetNumber();
@@ -146,7 +141,7 @@ FArchive& ULinkerSave::operator<<( FName& InName )
 	return Ar << Save << Number;
 }
 
-FArchive& ULinkerSave::operator<<( UObject*& Obj )
+FArchive& FLinkerSave::operator<<( UObject*& Obj )
 {
 	FPackageIndex Save;
 	if (Obj)
@@ -156,14 +151,14 @@ FArchive& ULinkerSave::operator<<( UObject*& Obj )
 	return *this << Save;
 }
 
-FArchive& ULinkerSave::operator<<( FLazyObjectPtr& LazyObjectPtr)
+FArchive& FLinkerSave::operator<<( FLazyObjectPtr& LazyObjectPtr)
 {
 	FUniqueObjectGuid ID;
 	ID = LazyObjectPtr.GetUniqueID();
 	return *this << ID;
 }
 
-FArchive& ULinkerSave::operator<<( FAssetPtr& AssetPtr)
+FArchive& FLinkerSave::operator<<( FAssetPtr& AssetPtr)
 {
 	FStringAssetReference ID;
 	UObject *Object = AssetPtr.Get();
@@ -181,8 +176,3 @@ FArchive& ULinkerSave::operator<<( FAssetPtr& AssetPtr)
 	ID.Serialize(*this);
 	return *this;
 }
-
-IMPLEMENT_CORE_INTRINSIC_CLASS(ULinkerSave, ULinker,
-	{
-	}
-);
diff --git a/Engine/Source/Runtime/CoreUObject/Private/UObject/Obj.cpp b/Engine/Source/Runtime/CoreUObject/Private/UObject/Obj.cpp
--- a/Engine/Source/Runtime/CoreUObject/Private/UObject/Obj.cpp
+++ b/Engine/Source/Runtime/CoreUObject/Private/UObject/Obj.cpp
@@ -14,11 +14,10 @@
 #include "FindStronglyConnected.h"
 #include "HotReloadInterface.h"
 #include "UObject/TlsObjectInitializers.h"
+#include "UObject/UObjectThreadContext.h"
 
 DEFINE_LOG_CATEGORY(LogObj);
 
-extern int32 GIsInConstructor;
-
 /** Stat group for dynamic objects cycle counters*/
 
 
@@ -82,7 +81,8 @@ void UObject::EnsureNotRetrievingVTablePtr() const
 
 UObject* UObject::CreateDefaultSubobject(FName SubobjectFName, UClass* ReturnType, UClass* ClassToCreateByDefault, bool bIsRequired, bool bAbstract, bool bIsTransient)
 {
-	UE_CLOG(!GIsInConstructor, LogObj, Fatal, TEXT("CreateDefultSubobject can only be used inside of UObject constructors. UObject constructing subobjects cannot be created using new or placement new operator."));
+	auto& ThreadContext = FUObjectThreadContext::Get();
+	UE_CLOG(!ThreadContext.IsInConstructor, LogObj, Fatal, TEXT("CreateDefultSubobject can only be used inside of UObject constructors. UObject constructing subobjects cannot be created using new or placement new operator."));
 	auto CurrentInitializer = FTlsObjectInitializers::Top();
 	UE_CLOG(!CurrentInitializer, LogObj, Fatal, TEXT("No object initializer found during construction."));
 	UE_CLOG(CurrentInitializer->Obj != this, LogObj, Fatal, TEXT("Using incorrect object initializer."));
@@ -534,7 +534,7 @@ void UObject::FinishDestroy()
 			);
 	}
 
-	check( GetLinker() == NULL );
+	check( !GetLinker() );
 	check( GetLinkerIndex()	== INDEX_NONE );
 
 	DestroyNonNativeProperties();
@@ -664,6 +664,8 @@ void UObject::ConditionalPostLoad()
 {
 	if( HasAnyFlags(RF_NeedPostLoad) )
 	{
+		check(IsInGameThread() || HasAnyFlags(RF_ClassDefaultObject|RF_ArchetypeObject) || IsPostLoadThreadSafe() || IsA(UClass::StaticClass()))
+
 #if !(UE_BUILD_SHIPPING || UE_BUILD_TEST)
 		checkSlow(!DebugPostLoad.Contains(this));
 		DebugPostLoad.Add(this);
@@ -784,7 +786,7 @@ bool UObject::Modify( bool bAlwaysMarkDirty/*=true*/ )
 {
 	bool bSavedToTransactionBuffer = false;
 
-	if (!GIsGarbageCollecting)
+	if (!IsGarbageCollecting())
 	{
 		// Do not consider PIE world objects or script packages, as they should never end up in the
 		// transaction buffer and we don't want to mark them dirty here either.
@@ -817,10 +819,6 @@ bool UObject::IsSelected() const
 
 void UObject::Serialize( FArchive& Ar )
 {
-#if !(UE_BUILD_SHIPPING || UE_BUILD_TEST)
-	DebugSerialize.RemoveSingle(this);
-#endif
-
 	// These three items are very special items from a serialization standpoint. They aren't actually serialized.
 	UClass *Class = GetClass();
 	UObject* LoadOuter = GetOuter();
@@ -834,7 +832,7 @@ void UObject::Serialize( FArchive& Ar )
 		// make sure this object's template data is loaded - the only objects
 		// this should actually affect are those that don't have any defaults
 		// to serialize.  for objects with defaults that actually require loading
-		// the class default object should be serialized in ULinkerLoad::Preload, before
+		// the class default object should be serialized in FLinkerLoad::Preload, before
 		// we've hit this code.
 		if ( !HasAnyFlags(RF_ClassDefaultObject) && Class->GetDefaultsCount() > 0 )
 		{
@@ -855,8 +853,11 @@ void UObject::Serialize( FArchive& Ar )
 			Ar << Class;
 		}
 		//@todo UE4 - This seems to be required and it should not be. Seems to be related to the texture streamer.
-		ULinkerLoad* LinkerLoad = GetLinker();
-		Ar << LinkerLoad;
+		FLinkerLoad* LinkerLoad = GetLinker();
+		if (LinkerLoad)
+		{
+			LinkerLoad->Serialize(Ar);
+		}
 	}
 	// Special support for supporting undo/redo of renaming and changing Archetype.
 	else if( Ar.IsTransacting() )
@@ -1273,11 +1274,8 @@ bool UObject::IsSafeForRootSet() const
 		return false;
 	}
 
-	const ULinkerLoad* LinkerLoad = dynamic_cast<const ULinkerLoad*>(this);
-
-	// Exclude linkers from root set if we're using seekfree loading
-	if( !HasAnyFlags(RF_PendingKill)
-		&& ( !FPlatformProperties::RequiresCookedData() || LinkerLoad == NULL || LinkerLoad->HasAnyFlags(RF_ClassDefaultObject) ) )
+	// Exclude linkers from root set if we're using seekfree loading		
+	if (!HasAnyFlags(RF_PendingKill))
 	{
 		return true;
 	}
@@ -3427,43 +3425,6 @@ bool StaticExec( UWorld* InWorld, const TCHAR* Cmd, FOutputDevice& Ar )
 				}
 			}
 		}
-		else if( FParse::Command(&Str,TEXT("LINKERS")) )
-		{
-			Ar.Logf( TEXT("Linkers:") );
-			for (TMap<UPackage*, ULinkerLoad*>::TIterator It(GObjLoaders); It; ++It)
-			{
-				ULinkerLoad* Linker = It.Value();
-				int32 NameSize = 0;
-				for( int32 j=0; j<Linker->NameMap.Num(); j++ )
-				{
-					if( Linker->NameMap[j] != NAME_None )
-					{
-						NameSize += FNameEntry::GetSize( *Linker->NameMap[j].ToString() );
-					}
-				}
-				Ar.Logf
-				(
-					TEXT("%s (%s): Names=%i (%iK/%iK) Imports=%i (%iK) Exports=%i (%iK) Gen=%i Bulk=%i"),
-					*Linker->Filename,
-					*Linker->LinkerRoot->GetFullName(),
-					Linker->NameMap.Num(),
-					Linker->NameMap.Num() * sizeof(FName) / 1024,
-					NameSize / 1024,
-					Linker->ImportMap.Num(),
-					Linker->ImportMap.Num() * sizeof(FObjectImport) / 1024,
-					Linker->ExportMap.Num(),
-					Linker->ExportMap.Num() * sizeof(FObjectExport) / 1024,
-					Linker->Summary.Generations.Num(),
-#if WITH_EDITOR
-					Linker->BulkDataLoaders.Num()
-#else
-					0
-#endif // WITH_EDITOR
-				);
-			}
-
-			return true;
-		}
 		else if ( FParse::Command(&Str,TEXT("FLAGS")) )
 		{
 			// Dump all object flags for objects rooted at the named object.
@@ -3629,7 +3590,7 @@ void StaticUObjectInit()
 //
 void StaticExit()
 {
-	check(GObjLoaded.Num()==0);
+	check(FUObjectThreadContext::Get().ObjLoaded.Num() == 0);
 	if (UObjectInitialized() == false)
 	{
 		return;
@@ -3672,7 +3633,7 @@ void StaticExit()
 	GExitPurge					= true;
 
 	// Route BeginDestroy. This needs to be a separate pass from marking as RF_Unreachable as code might rely on RF_Unreachable to be 
-	// set on all objects that are about to be deleted. One example is ULinkerLoad detaching textures - the SetLinker call needs to 
+	// set on all objects that are about to be deleted. One example is FLinkerLoad detaching textures - the SetLinker call needs to 
 	// not kick off texture streaming.
 	//
 	for ( FRawObjectIterator It; It; ++It )
@@ -3710,7 +3671,7 @@ void StaticExit()
 
 	UObjectBaseShutdown();
 	// Empty arrays to prevent falsely-reported memory leaks.
-	GObjLoaded			.Empty();
+	FUObjectThreadContext::Get().ObjLoaded.Empty();
 	UE_LOG(LogExit, Log, TEXT("Object subsystem successfully closed.") );
 }
 
@@ -3750,7 +3711,7 @@ void MarkObjectsToDisregardForGC()
 	}
 
 	UE_LOG(LogObj, Log, TEXT("%i objects as part of root set at end of initial load."), NumAlwaysLoadedObjects);
-	if (GUObjectArray.DisregardForGCEnabled())
+	if (GetUObjectArray().DisregardForGCEnabled())
 	{
 		UE_LOG(LogObj, Log, TEXT("%i objects are not in the root set, but can never be destroyed because they are in the DisregardForGC set."), NumAlwaysLoadedObjects - NumRootObjects);
 	}
diff --git a/Engine/Source/Runtime/CoreUObject/Private/UObject/ObjectBaseUtility.cpp b/Engine/Source/Runtime/CoreUObject/Private/UObject/ObjectBaseUtility.cpp
--- a/Engine/Source/Runtime/CoreUObject/Private/UObject/ObjectBaseUtility.cpp
+++ b/Engine/Source/Runtime/CoreUObject/Private/UObject/ObjectBaseUtility.cpp
@@ -13,28 +13,30 @@
  */
 int32 UObjectBaseUtility::GetLinkerUE4Version() const
 {
-	// Catch code errors like GetLinker()->GetLinkerUE4Version()
-	check(!IsA(ULinkerLoad::StaticClass()));
-
-	ULinkerLoad* Loader = GetLinker();
+	FLinkerLoad* Loader = GetLinker();
 
 	// No linker.
-	if( Loader == NULL )
+	if (Loader == nullptr)
 	{
 		// the _Linker reference is never set for the top-most UPackage of a package (the linker root), so if this object
 		// is the linker root, find our loader in the global list.
-		if( GetOutermost() == this )
+		if (GetOutermost() == this)
 		{
-			Loader = ULinkerLoad::FindExistingLinkerForPackage(const_cast<UPackage*>(CastChecked<UPackage>((const UObject*)this)));
+			Loader = FLinkerLoad::FindExistingLinkerForPackage(const_cast<UPackage*>(CastChecked<UPackage>((const UObject*)this)));
 		}
 	}
 
-	if ( Loader != NULL )
+	if (Loader != nullptr)
 	{
 		// We have a linker so we can return its version.
 		return Loader->UE4Ver();
 
 	}
+	else if (GetOutermost())
+	{
+		// Get the linker version associated with the package this object lives in
+		return GetOutermost()->LinkerPackageVersion;
+	}
 	else
 	{
 		// We don't have a linker associated as we e.g. might have been saved or had loaders reset, ...
@@ -44,7 +46,7 @@ int32 UObjectBaseUtility::GetLinkerUE4Version() const
 
 int32 UObjectBaseUtility::GetLinkerCustomVersion(FGuid CustomVersionKey) const
 {
-	ULinkerLoad* Loader = GetLinker();
+	FLinkerLoad* Loader = GetLinker();
 
 	// No linker.
 	if( Loader == NULL )
@@ -53,7 +55,7 @@ int32 UObjectBaseUtility::GetLinkerCustomVersion(FGuid CustomVersionKey) const
 		// is the linker root, find our loader in the global list.
 		if( GetOutermost() == this )
 		{
-			Loader = ULinkerLoad::FindExistingLinkerForPackage(const_cast<UPackage*>(CastChecked<UPackage>((const UObject*)this)));
+			Loader = FLinkerLoad::FindExistingLinkerForPackage(const_cast<UPackage*>(CastChecked<UPackage>((const UObject*)this)));
 		}
 	}
 
@@ -63,7 +65,12 @@ int32 UObjectBaseUtility::GetLinkerCustomVersion(FGuid CustomVersionKey) const
 		auto* CustomVersion = Loader->Summary.GetCustomVersionContainer().GetVersion(CustomVersionKey);
 		return CustomVersion ? CustomVersion->Version : -1;
 	}
-
+	else if (GetOutermost() && GetOutermost()->LinkerCustomVersion.GetAllVersions().Num())
+	{
+		// Get the linker version associated with the package this object lives in
+		auto* CustomVersion = GetOutermost()->LinkerCustomVersion.GetVersion(CustomVersionKey);
+		return CustomVersion ? CustomVersion->Version : -1;
+	}
 	// We don't have a linker associated as we e.g. might have been saved or had loaders reset, ...
 	// We must have a current version for this tag.
 	auto* CustomVersion = FCustomVersionContainer::GetRegistered().GetVersion(CustomVersionKey);
@@ -82,7 +89,7 @@ int32 UObjectBaseUtility::GetLinkerCustomVersion(FGuid CustomVersionKey) const
  */
 int32 UObjectBaseUtility::GetLinkerLicenseeUE4Version() const
 {
-	ULinkerLoad* Loader = GetLinker();
+	FLinkerLoad* Loader = GetLinker();
 
 	// No linker.
 	if( Loader == NULL )
@@ -91,7 +98,7 @@ int32 UObjectBaseUtility::GetLinkerLicenseeUE4Version() const
 		// is the linker root, find our loader in the global list.
 		if( GetOutermost() == this )
 		{
-			Loader = ULinkerLoad::FindExistingLinkerForPackage(const_cast<UPackage*>(CastChecked<UPackage>((const UObject*)this)));
+			Loader = FLinkerLoad::FindExistingLinkerForPackage(const_cast<UPackage*>(CastChecked<UPackage>((const UObject*)this)));
 		}
 	}
 
@@ -99,7 +106,11 @@ int32 UObjectBaseUtility::GetLinkerLicenseeUE4Version() const
 	{
 		// We have a linker so we can return its version.
 		return Loader->LicenseeUE4Ver();
-
+	}
+	else if (GetOutermost())
+	{
+		// Get the linker version associated with the package this object lives in
+		return GetOutermost()->LinkerLicenseeVersion;
 	}
 	else
 	{
diff --git a/Engine/Source/Runtime/CoreUObject/Private/UObject/Package.cpp b/Engine/Source/Runtime/CoreUObject/Private/UObject/Package.cpp
--- a/Engine/Source/Runtime/CoreUObject/Private/UObject/Package.cpp
+++ b/Engine/Source/Runtime/CoreUObject/Private/UObject/Package.cpp
@@ -26,10 +26,12 @@ void UPackage::PostInitProperties()
 	Super::PostInitProperties();
 	if ( !HasAnyFlags(RF_ClassDefaultObject) )
 	{
-		bDirty = 0;
+		bDirty = false;
 	}
 
 	MetaData = NULL;
+	LinkerPackageVersion = GPackageFileUE4Version;
+	LinkerLicenseeVersion = GPackageFileLicenseeUE4Version;
 
 #if WITH_EDITOR
 	PIEInstanceID = INDEX_NONE;
@@ -159,11 +161,6 @@ void UPackage::TagSubobjects(EObjectFlags NewFlags)
 	{
 		MetaData->SetFlags(NewFlags);
 	}
-
-	if (GetLinker())
-	{
-		GetLinker()->SetFlags(NewFlags);
-	}
 }
 
 /**
@@ -191,6 +188,17 @@ bool UPackage::IsFullyLoaded()
 	return bHasBeenFullyLoaded;
 }
 
+void UPackage::BeginDestroy()
+{
+	// Detach linker if still attached
+	if (LinkerLoad)
+	{
+		delete LinkerLoad;
+		LinkerLoad = nullptr;
+	}
+
+	Super::BeginDestroy();
+}
 
 IMPLEMENT_CORE_INTRINSIC_CLASS(UPackage, UObject,
 	{
diff --git a/Engine/Source/Runtime/CoreUObject/Private/UObject/PropertyArray.cpp b/Engine/Source/Runtime/CoreUObject/Private/UObject/PropertyArray.cpp
--- a/Engine/Source/Runtime/CoreUObject/Private/UObject/PropertyArray.cpp
+++ b/Engine/Source/Runtime/CoreUObject/Private/UObject/PropertyArray.cpp
@@ -10,7 +10,7 @@
 
 void UArrayProperty::LinkInternal(FArchive& Ar)
 {
-	ULinkerLoad* MyLinker = GetLinker();
+	FLinkerLoad* MyLinker = GetLinker();
 	if( MyLinker )
 	{
 		MyLinker->Preload(this);
diff --git a/Engine/Source/Runtime/CoreUObject/Private/UObject/PropertyBaseObject.cpp b/Engine/Source/Runtime/CoreUObject/Private/UObject/PropertyBaseObject.cpp
--- a/Engine/Source/Runtime/CoreUObject/Private/UObject/PropertyBaseObject.cpp
+++ b/Engine/Source/Runtime/CoreUObject/Private/UObject/PropertyBaseObject.cpp
@@ -430,7 +430,7 @@ void UObjectPropertyBase::CheckValidObject(void* Value) const
 		bool bIsReplacingClassRefs = PropertyClass && PropertyClass->HasAnyClassFlags(CLASS_NewerVersionExists) != ObjectClass->HasAnyClassFlags(CLASS_NewerVersionExists);
 		
 #if USE_CIRCULAR_DEPENDENCY_LOAD_DEFERRING
-		ULinkerLoad* PropertyLinker = GetLinker();
+		FLinkerLoad* PropertyLinker = GetLinker();
 		bool const bIsDeferringValueLoad = ((PropertyLinker == nullptr) || (PropertyLinker->LoadFlags & LOAD_DeferDependencyLoads)) &&
 			Object->IsA<ULinkerPlaceholderExportObject>();
 
diff --git a/Engine/Source/Runtime/CoreUObject/Private/UObject/PropertyStruct.cpp b/Engine/Source/Runtime/CoreUObject/Private/UObject/PropertyStruct.cpp
--- a/Engine/Source/Runtime/CoreUObject/Private/UObject/PropertyStruct.cpp
+++ b/Engine/Source/Runtime/CoreUObject/Private/UObject/PropertyStruct.cpp
@@ -8,14 +8,14 @@ static inline void PreloadInnerStructMembers(UStructProperty* StructProperty)
 {
 #if USE_CIRCULAR_DEPENDENCY_LOAD_DEFERRING
 	uint32 PropagatedLoadFlags = 0;
-	if (ULinkerLoad* Linker = StructProperty->GetLinker())
+	if (FLinkerLoad* Linker = StructProperty->GetLinker())
 	{
 		PropagatedLoadFlags |= (Linker->LoadFlags & LOAD_DeferDependencyLoads);
 	}
 
 	if (UScriptStruct* Struct = StructProperty->Struct)
 	{
-		if (ULinkerLoad* StructLinker = Struct->GetLinker())
+		if (FLinkerLoad* StructLinker = Struct->GetLinker())
 		{
 			TGuardValue<uint32> LoadFlagGuard(StructLinker->LoadFlags, StructLinker->LoadFlags | PropagatedLoadFlags);
 			Struct->RecursivelyPreload();
diff --git a/Engine/Source/Runtime/CoreUObject/Private/UObject/ReferenceChainSearch.cpp b/Engine/Source/Runtime/CoreUObject/Private/UObject/ReferenceChainSearch.cpp
--- a/Engine/Source/Runtime/CoreUObject/Private/UObject/ReferenceChainSearch.cpp
+++ b/Engine/Source/Runtime/CoreUObject/Private/UObject/ReferenceChainSearch.cpp
@@ -162,7 +162,7 @@ void FReferenceChainSearch::PrintReferencers( FReferenceChain& Referencer )
 			ObjectReachability += TEXT("(standalone) ");
 		}
 
-		if (GUObjectArray.IsDisregardForGC(RefInfo.ReferencedBy))
+		if (GetUObjectArray().IsDisregardForGC(RefInfo.ReferencedBy))
 		{
 			ObjectReachability += TEXT("(NeverGCed) ");
 		}
@@ -816,7 +816,7 @@ FString FReferenceChainSearch::FReferenceChainLink::ToString() const
 		ObjectReachability += TEXT("(standalone) ");
 	}
 
-	if (GUObjectArray.IsDisregardForGC(ReferencedBy))
+	if (GetUObjectArray().IsDisregardForGC(ReferencedBy))
 	{
 		ObjectReachability += TEXT("(NeverGCed) ");
 	}
