Merging UE4-Streaming to UE4

diff --git a/Engine/Source/Runtime/CoreUObject/Private/UObject/SavePackage.cpp b/Engine/Source/Runtime/CoreUObject/Private/UObject/SavePackage.cpp
--- a/Engine/Source/Runtime/CoreUObject/Private/UObject/SavePackage.cpp
+++ b/Engine/Source/Runtime/CoreUObject/Private/UObject/SavePackage.cpp
@@ -4,6 +4,7 @@
 #include "UObject/UTextProperty.h"
 #include "Interface.h"
 #include "TargetPlatform.h"
+#include "UObject/UObjectThreadContext.h"
 
 
 DEFINE_LOG_CATEGORY_STATIC(LogSavePackage, Log, All);
@@ -42,10 +43,10 @@ static void CheckObjectPriorToSave(FArchiveUObject& Ar, UObject* InObj, UPackage
 	{
 		return;
 	}
-	extern UObject* GSerializedObject;
+	FUObjectThreadContext& ThreadContext = FUObjectThreadContext::Get();
 	if (!InObj->IsValidLowLevelFast() || !InObj->IsValidLowLevel())
 	{
-		UE_LOG(LogLinker, Fatal, TEXT("Attempt to save bogus object %p GSerializedObject=%s  SerializedProperty=%s"), (void*)InObj, *GetFullNameSafe(GSerializedObject), *GetFullNameSafe(Ar.GetSerializedProperty()));
+		UE_LOG(LogLinker, Fatal, TEXT("Attempt to save bogus object %p ThreadContext.SerializedObject=%s  SerializedProperty=%s"), (void*)InObj, *GetFullNameSafe(ThreadContext.SerializedObject), *GetFullNameSafe(Ar.GetSerializedProperty()));
 		return;
 	}
 	// if the object class is abstract or has been marked as deprecated, mark this
@@ -60,7 +61,7 @@ static void CheckObjectPriorToSave(FArchiveUObject& Ar, UObject* InObj, UPackage
 		{
 			TArray<UObject*> ComponentReferences;
 			FReferenceFinder ComponentCollector(ComponentReferences, InObj, false, true, true);
-			ComponentCollector.FindReferences(InObj, GSerializedObject, Ar.GetSerializedProperty());
+			ComponentCollector.FindReferences(InObj, ThreadContext.SerializedObject, Ar.GetSerializedProperty());
 
 			for ( int32 Index = 0; Index < ComponentReferences.Num(); Index++ )
 			{
@@ -102,7 +103,7 @@ static void CheckObjectPriorToSave(FArchiveUObject& Ar, UObject* InObj, UPackage
 	}
 }
 
-static bool EndSavingIfCancelled( ULinkerSave* Linker, const FString& TempFilename ) 
+static bool EndSavingIfCancelled( FLinkerSave* Linker, const FString& TempFilename ) 
 {
 	if ( GWarn->ReceivedUserCancel() )
 	{
@@ -210,7 +211,7 @@ public:
 	 */
 	void MarkNameAsReferenced(const FName& Name)
 	{
-		// We need to store the FName without the number, as the number is stored separately by ULinkerSave 
+		// We need to store the FName without the number, as the number is stored separately by FLinkerSave 
 		// and we don't want duplicate entries in the name table just because of the number
 		const FName NameNoNumber(Name, 0);
 		ReferencedNames.Add(NameNoNumber);
@@ -220,7 +221,7 @@ public:
 	 * Add the marked names to a linker
 	 * @param Linker linker to save the marked names to
 	 */
-	void UpdateLinkerWithMarkedNames(ULinkerSave* Linker)
+	void UpdateLinkerWithMarkedNames(FLinkerSave* Linker)
 	{
 		Linker->NameMap.Reserve(Linker->NameMap.Num() + ReferencedNames.Num());
 		for (const FName& Name : ReferencedNames)
@@ -514,11 +515,11 @@ void FArchiveSaveTagExports::ProcessTaggedObjects()
 class FArchiveSaveTagImports : public FArchiveUObject
 {
 public:
-	ULinkerSave* Linker;
+	FLinkerSave* Linker;
 	TArray<UObject*> Dependencies;
 	TArray<FString>& StringAssetReferencesMap;
 
-	FArchiveSaveTagImports(ULinkerSave* InLinker, TArray<FString>& InStringAssetReferencesMap)
+	FArchiveSaveTagImports(FLinkerSave* InLinker, TArray<FString>& InStringAssetReferencesMap)
 		: Linker(InLinker), StringAssetReferencesMap(InStringAssetReferencesMap)
 	{
 		ArIsSaving				= true;
@@ -665,11 +666,11 @@ public:
 	 *
 	 * @param	SrcFilename		Filename to compress
 	 * @param	DstFilename		Output name of compressed file, cannot be the same as SrcFilename
-	 * @param	SrcLinker		ULinkerSave object used to save src file
+	 * @param	SrcLinker		FLinkerSave object used to save src file
 	 *
 	 * @return true if sucessful, false otherwise
 	 */
-	bool CompressFile( const TCHAR* SrcFilename, const TCHAR* DstFilename, ULinkerSave* SrcLinker )
+	bool CompressFile( const TCHAR* SrcFilename, const TCHAR* DstFilename, FLinkerSave* SrcLinker )
 	{
 		FString TmpFilename = FPaths::GetPath( DstFilename ) / ( FPaths::GetBaseFilename( DstFilename ) + TEXT( "_SaveCompressed.tmp" ));
 
@@ -699,14 +700,14 @@ public:
 		return bMoveSucceded;
 	}
 	/**
-	 * Compresses the passed in src file and writes it to destination. The file comes from the "Saver" in the ULinker
+	 * Compresses the passed in src file and writes it to destination. The file comes from the "Saver" in the FLinker
 	 *
 	 * @param	DstFilename		Output name of compressed file, cannot be the same as SrcFilename
-	 * @param	SrcLinker		ULinkerSave object used to save src file
+	 * @param	SrcLinker		FLinkerSave object used to save src file
 	 *
 	 * @return true if sucessful, false otherwise
 	 */
-	bool CompressFile( const TCHAR* DstFilename, ULinkerSave* SrcLinker )
+	bool CompressFile( const TCHAR* DstFilename, FLinkerSave* SrcLinker )
 	{
 		FString TmpFilename = FPaths::GetPath( DstFilename ) / ( FPaths::GetBaseFilename( DstFilename ) + TEXT( "_SaveCompressed.tmp" ));
 		// Create file reader and writer...
@@ -738,11 +739,11 @@ public:
 	 *
 	 * @param	FileReader		archive to read from
 	 * @param	FileWriter		archive to write to
-	 * @param	SrcLinker		ULinkerSave object used to save src file
+	 * @param	SrcLinker		FLinkerSave object used to save src file
 	 *
 	 * @return true if sucessful, false otherwise
 	 */
-	void CompressArchive( FArchive* FileReader, FArchive* FileWriter, ULinkerSave* SrcLinker )
+	void CompressArchive( FArchive* FileReader, FArchive* FileWriter, FLinkerSave* SrcLinker )
 	{
 
 		// Read package file summary from source file.
@@ -1018,7 +1019,7 @@ public:
 	 * @param	Linker				linker containing the names that need to be sorted
 	 * @param	LinkerToConformTo	optional linker to conform against.
 	 */
-	void SortNames( ULinkerSave* Linker, ULinkerLoad* LinkerToConformTo=NULL )
+	void SortNames( FLinkerSave* Linker, FLinkerLoad* LinkerToConformTo=NULL )
 	{
 		int32 SortStartPosition = 0;
 
@@ -1105,7 +1106,7 @@ public:
 	 * @param	Linker				linker containing the imports that need to be sorted
 	 * @param	LinkerToConformTo	optional linker to conform against.
 	 */
-	void SortImports( ULinkerSave* Linker, ULinkerLoad* LinkerToConformTo=NULL )
+	void SortImports( FLinkerSave* Linker, FLinkerLoad* LinkerToConformTo=NULL )
 	{
 		int32 SortStartPosition=0;
 		TArray<FObjectImport>& Imports = Linker->ImportMap;
@@ -1236,7 +1237,7 @@ public:
 	 * @param	Linker				linker containing the exports that need to be sorted
 	 * @param	LinkerToConformTo	optional linker to conform against.
 	 */
-	void SortExports( ULinkerSave* Linker, ULinkerLoad* LinkerToConformTo=NULL )
+	void SortExports( FLinkerSave* Linker, FLinkerLoad* LinkerToConformTo=NULL )
 	{
 		int32 SortStartPosition=0;
 		if ( LinkerToConformTo )
@@ -1982,7 +1983,7 @@ struct FObjectExportSeekFreeSorter
 	 * @param	Linker				LinkerSave to sort export map
 	 * @param	LinkerToConformTo	LinkerLoad to conform LinkerSave to if non- NULL
 	 */
-	void SortExports( ULinkerSave* Linker, ULinkerLoad* LinkerToConformTo )
+	void SortExports( FLinkerSave* Linker, FLinkerLoad* LinkerToConformTo )
 	{
 		SortArchive.SetCookingTarget(Linker->CookingTarget());
 
@@ -2183,7 +2184,7 @@ struct FPackageExportTagger
  * @param Error - log device to send any errors
  * @return whether NewPackage can be conformed
  */
-static bool ValidateConformCompatibility(UPackage* NewPackage, ULinkerLoad* OldLinker, FOutputDevice* Error)
+static bool ValidateConformCompatibility(UPackage* NewPackage, FLinkerLoad* OldLinker, FOutputDevice* Error)
 {
 	// various assumptions made about Core and its contents prevent loading a version mapped to a different name from working correctly
 	// Core script typically doesn't have any replication related definitions in it anyway
@@ -2380,7 +2381,7 @@ EObjectMark UPackage::GetObjectMarksForTargetPlatform( const class ITargetPlatfo
 #endif
 
 bool UPackage::SavePackage( UPackage* InOuter, UObject* Base, EObjectFlags TopLevelFlags, const TCHAR* Filename, 
-	FOutputDevice* Error, ULinkerLoad* Conform, bool bForceByteSwapping, bool bWarnOfLongFilename, uint32 SaveFlags, 
+	FOutputDevice* Error, FLinkerLoad* Conform, bool bForceByteSwapping, bool bWarnOfLongFilename, uint32 SaveFlags, 
 	const class ITargetPlatform* TargetPlatform, const FDateTime&  FinalTimeStamp, bool bSlowTask)
 {
 	UE_START_LOG_COOK_TIME( Filename );
@@ -2645,17 +2646,17 @@ bool UPackage::SavePackage( UPackage* InOuter, UObject* Base, EObjectFlags TopLe
 
 				bool bSaveUnversioned = !!(SaveFlags & SAVE_Unversioned);
 
-				ULinkerSave* Linker = NULL;
+				FLinkerSave* Linker = nullptr;
 				
 				if (bCompressFromMemory || bSaveAsync)
 				{
 					// Allocate the linker with a memory writer, forcing byte swapping if wanted.
-					Linker = new(EC_InternalUseOnlyConstructor) ULinkerSave(FObjectInitializer(), InOuter, bForceByteSwapping, bSaveUnversioned);
+					Linker = new FLinkerSave(InOuter, bForceByteSwapping, bSaveUnversioned);
 				}
 				else
 				{
 					// Allocate the linker, forcing byte swapping if wanted.
-					Linker = new(EC_InternalUseOnlyConstructor) ULinkerSave(FObjectInitializer(), InOuter, *TempFilename, bForceByteSwapping, bSaveUnversioned);
+					Linker = new FLinkerSave(InOuter, *TempFilename, bForceByteSwapping, bSaveUnversioned);
 				}
 
 				// Use the custom versions we had previously gleaned from the export tag pass
@@ -2665,6 +2666,11 @@ bool UPackage::SavePackage( UPackage* InOuter, UObject* Base, EObjectFlags TopLe
 				Linker->SetFilterEditorOnly( FilterEditorOnly );
 				Linker->SetCookingTarget(TargetPlatform);
 
+				// Make sure the package has the same version as the linker
+				InOuter->LinkerPackageVersion = Linker->UE4Ver();
+				InOuter->LinkerLicenseeVersion = Linker->LicenseeUE4Ver();
+				InOuter->LinkerCustomVersion = Linker->GetCustomVersions();
+
 				if ( EndSavingIfCancelled( Linker, TempFilename ) ) { return false; }
 				SlowTask.EnterProgressFrame();
 			
@@ -3321,7 +3327,7 @@ bool UPackage::SavePackage( UPackage* InOuter, UObject* Base, EObjectFlags TopLe
 				SlowTask.EnterProgressFrame(1, NSLOCTEXT("Core", "ProcessingExports", "ProcessingExports..."));
 
 				// look for this package in the list of packages to generate script SHA for 
-				TArray<uint8>* ScriptSHABytes = ULinkerSave::PackagesToScriptSHAMap.Find(*FPaths::GetBaseFilename(Filename));
+				TArray<uint8>* ScriptSHABytes = FLinkerSave::PackagesToScriptSHAMap.Find(*FPaths::GetBaseFilename(Filename));
 
 				// if we want to generate the SHA key, start tracking script writes
 				if (ScriptSHABytes)
@@ -3443,7 +3449,7 @@ bool UPackage::SavePackage( UPackage* InOuter, UObject* Base, EObjectFlags TopLe
 					{
 						BulkDataFeedback.EnterProgressFrame();
 
-						ULinkerSave::FBulkDataStorageInfo& BulkData = Linker->BulkDataToAppend[i];
+						FLinkerSave::FBulkDataStorageInfo& BulkData = Linker->BulkDataToAppend[i];
 
 						int64 BulkStartOffset = Linker->Tell();
 						int64 StoredBulkStartOffset = BulkStartOffset - StartOfBulkDataArea;
@@ -3750,7 +3756,7 @@ bool UPackage::SavePackage( UPackage* InOuter, UObject* Base, EObjectFlags TopLe
  * @param	InOuter							the outer to use for the new package
  * @param	Linker							linker we're currently saving with
  */
-void UPackage::SaveThumbnails( UPackage* InOuter, ULinkerSave* Linker )
+void UPackage::SaveThumbnails( UPackage* InOuter, FLinkerSave* Linker )
 {
 	Linker->Summary.ThumbnailTableOffset = 0;
 
@@ -3859,7 +3865,7 @@ void UPackage::SaveThumbnails( UPackage* InOuter, ULinkerSave* Linker )
 	}
 }
 
-void UPackage::SaveAssetRegistryData( UPackage* InOuter, ULinkerSave* Linker )
+void UPackage::SaveAssetRegistryData( UPackage* InOuter, FLinkerSave* Linker )
 {
 	// Make a copy of the tag map
 	TArray<UObject*> AssetObjects;
@@ -3912,7 +3918,7 @@ void UPackage::SaveAssetRegistryData( UPackage* InOuter, ULinkerSave* Linker )
 	}
 }
 
-void UPackage::SaveWorldLevelInfo( UPackage* InOuter, ULinkerSave* Linker )
+void UPackage::SaveWorldLevelInfo( UPackage* InOuter, FLinkerSave* Linker )
 {
 	Linker->Summary.WorldTileInfoDataOffset = 0;
 	
diff --git a/Engine/Source/Runtime/CoreUObject/Private/UObject/ScriptCore.cpp b/Engine/Source/Runtime/CoreUObject/Private/UObject/ScriptCore.cpp
--- a/Engine/Source/Runtime/CoreUObject/Private/UObject/ScriptCore.cpp
+++ b/Engine/Source/Runtime/CoreUObject/Private/UObject/ScriptCore.cpp
@@ -7,6 +7,7 @@
 #include "CoreUObjectPrivate.h"
 #include "MallocProfiler.h"
 #include "HotReloadInterface.h"
+#include "UObject/UObjectThreadContext.h"
 
 DEFINE_LOG_CATEGORY(LogScriptFrame);
 DEFINE_LOG_CATEGORY_STATIC(LogScriptCore, Log, All);
@@ -827,7 +828,7 @@ void UObject::ProcessEvent( UFunction* Function, void* Parms )
 	static int32 ScriptEntryTag = 0;
 
 	checkf(!HasAnyFlags(RF_Unreachable),TEXT("%s  Function: '%s'"), *GetFullName(), *Function->GetPathName());
-	checkf(!GIsRoutingPostLoad, TEXT("Cannot call UnrealScript (%s - %s) while PostLoading objects"), *GetFullName(), *Function->GetFullName());
+	checkf(!FUObjectThreadContext::Get().IsRoutingPostLoad, TEXT("Cannot call UnrealScript (%s - %s) while PostLoading objects"), *GetFullName(), *Function->GetFullName());
 
 	// Reject.
 	if (IsPendingKill())
diff --git a/Engine/Source/Runtime/CoreUObject/Private/UObject/StructScriptLoader.cpp b/Engine/Source/Runtime/CoreUObject/Private/UObject/StructScriptLoader.cpp
--- a/Engine/Source/Runtime/CoreUObject/Private/UObject/StructScriptLoader.cpp
+++ b/Engine/Source/Runtime/CoreUObject/Private/UObject/StructScriptLoader.cpp
@@ -49,7 +49,7 @@ public:
 	 * @param  TargetScriptContainer    The struct that the script should ultimately be serialized into
 	 * @param  ScriptLoader				The script serialization helper that contains info on the script's serializtion offset (buffer size, etc.)
 	 */
-	void AddDeferredScriptObject(ULinkerLoad* Linker, UStruct* TargetScriptContainer, const FStructScriptLoader& ScriptLoader);
+	void AddDeferredScriptObject(FLinkerLoad* Linker, UStruct* TargetScriptContainer, const FStructScriptLoader& ScriptLoader);
 
 	/**
 	 * Goes through every deferred script load associated with the specified 
@@ -59,16 +59,16 @@ public:
 	 * @param  Linker    The linker that may have deferred script serialization (possibly for many functions).
 	 * @return The number of script loads that were successfully resolved.
 	 */
-	int32 ResolveDeferredScripts(ULinkerLoad* Linker);
+	int32 ResolveDeferredScripts(FLinkerLoad* Linker);
 
 private:
 #if USE_DEFERRED_DEPENDENCY_CHECK_VERIFICATION_TESTS
 	/** Used to catch any deferred script loads that are added during a call to ResolveDeferredScripts() */
-	ULinkerLoad* ResolvingLinker;
+	FLinkerLoad* ResolvingLinker;
 #endif // USE_DEFERRED_DEPENDENCY_CHECK_VERIFICATION_TESTS
 
 	/** Keeps track of scripts (and their target containers) that need to be serialized in later */
-	TMultiMap<ULinkerLoad*, FDeferredScriptLoader> DeferredScriptLoads;
+	TMultiMap<FLinkerLoad*, FDeferredScriptLoader> DeferredScriptLoads;
 };
 
 //------------------------------------------------------------------------------
@@ -80,7 +80,7 @@ FDeferredScriptTracker::FDeferredScriptTracker()
 }
 
 //------------------------------------------------------------------------------
-void FDeferredScriptTracker::AddDeferredScriptObject(ULinkerLoad* Linker, UStruct* TargetScriptContainer, const FStructScriptLoader& ScriptLoader)
+void FDeferredScriptTracker::AddDeferredScriptObject(FLinkerLoad* Linker, UStruct* TargetScriptContainer, const FStructScriptLoader& ScriptLoader)
 {
 #if USE_DEFERRED_DEPENDENCY_CHECK_VERIFICATION_TESTS
 	check(ResolvingLinker == nullptr);
@@ -90,7 +90,7 @@ void FDeferredScriptTracker::AddDeferredScriptObject(ULinkerLoad* Linker, UStruc
 }
 
 //------------------------------------------------------------------------------
-int32 FDeferredScriptTracker::ResolveDeferredScripts(ULinkerLoad* Linker)
+int32 FDeferredScriptTracker::ResolveDeferredScripts(FLinkerLoad* Linker)
 {
 	FArchive& Ar = *Linker;
 	if (FStructScriptLoader::ShouldDeferScriptSerialization(Ar))
@@ -99,7 +99,7 @@ int32 FDeferredScriptTracker::ResolveDeferredScripts(ULinkerLoad* Linker)
 	}
 
 #if USE_DEFERRED_DEPENDENCY_CHECK_VERIFICATION_TESTS
-	TGuardValue<ULinkerLoad*> ScopedResolvingLinker(ResolvingLinker, Linker);
+	TGuardValue<FLinkerLoad*> ScopedResolvingLinker(ResolvingLinker, Linker);
 #endif // USE_DEFERRED_DEPENDENCY_CHECK_VERIFICATION_TESTS
 
 	TArray<FDeferredScriptLoader> DefferedLinkerScripts;
@@ -178,7 +178,7 @@ bool FStructScriptLoader::ShouldDeferScriptSerialization(FArchive& Ar)
 #if USE_CIRCULAR_DEPENDENCY_LOAD_DEFERRING
 	if (Ar.IsLoading() && Ar.IsPersistent())
 	{
-		if (ULinkerLoad* Linker = Cast<ULinkerLoad>(Ar.GetLinker()))
+		if (auto Linker = Cast<FLinkerLoad>(Ar.GetLinker()))
 		{
 			return ((Linker->LoadFlags & LOAD_DeferDependencyLoads) != 0);
 		}
@@ -225,7 +225,7 @@ bool FStructScriptLoader::LoadStructWithScript(UStruct* DestScriptContainer, FAr
 	bAllowDeferredSerialization &= bIsLinkerLoader;
 	if (bAllowDeferredSerialization && ShouldDeferScriptSerialization(Ar))
 	{
-		ULinkerLoad* Linker = CastChecked<ULinkerLoad>(Ar.GetLinker());
+		FLinkerLoad* Linker = CastChecked<FLinkerLoad>(Ar.GetLinker());
 		FDeferredScriptTracker::Get().AddDeferredScriptObject(Linker, DestScriptContainer, *this);
 
 		// we have to at least move the archiver forward, so it is positioned 
@@ -237,7 +237,7 @@ bool FStructScriptLoader::LoadStructWithScript(UStruct* DestScriptContainer, FAr
 	Ar.Seek(ScriptSerializationOffset);
 	if (bIsLinkerLoader)
 	{
-		ULinkerLoad* LinkerLoad = CastChecked<ULinkerLoad>(Ar.GetLinker());
+		auto LinkerLoad = CastChecked<FLinkerLoad>(Ar.GetLinker());
 
 		TArray<uint8> ShaScriptBuffer;
 		ShaScriptBuffer.AddUninitialized(SerializedScriptSize);
@@ -260,7 +260,7 @@ bool FStructScriptLoader::LoadStructWithScript(UStruct* DestScriptContainer, FAr
 	ensure(ScriptEndOffset == Ar.Tell());
 	checkf(BytecodeIndex == BytecodeBufferSize, TEXT("'%s' script expression-count mismatch; Expected: %i, Got: %i"), *DestScriptContainer->GetName(), BytecodeBufferSize, BytecodeIndex);
 
-	if (!GUObjectArray.IsDisregardForGC(DestScriptContainer))
+	if (!GetUObjectArray().IsDisregardForGC(DestScriptContainer))
 	{
 		DestScriptContainer->ScriptObjectReferences.Empty();
 		FArchiveScriptReferenceCollector ObjRefCollector(DestScriptContainer->ScriptObjectReferences);
@@ -277,7 +277,7 @@ bool FStructScriptLoader::LoadStructWithScript(UStruct* DestScriptContainer, FAr
 }
 
 //------------------------------------------------------------------------------
-int32 FStructScriptLoader::ResolveDeferredScriptLoads(ULinkerLoad* Linker)
+int32 FStructScriptLoader::ResolveDeferredScriptLoads(FLinkerLoad* Linker)
 {
 	return FDeferredScriptTracker::Get().ResolveDeferredScripts(Linker);
 }
diff --git a/Engine/Source/Runtime/CoreUObject/Private/UObject/StructScriptLoader.h b/Engine/Source/Runtime/CoreUObject/Private/UObject/StructScriptLoader.h
--- a/Engine/Source/Runtime/CoreUObject/Private/UObject/StructScriptLoader.h
+++ b/Engine/Source/Runtime/CoreUObject/Private/UObject/StructScriptLoader.h
@@ -69,7 +69,7 @@ public:
 	 * @param  Linker    The linker/archiver that the structs were originally meant to be loaded by.
 	 * @return The number of structs that were loaded with new bytecode.
 	 */
-	static int32 ResolveDeferredScriptLoads(ULinkerLoad* Linker);
+	static int32 ResolveDeferredScriptLoads(FLinkerLoad* Linker);
 
 private:
 	/**
diff --git a/Engine/Source/Runtime/CoreUObject/Private/UObject/TextureAllocations.cpp b/Engine/Source/Runtime/CoreUObject/Private/UObject/TextureAllocations.cpp
--- a/Engine/Source/Runtime/CoreUObject/Private/UObject/TextureAllocations.cpp
+++ b/Engine/Source/Runtime/CoreUObject/Private/UObject/TextureAllocations.cpp
@@ -182,7 +182,7 @@ FArchive& operator<<( FArchive& Ar, FTextureAllocations& TextureAllocations )
 		// Are we cooking a package?
 		if ( Ar.IsSaving() && !Ar.IsTransacting() )
 		{
-			ULinker* Linker = Ar.GetLinker();
+			FLinker* Linker = Ar.GetLinker();
 
 			// Do we need to build the texture allocation data?
 			if ( TextureAllocations.TextureTypes.Num() == 0 )
@@ -213,7 +213,7 @@ FArchive& operator<<( FArchive& Ar, FTextureAllocations& TextureAllocations )
 			else if ( Ar.GetLinker() )
 			{
 				NumSummaryTextures = 0;
-				ULinker* Linker = Ar.GetLinker();
+				FLinker* Linker = Ar.GetLinker();
 				for ( int32 TypeIndex=0; TypeIndex < TextureAllocations.TextureTypes.Num(); ++TypeIndex )
 				{
 					FTextureAllocations::FTextureType& TextureType = TextureAllocations.TextureTypes( TypeIndex );
diff --git a/Engine/Source/Runtime/CoreUObject/Private/UObject/UObjectArray.cpp b/Engine/Source/Runtime/CoreUObject/Private/UObject/UObjectArray.cpp
--- a/Engine/Source/Runtime/CoreUObject/Private/UObject/UObjectArray.cpp
+++ b/Engine/Source/Runtime/CoreUObject/Private/UObject/UObjectArray.cpp
@@ -8,60 +8,73 @@
 
 DEFINE_LOG_CATEGORY_STATIC(LogUObjectArray, Log, All);
 
-/** Global UObject array							*/
-FUObjectArray GUObjectArray;
-
 void FUObjectArray::AllocatePermanentObjectPool(int32 MaxObjectsNotConsideredByGC)
 {
+	check(IsInGameThread());
+
 	// GObjFirstGCIndex is the index at which the garbage collector will start for the mark phase.
-	ObjFirstGCIndex			= MaxObjectsNotConsideredByGC;
+	ObjFirstGCIndex = MaxObjectsNotConsideredByGC;
 
 	// Presize array.
-	check( ObjObjects.Num() == 0 );
-	if( ObjFirstGCIndex )
+	check(ObjObjects.Num() == 0);
+	if (ObjFirstGCIndex >= 0)
 	{
-		ObjObjects.Reserve( ObjFirstGCIndex );
+		ObjObjects.Reserve(ObjFirstGCIndex);
 	}
 	FWeakObjectPtr::Init(); // this adds a delete listener
 }
 
 void FUObjectArray::CloseDisregardForGC()
 {
+	check(IsInGameThread());
+
 	OpenForDisregardForGC = false;
 	GIsInitialLoad = false;
 	// Make sure the first GC index matches the last non-GC index after DisregardForGC is closed
 	ObjFirstGCIndex = ObjLastNonGCIndex + 1;
 }
 
-void FUObjectArray::AllocateUObjectIndex(UObjectBase* Object)
+void FUObjectArray::AllocateUObjectIndex(UObjectBase* Object, bool bMergingThreads /*= false*/)
 {
-	int32 Index;
-	check(Object->InternalIndex == INDEX_NONE);
+	int32 Index = INDEX_NONE;
+	check(Object->InternalIndex == INDEX_NONE || bMergingThreads);
 
 	// Special non- garbage collectable range.
 	if (OpenForDisregardForGC && DisregardForGCEnabled())
 	{
-		Index = ObjObjects.AddUninitialized();
+		// Disregard from GC pool is only available from the game thread, at least for now
+		check(IsInGameThread());
+		Index = ObjObjects.AddZeroed(1);
 		ObjLastNonGCIndex = Index;
 		ObjFirstGCIndex = FMath::Max(ObjFirstGCIndex, Index + 1);
 	}
 	// Regular pool/ range.
 	else
-	{
-		if(ObjAvailable.Num())
+	{		
+		int32* AvailableIndex = ObjAvailableList.Pop();
+		if (AvailableIndex)
 		{
-			Index = ObjAvailable.Pop();
-			check(ObjObjects[Index]==NULL);
+#if WITH_EDITOR
+			ObjAvailableCount.Decrement();
+			checkSlow(ObjAvailableCount.GetValue() >= 0);
+#endif
+			Index = (int32)(uintptr_t)AvailableIndex;
+			check(ObjObjects[Index]==nullptr);
 		}
 		else
 		{
-			Index = ObjObjects.AddUninitialized();
+			FScopeLock ObjObjectsLock(&ObjObjectsCritical);
+			Index = ObjObjects.AddZeroed(1);
 		}
 		check(Index >= ObjFirstGCIndex);
 	}
 	// Add to global table.
-	ObjObjects[Index] = Object;
+	if (FPlatformAtomics::InterlockedCompareExchangePointer((void**)&ObjObjects[Index], Object, NULL) != NULL) // we use an atomic operation to check for unexpected concurrency, verify alignment, etc
+	{
+		UE_LOG(LogUObjectArray, Fatal, TEXT("Unexpected concurency while adding new object"));
+	}
 	Object->InternalIndex = Index;
+	//  @todo: threading: lock UObjectCreateListeners
 	for (int32 ListenerIndex = 0; ListenerIndex < UObjectCreateListeners.Num(); ListenerIndex++)
 	{
 		UObjectCreateListeners[ListenerIndex]->NotifyUObjectCreated(Object,Index);
@@ -75,15 +88,29 @@ void FUObjectArray::AllocateUObjectIndex(UObjectBase* Object)
  */
 void FUObjectArray::FreeUObjectIndex(UObjectBase* Object)
 {
+	// This should only be happening on the game thread (GC runs only on game thread when it's freeing objects)
+	check(IsInGameThread());
+
 	int32 Index = Object->InternalIndex;
-	ObjObjects[Index] = NULL;
+	// At this point no two objects exist with the same index so no need to lock here
+	if (FPlatformAtomics::InterlockedCompareExchangePointer((void**)&ObjObjects[Index], NULL, Object) == NULL) // we use an atomic operation to check for unexpected concurrency, verify alignment, etc
+	{
+		UE_LOG(LogUObjectArray, Fatal, TEXT("Unexpected concurency while adding new object"));
+	}
+
+	// @todo: threading: delete listeners should be locked while we're doing this
 	for (int32 ListenerIndex = 0; ListenerIndex < UObjectDeleteListeners.Num(); ListenerIndex++)
 	{
-		UObjectDeleteListeners[ListenerIndex]->NotifyUObjectDeleted(Object,Index);
+		UObjectDeleteListeners[ListenerIndex]->NotifyUObjectDeleted(Object, Index);
 	}
-	if (Index > ObjLastNonGCIndex)  // you cannot safely recycle indicies in the non-GC range
+	// You cannot safely recycle indicies in the non-GC range
+	// No point in filling this list when doing exit purge. Nothing should be allocated afterwards anyway.
+	if (Index > ObjLastNonGCIndex && !GExitPurge)  
 	{
-		ObjAvailable.Add(Index);
+		ObjAvailableList.Push((int32*)Index);
+#if WITH_EDITOR
+		ObjAvailableCount.Increment();
+#endif
 	}
 }
 
@@ -170,11 +197,9 @@ bool FUObjectArray::IsValid(const UObjectBase* Object) const
  */
 void FUObjectArray::ShutdownUObjectArray()
 {
-	ObjObjects.Empty();
-	ObjAvailable.Empty();
 }
 
-TArray<UObjectBase*>* FUObjectArray::GetObjectArrayForDebugVisualizers()
+UObjectBase*** FUObjectArray::GetObjectArrayForDebugVisualizers()
 {
-	return &GUObjectArray.ObjObjects;
+	return GetUObjectArray().ObjObjects.GetRootBlockForDebuggerVisualizers();
 }
diff --git a/Engine/Source/Runtime/CoreUObject/Private/UObject/UObjectBase.cpp b/Engine/Source/Runtime/CoreUObject/Private/UObject/UObjectBase.cpp
--- a/Engine/Source/Runtime/CoreUObject/Private/UObject/UObjectBase.cpp
+++ b/Engine/Source/Runtime/CoreUObject/Private/UObject/UObjectBase.cpp
@@ -90,7 +90,7 @@ UObjectBase::~UObjectBase()
 		// Validate it.
 		check(IsValidLowLevel());
 		LowLevelRename(NAME_None);
-		GUObjectArray.FreeUObjectIndex(this);
+		GetUObjectArray().FreeUObjectIndex(this);
 	}
 }
 
@@ -153,7 +153,7 @@ void UObjectBase::DeferredRegister(UClass *UClassStaticClass,const TCHAR* Packag
 	AddObject(FName(InName));
 
 	// Make sure that objects disregarded for GC are part of root set.
-	check(!GUObjectArray.IsDisregardForGC(this) || (GetFlags() & RF_RootSet) );
+	check(!GetUObjectArray().IsDisregardForGC(this) || (GetFlags() & RF_RootSet) );
 }
 
 /**
@@ -164,7 +164,11 @@ void UObjectBase::DeferredRegister(UClass *UClassStaticClass,const TCHAR* Packag
 void UObjectBase::AddObject(FName InName)
 {
 	Name = InName;
-	GUObjectArray.AllocateUObjectIndex(this);
+	if (!IsInGameThread())
+	{
+		ObjectFlags |= RF_Async;
+	}
+	AllocateUObjectIndexForCurrentThread(this);
 	check(InName != NAME_None && InternalIndex >= 0);
 	HashObject(this);
 	check(IsValidLowLevel());
@@ -221,7 +225,7 @@ bool UObjectBase::IsValidLowLevel() const
 		UE_LOG(LogUObjectBase, Warning, TEXT("Object is not registered") );
 		return false;
 	}
-	return GUObjectArray.IsValid(this);
+	return GetUObjectArray().IsValid(this);
 }
 
 bool UObjectBase::IsValidLowLevelFast(bool bRecursive /*= true*/) const
@@ -266,7 +270,7 @@ bool UObjectBase::IsValidLowLevelFast(bool bRecursive /*= true*/) const
 		return false;
 	}
 	// Lightweight versions of index checks.
-	if (!GUObjectArray.IsValidIndex(this) || !Name.IsValidIndexFast())
+	if (!GetUObjectArray().IsValidIndex(this) || !Name.IsValidIndexFast())
 	{
 		UE_LOG(LogUObjectBase, Error, TEXT("Object array index or name index is invalid."));
 		return false;
@@ -784,7 +788,10 @@ void UObjectBaseInit()
 	UE_LOG(LogInit, Log, TEXT("Presizing for %i objects not considered by GC, pre-allocating %i bytes."), MaxObjectsNotConsideredByGC, SizeOfPermanentObjectPool );
 
 	GUObjectAllocator.AllocatePermanentObjectPool(SizeOfPermanentObjectPool);
-	GUObjectArray.AllocatePermanentObjectPool(MaxObjectsNotConsideredByGC);
+	GetUObjectArray().AllocatePermanentObjectPool(MaxObjectsNotConsideredByGC);
+
+	void InitAsyncThread();
+	InitAsyncThread();
 
 	// Note initialized.
 	Internal::GObjInitialized = true;
@@ -797,7 +804,7 @@ void UObjectBaseInit()
  */
 void UObjectBaseShutdown()
 {
-	GUObjectArray.ShutdownUObjectArray();
+	GetUObjectArray().ShutdownUObjectArray();
 	Internal::GObjInitialized = false;
 }
 
diff --git a/Engine/Source/Runtime/CoreUObject/Private/UObject/UObjectGlobals.cpp b/Engine/Source/Runtime/CoreUObject/Private/UObject/UObjectGlobals.cpp
--- a/Engine/Source/Runtime/CoreUObject/Private/UObject/UObjectGlobals.cpp
+++ b/Engine/Source/Runtime/CoreUObject/Private/UObject/UObjectGlobals.cpp
@@ -9,6 +9,7 @@
 #include "UObject/ConstructorHelpers.h"
 #include "UObject/UObjectAnnotation.h"
 #include "UObject/TlsObjectInitializers.h"
+#include "UObject/UObjectThreadContext.h"
 
 DEFINE_LOG_CATEGORY(LogUObjectGlobals);
 
@@ -18,21 +19,6 @@ DEFINE_LOG_CATEGORY(LogUObjectGlobals);
 #endif
 
 bool						GIsSavingPackage = false;
-int32							GImportCount							= 0;
-/** Forced exports for EndLoad optimization.											*/
-int32							GForcedExportCount						= 0;
-TArray<UObject*>			GObjLoaded;
-
-/** List of linkers that we want to close the loaders for (to free file handles) - needs to be delayed until EndLoad is called with GObjBeginLoadCount of 0 */
-TArray<ULinkerLoad*>		GDelayedLinkerClosePackages;
-
-/** Count for BeginLoad multiple loads.									*/
-static int32					GObjBeginLoadCount						= 0;
-
-#if !(UE_BUILD_SHIPPING || UE_BUILD_TEST)
-/** Used to verify that the Super::Serialize chain is intact.			*/
-TArray<UObject*,TInlineAllocator<16> >		DebugSerialize;
-#endif
 
 /** Object annotation used by the engine to keep track of which objects are selected */
 FUObjectAnnotationSparseBool GSelectedAnnotation;
@@ -84,7 +70,16 @@ FCoreUObjectDelegates::FPackageCreatedForLoad FCoreUObjectDelegates::PackageCrea
 /** Check wehether we should report progress or not */
 bool ShouldReportProgress()
 {
-	return GIsEditor && !IsRunningCommandlet() && !GIsAsyncLoading;
+	return GIsEditor && !IsRunningCommandlet() && !IsAsyncLoading();
+}
+
+/**
+ * Returns true if code is called form the game thread while collecting garbage.
+ * We only have tp guard agains StaticFindObject on the game thread as other threads will be blocked anyway
+ */
+static FORCEINLINE bool IsGarbageCollectingOnGameThread()
+{
+	return IsInGameThread() && IsGarbageCollecting();
 }
 
 /**
@@ -101,13 +96,13 @@ bool ShouldReportProgress()
  */
 UObject* StaticFindObjectFast( UClass* ObjectClass, UObject* ObjectPackage, FName ObjectName, bool ExactClass, bool AnyPackage, EObjectFlags ExclusiveFlags )
 {
-	if (GIsSavingPackage || GIsGarbageCollecting)
+	if (GIsSavingPackage || IsGarbageCollectingOnGameThread())
 	{
 		UE_LOG(LogUObjectGlobals, Fatal,TEXT("Illegal call to StaticFindObjectFast() while serializing object data or garbage collecting!"));
 	}
 
 	// We don't want to return any objects that are currently being background loaded unless we're using FindObject during async loading.
-	ExclusiveFlags |= GIsAsyncLoading ? RF_NoFlags : RF_AsyncLoading;
+	ExclusiveFlags |= IsAsyncLoading() ? RF_NoFlags : RF_AsyncLoading;
 	return StaticFindObjectFastInternal( ObjectClass, ObjectPackage, ObjectName, ExactClass, AnyPackage, ExclusiveFlags );
 }
 
@@ -123,7 +118,7 @@ UObject* StaticFindObject( UClass* ObjectClass, UObject* InObjectPackage, const
 		UE_LOG(LogUObjectGlobals, Fatal,TEXT("Illegal call to StaticFindObject() while serializing object data!"));
 	}
 
-	if (GIsGarbageCollecting)
+	if (IsGarbageCollectingOnGameThread())
 	{
 		UE_LOG(LogUObjectGlobals, Fatal,TEXT("Illegal call to StaticFindObject() while collecting garbage!"));
 	}
@@ -204,11 +199,11 @@ UObject* StaticFindObjectChecked( UClass* ObjectClass, UObject* ObjectParent, co
 }
 
 //
-// Find an object; won't assert on GIsSavingPackage or GIsGarbageCollecting
+// Find an object; won't assert on GIsSavingPackage or IsGarbageCollecting()
 //
 UObject* StaticFindObjectSafe( UClass* ObjectClass, UObject* ObjectParent, const TCHAR* InName, bool ExactClass )
 {
-	if ( !GIsSavingPackage && !GIsGarbageCollecting )
+	if (!GIsSavingPackage && !IsGarbageCollectingOnGameThread())
 	{
 		return StaticFindObject( ObjectClass, ObjectParent, InName, ExactClass );
 	}
@@ -770,7 +765,7 @@ UObject* StaticLoadObjectInternal(UClass* ObjectClass, UObject* InOuter, const T
 
 UObject* StaticLoadObject(UClass* ObjectClass, UObject* InOuter, const TCHAR* InName, const TCHAR* Filename, uint32 LoadFlags, UPackageMap* Sandbox, bool bAllowObjectReconciliation )
 {
-	//UE_CLOG(GIsRoutingPostLoad, LogUObjectGlobals, Fatal, TEXT("Calling StaticLoadObject during PostLoad is forbidden."));
+	UE_CLOG(FUObjectThreadContext::Get().IsRoutingPostLoad, LogUObjectGlobals, Fatal, TEXT("Calling StaticLoadObject during PostLoad is forbidden."));
 
 	UObject* Result = StaticLoadObjectInternal(ObjectClass, InOuter, InName, Filename, LoadFlags, Sandbox, bAllowObjectReconciliation);
 	if (!Result)
@@ -820,7 +815,7 @@ UClass* StaticLoadClass( UClass* BaseClass, UObject* InOuter, const TCHAR* InNam
 * @param	ImportLinker	Linker that requests this package through one of its imports
 * @return	Loaded package if successful, NULL otherwise
 */
-UPackage* LoadPackageInternal(UPackage* InOuter, const TCHAR* InLongPackageName, uint32 LoadFlags, ULinkerLoad* ImportLinker)
+UPackage* LoadPackageInternal(UPackage* InOuter, const TCHAR* InLongPackageName, uint32 LoadFlags, FLinkerLoad* ImportLinker)
 {
 	DECLARE_SCOPE_CYCLE_COUNTER(TEXT("LoadPackageInternal"), STAT_LoadPackageInternal, STATGROUP_ObjectVerbose);
 
@@ -869,12 +864,15 @@ UPackage* LoadPackageInternal(UPackage* InOuter, const TCHAR* InLongPackageName,
 	bool bFullyLoadSkipped = false;
 
 	SlowTask.EnterProgressFrame(30);
+
+	// Declare here so that the linker does not get destroyed before ResetLoaders is called
+	FLinkerLoad* Linker = nullptr;
 	{
 		// Keep track of start time.
 		const double StartTime = FPlatformTime::Seconds();
 
 		// Create a new linker object which goes off and tries load the file.
-		ULinkerLoad* Linker = GetPackageLinker(InOuter, *FileToLoad, LoadFlags, nullptr, nullptr);
+		Linker = GetPackageLinker(InOuter, *FileToLoad, LoadFlags, nullptr, nullptr);
 		if (!Linker)
 		{
 			EndLoad();
@@ -992,7 +990,7 @@ UPackage* LoadPackageInternal(UPackage* InOuter, const TCHAR* InLongPackageName,
 		// fail, so we only currently do this when loading on consoles
 		if (FPlatformProperties::RequiresCookedData())
 		{
-			if (GObjBeginLoadCount == 0)
+			if (FUObjectThreadContext::Get().ObjBeginLoadCount == 0)
 			{
 				if (Result && Linker->Loader)
 				{
@@ -1001,9 +999,10 @@ UPackage* LoadPackageInternal(UPackage* InOuter, const TCHAR* InLongPackageName,
 				delete Linker->Loader;
 				Linker->Loader = NULL;
 			}
-			else
+			// Async loading removes delayed linkers on the game thread after streaming has finished
+			else if (!IsInAsyncLoadingThread())
 			{
-				GDelayedLinkerClosePackages.Add(Linker);
+				FUObjectThreadContext::Get().DelayedLinkerClosePackages.Add(Linker);
 			}
 		}
 	}
@@ -1034,8 +1033,8 @@ UPackage* LoadPackage(UPackage* InOuter, const TCHAR* InLongPackageName, uint32
  */
 bool IsLoading()
 {
-	check(GObjBeginLoadCount>=0);
-	return GObjBeginLoadCount > 0;
+	check(FUObjectThreadContext::Get().ObjBeginLoadCount >= 0);
+	return FUObjectThreadContext::Get().ObjBeginLoadCount > 0;
 }
 
 //
@@ -1044,13 +1043,14 @@ bool IsLoading()
 //
 void BeginLoad()
 {
-	if( ++GObjBeginLoadCount == 1 && !GIsAsyncLoading )
+	auto& ThreadContext = FUObjectThreadContext::Get();
+	if (++ThreadContext.ObjBeginLoadCount == 1 && !IsInAsyncLoadingThread())
 	{
 		// Make sure we're finishing up all pending async loads, and trigger texture streaming next tick if necessary.
 		FlushAsyncLoading( NAME_None );
 
 		// Validate clean load state.
-		check(GObjLoaded.Num()==0);
+		check(ThreadContext.ObjLoaded.Num() == 0);
 	}
 }
 
@@ -1059,8 +1059,8 @@ struct FCompareUObjectByLinkerAndOffset
 {
 	FORCEINLINE bool operator()( const UObject& A, const UObject &B ) const
 	{
-		ULinker* LinkerA = A.GetLinker();
-		ULinker* LinkerB = B.GetLinker();
+		FLinker* LinkerA = A.GetLinker();
+		FLinker* LinkerB = B.GetLinker();
 
 		// Both objects have linkers.
 		if( LinkerA && LinkerB )
@@ -1072,10 +1072,10 @@ struct FCompareUObjectByLinkerAndOffset
 				FObjectExport& ExportB = LinkerB->ExportMap[ B.GetLinkerIndex() ];
 				return ExportA.SerialOffset < ExportB.SerialOffset;
 			}
-			// Sort by linker name.
+			// Sort by pointer address.
 			else
 			{
-				return LinkerA->GetFName().GetComparisonIndex() < LinkerB->GetFName().GetComparisonIndex();
+				return LinkerA < LinkerB;
 			}
 		}
 		// Neither objects have a linker, don't do anything.
@@ -1096,10 +1096,12 @@ struct FCompareUObjectByLinkerAndOffset
 //
 void EndLoad()
 {
-	check(GObjBeginLoadCount>0);
-	if (GIsAsyncLoading)
+	auto& ThreadContext = FUObjectThreadContext::Get();
+
+	check(ThreadContext.ObjBeginLoadCount > 0);
+	if (IsInAsyncLoadingThread())
 	{
-		GObjBeginLoadCount--;
+		ThreadContext.ObjBeginLoadCount--;
 		return;
 	}
 
@@ -1109,20 +1111,20 @@ void EndLoad()
 	int32 NumObjectsLoaded = 0, NumObjectsFound = 0;
 #endif
 
-	while( --GObjBeginLoadCount == 0 && (GObjLoaded.Num() || GImportCount || GForcedExportCount) )
+	while (--ThreadContext.ObjBeginLoadCount == 0 && (ThreadContext.ObjLoaded.Num() || ThreadContext.ImportCount || ThreadContext.ForcedExportCount))
 	{
 		// Make sure we're not recursively calling EndLoad as e.g. loading a config file could cause
 		// BeginLoad/EndLoad to be called.
-		GObjBeginLoadCount++;
+		ThreadContext.ObjBeginLoadCount++;
 
 		// Temporary list of loaded objects as GObjLoaded might expand during iteration.
 		TArray<UObject*> ObjLoaded;
-		TSet<ULinkerLoad*> LoadedLinkers;
-		while( GObjLoaded.Num() )
+		TSet<FLinkerLoad*> LoadedLinkers;
+		while (ThreadContext.ObjLoaded.Num())
 		{
 			// Accumulate till GObjLoaded no longer increases.
-			ObjLoaded += GObjLoaded;
-			GObjLoaded.Empty();
+			ObjLoaded += ThreadContext.ObjLoaded;
+			ThreadContext.ObjLoaded.Empty();
 
 			// Sort by Filename and Offset.
 			ObjLoaded.Sort( FCompareUObjectByLinkerAndOffset() );
@@ -1141,7 +1143,7 @@ void EndLoad()
 
 			// Start over again as new objects have been loaded that need to have "Preload" called on them before
 			// we can safely PostLoad them.
-			if(GObjLoaded.Num())
+			if (ThreadContext.ObjLoaded.Num())
 			{
 				continue;
 			}
@@ -1166,7 +1168,7 @@ void EndLoad()
 
 			{
 				// set this so that we can perform certain operations in which are only safe once all objects have been de-serialized.
-				TGuardValue<bool> GuardIsRoutingPostLoad(GIsRoutingPostLoad, true);
+				TGuardValue<bool> GuardIsRoutingPostLoad(FUObjectThreadContext::Get().IsRoutingPostLoad, true);
 
 				// Postload objects.
 				for(int32 i = 0; i < ObjLoaded.Num(); i++)
@@ -1198,23 +1200,22 @@ void EndLoad()
 #endif	// WITH_EDITOR
 
 			// Empty array before next iteration as we finished postloading all objects.
-			ObjLoaded.Empty( GObjLoaded.Num() );
+			ObjLoaded.Empty(ThreadContext.ObjLoaded.Num());
 		}
 
 		if ( GIsEditor && LoadedLinkers.Num() > 0 )
 		{
-			for ( TSet<ULinkerLoad*>::TIterator It(LoadedLinkers); It; ++It )
+			for (auto LoadedLinker : LoadedLinkers)
 			{
-				ULinkerLoad* LoadedLinker = *It;
-				check(LoadedLinker);
+				check(LoadedLinker != nullptr);
 
-				if ( LoadedLinker->LinkerRoot != NULL && !LoadedLinker->LinkerRoot->IsFullyLoaded() )
+				if (LoadedLinker->LinkerRoot != nullptr && !LoadedLinker->LinkerRoot->IsFullyLoaded())
 				{
 					bool bAllExportsCreated = true;
 					for ( int32 ExportIndex = 0; ExportIndex < LoadedLinker->ExportMap.Num(); ExportIndex++ )
 					{
 						FObjectExport& Export = LoadedLinker->ExportMap[ExportIndex];
-						if ( !Export.bForcedExport && Export.Object == NULL )
+						if ( !Export.bForcedExport && Export.Object == nullptr )
 						{
 							bAllExportsCreated = false;
 							break;
@@ -1234,7 +1235,7 @@ void EndLoad()
 		DissociateImportsAndForcedExports();
 
 		// close any linkers' loaders that were requested to be closed once GObjBeginLoadCount goes to 0
-		auto PackagesToClose = MoveTemp(GDelayedLinkerClosePackages);
+		auto PackagesToClose = MoveTemp(ThreadContext.DelayedLinkerClosePackages);
 		for (auto Linker: PackagesToClose)
 		{
 			if (Linker)
@@ -1791,7 +1792,7 @@ UObject* StaticAllocateObject
 		}
 	}
 
-	ULinkerLoad*	Linker						= NULL;
+	FLinkerLoad*	Linker						= NULL;
 	int32				LinkerIndex					= INDEX_NONE;
 	bool			bWasConstructedOnOldObject	= false;
 	// True when the object to be allocated already exists and is a subobject.
@@ -1895,7 +1896,7 @@ UObject* StaticAllocateObject
 		}
 	}
 
-	if (GIsAsyncLoading)
+	if (IsInAsyncLoadingThread())
 	{
 		NotifyConstructedDuringAsyncLoading(Obj, bSubObject);
 	}
@@ -1910,16 +1911,13 @@ UObject* StaticAllocateObject
 }
 
 //@todo UE4 - move this stuff to UnObj.cpp or something
-#if !(UE_BUILD_SHIPPING || UE_BUILD_TEST)
-/** Stack to ensure that PostInitProperties is routed through Super:: calls. **/
-TArray<UObject*> PostInitPropertiesCheck;
-#endif
+
 
 
 void UObject::PostInitProperties()
 {
 #if !(UE_BUILD_SHIPPING || UE_BUILD_TEST)
-	PostInitPropertiesCheck.Push(this);
+	FUObjectThreadContext::Get().PostInitPropertiesCheck.Push(this);
 #endif
 #if USE_UBER_GRAPH_PERSISTENT_FRAME
 	GetClass()->CreatePersistentUberGraphFrame(this);
@@ -1951,23 +1949,22 @@ UObject::UObject(const FObjectInitializer& ObjectInitializer)
 	const_cast<FObjectInitializer&>(ObjectInitializer).FinalizeSubobjectClassInitialization();
 }
 
-/* Global flag so that FObjectFinders know if they are called from inside the UObject constructors or not. */
-int32 GIsInConstructor = 0;
-/* Object that is currently being constructed with ObjectInitializer */
-static UObject* GConstructedObject = NULL;
+
 
 FObjectInitializer::FObjectInitializer() :
-	Obj(NULL),
-	ObjectArchetype(NULL),
+Obj(nullptr),
+	ObjectArchetype(nullptr),
 	bCopyTransientsFromClassDefaults(false),
 	bShouldIntializePropsFromArchetype(true),
 	bSubobjectClassInitializationAllowed(true),
-	InstanceGraph(NULL),
-	LastConstructedObject(GConstructedObject)
+	InstanceGraph(nullptr),
+	LastConstructedObject(nullptr)
 {
+	auto& ThreadContext = FUObjectThreadContext::Get();
 	// Mark we're in the constructor now.	
-	GIsInConstructor++;
-	GConstructedObject = Obj;
+	ThreadContext.IsInConstructor++;
+	LastConstructedObject = ThreadContext.ConstructedObject;
+	ThreadContext.ConstructedObject = Obj;
 	FTlsObjectInitializers::Push(this);
 }	
 
@@ -1979,11 +1976,13 @@ FObjectInitializer::FObjectInitializer(UObject* InObj, UObject* InObjectArchetyp
 	bShouldIntializePropsFromArchetype(bInShouldIntializeProps),
 	bSubobjectClassInitializationAllowed(true),
 	InstanceGraph(InInstanceGraph),
-	LastConstructedObject(GConstructedObject)
+	LastConstructedObject(nullptr)
 {
+	auto& ThreadContext = FUObjectThreadContext::Get();
 	// Mark we're in the constructor now.
-	GIsInConstructor++;
-	GConstructedObject = Obj;
+	ThreadContext.IsInConstructor++;
+	LastConstructedObject = ThreadContext.ConstructedObject;
+	ThreadContext.ConstructedObject = Obj;
 	FTlsObjectInitializers::Push(this);
 }
 
@@ -1993,10 +1992,11 @@ FObjectInitializer::FObjectInitializer(UObject* InObj, UObject* InObjectArchetyp
 FObjectInitializer::~FObjectInitializer()
 {
 	FTlsObjectInitializers::Pop();
+	auto& ThreadContext = FUObjectThreadContext::Get();
 	// Let the FObjectFinders know we left the constructor.
-	GIsInConstructor--;
-	check(GIsInConstructor >= 0);
-	GConstructedObject = LastConstructedObject;
+	ThreadContext.IsInConstructor--;
+	check(ThreadContext.IsInConstructor >= 0);
+	ThreadContext.ConstructedObject = LastConstructedObject;
 
 	SCOPE_CYCLE_COUNTER(STAT_PostConstructInitializeProperties);
 	check(Obj);
@@ -2069,7 +2069,7 @@ FObjectInitializer::~FObjectInitializer()
 	Obj->PostInitProperties();
 
 #if !(UE_BUILD_SHIPPING || UE_BUILD_TEST)
-	if (!PostInitPropertiesCheck.Num() || (PostInitPropertiesCheck.Pop() != Obj))
+	if (!FUObjectThreadContext::Get().PostInitPropertiesCheck.Num() || (FUObjectThreadContext::Get().PostInitPropertiesCheck.Pop() != Obj))
 	{
 		UE_LOG(LogUObjectGlobals, Fatal, TEXT("%s failed to route PostInitProperties"), *Obj->GetClass()->GetName() );
 	}
@@ -2329,12 +2329,14 @@ UObject* StaticConstructObject
 
 void FObjectInitializer::AssertIfInConstructor(UObject* Outer, const TCHAR* ErrorMessage)
 {
-	UE_CLOG(GIsInConstructor && Outer == GConstructedObject, LogUObjectGlobals, Fatal, TEXT("%s"), ErrorMessage);
+	auto& ThreadContext = FUObjectThreadContext::Get();
+	UE_CLOG(ThreadContext.IsInConstructor && Outer == ThreadContext.ConstructedObject, LogUObjectGlobals, Fatal, TEXT("%s"), ErrorMessage);
 }
 
 FObjectInitializer& FObjectInitializer::Get()
 {
-	UE_CLOG(!GIsInConstructor, LogUObjectGlobals, Fatal, TEXT("FObjectInitializer::Get() can only be used inside of UObject-derived class constructor."));
+	auto& ThreadContext = FUObjectThreadContext::Get();
+	UE_CLOG(!ThreadContext.IsInConstructor, LogUObjectGlobals, Fatal, TEXT("FObjectInitializer::Get() can only be used inside of UObject-derived class constructor."));
 	return FTlsObjectInitializers::TopChecked();
 }
 
@@ -2400,7 +2402,8 @@ void ConstructorHelpers::CheckFoundViaRedirect(UObject *Object, const FString& P
 
 void ConstructorHelpers::CheckIfIsInConstructor(const TCHAR* ObjectToFind)
 {
-	UE_CLOG(!GIsInConstructor, LogUObjectGlobals, Fatal, TEXT("FObjectFinders can't be used outside of constructors to find %s"), ObjectToFind);
+	auto& ThreadContext = FUObjectThreadContext::Get();
+	UE_CLOG(!ThreadContext.IsInConstructor, LogUObjectGlobals, Fatal, TEXT("FObjectFinders can't be used outside of constructors to find %s"), ObjectToFind);
 }
 
 void ConstructorHelpers::StripObjectClass( FString& PathName, bool bAssertOnBadPath /*= false */ )
@@ -2466,7 +2469,6 @@ public:
 		ReferenceSearchFlags = SearchFlags;
 		FoundReferencesList = FoundReferences;
 
-		checkSlow(!DebugSerialize.Num()); // should only be filled int he scope of this function
 		// Iterate over all objects.
 		for( FObjectIterator It; It; ++It )
 		{
@@ -2509,7 +2511,6 @@ public:
 			// Serialize object.
 			FindReferences( CurrentObject );
 		}
-		checkSlow(!DebugSerialize.Num()); // everyone should have been removed, or we should have errored
 	}
 
 private:
@@ -2751,7 +2752,6 @@ UObject* FObjectInitializer::CreateDefaultSubobject(UObject* Outer, FName Subobj
 	}
 	return Result;
 }
-
 UObject* FObjectInitializer::CreateEditorOnlyDefaultSubobject(UObject* Outer, FName SubobjectName, UClass* ReturnType, bool bTransient /*= false*/) const
 {
 #if WITH_EDITOR
diff --git a/Engine/Source/Runtime/CoreUObject/Private/UObject/UObjectHash.cpp b/Engine/Source/Runtime/CoreUObject/Private/UObject/UObjectHash.cpp
--- a/Engine/Source/Runtime/CoreUObject/Private/UObject/UObjectHash.cpp
+++ b/Engine/Source/Runtime/CoreUObject/Private/UObject/UObjectHash.cpp
@@ -8,6 +8,15 @@
 
 DEFINE_LOG_CATEGORY_STATIC(LogUObjectHash, Log, All);
 
+DECLARE_STATS_GROUP(TEXT("UObject Hash"), STATGROUP_UObjectHash, STATCAT_Advanced);
+DECLARE_FLOAT_ACCUMULATOR_STAT(TEXT("StaticFindObjectFastInternal"), STAT_Hash_StaticFindObjectFastInternal, STATGROUP_UObjectHash);
+DECLARE_FLOAT_ACCUMULATOR_STAT(TEXT("StaticFindObjectFastExplicit"), STAT_Hash_StaticFindObjectFastExplicit, STATGROUP_UObjectHash);
+DECLARE_FLOAT_ACCUMULATOR_STAT(TEXT("GetObjectsWithOuter"), STAT_Hash_GetObjectsWithOuter, STATGROUP_UObjectHash);
+DECLARE_FLOAT_ACCUMULATOR_STAT(TEXT("FindObjectWithOuter"), STAT_Hash_FindObjectWithOuter, STATGROUP_UObjectHash);
+DECLARE_FLOAT_ACCUMULATOR_STAT(TEXT("GetObjectsOfClass"), STAT_Hash_GetObjectsOfClass, STATGROUP_UObjectHash);
+DECLARE_FLOAT_ACCUMULATOR_STAT(TEXT("HashObject"), STAT_Hash_HashObject, STATGROUP_UObjectHash);
+DECLARE_FLOAT_ACCUMULATOR_STAT(TEXT("UnhashObject"), STAT_Hash_UnhashObject, STATGROUP_UObjectHash);
+
 /**
  * This implementation will use more space than the UE3 implementation. The goal was to make UObjects smaller to save L2 cache space. 
  * The hash is rarely used at runtime. A more space-efficient implementation is possible.
@@ -22,8 +31,70 @@ DEFINE_LOG_CATEGORY_STATIC(LogUObjectHash, Log, All);
 */
 #define OBJECT_HASH_BINS (1024*1024)
 
-static TMultiMap<int32,class UObjectBase*> ObjectHash;
-static TMultiMap<int32,class UObjectBase*> ObjectHashOuter;
+class FUObjectHashTables
+{
+	FCriticalSection CriticalSection;
+
+public:
+
+	/** Hash sets */
+	TMultiMap<int32, class UObjectBase*> Hash;
+	TMultiMap<int32, class UObjectBase*> HashOuter;
+
+	/** Map of object to their outers, used to avoid an object iterator to find such things. **/
+	TMap<UObjectBase*, TSet<UObjectBase*> > ObjectOuterMap;
+	TMap<UClass*, TSet<UObjectBase*> > ClassToObjectListMap;
+	TMap<UClass*, TSet<UClass*> > ClassToChildListMap;
+
+	FUObjectHashTables()
+	{
+	}
+
+	FORCEINLINE void Lock()
+	{
+		CriticalSection.Lock();
+	}
+
+	FORCEINLINE void Unlock()
+	{
+		CriticalSection.Unlock();
+	}
+
+	static FUObjectHashTables& Get()
+	{
+		static FUObjectHashTables Singleton;
+		return Singleton;
+	}
+};
+
+class FHashTableLock
+{
+	FUObjectHashTables& Tables;
+public:
+	FORCEINLINE FHashTableLock(FUObjectHashTables& InTables)
+		: Tables(InTables)
+	{
+		// GC locks everything on the main thread so no need to lock here
+		if (!(IsGarbageCollecting() && IsInGameThread()))
+		{
+			Tables.Lock();
+		}
+	}
+	FORCEINLINE ~FHashTableLock()
+	{
+		if (!(IsGarbageCollecting() && IsInGameThread()))
+		{
+			Tables.Unlock();
+		}
+	}
+};
+
+FUObjectArray& GetUObjectArray()
+{
+	static FUObjectArray GlobalUObjectArray;
+	return GlobalUObjectArray;
+}
+
 
 /**
  * Calculates the object's hash just using the object's name index
@@ -49,40 +120,30 @@ static FORCEINLINE int32 GetObjectOuterHash(FName ObjName,PTRINT Outer)
 	return ((ObjName.GetComparisonIndex() ^ ObjName.GetNumber()) ^ (Outer >> 4)) & (OBJECT_HASH_BINS - 1);
 }
 
-/**
- * Variation of StaticFindObjectFast that uses explicit path.
- *
- * @param	ObjectClass		The to be found object's class
- * @param	ObjectName		The to be found object's class
- * @param	ObjectPathName	Full path name for the object to search for
- * @param	ExactClass		Whether to require an exact match with the passed in class
- * @param	ExclusiveFlags	Ignores objects that contain any of the specified exclusive flags
- * @return	Returns a pointer to the found object or NULL if none could be found
- */
-UObject* StaticFindObjectFastExplicit( UClass* ObjectClass, FName ObjectName, const FString& ObjectPathName, bool bExactClass, EObjectFlags ExcludeFlags/*=0*/ )
+UObject* StaticFindObjectFastExplicitThreadSafe(FUObjectHashTables& ThreadHash, UClass* ObjectClass, FName ObjectName, const FString& ObjectPathName, bool bExactClass, EObjectFlags ExcludeFlags/*=0*/)
 {
-	checkSlow(FPackageName::IsShortPackageName(ObjectName)); //@Package name transition, we aren't checking the name here because we know this is only used for texture
 	// Find an object with the specified name and (optional) class, in any package; if bAnyPackage is false, only matches top-level packages
-	int32 Hash = GetObjectHash( ObjectName );
-	for(TMultiMap<int32,class UObjectBase*>::TConstKeyIterator HashIt(ObjectHash,Hash); HashIt; ++HashIt)
+	int32 Hash = GetObjectHash(ObjectName);
+	FHashTableLock HashLock(ThreadHash);
+	for (TMultiMap<int32, class UObjectBase*>::TConstKeyIterator HashIt(ThreadHash.Hash, Hash); HashIt; ++HashIt)
 	{
 		UObject *Object = (UObject *)HashIt.Value();
 		if
-		(	(Object->GetFName()==ObjectName)
+			((Object->GetFName() == ObjectName)
 
-		/* Don't return objects that have any of the exclusive flags set */
-		&&	!Object->HasAnyFlags(ExcludeFlags)
+			/* Don't return objects that have any of the exclusive flags set */
+			&& !Object->HasAnyFlags(ExcludeFlags)
 
-		/** If a class was specified, check that the object is of the correct class */
-		&&	(ObjectClass==NULL || (bExactClass ? Object->GetClass()==ObjectClass : Object->IsA(ObjectClass)))
-		)
+			/** If a class was specified, check that the object is of the correct class */
+			&& (ObjectClass == NULL || (bExactClass ? Object->GetClass() == ObjectClass : Object->IsA(ObjectClass)))
+			)
 
 		{
 			FString ObjectPath = Object->GetPathName();
 			/** Finally check the explicit path */
 			if (ObjectPath == ObjectPathName)
 			{
-				checkf( !Object->HasAnyFlags(RF_Unreachable), TEXT("%s"), *Object->GetFullName() );
+				checkf(!Object->HasAnyFlags(RF_Unreachable), TEXT("%s"), *Object->GetFullName());
 				return Object;
 			}
 		}
@@ -91,32 +152,57 @@ UObject* StaticFindObjectFastExplicit( UClass* ObjectClass, FName ObjectName, co
 	return NULL;
 }
 
-UObject* StaticFindObjectFastInternal( UClass* ObjectClass, UObject* ObjectPackage, FName ObjectName, bool bExactClass, bool bAnyPackage, EObjectFlags ExcludeFlags )
+/**
+ * Variation of StaticFindObjectFast that uses explicit path.
+ *
+ * @param	ObjectClass		The to be found object's class
+ * @param	ObjectName		The to be found object's class
+ * @param	ObjectPathName	Full path name for the object to search for
+ * @param	ExactClass		Whether to require an exact match with the passed in class
+ * @param	ExclusiveFlags	Ignores objects that contain any of the specified exclusive flags
+ * @return	Returns a pointer to the found object or NULL if none could be found
+ */
+UObject* StaticFindObjectFastExplicit( UClass* ObjectClass, FName ObjectName, const FString& ObjectPathName, bool bExactClass, EObjectFlags ExcludeFlags/*=0*/ )
+{
+	checkSlow(FPackageName::IsShortPackageName(ObjectName)); //@Package name transition, we aren't checking the name here because we know this is only used for texture
+	// Find an object with the specified name and (optional) class, in any package; if bAnyPackage is false, only matches top-level packages
+	int32 Hash = GetObjectHash( ObjectName );
+	UObject* Result = nullptr;
+	STAT(double ThisTime = 0);
+	{
+		SCOPE_SECONDS_COUNTER(ThisTime);
+		auto& ThreadHash = FUObjectHashTables::Get();
+		Result = StaticFindObjectFastExplicitThreadSafe(ThreadHash, ObjectClass, ObjectName, ObjectPathName, bExactClass, ExcludeFlags);
+	}
+	INC_FLOAT_STAT_BY(STAT_Hash_StaticFindObjectFastExplicit, (float)ThisTime);
+	return Result;
+}
+
+UObject* StaticFindObjectFastInternalThreadSafe(FUObjectHashTables& ThreadHash, UClass* ObjectClass, UObject* ObjectPackage, FName ObjectName, bool bExactClass, bool bAnyPackage, EObjectFlags ExcludeFlags)
 {
-	INC_DWORD_STAT(STAT_FindObjectFast);
-	check(ObjectPackage != ANY_PACKAGE); // this could never have returned anything but NULL
 	// If they specified an outer use that during the hashing
 	UObject* Result = NULL;
 	if (ObjectPackage != NULL)
 	{
-		int32 Hash = GetObjectOuterHash( ObjectName, (PTRINT)ObjectPackage );
-		for(TMultiMap<int32,class UObjectBase*>::TConstKeyIterator HashIt(ObjectHashOuter,Hash); HashIt; ++HashIt)
+		int32 Hash = GetObjectOuterHash(ObjectName, (PTRINT)ObjectPackage);
+		FHashTableLock HashLock(ThreadHash);
+		for (TMultiMap<int32, class UObjectBase*>::TConstKeyIterator HashIt(ThreadHash.HashOuter, Hash); HashIt; ++HashIt)
 		{
 			UObject *Object = (UObject *)HashIt.Value();
 			if
-			/* check that the name matches the name we're searching for */
-			(	(Object->GetFName()==ObjectName)
+				/* check that the name matches the name we're searching for */
+				((Object->GetFName() == ObjectName)
 
-			/* Don't return objects that have any of the exclusive flags set */
-			&&	!Object->HasAnyFlags(ExcludeFlags)
+				/* Don't return objects that have any of the exclusive flags set */
+				&& !Object->HasAnyFlags(ExcludeFlags)
 
-			/* check that the object has the correct Outer */
-			&&	Object->GetOuter() == ObjectPackage
+				/* check that the object has the correct Outer */
+				&& Object->GetOuter() == ObjectPackage
 
-			/** If a class was specified, check that the object is of the correct class */
-			&&	(ObjectClass==NULL || (bExactClass ? Object->GetClass()==ObjectClass : Object->IsA(ObjectClass))) )
+				/** If a class was specified, check that the object is of the correct class */
+				&& (ObjectClass == NULL || (bExactClass ? Object->GetClass() == ObjectClass : Object->IsA(ObjectClass))))
 			{
-				checkf( !Object->HasAnyFlags(RF_Unreachable), TEXT("%s"), *Object->GetFullName() );
+				checkf(!Object->HasAnyFlags(RF_Unreachable), TEXT("%s"), *Object->GetFullName());
 				if (Result)
 				{
 					UE_LOG(LogUObjectHash, Warning, TEXT("Ambiguous search, could be %s or %s"), *GetFullNameSafe(Result), *GetFullNameSafe(Object));
@@ -136,35 +222,36 @@ UObject* StaticFindObjectFastInternal( UClass* ObjectClass, UObject* ObjectPacka
 		// Find an object with the specified name and (optional) class, in any package; if bAnyPackage is false, only matches top-level packages
 		FName ActualObjectName = ObjectName;
 		const FString ObjectNameString = ObjectName.ToString();
-		const int32 DotIndex = FMath::Max<int32>(ObjectNameString.Find(TEXT("."), ESearchCase::CaseSensitive, ESearchDir::FromEnd), 
-												 ObjectNameString.Find(TEXT(":"), ESearchCase::CaseSensitive, ESearchDir::FromEnd));
+		const int32 DotIndex = FMath::Max<int32>(ObjectNameString.Find(TEXT("."), ESearchCase::CaseSensitive, ESearchDir::FromEnd),
+			ObjectNameString.Find(TEXT(":"), ESearchCase::CaseSensitive, ESearchDir::FromEnd));
 		if (DotIndex != INDEX_NONE)
 		{
 			ActualObjectName = FName(*ObjectNameString.Mid(DotIndex + 1));
 		}
-		const int32 Hash = GetObjectHash( ActualObjectName );
-		for(TMultiMap<int32,class UObjectBase*>::TConstKeyIterator HashIt(ObjectHash,Hash); HashIt; ++HashIt)
+		const int32 Hash = GetObjectHash(ActualObjectName);
+		FHashTableLock HashLock(ThreadHash);
+		for (TMultiMap<int32, class UObjectBase*>::TConstKeyIterator HashIt(ThreadHash.Hash, Hash); HashIt; ++HashIt)
 		{
 			UObject *Object = (UObject *)HashIt.Value();
 			if
-			(	(Object->GetFName()==ActualObjectName)
+				((Object->GetFName() == ActualObjectName)
 
-			/* Don't return objects that have any of the exclusive flags set */
-			&&	!Object->HasAnyFlags(ExcludeFlags)
+				/* Don't return objects that have any of the exclusive flags set */
+				&& !Object->HasAnyFlags(ExcludeFlags)
 
-			/*If there is no package (no InObjectPackage specified, and InName's package is "")
+				/*If there is no package (no InObjectPackage specified, and InName's package is "")
 				and the caller specified any_package, then accept it, regardless of its package.
 				Or, if the object is a top-level package then accept it immediately.*/
-			&&	(bAnyPackage ||	!Object->GetOuter())
-			
+				&& (bAnyPackage || !Object->GetOuter())
 
-			/** If a class was specified, check that the object is of the correct class */
-			&&	(ObjectClass==NULL || (bExactClass ? Object->GetClass()==ObjectClass : Object->IsA(ObjectClass))) 
 
-			/** Ensure that the partial path provided matches the object found */
-			&&  (Object->GetPathName().EndsWith(ObjectNameString)) )
+				/** If a class was specified, check that the object is of the correct class */
+				&& (ObjectClass == NULL || (bExactClass ? Object->GetClass() == ObjectClass : Object->IsA(ObjectClass)))
+
+				/** Ensure that the partial path provided matches the object found */
+				&& (Object->GetPathName().EndsWith(ObjectNameString)))
 			{
-				checkf( !Object->HasAnyFlags(RF_Unreachable), TEXT("%s"), *Object->GetFullName() );
+				checkf(!Object->HasAnyFlags(RF_Unreachable), TEXT("%s"), *Object->GetFullName());
 				if (Result)
 				{
 					UE_LOG(LogUObjectHash, Warning, TEXT("Ambiguous search, could be %s or %s"), *GetFullNameSafe(Result), *GetFullNameSafe(Object));
@@ -183,24 +270,38 @@ UObject* StaticFindObjectFastInternal( UClass* ObjectClass, UObject* ObjectPacka
 	return Result;
 }
 
-/** Map of object to their outers, used to avoid an object iterator to find such things. **/
-static TMap<UObjectBase*, TSet<UObjectBase*> > ObjectOuterMap;
-static TMap<UClass*, TSet<UObjectBase*> > ClassToObjectListMap;
-static TMap<UClass*, TSet<UClass*> > ClassToChildListMap;
+UObject* StaticFindObjectFastInternal( UClass* ObjectClass, UObject* ObjectPackage, FName ObjectName, bool bExactClass, bool bAnyPackage, EObjectFlags ExcludeFlags )
+{
+	INC_DWORD_STAT(STAT_FindObjectFast);
+	check(ObjectPackage != ANY_PACKAGE); // this could never have returned anything but NULL
+	// If they specified an outer use that during the hashing
+
+	UObject* Result = nullptr;
+	STAT(double ThisTime = 0);
+	{
+		SCOPE_SECONDS_COUNTER(ThisTime);
+		auto& ThreadHash = FUObjectHashTables::Get();
+		Result = StaticFindObjectFastInternalThreadSafe(ThreadHash, ObjectClass, ObjectPackage, ObjectName, bExactClass, bAnyPackage, ExcludeFlags);
+	}
+	INC_FLOAT_STAT_BY(STAT_Hash_StaticFindObjectFastInternal, (float)ThisTime);
+	return Result;
+}
 
-static void AddToOuterMap(UObjectBase* Object)
+// Assumes that ThreadHash's critical is already locked
+static void AddToOuterMap(FUObjectHashTables& ThreadHash, UObjectBase* Object)
 {
-	TSet<UObjectBase*>& Inners = ObjectOuterMap.FindOrAdd(Object->GetOuter());
+	TSet<UObjectBase*>& Inners = ThreadHash.ObjectOuterMap.FindOrAdd(Object->GetOuter());
 	bool bIsAlreadyInSetPtr = false;
 	Inners.Add(Object, &bIsAlreadyInSetPtr);
 	check(!bIsAlreadyInSetPtr); // if it already exists, something is wrong with the external code
 }
 
-static void AddToClassMap(UObjectBase* Object)
+// Assumes that ThreadHash's critical is already locked
+static void AddToClassMap(FUObjectHashTables& ThreadHash, UObjectBase* Object)
 {
 	{
 		check(Object->GetClass());
-		TSet<UObjectBase*>& ObjectList = ClassToObjectListMap.FindOrAdd(Object->GetClass());
+		TSet<UObjectBase*>& ObjectList = ThreadHash.ClassToObjectListMap.FindOrAdd(Object->GetClass());
 		bool bIsAlreadyInSetPtr = false;
 		ObjectList.Add(Object, &bIsAlreadyInSetPtr);
 		check(!bIsAlreadyInSetPtr); // if it already exists, something is wrong with the external code
@@ -213,7 +314,7 @@ static void AddToClassMap(UObjectBase* Object)
 		UClass* SuperClass = Class->GetSuperClass();
 		if ( SuperClass )
 		{
-			TSet<UClass*>& ChildList = ClassToChildListMap.FindOrAdd(SuperClass);
+			TSet<UClass*>& ChildList = ThreadHash.ClassToChildListMap.FindOrAdd(SuperClass);
 			bool bIsAlreadyInSetPtr = false;
 			ChildList.Add(Class, &bIsAlreadyInSetPtr);
 			check(!bIsAlreadyInSetPtr); // if it already exists, something is wrong with the external code
@@ -221,27 +322,29 @@ static void AddToClassMap(UObjectBase* Object)
 	}
 }
 
-static void RemoveFromOuterMap(UObjectBase* Object)
+// Assumes that ThreadHash's critical is already locked
+static void RemoveFromOuterMap(FUObjectHashTables& ThreadHash, UObjectBase* Object)
 {
-	TSet<UObjectBase*>& Inners = ObjectOuterMap.FindOrAdd(Object->GetOuter());
+	TSet<UObjectBase*>& Inners = ThreadHash.ObjectOuterMap.FindOrAdd(Object->GetOuter());
 	int32 NumRemoved = Inners.Remove(Object);
-    if (NumRemoved != 1)
+	if (NumRemoved != 1)
 	{
 		UE_LOG(LogUObjectHash, Error, TEXT("Internal Error: RemoveFromOuterMap NumRemoved = %d  for %s"), NumRemoved, *GetFullNameSafe((UObjectBaseUtility*)Object));
 	}
 	check(NumRemoved == 1); // must have existed, else something is wrong with the external code
 	if (!Inners.Num())
 	{
-		ObjectOuterMap.Remove(Object->GetOuter());
+		ThreadHash.ObjectOuterMap.Remove(Object->GetOuter());
 	}
 }
 
-static void RemoveFromClassMap(UObjectBase* Object)
+// Assumes that ThreadHash's critical is already locked
+static void RemoveFromClassMap(FUObjectHashTables& ThreadHash, UObjectBase* Object)
 {
 	UObjectBaseUtility* ObjectWithUtility = static_cast<UObjectBaseUtility*>(Object);
 
 	{
-		TSet<UObjectBase*>& ObjectList = ClassToObjectListMap.FindOrAdd(Object->GetClass());
+		TSet<UObjectBase*>& ObjectList = ThreadHash.ClassToObjectListMap.FindOrAdd(Object->GetClass());
 		int32 NumRemoved = ObjectList.Remove(Object);
 		if (NumRemoved != 1)
 		{
@@ -250,7 +353,7 @@ static void RemoveFromClassMap(UObjectBase* Object)
 		check(NumRemoved == 1); // must have existed, else something is wrong with the external code
 		if (!ObjectList.Num())
 		{
-			ClassToObjectListMap.Remove(Object->GetClass());
+			ThreadHash.ClassToObjectListMap.Remove(Object->GetClass());
 		}
 	}
 
@@ -261,7 +364,7 @@ static void RemoveFromClassMap(UObjectBase* Object)
 		if ( SuperClass )
 		{
 			// Remove the class from the SuperClass' child list
-			TSet<UClass*>& ChildList = ClassToChildListMap.FindOrAdd(SuperClass);
+			TSet<UClass*>& ChildList = ThreadHash.ClassToChildListMap.FindOrAdd(SuperClass);
 			int32 NumRemoved = ChildList.Remove(Class);
 			if (NumRemoved != 1)
 			{
@@ -270,7 +373,7 @@ static void RemoveFromClassMap(UObjectBase* Object)
 			check(NumRemoved == 1); // must have existed, else something is wrong with the external code
 			if (!ChildList.Num())
 			{
-				ClassToChildListMap.Remove(SuperClass);
+				ThreadHash.ClassToChildListMap.Remove(SuperClass);
 			}
 		}
 	}
@@ -279,91 +382,108 @@ static void RemoveFromClassMap(UObjectBase* Object)
 void GetObjectsWithOuter(const class UObjectBase* Outer, TArray<UObject *>& Results, bool bIncludeNestedObjects, EObjectFlags ExclusionFlags)
 {
 	// We don't want to return any objects that are currently being background loaded unless we're using the object iterator during async loading.
-	ExclusionFlags |= RF_Unreachable;
-	if( !GIsAsyncLoading )
+	STAT(double ThisTime = 0);
 	{
-		ExclusionFlags = EObjectFlags(ExclusionFlags | RF_AsyncLoading);
-	}
-	int32 StartNum = Results.Num();
-	TSet<UObjectBase*> const* Inners = ObjectOuterMap.Find(Outer);
-	if (Inners)
-	{
-		for(TSet<UObjectBase*>::TConstIterator It(*Inners); It; ++It)
+		SCOPE_SECONDS_COUNTER(ThisTime);
+
+		ExclusionFlags |= RF_Unreachable;
+		if (!IsAsyncLoading())
 		{
-			UObject *Object = static_cast<UObject *>(*It);
-			if (!Object->HasAnyFlags(ExclusionFlags))
-			{
-				Results.Add(Object);
-			}
+			ExclusionFlags = EObjectFlags(ExclusionFlags | RF_AsyncLoading);
 		}
-		int32 MaxResults = GUObjectArray.GetObjectArrayNum();
-		while (StartNum != Results.Num() && bIncludeNestedObjects) 
+		int32 StartNum = Results.Num();
+		auto& ThreadHash = FUObjectHashTables::Get();
+		FHashTableLock HashLock(ThreadHash);
+		TSet<UObjectBase*> const* Inners = ThreadHash.ObjectOuterMap.Find(Outer);
+		if (Inners)
 		{
-			int32 RangeStart = StartNum;
-			int32 RangeEnd = Results.Num();
-			StartNum = RangeEnd;
-			for (int32 Index = RangeStart; Index < RangeEnd; Index++)
+			for (TSet<UObjectBase*>::TConstIterator It(*Inners); It; ++It)
 			{
-				TSet<UObjectBase*> const* InnerInners = ObjectOuterMap.Find(Results[Index]);
-				if (InnerInners)
+				UObject *Object = static_cast<UObject *>(*It);
+				if (!Object->HasAnyFlags(ExclusionFlags))
 				{
-					for(TSet<UObjectBase*>::TConstIterator It(*InnerInners); It; ++It)
+					Results.Add(Object);
+				}
+			}
+			int32 MaxResults = GetUObjectArray().GetObjectArrayNum();
+			while (StartNum != Results.Num() && bIncludeNestedObjects)
+			{
+				int32 RangeStart = StartNum;
+				int32 RangeEnd = Results.Num();
+				StartNum = RangeEnd;
+				for (int32 Index = RangeStart; Index < RangeEnd; Index++)
+				{
+					TSet<UObjectBase*> const* InnerInners = ThreadHash.ObjectOuterMap.Find(Results[Index]);
+					if (InnerInners)
 					{
-						UObject *Object = static_cast<UObject *>(*It);
-						if (!Object->HasAnyFlags(ExclusionFlags))
+						for (TSet<UObjectBase*>::TConstIterator It(*InnerInners); It; ++It)
 						{
-							Results.Add(Object);
+							UObject *Object = static_cast<UObject *>(*It);
+							if (!Object->HasAnyFlags(ExclusionFlags))
+							{
+								Results.Add(Object);
+							}
 						}
 					}
 				}
+				check(Results.Num() <= MaxResults); // otherwise we have a cycle in the outer chain, which should not be possible
 			}
-			check(Results.Num() <= MaxResults); // otherwise we have a cycle in the outer chain, which should not be possible
-		} 
+		}
 	}
+	INC_FLOAT_STAT_BY(STAT_Hash_GetObjectsWithOuter, (float)ThisTime);
 }
 
 UObjectBase* FindObjectWithOuter(class UObjectBase* Outer, class UClass* ClassToLookFor, FName NameToLookFor)
 {
-	check(Outer && ClassToLookFor);
-	// We don't want to return any objects that are currently being background loaded unless we're using the object iterator during async loading.
-	EObjectFlags ExclusionFlags = RF_Unreachable;
-	if( !GIsAsyncLoading )
-	{
-		ExclusionFlags = EObjectFlags(ExclusionFlags | RF_AsyncLoading);
-	}
-
-	if (NameToLookFor != NAME_None)
+	UObject* Result = nullptr;
+	STAT(double ThisTime = 0);
 	{
-		return StaticFindObjectFastInternal( ClassToLookFor, static_cast<UObject*>(Outer), NameToLookFor, false, false, ExclusionFlags );
-	}
+		SCOPE_SECONDS_COUNTER(ThisTime);
 
-	UObject *Result = NULL;
-	TSet<UObjectBase*> const* Inners = ObjectOuterMap.Find(Outer);
-	if (Inners)
-	{
+		check(Outer && ClassToLookFor);
+		// We don't want to return any objects that are currently being background loaded unless we're using the object iterator during async loading.
+		EObjectFlags ExclusionFlags = RF_Unreachable;
+		if (!IsAsyncLoading())
+		{
+			ExclusionFlags = EObjectFlags(ExclusionFlags | RF_AsyncLoading);
+		}
 
-		for(TSet<UObjectBase*>::TConstIterator It(*Inners); It; ++It)
+		if (NameToLookFor != NAME_None)
 		{
-			UObject *Object = static_cast<UObject *>(*It);
-			if (Object->HasAnyFlags(ExclusionFlags))
-			{
-				continue;
-			}
-			if (!Object->IsA(ClassToLookFor))
+			Result = StaticFindObjectFastInternal(ClassToLookFor, static_cast<UObject*>(Outer), NameToLookFor, false, false, ExclusionFlags);
+		}
+		else
+		{
+			auto& ThreadHash = FUObjectHashTables::Get();
+			FHashTableLock HashLock(ThreadHash);
+			TSet<UObjectBase*> const* Inners = ThreadHash.ObjectOuterMap.Find(Outer);
+			if (Inners)
 			{
-				continue;
+				for (TSet<UObjectBase*>::TConstIterator It(*Inners); It; ++It)
+				{
+					UObject *Object = static_cast<UObject *>(*It);
+					if (Object->HasAnyFlags(ExclusionFlags))
+					{
+						continue;
+					}
+					if (!Object->IsA(ClassToLookFor))
+					{
+						continue;
+					}
+					Result = Object;
+					break;
+				}
 			}
-			Result = Object;
-			break;
 		}
 	}
+	INC_FLOAT_STAT_BY(STAT_Hash_FindObjectWithOuter, (float)ThisTime);
 	return Result;
 }
 
 /** Helper function that returns all the children of the specified class recursively */
-static void RecursivelyPopulateDerivedClasses(UClass* ParentClass, TSet<UClass*>& OutAllDerivedClass)
+static void RecursivelyPopulateDerivedClasses(FUObjectHashTables& ThreadHash, UClass* ParentClass, TSet<UClass*>& OutAllDerivedClass)
 {
-	TSet<UClass*>* ChildSet = ClassToChildListMap.Find(ParentClass);
+	TSet<UClass*>* ChildSet = ThreadHash.ClassToChildListMap.Find(ParentClass);
 	if ( ChildSet )
 	{
 		for ( auto ChildIt = ChildSet->CreateConstIterator(); ChildIt; ++ChildIt )
@@ -372,37 +492,21 @@ static void RecursivelyPopulateDerivedClasses(UClass* ParentClass, TSet<UClass*>
 			if ( !OutAllDerivedClass.Contains(ChildClass) )
 			{
 				OutAllDerivedClass.Add(ChildClass);
-				RecursivelyPopulateDerivedClasses(ChildClass, OutAllDerivedClass);
+				RecursivelyPopulateDerivedClasses(ThreadHash, ChildClass, OutAllDerivedClass);
 			}
 		}
 	}
 }
 
-void GetObjectsOfClass(UClass* ClassToLookFor, TArray<UObject *>& Results, bool bIncludeDerivedClasses, EObjectFlags AdditionalExcludeFlags)
+static void GetObjectsOfClassThreadSafe(FUObjectHashTables& ThreadHash, TSet<UClass*>& ClassesToSearch, TArray<UObject *>& Results, EObjectFlags ExclusionFlags)
 {
-	// We don't want to return any objects that are currently being background loaded unless we're using the object iterator during async loading.
-	EObjectFlags ExclusionFlags = RF_Unreachable;
-	if( !GIsAsyncLoading )
-	{
-		ExclusionFlags |= RF_AsyncLoading;
-	}
-	ExclusionFlags |= AdditionalExcludeFlags;
-
-	TSet<UClass*> ClassesToSearch;
-	ClassesToSearch.Add(ClassToLookFor);
-	if ( bIncludeDerivedClasses )
+	FHashTableLock HashLock(ThreadHash);
+	for (auto ClassIt = ClassesToSearch.CreateConstIterator(); ClassIt; ++ClassIt)
 	{
-		RecursivelyPopulateDerivedClasses(ClassToLookFor, ClassesToSearch);
-	}
-
-	const int32 MaxResults = GUObjectArray.GetObjectArrayNum();
-	for ( auto ClassIt = ClassesToSearch.CreateConstIterator(); ClassIt; ++ClassIt )
-	{
-		TSet<UObjectBase*> const* List = ClassToObjectListMap.Find(*ClassIt);
-
-		if ( List )
+		TSet<UObjectBase*> const* List = ThreadHash.ClassToObjectListMap.Find(*ClassIt);
+		if (List)
 		{
-			for( auto ObjectIt = List->CreateConstIterator(); ObjectIt; ++ObjectIt )
+			for (auto ObjectIt = List->CreateConstIterator(); ObjectIt; ++ObjectIt)
 			{
 				UObject *Object = static_cast<UObject *>(*ObjectIt);
 				if (!Object->HasAnyFlags(ExclusionFlags))
@@ -412,21 +516,57 @@ void GetObjectsOfClass(UClass* ClassToLookFor, TArray<UObject *>& Results, bool
 			}
 		}
 	}
+}
+void GetObjectsOfClass(UClass* ClassToLookFor, TArray<UObject *>& Results, bool bIncludeDerivedClasses, EObjectFlags AdditionalExcludeFlags)
+{
+	STAT(double ThisTime = 0);
+	{
+		SCOPE_SECONDS_COUNTER(ThisTime);
+
+		// We don't want to return any objects that are currently being background loaded unless we're using the object iterator during async loading.
+		EObjectFlags ExclusionFlags = RF_Unreachable;
+		if (!IsAsyncLoading())
+		{
+			ExclusionFlags |= RF_AsyncLoading;
+		}
+		ExclusionFlags |= AdditionalExcludeFlags;
+
+		TSet<UClass*> ClassesToSearch;
+		ClassesToSearch.Add(ClassToLookFor);
+		if (bIncludeDerivedClasses)
+		{
+			auto& ThreadHash = FUObjectHashTables::Get();
+			FHashTableLock HashLock(ThreadHash);
+			RecursivelyPopulateDerivedClasses(ThreadHash, ClassToLookFor, ClassesToSearch);
+		}
+
+		GetObjectsOfClassThreadSafe(FUObjectHashTables::Get(), ClassesToSearch, Results, ExclusionFlags);
 
-	check(Results.Num() <= MaxResults); // otherwise we have a cycle in the outer chain, which should not be possible
+		check(Results.Num() <= GetUObjectArray().GetObjectArrayNum()); // otherwise we have a cycle in the outer chain, which should not be possible
+	}
+#if STATS
+	if (FThreadStats::IsCollectingData())
+	{
+		INC_FLOAT_STAT_BY(STAT_Hash_GetObjectsOfClass, (float)ThisTime);
+	}
+#endif
 }
 
 void GetDerivedClasses(UClass* ClassToLookFor, TArray<UClass *>& Results, bool bRecursive)
 {
-	if ( bRecursive )
+	if (bRecursive)
 	{
 		TSet<UClass*> AllDerivedClasses;
-		RecursivelyPopulateDerivedClasses(ClassToLookFor, AllDerivedClasses);
+		auto& ThreadHash = FUObjectHashTables::Get();
+		FHashTableLock HashLock(ThreadHash);
+		RecursivelyPopulateDerivedClasses(ThreadHash, ClassToLookFor, AllDerivedClasses);
 		Results.Append( AllDerivedClasses.Array() );
 	}
 	else
 	{
-		TSet<UClass*>* DerivedClasses = ClassToChildListMap.Find(ClassToLookFor);
+		auto& ThreadHash = FUObjectHashTables::Get();
+		FHashTableLock HashLock(ThreadHash);
+		TSet<UClass*>* DerivedClasses = ThreadHash.ClassToChildListMap.Find(ClassToLookFor);
 		if ( DerivedClasses )
 		{
 			Results.Append( DerivedClasses->Array() );
@@ -434,29 +574,41 @@ void GetDerivedClasses(UClass* ClassToLookFor, TArray<UClass *>& Results, bool b
 	}
 }
 
-/**
- * Add an object to the name hash tables
- *
- * @param	Object		Object to add to the hash tables
- */
+void AllocateUObjectIndexForCurrentThread(UObjectBase* Object)
+{
+	GetUObjectArray().AllocateUObjectIndex(Object);
+}
+
 void HashObject(UObjectBase* Object)
 {
-	FName Name = Object->GetFName();
-	if (Name == NAME_None)
+	STAT(double ThisTime = 0);
 	{
-		return;
-	}
+		SCOPE_SECONDS_COUNTER(ThisTime);
+		FName Name = Object->GetFName();
+		if (Name != NAME_None)
+		{
+			int32 Hash = GetObjectHash(Name);
+			auto& ThreadHash = FUObjectHashTables::Get();
 
-	int32 Hash = GetObjectHash( Name );
-	checkSlow(!ObjectHash.FindPair(Hash,Object));  // if it already exists, something is wrong with the external code
-	ObjectHash.Add(Hash,Object);
+			FHashTableLock HashLock(ThreadHash);
+			checkSlow(!ThreadHash.Hash.FindPair(Hash, Object));  // if it already exists, something is wrong with the external code
 
-	Hash = GetObjectOuterHash(Name,(PTRINT)Object->GetOuter());
-	checkSlow(!ObjectHashOuter.FindPair(Hash,Object));  // if it already exists, something is wrong with the external code
-	ObjectHashOuter.Add(Hash,Object);
+			ThreadHash.Hash.Add(Hash, Object);
 
-	AddToOuterMap(Object);
-	AddToClassMap(Object);
+			Hash = GetObjectOuterHash(Name, (PTRINT)Object->GetOuter());
+			checkSlow(!ThreadHash.HashOuter.FindPair(Hash, Object));  // if it already exists, something is wrong with the external code
+			ThreadHash.HashOuter.Add(Hash, Object);
+
+			AddToOuterMap(ThreadHash, Object);
+			AddToClassMap(ThreadHash, Object);
+		}
+	}
+#if STATS
+	if (FThreadStats::IsCollectingData())
+	{
+		INC_FLOAT_STAT_BY(STAT_Hash_HashObject, (float)ThisTime);
+	}
+#endif
 }
 
 /**
@@ -466,20 +618,47 @@ void HashObject(UObjectBase* Object)
  */
 void UnhashObject(UObjectBase* Object)
 {
-	FName Name = Object->GetFName();
-	if (Name == NAME_None)
+	STAT(double ThisTime = 0);
 	{
-		return;
-	}
+		SCOPE_SECONDS_COUNTER(ThisTime);
+		FName Name = Object->GetFName();
+		if (Name != NAME_None)
+		{
 
-	int32 Hash = GetObjectHash( Name );
-	int32 NumRemoved = ObjectHash.RemoveSingle(Hash,Object);
-	check(NumRemoved == 1); // must have existed, else something is wrong with the external code
+			int32 Hash = GetObjectHash(Name);
+			auto& ThreadHash = FUObjectHashTables::Get();
+			FHashTableLock LockHash(ThreadHash);
+			int32 NumRemoved = ThreadHash.Hash.RemoveSingle(Hash, Object);
+			check(NumRemoved == 1); // must have existed, else something is wrong with the external code
 
-	Hash = GetObjectOuterHash(Name,(PTRINT)Object->GetOuter());
-	NumRemoved = ObjectHashOuter.RemoveSingle(Hash,Object);
-	check(NumRemoved == 1); // must have existed, else something is wrong with the external code
+			Hash = GetObjectOuterHash(Name, (PTRINT)Object->GetOuter());
+			NumRemoved = ThreadHash.HashOuter.RemoveSingle(Hash, Object);
+			check(NumRemoved == 1); // must have existed, else something is wrong with the external code
 
-	RemoveFromOuterMap(Object);
-	RemoveFromClassMap(Object);
+			RemoveFromOuterMap(ThreadHash, Object);
+			RemoveFromClassMap(ThreadHash, Object);
+		}
+	}
+#if STATS
+	if (FThreadStats::IsCollectingData())
+	{
+		INC_FLOAT_STAT_BY(STAT_Hash_UnhashObject, (float)ThisTime);
+	}
+#endif
+}
+
+/**
+ * Prevents any other threads from finding/adding UObjects while GC is running
+*/
+void LockUObjectHashTablesForGC()
+{
+	FUObjectHashTables::Get().Lock();
 }
+
+/**
+* Releases UObject hash tables lock after GC has finished running
+*/
+void UnlockUObjectHashTablesForGC()
+{
+	FUObjectHashTables::Get().Unlock();
+}
\ No newline at end of file
diff --git a/Engine/Source/Runtime/CoreUObject/Private/UObject/UObjectLinker.cpp b/Engine/Source/Runtime/CoreUObject/Private/UObject/UObjectLinker.cpp
--- a/Engine/Source/Runtime/CoreUObject/Private/UObject/UObjectLinker.cpp
+++ b/Engine/Source/Runtime/CoreUObject/Private/UObject/UObjectLinker.cpp
@@ -31,7 +31,7 @@ struct FLinkerIndexPair
 	FORCEINLINE bool IsDefault()
 	{
 		CheckInvariants();
-		return Linker == NULL;
+		return Linker == nullptr;
 	}
 
 	/**
@@ -39,7 +39,7 @@ struct FLinkerIndexPair
 	 * @param InLinker linker to assign
 	 * @param InLinkerIndex linker index to assign
 	 */
-	FLinkerIndexPair(ULinkerLoad* InLinker, int32 InLinkerIndex) :
+	FLinkerIndexPair(FLinkerLoad* InLinker, int32 InLinkerIndex) :
 		Linker(InLinker),
 		LinkerIndex(InLinkerIndex)
 	{
@@ -51,7 +51,7 @@ struct FLinkerIndexPair
 	 */
 	FORCEINLINE void CheckInvariants()
 	{
-		check(!((Linker == 0) ^ (LinkerIndex == INDEX_NONE))); // you need either a valid linker and index or neither valid
+		check(!((Linker == nullptr) ^ (LinkerIndex == INDEX_NONE))); // you need either a valid linker and index or neither valid
 	}
 
 	/**
@@ -60,7 +60,7 @@ struct FLinkerIndexPair
 	 * in an Unreal package), or if this object has been detached from its
 	 * linker, for e.g. renaming operations, saving the package, etc.
 	 */
-	ULinkerLoad*				Linker; 
+	FLinkerLoad*			Linker; 
 
 	/**
 	 * Index into the linker's ExportMap array for the FObjectExport resource
@@ -69,7 +69,7 @@ struct FLinkerIndexPair
 	int32							LinkerIndex; 
 };
 
-template <> struct TIsPODType<FLinkerIndexPair> { enum { Value = true }; };
+template <> struct TIsPODType<FLinkerIndexPair> { enum { Value = false }; };
 
 
 /**
@@ -86,7 +86,7 @@ static FUObjectAnnotationDense<FLinkerIndexPair,false> LinkerAnnotation;
 
 
 
-void UObject::SetLinker( ULinkerLoad* LinkerLoad, int32 LinkerIndex, bool bShouldDetachExisting )
+void UObject::SetLinker( FLinkerLoad* LinkerLoad, int32 LinkerIndex, bool bShouldDetachExisting )
 {
 	FLinkerIndexPair Existing = LinkerAnnotation.GetAnnotation(this);
 	Existing.CheckInvariants();
@@ -110,13 +110,11 @@ void UObject::SetLinker( ULinkerLoad* LinkerLoad, int32 LinkerIndex, bool bShoul
 		PostLinkerChange();
 #else
 		UE_CLOG(Existing.Linker && LinkerLoad, LogUObjectLinker, Fatal,
-			TEXT("It is only legal to change linkers in the editor. Trying to change linker on %s from %s (%s->LinkerRoot=%s) to %s (%s->LinkerRoot=%s)"),
+			TEXT("It is only legal to change linkers in the editor. Trying to change linker on %s from %s (Existing->LinkerRoot=%s) to %s (LinkerLoad->LinkerRoot=%s)"),
 			*GetFullName(),
 			*Existing.Linker->Filename,
-			*Existing.Linker->GetName(),
 			*GetNameSafe(Existing.Linker->LinkerRoot),
 			*LinkerLoad->Filename,
-			*LinkerLoad->GetName(),
 			*GetNameSafe(LinkerLoad->LinkerRoot));
 #endif
 	}
@@ -127,7 +125,7 @@ void UObject::SetLinker( ULinkerLoad* LinkerLoad, int32 LinkerIndex, bool bShoul
  *
  * @return	a pointer to the linker for this object, or NULL if this object has no linker
  */
-ULinkerLoad* UObjectBaseUtility::GetLinker() const
+FLinkerLoad* UObjectBaseUtility::GetLinker() const
 {
 	FLinkerIndexPair Existing = LinkerAnnotation.GetAnnotation(this);
 	Existing.CheckInvariants();
diff --git a/Engine/Source/Runtime/CoreUObject/Private/UObject/UObjectMarks.cpp b/Engine/Source/Runtime/CoreUObject/Private/UObject/UObjectMarks.cpp
--- a/Engine/Source/Runtime/CoreUObject/Private/UObject/UObjectMarks.cpp
+++ b/Engine/Source/Runtime/CoreUObject/Private/UObject/UObjectMarks.cpp
@@ -120,7 +120,7 @@ void GetObjectsWithAllMarks(TArray<UObject *>& Results, EObjectMark Marks)
 {
 	// We don't want to return any objects that are currently being background loaded unless we're using the object iterator during async loading.
 	EObjectFlags ExclusionFlags = RF_Unreachable;
-	if( !GIsAsyncLoading )
+	if (!IsAsyncLoading())
 	{
 		ExclusionFlags = EObjectFlags(ExclusionFlags | RF_AsyncLoading);
 	}
@@ -143,7 +143,7 @@ void GetObjectsWithAnyMarks(TArray<UObject *>& Results, EObjectMark Marks)
 {
 	// We don't want to return any objects that are currently being background loaded unless we're using the object iterator during async loading.
 	EObjectFlags ExclusionFlags = RF_Unreachable;
-	if( !GIsAsyncLoading )
+	if (!IsAsyncLoading())
 	{
 		ExclusionFlags = EObjectFlags(ExclusionFlags | RF_AsyncLoading);
 	}
diff --git a/Engine/Source/Runtime/CoreUObject/Private/UObject/UObjectThreadContext.cpp b/Engine/Source/Runtime/CoreUObject/Private/UObject/UObjectThreadContext.cpp
new file mode 100644
--- /dev/null
+++ b/Engine/Source/Runtime/CoreUObject/Private/UObject/UObjectThreadContext.cpp
@@ -0,0 +1,23 @@
+// Copyright 1998-2014 Epic Games, Inc. All Rights Reserved.
+
+/*=============================================================================
+	UObjectThreadContext.cpp: Unreal object globals
+=============================================================================*/
+
+#include "CoreUObjectPrivate.h"
+#include "UObject/UObjectThreadContext.h"
+
+FUObjectThreadContext::FUObjectThreadContext()
+: ImportCount(0)
+, ForcedExportCount(0)
+, ObjBeginLoadCount(0)
+, IsRoutingPostLoad(false)
+, IsInConstructor(0)
+, ConstructedObject(nullptr)
+, SerializedObject(nullptr)
+, SerializedPackageLinker(nullptr)
+, SerializedImportIndex(INDEX_NONE)
+, SerializedImportLinker(nullptr)
+, SerializedExportIndex(INDEX_NONE)
+, SerializedExportLinker(nullptr)
+{}
\ No newline at end of file
diff --git a/Engine/Source/Runtime/CoreUObject/Private/UObject/WeakObjectPtr.cpp b/Engine/Source/Runtime/CoreUObject/Private/UObject/WeakObjectPtr.cpp
--- a/Engine/Source/Runtime/CoreUObject/Private/UObject/WeakObjectPtr.cpp
+++ b/Engine/Source/Runtime/CoreUObject/Private/UObject/WeakObjectPtr.cpp
@@ -149,18 +149,18 @@ public:
 	 * @param Index	index of object that is being deleted
 	 */
 	virtual void NotifyUObjectDeleted(const UObjectBase *Object, int32 Index)
-	{
-		checkSlow(IsInGameThread());
+	{		
 		int32 Block, SubIndex;
 		GetBlock( Index, Block, SubIndex );
 		volatile FSerialNumberBlock** BlockPtr = (volatile FSerialNumberBlock**)&Blocks[Block];
 		FSerialNumberBlock* BlockToUse = (FSerialNumberBlock*)*BlockPtr;
 		if (BlockToUse)
-		{
+		{			
 			volatile int32 *SerialNumberPtr = &BlockToUse->SerialNumbers[SubIndex];
 			int32 SerialNumber = *SerialNumberPtr;
 			if (SerialNumber)
 			{
+				checkSlow(IsInGameThread());
 				*SerialNumberPtr = 0;
 				FPlatformMisc::MemoryBarrier();
 				checkSlow(!*SerialNumberPtr); // nobody should be creating these while we are zeroing them!
@@ -186,7 +186,7 @@ static FSerialNumberManager GSerialNumberManager;
 **/
 void FWeakObjectPtr::Init()
 {
-	GUObjectArray.AddUObjectDeleteListener(&GSerialNumberManager);
+	GetUObjectArray().AddUObjectDeleteListener(&GSerialNumberManager);
 }
 
 /**  
@@ -198,7 +198,7 @@ void FWeakObjectPtr::operator=(const class UObject *Object)
 	if (Object // && UObjectInitialized() we might need this at some point, but it is a speed hit we would prefer to avoid
 		)
 	{
-		ObjectIndex = GUObjectArray.ObjectToIndex((UObjectBase*)Object);
+		ObjectIndex = GetUObjectArray().ObjectToIndex((UObjectBase*)Object);
 		ObjectSerialNumber = GSerialNumberManager.GetAndAllocateSerialNumber(ObjectIndex);
 		checkSlow(SerialNumbersMatch());
 	}
@@ -240,7 +240,7 @@ bool FWeakObjectPtr::IsValid(bool bEvenIfPendingKill, bool bThreadsafeTest) cons
 	{
 		return true;
 	}
-	return GUObjectArray.IsValid(ObjectIndex, bEvenIfPendingKill);
+	return GetUObjectArray().IsValid(ObjectIndex, bEvenIfPendingKill);
 }
 
 bool FWeakObjectPtr::IsStale(bool bEvenIfPendingKill, bool bThreadsafeTest) const
@@ -262,7 +262,7 @@ bool FWeakObjectPtr::IsStale(bool bEvenIfPendingKill, bool bThreadsafeTest) cons
 	{
 		return false;
 	}
-	return GUObjectArray.IsStale(ObjectIndex, bEvenIfPendingKill);
+	return GetUObjectArray().IsStale(ObjectIndex, bEvenIfPendingKill);
 }
 
 UObject* FWeakObjectPtr::Get(bool bEvenIfPendingKill) const
@@ -270,7 +270,7 @@ UObject* FWeakObjectPtr::Get(bool bEvenIfPendingKill) const
 	UObject* Result = nullptr;
 	if (IsValid(true))
 	{
-		Result = (UObject*)(GUObjectArray.IndexToObject(GetObjectIndex(), bEvenIfPendingKill));
+		Result = (UObject*)(GetUObjectArray().IndexToObject(GetObjectIndex(), bEvenIfPendingKill));
 	}
 	return Result;
 }
@@ -280,7 +280,7 @@ UObject* FWeakObjectPtr::GetEvenIfUnreachable() const
 	UObject* Result = nullptr;
 	if (IsValid(true, true))
 	{
-		Result = static_cast<UObject*>(GUObjectArray.IndexToObject(GetObjectIndex(), true));
+		Result = static_cast<UObject*>(GetUObjectArray().IndexToObject(GetObjectIndex(), true));
 	}
 	return Result;
 }
diff --git a/Engine/Source/Runtime/CoreUObject/Public/Blueprint/BlueprintSupport.h b/Engine/Source/Runtime/CoreUObject/Public/Blueprint/BlueprintSupport.h
--- a/Engine/Source/Runtime/CoreUObject/Public/Blueprint/BlueprintSupport.h
+++ b/Engine/Source/Runtime/CoreUObject/Public/Blueprint/BlueprintSupport.h
@@ -31,7 +31,7 @@ struct FBlueprintSupport
 #if WITH_EDITOR
 /**
  * This is a helper struct that allows us to gather all previously unloaded class dependencies of a UClass
- * The first time we create a new UClass object in ULinkerLoad::CreateExport(), we register it as a dependency
+ * The first time we create a new UClass object in FLinkerLoad::CreateExport(), we register it as a dependency
  * master.  Any subsequent UClasses that are created for the first time during the preload of that class are
  * added to the list as potential cyclic referencers.  We then step over the list at the end of the load, and
  * recompile any classes that may depend on each other a second time to ensure that that functions and properties
diff --git a/Engine/Source/Runtime/CoreUObject/Public/Serialization/ArchiveUObject.h b/Engine/Source/Runtime/CoreUObject/Public/Serialization/ArchiveUObject.h
--- a/Engine/Source/Runtime/CoreUObject/Public/Serialization/ArchiveUObject.h
+++ b/Engine/Source/Runtime/CoreUObject/Public/Serialization/ArchiveUObject.h
@@ -153,7 +153,7 @@ protected:
  * of the UObject data being saved/loaded.
  * <p>
  * UObject references are not serialized directly into the memory archive.  Instead, we use
- * a system similar to the Export/ImportMap of ULinker - the pointer to the UObject is added
+ * a system similar to the Export/ImportMap of FLinker - the pointer to the UObject is added
  * to a persistent (from the standpoint of the FReloadObjectArc) array.  The location into
  * this array is what is actually stored in the archive's buffer.
  * <p>
diff --git a/Engine/Source/Runtime/CoreUObject/Public/Serialization/AsyncLoading.h b/Engine/Source/Runtime/CoreUObject/Public/Serialization/AsyncLoading.h
--- a/Engine/Source/Runtime/CoreUObject/Public/Serialization/AsyncLoading.h
+++ b/Engine/Source/Runtime/CoreUObject/Public/Serialization/AsyncLoading.h
@@ -6,17 +6,19 @@
 
 #pragma once
 
+#include "AsyncPackage.h"
+
 #define PERF_TRACK_DETAILED_ASYNC_STATS (0)
 /**
  * Structure containing intermediate data required for async loading of all imports and exports of a
- * ULinkerLoad.
+ * FLinkerLoad.
  */
-struct FAsyncPackage : public FGCObject
+struct FAsyncPackage
 {
 	/**
 	 * Constructor
 	 */
-	FAsyncPackage(const FName& InPackageName, const FGuid* InPackageGuid, FName InPackageType, const FName& InPackageNameToLoad);
+	FAsyncPackage(const FAsyncPackageDesc& InDesc);
 
 	/**
 	 * Ticks the async loading code.
@@ -32,7 +34,10 @@ struct FAsyncPackage : public FGCObject
 	/**
 	 * @return Estimated load completion percentage.
 	 */
-	float GetLoadPercentage() const;
+	FORCEINLINE float GetLoadPercentage() const
+	{
+		return LoadPercentage;
+	}
 
 	/**
 	 * @return Time load begun. This is NOT the time the load was requested in the case of other pending requests.
@@ -49,7 +54,7 @@ struct FAsyncPackage : public FGCObject
 	 */
 	FORCEINLINE const FName& GetPackageName() const
 	{
-		return PackageName;
+		return Desc.Name;
 	}
 
 	/**
@@ -57,48 +62,33 @@ struct FAsyncPackage : public FGCObject
 	 */
 	const FName &GetPackageType() const
 	{
-		return PackageType;
+		return Desc.Type;
 	}
 
-	void AddCompletionCallback(const FLoadPackageAsyncDelegate& Callback)
-	{
-		// This is to ensure that there is no one trying to subscribe to a already loaded package
-		check(!bLoadHasFinished && !bLoadHasFailed);
-		
-		CompletionCallbacks.Add(Callback);
-	}
+	void AddCompletionCallback(const FLoadPackageAsyncDelegate& Callback, bool bInternal);
 
-	/**
-	 * Serialize the linke for garbage collection.
-	 * 
-	 * @param Ar The archive to serialize with
-	 */
-	virtual void AddReferencedObjects( FReferenceCollector& Collector ) override
+	/** Gets the number of references to this package from other packages in the dependency tree. */
+	FORCEINLINE int32 GetDependencyRefCount() const
 	{
-		Collector.AddReferencedObject( Linker );
+		return DependencyRefCount;
 	}
 
-	/**
-	 *	Sets the flags and PIE ID for the package
-	 */
-	void SetPackageData(uint32 InPackageFlags, int32 InPIEInstanceID)
+	/** Returns true if the package has finished loading. */
+	FORCEINLINE bool HasFinishedLoading() const
 	{
-		PackageFlags = InPackageFlags;
-#if WITH_EDITOR
-		PIEInstanceID = InPIEInstanceID;
-#endif
+		return bLoadHasFinished;
 	}
 
-	/** Gets the number of references to this package from other packages in the dependency tree. */
-	FORCEINLINE int32 GetDependencyRefCount() const
+	/** Returns package loading priority. */
+	FORCEINLINE uint32 GetPriority() const
 	{
-		return DependencyRefCount;
+		return Desc.Priority;
 	}
 
-	/** Returns true if the package has finished loading. */
-	FORCEINLINE int32 HasFinishedLoading() const
+	/** Returns true if loading has failed */
+	FORCEINLINE bool HasLoadFailed() const
 	{
-		return bLoadHasFinished;
+		return bLoadHasFailed;
 	}
 
 	/**
@@ -107,18 +97,28 @@ struct FAsyncPackage : public FGCObject
 	void Cancel();
 
 private:
-	/** Name of the UPackage to create.																	*/
-	FName						PackageName;
-	/** Name of the package to load.																	*/
-	FName						PackageNameToLoad;
-	/** GUID of the package to load, or the zeroed invalid GUID for "don't care" */
-	FGuid						PackageGuid;
-	/** An abstract type name associated with this package, for tagging use								*/
-	FName						PackageType;
+
+	struct FCompletionCallback
+	{
+		bool bIsInternal;
+		FLoadPackageAsyncDelegate Callback;
+
+		FCompletionCallback()
+		{
+		}
+		FCompletionCallback(bool bInInternal, FLoadPackageAsyncDelegate InCallback)
+			: bIsInternal(bInInternal)
+			, Callback(InCallback)
+		{
+		}
+	};
+
+	/** Basic information associated with this package */
+	FAsyncPackageDesc Desc;
 	/** Linker which is going to have its exports and imports loaded									*/
-	ULinkerLoad*				Linker;
+	FLinkerLoad*				Linker;
 	/** Call backs called when we finished loading this package											*/
-	TArray<FLoadPackageAsyncDelegate>	CompletionCallbacks;
+	TArray<FCompletionCallback>	CompletionCallbacks;
 	/** Pending Import packages - we wait until all of them have been fully loaded. */
 	TArray<FAsyncPackage*> PendingImportedPackages;
 	/** Referenced imports - list of packages we need until we finish loading this package. */
@@ -135,6 +135,8 @@ private:
 	static int32					PreLoadIndex;
 	/** Current index into GObjLoaded array used to spread routing PostLoad over several frames			*/
 	static int32					PostLoadIndex;
+	/** Current index into DeferredPostLoadObjects array used to spread routing PostLoad over several frames			*/
+	int32						DeferredPostLoadIndex;
 	/** Currently used time limit for this tick.														*/
 	float						TimeLimit;
 	/** Whether we are using a time limit for this tick.												*/
@@ -157,13 +159,8 @@ private:
 	double						LoadStartTime;
 	/** Estimated load percentage.																		*/
 	float						LoadPercentage;
-	/** The flags that should be applied to the package */
-	uint32						PackageFlags;
-#if WITH_EDITOR
-	/** Editor only: PIE instance ID this package belongs to, INDEX_NONE otherwise */
-	int32						PIEInstanceID;
-#endif
-
+	/** Objects to be post loaded on the game thread */
+	TArray<UObject*> DeferredPostLoadObjects;
 
 public:
 #if PERF_TRACK_DETAILED_ASYNC_STATS
@@ -206,6 +203,15 @@ public:
 
 #endif
 
+	void CallCompletionCallbacks(bool bInternalOnly, EAsyncLoadingResult::Type LoadingResult);
+
+	/**
+	* Route PostLoad to deferred objects.
+	*
+	* @return true if we finished calling PostLoad on all loaded objects and no new ones were created, false otherwise
+	*/
+	EAsyncPackageState::Type PostLoadDeferredObjects(double InTickStartTime, bool bUseTimeLimit, float& InOutTimeLimit);
+
 private:
 	/**
 	 * Gives up time slice if time limit is enabled.
@@ -283,7 +289,8 @@ private:
 	 *
 	 * @return true
 	 */
-	EAsyncPackageState::Type FinishObjects();
+	EAsyncPackageState::Type FinishObjects();	
+
 	/**
 	 * Function called when pending import package has been loaded.
 	 */
@@ -311,6 +318,11 @@ private:
 	 */
 	void FreeReferencedImports();
 
+	/**
+	* Updates load percentage stat
+	*/
+	void UpdateLoadPercentage();
+
 #if PERF_TRACK_DETAILED_ASYNC_STATS
 	/** Add this time taken for object of class Class to have CreateExport called, to the stats we track. */
 	void TrackCreateExportTimeForClass(const UClass* Class, double Time);
diff --git a/Engine/Source/Runtime/CoreUObject/Public/Serialization/AsyncPackage.h b/Engine/Source/Runtime/CoreUObject/Public/Serialization/AsyncPackage.h
new file mode 100644
--- /dev/null
+++ b/Engine/Source/Runtime/CoreUObject/Public/Serialization/AsyncPackage.h
@@ -0,0 +1,44 @@
+// Copyright 1998-2014 Epic Games, Inc. All Rights Reserved.
+
+/*=============================================================================
+	AsyncPackage.h: Unreal async loading definitions.
+=============================================================================*/
+
+#pragma once
+
+struct FAsyncPackageDesc
+{
+	/** Name of the UPackage to create. */
+	FName Name;
+	/** An abstract type name associated with this package, for tagging use	*/
+	FName Type;
+	/** Name of the package to load. */
+	FName NameToLoad;
+	/** GUID of the package to load, or the zeroed invalid GUID for "don't care" */
+	FGuid Guid;
+	/** Delegate called on completion of loading */
+	FLoadPackageAsyncDelegate PackageLoadedDelegate;
+	/** The flags that should be applied to the package */
+	uint32 Flags;
+	/** Package loading priority. Higher number is higher priority. */
+	uint32 Priority;
+#if WITH_EDITOR
+	/** Editor only: PIE instance ID this package belongs to, INDEX_NONE otherwise */
+	int32 PIEInstanceID;
+#endif
+
+
+	FAsyncPackageDesc(const FName& InName, const FGuid& InGuid = FGuid(), FName InType = NAME_None, FName InPackageToLoadFrom = NAME_None, FLoadPackageAsyncDelegate InCompletionDelegate = FLoadPackageAsyncDelegate(), uint32 InFlags = 0, int32 InPIEInstanceID = INDEX_NONE, uint32 InPriority = 0)
+		: Name(InName)
+		, Type(InType)
+		, NameToLoad(InPackageToLoadFrom)
+		, Guid(InGuid)
+		, PackageLoadedDelegate(InCompletionDelegate)
+		, Flags(InFlags)
+		, Priority(InPriority)
+#if WITH_EDITOR
+		, PIEInstanceID(InPIEInstanceID)
+#endif
+	{
+	}
+};
diff --git a/Engine/Source/Runtime/CoreUObject/Public/Serialization/BulkData.h b/Engine/Source/Runtime/CoreUObject/Public/Serialization/BulkData.h
--- a/Engine/Source/Runtime/CoreUObject/Public/Serialization/BulkData.h
+++ b/Engine/Source/Runtime/CoreUObject/Public/Serialization/BulkData.h
@@ -62,7 +62,7 @@ enum EBulkDataLockFlags
  */
 struct COREUOBJECT_API FUntypedBulkData
 {
-	friend class ULinkerLoad;
+	friend class FLinkerLoad;
 
 	/*-----------------------------------------------------------------------------
 		Constructors and operators
@@ -428,10 +428,10 @@ protected:
 	/** Archive associated with bulk data for serialization																*/
 	FArchive*			AttachedAr;
 	/** Used to make sure the linker doesn't get garbage collected at runtime for things with attached archives			*/
-	ULinkerLoad*		Linker;
+	FLinkerLoad*		Linker;
 #else
 	/** weak pointer to the linker this bulk data originally belonged to. */
-	TWeakObjectPtr<ULinkerLoad> Linker;
+	TWeakObjectPtr<UPackage> Package;
 #endif // WITH_EDITOR
 };
 
