Merging UE4-Streaming to UE4

diff --git a/Engine/Source/Runtime/CoreUObject/Public/UObject/CoreObject.h b/Engine/Source/Runtime/CoreUObject/Public/UObject/CoreObject.h
--- a/Engine/Source/Runtime/CoreUObject/Public/UObject/CoreObject.h
+++ b/Engine/Source/Runtime/CoreUObject/Public/UObject/CoreObject.h
@@ -78,6 +78,18 @@ public:
 	/** The name of the file that this package was loaded from */
 	FName	FileName;
 
+	/** Linker load associated with this package */
+	class FLinkerLoad* LinkerLoad;
+
+	/** Linker package version this package has been serialized with. This is mostly used by PostLoad **/
+	int32 LinkerPackageVersion;
+
+	/** Linker licensee version this package has been serialized with. This is mostly used by PostLoad **/
+	int32 LinkerLicenseeVersion;
+
+	/** Linker custom version container this package has been serialized with. This is mostly used by PostLoad **/
+	FCustomVersionContainer LinkerCustomVersion;
+
 	/** size of the file for this package; if the package was not loaded from a file or was a forced export in another package, this will be zero */
 	int64 FileSize;
 	
@@ -101,6 +113,8 @@ public:
 	 */
 	virtual void PostInitProperties() override;
 
+	virtual void BeginDestroy() override;
+
 	/** Serializer */
 	virtual void Serialize( FArchive& Ar ) override;
 
@@ -342,7 +356,7 @@ public:
 	 * @return	true if the package was saved successfully.
 	 */
 	static bool SavePackage( UPackage* InOuter, UObject* Base, EObjectFlags TopLevelFlags, const TCHAR* Filename, 
-		FOutputDevice* Error=GError, ULinkerLoad* Conform=NULL, bool bForceByteSwapping=false, bool bWarnOfLongFilename=true, 
+		FOutputDevice* Error=GError, FLinkerLoad* Conform=NULL, bool bForceByteSwapping=false, bool bWarnOfLongFilename=true, 
 		uint32 SaveFlags=SAVE_None, const class ITargetPlatform* TargetPlatform = NULL, const FDateTime& FinalTimeStamp = FDateTime::MinValue(), bool bSlowTask = true );
 
 	/** Wait for any SAVE_Async file writes to complete **/
@@ -354,7 +368,7 @@ public:
 	 * @param	InOuter							the outer to use for the new package
 	 * @param	Linker							linker we're currently saving with
 	 */
-	static void SaveThumbnails( UPackage* InOuter, ULinkerSave* Linker );
+	static void SaveThumbnails( UPackage* InOuter, FLinkerSave* Linker );
 
 	/**
 	 * Static: Saves asset registry data for the specified package outer and linker
@@ -362,7 +376,7 @@ public:
 	 * @param	InOuter							the outer to use for the new package
 	 * @param	Linker							linker we're currently saving with
 	 */
-	static void SaveAssetRegistryData( UPackage* InOuter, ULinkerSave* Linker );
+	static void SaveAssetRegistryData( UPackage* InOuter, FLinkerSave* Linker );
 
 	/**
 	 * Static: Saves the level information used by the World browser
@@ -370,7 +384,7 @@ public:
 	 * @param	InOuter							the outer to use for the new package
 	 * @param	Linker							linker we're currently saving with
 	 */
-	static void SaveWorldLevelInfo( UPackage* InOuter, ULinkerSave* Linker );
+	static void SaveWorldLevelInfo( UPackage* InOuter, FLinkerSave* Linker );
 
 	/**
 	  * Determines if a package contains no more assets.
diff --git a/Engine/Source/Runtime/CoreUObject/Public/UObject/GarbageCollection.h b/Engine/Source/Runtime/CoreUObject/Public/UObject/GarbageCollection.h
--- a/Engine/Source/Runtime/CoreUObject/Public/UObject/GarbageCollection.h
+++ b/Engine/Source/Runtime/CoreUObject/Public/UObject/GarbageCollection.h
@@ -6,7 +6,7 @@
 =============================================================================*/
 
 /** Context sensitive keep flags for garbage collection */
-#define GARBAGE_COLLECTION_KEEPFLAGS	(GIsEditor ? RF_Native|RF_AsyncLoading|RF_Standalone : RF_Native|RF_AsyncLoading)
+#define GARBAGE_COLLECTION_KEEPFLAGS	(GIsEditor ? RF_Native|RF_AsyncLoading|RF_Standalone|RF_Async : RF_Native|RF_AsyncLoading|RF_Async)
 
 /*-----------------------------------------------------------------------------
 	Realtime garbage collection helper classes.
@@ -467,3 +467,11 @@ protected:
 	/** Object which is performing the serialization. */
 	const UObject* SerializingObject;
 };
+
+/** Prevent GC from running in the current scope */
+class COREUOBJECT_API FGCScopeGuard
+{
+public:
+	FGCScopeGuard();
+	~FGCScopeGuard();
+};
\ No newline at end of file
diff --git a/Engine/Source/Runtime/CoreUObject/Public/UObject/Linker.h b/Engine/Source/Runtime/CoreUObject/Public/UObject/Linker.h
--- a/Engine/Source/Runtime/CoreUObject/Public/UObject/Linker.h
+++ b/Engine/Source/Runtime/CoreUObject/Public/UObject/Linker.h
@@ -8,14 +8,6 @@
 
 DECLARE_LOG_CATEGORY_EXTERN(LogLinker, Log, All);
 
-/** Map of packages to their open linkers **/ 
-extern TMap<class UPackage*, class ULinkerLoad*>		GObjLoaders;
-/** Map of packages to their open linkers which have not yet been finalized (during async loading) **/
-extern TMap<class UPackage*, class ULinkerLoad*>		GObjPendingLoaders;
-/** List of loaders that have new imports **/
-extern TSet<class ULinkerLoad*>							GObjLoadersWithNewImports;
-
-
 /**
  * Wrapper for index into a ULnker's ImportMap or ExportMap.
  * Values greater than zero indicate that this is an index into the ExportMap.  The
@@ -121,7 +113,7 @@ public:
 
 /**
  * Base class for UObject resource types.  FObjectResources are used to store UObjects on disk
- * via ULinker's ImportMap (for resources contained in other packages) and ExportMap (for resources
+ * via FLinker's ImportMap (for resources contained in other packages) and ExportMap (for resources
  * contained within the same package)
  */
 struct FObjectResource
@@ -195,21 +187,21 @@ struct FObjectExport : public FObjectResource
 	int32         	SerialSize;
 
 	/**
-	 * The location (into the ULinker's underlying file reader archive) of the beginning of the
+	 * The location (into the FLinker's underlying file reader archive) of the beginning of the
 	 * data for this export's UObject.  Used for verification only.
 	 * Serialized
 	 */
 	int32         	SerialOffset;
 
 	/**
-	 * The location (into the ULinker's underlying file reader archive) of the beginning of the
+	 * The location (into the FLinker's underlying file reader archive) of the beginning of the
 	 * portion of this export's data that is serialized using script serialization.
 	 * Transient
 	 */
 	int32				ScriptSerializationStartOffset;
 
 	/**
-	 * The location (into the ULinker's underlying file reader archive) of the end of the
+	 * The location (into the FLinker's underlying file reader archive) of the end of the
 	 * portion of this export's data that is serialized using script serialization.
 	 * Transient
 	 */
@@ -222,7 +214,7 @@ struct FObjectExport : public FObjectResource
 	UObject*		Object;
 
 	/**
-	 * The index into the ULinker's ExportMap for the next export in the linker's export hash table.
+	 * The index into the FLinker's ExportMap for the next export in the linker's export hash table.
 	 * Transient
 	 */
 	int32				HashNext;
@@ -308,7 +300,7 @@ struct FObjectImport : public FObjectResource
 	 * The linker that contains the original FObjectExport resource associated with this import.
 	 * Transient
 	 */
-	ULinkerLoad*	SourceLinker;
+	FLinkerLoad*	SourceLinker;
 
 	/**
 	 * Index into SourceLinker's ExportMap for the export associated with this import's UObject.
@@ -548,7 +540,7 @@ private:
 
 public:
 	/**
-	 * Total size of all information that needs to be read in to create a ULinkerLoad. This includes
+	 * Total size of all information that needs to be read in to create a FLinkerLoad. This includes
 	 * the package file summary, name table and import & export maps.
 	 */
 	int32		TotalHeaderSize;
@@ -916,16 +908,33 @@ struct TLinkerNameMapKeyFuncs : TDefaultMapKeyFuncs<FName, ValueType, false>
 
 
 /*----------------------------------------------------------------------------
-	ULinker.
+	FLinker.
 ----------------------------------------------------------------------------*/
+namespace ELinkerType
+{
+	enum Type
+	{
+		None,
+		Load,
+		Save
+	};
+}
 
 /**
  * Manages the data associated with an Unreal package.  Acts as the bridge between
  * the file on disk and the UPackage object in memory for all Unreal package types.
  */
-class ULinker : public UObject, public FLinkerTables
+class FLinker : public FLinkerTables
 {
-	DECLARE_CLASS_INTRINSIC(ULinker,UObject,CLASS_Transient,CoreUObject)
+private:	
+
+	ELinkerType::Type LinkerType;
+
+protected:
+
+	bool bDestroyed;
+
+public:
 
 	/** The top-level UPackage object for the package associated with this linker */
 	UPackage*				LinkerRoot;
@@ -949,15 +958,22 @@ class ULinker : public UObject, public FLinkerTables
 	class FSHA1*			ScriptSHA;
 
 	/** Constructor. */
-	ULinker(const FObjectInitializer& ObjectInitializer, UPackage* InRoot, const TCHAR* InFilename );
+	FLinker(ELinkerType::Type InType, UPackage* InRoot, const TCHAR* InFilename);
+
+	virtual ~FLinker();
+
+	FORCEINLINE ELinkerType::Type GetType() const
+	{
+		return LinkerType;
+	}
+
 	/**
 	 * I/O function
-	 * 
+	*
 	 * @param	Ar	the archive to read/write into
 	 */
-	void Serialize( FArchive& Ar ) override;
-	static void AddReferencedObjects(UObject* InThis, FReferenceCollector& Collector);
-
+	void Serialize(FArchive& Ar);
+	void AddReferencedObjects(FReferenceCollector& Collector);
 	/**
 	 * Return the path name of the UObject represented by the specified import. 
 	 * (can be used with StaticFindObject)
@@ -1128,9 +1144,6 @@ class ULinker : public UObject, public FLinkerTables
 	 */
 	void GetScriptSHAKey(uint8* OutKey);
 
-	// UObject interface.
-	virtual void BeginDestroy() override;
-
 	/**
 	 * Test and object against the load flag filters
 	 *
@@ -1164,7 +1177,7 @@ class ULinker : public UObject, public FLinkerTables
 };
 
 /*----------------------------------------------------------------------------
-	ULinkerLoad.
+	FLinkerLoad.
 ----------------------------------------------------------------------------*/
 
 /**
@@ -1173,7 +1186,7 @@ class ULinker : public UObject, public FLinkerTables
 struct FDependencyRef
 {
 	/** The Linker the export lives in */
-	ULinkerLoad* Linker;
+	FLinkerLoad* Linker;
 
 	/** Index into Linker's ExportMap for this object */
 	int32 ExportIndex;
@@ -1198,10 +1211,8 @@ struct FDependencyRef
 /**
  * Handles loading Unreal package files, including reading UObject data from disk.
  */
-class ULinkerLoad : public ULinker, public FArchiveUObject
+class FLinkerLoad : public FLinker, public FArchiveUObject
 {
-	DECLARE_CASTED_CLASS_INTRINSIC_WITH_API(ULinkerLoad,ULinker,CLASS_Transient|0,CoreUObject,CASTCLASS_None,COREUOBJECT_API)
-
 	// Friends.
 	friend class UObject;
 	friend class UPackageMap;
@@ -1231,6 +1242,17 @@ class ULinkerLoad : public ULinker, public FArchiveUObject
 
 	// Variables.
 public:
+
+	/** Initialize everything related to linkers */
+	static void InitLinkers();
+
+	FORCEINLINE static ELinkerType::Type StaticType()
+	{
+		return ELinkerType::Load;
+	}
+
+	virtual ~FLinkerLoad();
+
 	/** Flags determining loading behavior.																					*/
 	uint32					LoadFlags;
 	/** Indicates whether the imports for this loader have been verified													*/
@@ -1329,6 +1351,9 @@ private:
 	/** Used for ActiveClassRedirects functionality */
 	bool					bFixupExportMapDone;
 
+	/** Id of the thread that created this linker. This is to guard against using this linker on other threads than the one it was created on **/
+	int32					OwnerThread;
+
 	/**
 	 * Helper struct to keep track of background file reads
 	 */
@@ -1376,14 +1401,11 @@ private:
 	/** Test whether we should report progress or not */
 	FORCEINLINE bool ShouldReportProgress() const
 	{
-		return !GIsAsyncLoading && ( LoadFlags & ( LOAD_Quiet | LOAD_SeekFree ) ) == 0;
+		return !IsAsyncLoading() && ( LoadFlags & ( LOAD_Quiet | LOAD_SeekFree ) ) == 0;
 	}
 #endif 
 
 public:
-	
-	/** Destructor */
-	~ULinkerLoad();
 
 	/**
 	 * Initialize the static variables
@@ -1391,7 +1413,7 @@ public:
 	COREUOBJECT_API static void StaticInit(UClass* InUTexture2DStaticClass);
 
 	/**
-	 * Add redirects to ULinkerLoad static map
+	 * Add redirects to FLinkerLoad static map
 	 */
 	COREUOBJECT_API static void CreateActiveRedirectsMap(const FString& GEngineIniName);
 
@@ -1423,26 +1445,32 @@ public:
 	 *
 	 * @return true if initialized, false if pending.
 	 */
-	bool HasFinishedInitialization() const
+	FORCEINLINE bool HasFinishedInitialization() const
 	{
         return bHasFinishedInitialization;
 	}
 
+	/** Returns ID of the thread that created this linker */
+	FORCEINLINE int32 GetOwnerThreadId() const
+	{
+		return OwnerThread;
+	}
+
 	/**
-	 * If this archive is a ULinkerLoad or ULinkerSave, returns a pointer to the ULinker portion.
+	 * If this archive is a FLinkerLoad or FLinkerSave, returns a pointer to the FLinker portion.
 	 */
-	virtual ULinker* GetLinker() override { return this; }
+	virtual FLinker* GetLinker() override { return this; }
 
 	/**
-	 * Creates and returns a ULinkerLoad object.
+	 * Creates and returns a FLinkerLoad object.
 	 *
 	 * @param	Parent		Parent object to load into, can be NULL (most likely case)
 	 * @param	Filename	Name of file on disk to load
 	 * @param	LoadFlags	Load flags determining behavior
 	 *
-	 * @return	new ULinkerLoad object for Parent/ Filename
+	 * @return	new FLinkerLoad object for Parent/ Filename
 	 */
-	COREUOBJECT_API static ULinkerLoad* CreateLinker( UPackage* Parent, const TCHAR* Filename, uint32 LoadFlags );
+	COREUOBJECT_API static FLinkerLoad* CreateLinker( UPackage* Parent, const TCHAR* Filename, uint32 LoadFlags );
 
 	void Verify();
 
@@ -1520,7 +1548,7 @@ public:
 	 * When this function exits, Object is guaranteed to contain the data stored that was stored on disk.
 	 *
 	 * @param	Object	The object to load data for.  If the data for this object isn't stored in this
-	 *					ULinkerLoad, routes the call to the appropriate linker.  Data serialization is 
+	 *					FLinkerLoad, routes the call to the appropriate linker.  Data serialization is 
 	 *					skipped if the object has already been loaded (as indicated by the RF_NeedLoad flag
 	 *					not set for the object), so safe to call on objects that have already been loaded.
 	 *					Note that this function assumes that Object has already been initialized against
@@ -1579,7 +1607,7 @@ public:
 	/**
 	 * Looks for an existing linker for the given package, without trying to make one if it doesn't exist
 	 */
-	COREUOBJECT_API static ULinkerLoad* FindExistingLinkerForPackage(UPackage* Package);
+	COREUOBJECT_API static FLinkerLoad* FindExistingLinkerForPackage(UPackage* Package);
 
 	/**
 	 * Replaces OldObject's entry in its linker with NewObject, so that all subsequent loads of OldObject will return NewObject.
@@ -1648,7 +1676,7 @@ private:
 	bool IsImportNative(const int32 ImportIndex) const;
 
 	/**
-	 * Attempts to lookup and return the corresponding ULinkerLoad object for 
+	 * Attempts to lookup and return the corresponding FLinkerLoad object for 
 	 * the specified import WITHOUT invoking  a load, or continuing to load 
 	 * the import package (will only return one if it has already been 
 	 * created... could still be in the process of loading).
@@ -1656,15 +1684,12 @@ private:
 	 * @param  ImportIndex    Specifies the import that you would like a linker for.
 	 * @return The imports associated linker (null if it hasn't been created yet).
 	 */
-	ULinkerLoad* FindExistingLinkerForImport(int32 ImportIndex) const;
+	FLinkerLoad* FindExistingLinkerForImport(int32 ImportIndex) const;
 
 	UObject* IndexToObject( FPackageIndex Index );
 
 	void DetachExport( int32 i );
 
-	// UObject interface.
-	virtual void BeginDestroy() override;
-
 	// FArchive interface.
 	/**
 	 * Hint the archive that the region starting at passed in offset and spanning the passed in size
@@ -1690,36 +1715,38 @@ private:
 	 * @param	Owner		UObject owning the bulk data
 	 * @param	BulkData	Bulk data object to associate
 	 */
-	virtual void AttachBulkData( UObject* Owner, FUntypedBulkData* BulkData ) override;
+	virtual void AttachBulkData(UObject* Owner, FUntypedBulkData* BulkData) override;
 	/**
 	 * Detaches the passed in bulk data object from the linker.
 	 *
 	 * @param	BulkData	Bulk data object to detach
 	 * @param	bEnsureBulkDataIsLoaded	Whether to ensure that the bulk data is loaded before detaching
 	 */
-	virtual void DetachBulkData( FUntypedBulkData* BulkData, bool bEnsureBulkDataIsLoaded ) override;
+	virtual void DetachBulkData(FUntypedBulkData* BulkData, bool bEnsureBulkDataIsLoaded) override;
 	/**
 	 * Detaches all attached bulk  data objects.
 	 *
 	 * @param	bEnsureBulkDataIsLoaded	Whether to ensure that the bulk data is loaded before detaching
 	 */
-	virtual void DetachAllBulkData( bool bEnsureBulkDataIsLoaded );
-#endif // WITH_EDITOR
+#endif
+	void DetachAllBulkData(bool bEnsureBulkDataIsLoaded);
 public:
 	/**
 	 * Detaches linker from bulk data/ exports and removes itself from array of loaders.
 	 *
 	 * @param	bEnsureAllBulkDataIsLoaded	Whether to load all bulk data first before detaching.
 	 */
-	virtual void Detach( bool bEnsureAllBulkDataIsLoaded );
+	void LoadAndDetachAllBulkData();
+
 private:
 
+	void Detach();
 	void Seek( int64 InPos ) override;
 	int64 Tell() override;
 	int64 TotalSize() override;
 	// this fixes the warning : 'ULinkerSave::Serialize' hides overloaded virtual function
-	using ULinker::Serialize;
-	void Serialize( void* V, int64 Length ) override;
+	using FLinker::Serialize;
+	virtual void Serialize( void* V, int64 Length ) override;
 	virtual FArchive& operator<<( UObject*& Object ) override;
 	virtual FArchive& operator<<( FLazyObjectPtr& LazyObjectPtr) override;
 	virtual FArchive& operator<<( FAssetPtr& AssetPtr) override;
@@ -1738,20 +1765,20 @@ private:
 	bool VerifyImportInner(const int32 ImportIndex, FString& WarningSuffix);
 
 	//
-	// ULinkerLoad creation helpers BEGIN
+	// FLinkerLoad creation helpers BEGIN
 	//
 
 	/**
-	 * Creates a ULinkerLoad object for async creation. Tick has to be called manually till it returns
+	 * Creates a FLinkerLoad object for async creation. Tick has to be called manually till it returns
 	 * true in which case the returned linker object has finished the async creation process.
 	 *
 	 * @param	Parent		Parent object to load into, can be NULL (most likely case)
 	 * @param	Filename	Name of file on disk to load
 	 * @param	LoadFlags	Load flags determining behavior
 	 *
-	 * @return	new ULinkerLoad object for Parent/ Filename
+	 * @return	new FLinkerLoad object for Parent/ Filename
 	 */
-	COREUOBJECT_API static ULinkerLoad* CreateLinkerAsync( UPackage* Parent, const TCHAR* Filename, uint32 LoadFlags );
+	COREUOBJECT_API static FLinkerLoad* CreateLinkerAsync( UPackage* Parent, const TCHAR* Filename, uint32 LoadFlags );
 
 	/**
 	 * Ticks an in-flight linker and spends InTimeLimit seconds on creation. This is a soft time limit used
@@ -1772,7 +1799,7 @@ private:
 	 * @param	Filename	Name of file on disk to load
 	 * @param	LoadFlags	Load flags determining behavior
 	 */
-	ULinkerLoad( const FObjectInitializer& ObjectInitializer, UPackage* InParent, const TCHAR* InFilename, uint32 InLoadFlags );
+	FLinkerLoad(UPackage* InParent, const TCHAR* InFilename, uint32 InLoadFlags);
 
 	/**
 	 * Returns whether the time limit allotted has been exceeded, if enabled.
@@ -2008,20 +2035,27 @@ public:
 	ELinkerStatus FinalizeCreation();
 
 	//
-	// ULinkerLoad creation helpers END
+	// FLinkerLoad creation helpers END
 	//
 };
 
 /*----------------------------------------------------------------------------
-	ULinkerSave.
+	FLinkerSave.
 ----------------------------------------------------------------------------*/
 
 /**
  * Handles saving Unreal package files.
  */
-class ULinkerSave : public ULinker, public FArchiveUObject
+class FLinkerSave : public FLinker, public FArchiveUObject
 {
-	DECLARE_CLASS_INTRINSIC(ULinkerSave,ULinker,CLASS_Transient|0,CoreUObject);
+public:
+
+	FORCEINLINE static ELinkerType::Type StaticType()
+	{
+		return ELinkerType::Save;
+	}
+
+	virtual ~FLinkerSave();
 
 	// Variables.
 	/** The archive that actually writes the data to disk. */
@@ -2049,10 +2083,9 @@ class ULinkerSave : public ULinker, public FArchiveUObject
 	COREUOBJECT_API static TMap<FString, TArray<uint8> > PackagesToScriptSHAMap;
 
 	/** Constructor */
-	ULinkerSave( const FObjectInitializer& ObjectInitializer, UPackage* InParent, const TCHAR* InFilename, bool bForceByteSwapping, bool bInSaveUnversioned = false );
+	FLinkerSave(UPackage* InParent, const TCHAR* InFilename, bool bForceByteSwapping, bool bInSaveUnversioned = false );
 	/** Constructor for memory writer */
-	ULinkerSave( const FObjectInitializer& ObjectInitializer, UPackage* InParent, bool bForceByteSwapping, bool bInSaveUnversioned = false );
-	void BeginDestroy();
+	FLinkerSave(UPackage* InParent, bool bForceByteSwapping, bool bInSaveUnversioned = false );
 
 	/** Returns the appropriate name index for the source name, or 0 if not found in NameIndices */
 	int32 MapName(const FName& Name) const;
@@ -2067,14 +2100,14 @@ class ULinkerSave : public ULinker, public FArchiveUObject
 	FArchive& operator<<( FAssetPtr& AssetPtr );
 
 	/**
-	 * If this archive is a ULinkerLoad or ULinkerSave, returns a pointer to the ULinker portion.
+	 * If this archive is a FLinkerLoad or FLinkerSave, returns a pointer to the FLinker portion.
 	 */
-	virtual ULinker* GetLinker() { return this; }
+	virtual FLinker* GetLinker() { return this; }
 
 	void Seek( int64 InPos );
 	int64 Tell();
-	// this fixes the warning : 'ULinkerSave::Serialize' hides overloaded virtual function
-	using ULinker::Serialize;
+	// this fixes the warning : 'FLinkerSave::Serialize' hides overloaded virtual function
+	using FLinker::Serialize;
 	void Serialize( void* V, int64 Length );
 
 	/**
@@ -2088,6 +2121,19 @@ class ULinkerSave : public ULinker, public FArchiveUObject
 	void ThisRequiresLocalizationGather();
 };
 
+template<typename T> 
+FORCEINLINE T* Cast(FLinker* Src)
+{
+	return Src && T::StaticType() == Src->GetType() ? (T*)Src : nullptr;
+}
+
+template<typename T>
+FORCEINLINE T* CastChecked(FLinker* Src)
+{
+	T* LinkerCastResult = Src && T::StaticType() == Src->GetType() ? (T*)Src : nullptr;
+	check(LinkerCastResult);
+	return LinkerCastResult;
+}
 
 /*-----------------------------------------------------------------------------
 	Lazy loading.
@@ -2129,7 +2175,7 @@ COREUOBJECT_API void ResetLoaders( UObject* InOuter );
  */
 void DissociateImportsAndForcedExports();
 
-COREUOBJECT_API ULinkerLoad* GetPackageLinker( UPackage* InOuter, const TCHAR* InLongPackageName, uint32 LoadFlags, UPackageMap* Sandbox, FGuid* CompatibleGuid );
+COREUOBJECT_API FLinkerLoad* GetPackageLinker( UPackage* InOuter, const TCHAR* InLongPackageName, uint32 LoadFlags, UPackageMap* Sandbox, FGuid* CompatibleGuid );
 
 /**
  * 
diff --git a/Engine/Source/Runtime/CoreUObject/Public/UObject/ObjectBase.h b/Engine/Source/Runtime/CoreUObject/Public/UObject/ObjectBase.h
--- a/Engine/Source/Runtime/CoreUObject/Public/UObject/ObjectBase.h
+++ b/Engine/Source/Runtime/CoreUObject/Public/UObject/ObjectBase.h
@@ -415,12 +415,13 @@ enum EObjectFlags
 	RF_FinishDestroyed			=0x00010000,	///< FinishDestroy has been called on the object.
 
 	// Misc. Flags
-	RF_BeingRegenerated			=0x00020000,	///< Flagged on UObjects that are used to create UClasses (e.g. Blueprints) while they are regenerating their UClass on load (See ULinkerLoad::CreateExport())
+	RF_BeingRegenerated			=0x00020000,	///< Flagged on UObjects that are used to create UClasses (e.g. Blueprints) while they are regenerating their UClass on load (See FLinkerLoad::CreateExport())
 	RF_DefaultSubObject			=0x00040000,	///< Flagged on subobjects that are defaults
 	RF_WasLoaded				=0x00080000,	///< Flagged on UObjects that were loaded
 	RF_TextExportTransient		=0x00100000,	///< Do not export object to text form (e.g. copy/paste). Generally used for sub-objects that can be regenerated from data in their parent object.
 	RF_LoadCompleted			=0x00200000,	///< Object has been completely serialized by linkerload at least once. DO NOT USE THIS FLAG, It should be replaced with RF_WasLoaded.
 	RF_InheritableComponentTemplate = 0x00400000, ///< Archetype of the object can be in its super class
+	RF_Async = 0x00800000, ///< Object exists only on a different thread than the game thread.
 
 	// Special all and none masks
 	RF_AllFlags					=0x00ffffff,	///< All flags, used mainly for error checking
diff --git a/Engine/Source/Runtime/CoreUObject/Public/UObject/ScriptSerialization.h b/Engine/Source/Runtime/CoreUObject/Public/UObject/ScriptSerialization.h
--- a/Engine/Source/Runtime/CoreUObject/Public/UObject/ScriptSerialization.h
+++ b/Engine/Source/Runtime/CoreUObject/Public/UObject/ScriptSerialization.h
@@ -331,7 +331,7 @@
 		case EX_SetArray:
 		{
 			// If not loading, or its a newer version
-			if((GetLinker() == NULL) || !Ar.IsLoading() || (Ar.UE4Ver() >= VER_UE4_CHANGE_SETARRAY_BYTECODE))
+			if((!GetLinker()) || !Ar.IsLoading() || (Ar.UE4Ver() >= VER_UE4_CHANGE_SETARRAY_BYTECODE))
 			{
 				// Array property to assign to
 				EExprToken TargetToken = SerializeExpr( iCode, Ar );
diff --git a/Engine/Source/Runtime/CoreUObject/Public/UObject/UObject.h b/Engine/Source/Runtime/CoreUObject/Public/UObject/UObject.h
--- a/Engine/Source/Runtime/CoreUObject/Public/UObject/UObject.h
+++ b/Engine/Source/Runtime/CoreUObject/Public/UObject/UObject.h
@@ -363,6 +363,17 @@ public:
 		return false;
 	}
 
+	/**
+	* Called during async load to determine if PostLoad can be called on the loading thread.
+	*
+	* @return	true if this object's PostLoad is thread safe
+	*/
+	virtual bool IsPostLoadThreadSafe() const
+	{
+		return false;
+	}
+
+
 	/** 
 	 *	Determines if you can create an object from the supplied template in the current context (editor, client only, dedicated server, game/listen) 
 	 *	This calls NeedsLoadForClient & NeedsLoadForServer
@@ -844,7 +855,7 @@ public:
 	 * @param LinkerIndex				New LinkerIndex to set
 	 * @param bShouldDetachExisting		If true, detach existing linker and call PostLinkerChange
 	 */
-	void SetLinker( ULinkerLoad* LinkerLoad, int32 LinkerIndex, bool bShouldDetachExisting=true );
+	void SetLinker( FLinkerLoad* LinkerLoad, int32 LinkerIndex, bool bShouldDetachExisting=true );
 
 	/**
 	 * Creates a new archetype based on this UObject.  The archetype's property values will match
diff --git a/Engine/Source/Runtime/CoreUObject/Public/UObject/UObjectAnnotation.h b/Engine/Source/Runtime/CoreUObject/Public/UObject/UObjectAnnotation.h
--- a/Engine/Source/Runtime/CoreUObject/Public/UObject/UObjectAnnotation.h
+++ b/Engine/Source/Runtime/CoreUObject/Public/UObject/UObjectAnnotation.h
@@ -88,7 +88,7 @@ public:
 				if (bAutoRemove)
 #endif
 				{
-					GUObjectArray.AddUObjectDeleteListener(this);
+					GetUObjectArray().AddUObjectDeleteListener(this);
 				}
 			}
 			AnnotationMap.Add(AnnotationCacheKey,AnnotationCacheValue);
@@ -115,7 +115,7 @@ public:
 			if (bAutoRemove)
 #endif
 			{
-				GUObjectArray.RemoveUObjectDeleteListener(this);
+				GetUObjectArray().RemoveUObjectDeleteListener(this);
 			}
 		}
 		return Result;
@@ -139,7 +139,7 @@ public:
 			if (bAutoRemove)
 #endif
 			{
-				GUObjectArray.RemoveUObjectDeleteListener(this);
+				GetUObjectArray().RemoveUObjectDeleteListener(this);
 			}
 		}
 	}
@@ -160,7 +160,7 @@ public:
 			if (bAutoRemove)
 #endif
 			{
-				GUObjectArray.RemoveUObjectDeleteListener(this);
+				GetUObjectArray().RemoveUObjectDeleteListener(this);
 			}
 		}
 	}
@@ -481,7 +481,7 @@ public:
 	void AddAnnotation(const UObjectBase *Object,TAnnotation Annotation)
 	{
 		check(Object);
-		AddAnnotation(GUObjectArray.ObjectToIndex(Object),Annotation);
+		AddAnnotation(GetUObjectArray().ObjectToIndex(Object),Annotation);
 	}
 	/**
 	 * Add an annotation to the annotation list. If the Annotation is the default, then the annotation is removed.
@@ -505,7 +505,7 @@ public:
 				if (bAutoRemove)
 #endif
 				{
-					GUObjectArray.AddUObjectDeleteListener(this);
+					GetUObjectArray().AddUObjectDeleteListener(this);
 				}
 			}
 			if (Index >= AnnotationArray.Num())
@@ -528,7 +528,7 @@ public:
 	void RemoveAnnotation(const UObjectBase *Object)
 	{
 		check(Object);
-		RemoveAnnotation(GUObjectArray.ObjectToIndex(Object));
+		RemoveAnnotation(GetUObjectArray().ObjectToIndex(Object));
 	}
 	/**
 	 * Removes an annotation from the annotation list. 
@@ -558,7 +558,7 @@ public:
 			if (bAutoRemove)
 #endif
 			{
-				GUObjectArray.RemoveUObjectDeleteListener(this);
+				GetUObjectArray().RemoveUObjectDeleteListener(this);
 			}
 		}
 	}
@@ -571,7 +571,7 @@ public:
 	FORCEINLINE TAnnotation GetAnnotation(const UObjectBase *Object)
 	{
 		check(Object);
-		return GetAnnotation(GUObjectArray.ObjectToIndex(Object));
+		return GetAnnotation(GetUObjectArray().ObjectToIndex(Object));
 	}
 
 	/**
@@ -598,7 +598,7 @@ public:
 	FORCEINLINE TAnnotation& GetAnnotationRef(const UObjectBase *Object)
 	{
 		check(Object);
-		return GetAnnotationRef(GUObjectArray.ObjectToIndex(Object));
+		return GetAnnotationRef(GetUObjectArray().ObjectToIndex(Object));
 	}
 
 	/**
@@ -662,11 +662,11 @@ public:
 	FORCEINLINE void Set(const UObjectBase *Object)
 	{
 		checkSlow(Object);
-		int32 Index = GUObjectArray.ObjectToIndex(Object);
+		int32 Index = GetUObjectArray().ObjectToIndex(Object);
 		checkSlow(Index >= 0);
 		if (AnnotationArray.Num() == 0)
 		{
-			GUObjectArray.AddUObjectDeleteListener(this);
+			GetUObjectArray().AddUObjectDeleteListener(this);
 		}
 		if (Index >= AnnotationArray.Num() * BitsPerElement)
 		{
@@ -686,7 +686,7 @@ public:
 	FORCEINLINE void Clear(const UObjectBase *Object)
 	{
 		checkSlow(Object);
-		int32 Index = GUObjectArray.ObjectToIndex(Object);
+		int32 Index = GetUObjectArray().ObjectToIndex(Object);
 		RemoveAnnotation(Index);
 	}
 
@@ -707,7 +707,7 @@ public:
 	FORCEINLINE bool Get(const UObjectBase *Object)
 	{
 		checkSlow(Object);
-		int32 Index = GUObjectArray.ObjectToIndex(Object);
+		int32 Index = GetUObjectArray().ObjectToIndex(Object);
 		checkSlow(Index >= 0);
 		if (Index < AnnotationArray.Num() * BitsPerElement)
 		{
@@ -740,7 +740,7 @@ private:
 		AnnotationArray.Empty();
 		if (bHadElements)
 		{
-			GUObjectArray.RemoveUObjectDeleteListener(this);
+			GetUObjectArray().RemoveUObjectDeleteListener(this);
 		}
 	}
 
diff --git a/Engine/Source/Runtime/CoreUObject/Public/UObject/UObjectArray.h b/Engine/Source/Runtime/CoreUObject/Public/UObject/UObjectArray.h
--- a/Engine/Source/Runtime/CoreUObject/Public/UObject/UObjectArray.h
+++ b/Engine/Source/Runtime/CoreUObject/Public/UObject/UObjectArray.h
@@ -77,6 +77,13 @@ public:
 	 */
 	void AllocatePermanentObjectPool(int32 MaxObjectsNotConsideredByGC);
 
+	/**
+	* Reserves array memory to hold the specified number of objects
+	*
+	* @param MaxObjectsNotConsideredByGC number of objects in the permanent object pool
+	*/
+	void ReserveUObjectPool(int32 InNumObjects);
+
 	/**
 	 * Disables the disregard for GC optimization. Commandlets can't use it.
 	 *
@@ -107,7 +114,7 @@ public:
 	 *
 	 * @param	Object Object to allocate an index for
 	 */
-	void AllocateUObjectIndex(class UObjectBase* Object);
+	void AllocateUObjectIndex(class UObjectBase* Object, bool bMergingThreads = false);
 
 	/**
 	 * Returns a UObject index top to the global uobject array
@@ -138,7 +145,7 @@ public:
 		check(Index >= 0);
 		if (Index < ObjObjects.Num())
 		{
-			return ObjObjects[Index];
+			return (UObjectBase*)ObjObjects[Index];
 		}
 		return NULL;
 	}
@@ -257,6 +264,7 @@ public:
 		return ObjLastNonGCIndex + 1;
 	}
 
+#if WITH_EDITOR
 	/**
 	 * Returns the number of actual object indices that are claimed (the total size of the global object array minus
 	 * the number of available object array elements
@@ -265,9 +273,9 @@ public:
 	 */
 	FORCEINLINE int32 GetObjectArrayNumMinusAvailable()
 	{
-		return ObjObjects.Num() - ObjAvailable.Num();
+		return ObjObjects.Num() - ObjAvailableCount.GetValue();
 	}
-
+#endif
 
 	/**
 	 * Clears some internal arrays to get rid of false memory leaks
@@ -346,7 +354,7 @@ public:
 		 */
 		FORCEINLINE UObjectBase* GetObject() const 
 		{ 
-			return Array.ObjObjects[Index];
+			return (UObjectBase*)Array.ObjObjects[Index];
 		}
 		/**
 		 * Iterator advance with ordinary name for clarity in subclasses
@@ -365,7 +373,7 @@ public:
 			return false;
 		}
 	private:
-		/** the array that we are iterating on, probably always GUObjectArray */
+		/** the array that we are iterating on, probably always GetUObjectArray() */
 		const FUObjectArray& Array;
 		/** index of the current element in the object array */
 		int32 Index;
@@ -373,23 +381,31 @@ public:
 
 private:
 
+	typedef TStaticIndirectArrayThreadSafeRead<UObjectBase, 8 * 1024 * 1024 /* Max 8M UObjects */, 16384 /* allocated in 64K/128K chunks */ > TUObjectArray;
+
 	/**
 	 * return the object array for use by debug visualizers
 	 */
-	static TArray<UObjectBase*>* GetObjectArrayForDebugVisualizers();
+	static UObjectBase*** GetObjectArrayForDebugVisualizers();
 
 	// note these variables are left with the Obj prefix so they can be related to the historical GObj versions
 
 	/** First index into objects array taken into account for GC.							*/
-	int32							ObjFirstGCIndex;
+	int32 ObjFirstGCIndex;
 	/** Index pointing to last object created in range disregarded for GC.					*/
-	int32							ObjLastNonGCIndex;
+	int32 ObjLastNonGCIndex;
 	/** If true this is the intial load and we should load objects int the disregarded for GC range.	*/
-	int32							OpenForDisregardForGC;
+	int32 OpenForDisregardForGC;
 	/** Array of all live objects.											*/
-	TArray<UObjectBase*>		ObjObjects;
+	TUObjectArray ObjObjects;
+	/** Synchronization object for all live objects.											*/
+	FCriticalSection ObjObjectsCritical;
 	/** Available object indices.											*/
-	TArray<int32>					ObjAvailable;	
+	TLockFreePointerList<int32> ObjAvailableList;
+#if WITH_EDITOR
+	/** Available object index count.										*/
+	FThreadSafeCounter ObjAvailableCount;
+#endif
 	/**
 	 * Array of things to notify when a UObjectBase is created
 	 */
@@ -401,7 +417,7 @@ private:
 };
 
 /** Global UObject allocator							*/
-extern COREUOBJECT_API FUObjectArray GUObjectArray;
+COREUOBJECT_API FUObjectArray& GetUObjectArray();
 
 /**
 	* Static version of IndexToObject for use with TWeakObjectPtr.
@@ -410,7 +426,7 @@ struct FIndexToObject
 {
 	static FORCEINLINE class UObjectBase* IndexToObject(int32 Index, bool bEvenIfPendingKill)
 	{
-		return GUObjectArray.IndexToObject(Index, bEvenIfPendingKill);
+		return GetUObjectArray().IndexToObject(Index, bEvenIfPendingKill);
 	}
 };
 
diff --git a/Engine/Source/Runtime/CoreUObject/Public/UObject/UObjectBase.h b/Engine/Source/Runtime/CoreUObject/Public/UObject/UObjectBase.h
--- a/Engine/Source/Runtime/CoreUObject/Public/UObject/UObjectBase.h
+++ b/Engine/Source/Runtime/CoreUObject/Public/UObject/UObjectBase.h
@@ -148,7 +148,6 @@ protected:
 	/**
 	 * Set the object flags directly
 	 *
-	 * @return Flags for this object
 	 **/
 	FORCEINLINE void SetFlagsTo( EObjectFlags NewFlags )
 	{
@@ -167,6 +166,40 @@ public:
 		return ObjectFlags;
 	}
 
+	/**
+	 *	Atomically adds the specified flags.
+	 *	Do not use unless you know what you are doing.
+	 *	Designed to be used only by parallel GC and UObject loading thread.
+	 */
+	FORCENOINLINE void AtomicallySetFlags( EObjectFlags FlagsToAdd )
+	{
+		int32 OldFlags = 0;
+		int32 NewFlags = 0;
+		do 
+		{
+			OldFlags = ObjectFlags;
+			NewFlags = OldFlags | FlagsToAdd;
+		}
+		while( FPlatformAtomics::InterlockedCompareExchange( (int32*)&ObjectFlags, NewFlags, OldFlags) != OldFlags );
+	}
+
+	/**
+	 *	Atomically clears the specified flags.
+	 *	Do not use unless you know what you are doing.
+	 *	Designed to be used only by parallel GC and UObject loading thread.
+	 */
+	FORCENOINLINE void AtomicallyClearFlags( EObjectFlags FlagsToClear )
+	{
+		int32 OldFlags = 0;
+		int32 NewFlags = 0;
+		do 
+		{
+			OldFlags = ObjectFlags;
+			NewFlags = OldFlags & ~FlagsToClear;
+		}
+		while( FPlatformAtomics::InterlockedCompareExchange( (int32*)&ObjectFlags, NewFlags, OldFlags) != OldFlags );
+	}
+
 	/**
 	 * Atomically clear the unreachable flag
 	 *
@@ -189,6 +222,7 @@ public:
 				bIChangedIt = true;
 				break;
 			}
+			// Remove later.
 			checkSlow(OldValue == (StartValue & ~int32(RF_Unreachable)));
 		}
 		return bIChangedIt;
diff --git a/Engine/Source/Runtime/CoreUObject/Public/UObject/UObjectBaseUtility.h b/Engine/Source/Runtime/CoreUObject/Public/UObject/UObjectBaseUtility.h
--- a/Engine/Source/Runtime/CoreUObject/Public/UObject/UObjectBaseUtility.h
+++ b/Engine/Source/Runtime/CoreUObject/Public/UObject/UObjectBaseUtility.h
@@ -340,7 +340,7 @@ public:
 	 *
 	 * @return	a pointer to the linker for this object, or NULL if this object has no linker
 	 */
-	ULinkerLoad* GetLinker() const;
+	class FLinkerLoad* GetLinker() const;
 	/**
 	 * Returns this object's LinkerIndex.
 	 *
diff --git a/Engine/Source/Runtime/CoreUObject/Public/UObject/UObjectGlobals.h b/Engine/Source/Runtime/CoreUObject/Public/UObject/UObjectGlobals.h
--- a/Engine/Source/Runtime/CoreUObject/Public/UObject/UObjectGlobals.h
+++ b/Engine/Source/Runtime/CoreUObject/Public/UObject/UObjectGlobals.h
@@ -45,19 +45,6 @@ typedef void (UObject::*Native)( FFrame& TheStack, RESULT_DECL );
 
 /** set while in SavePackage() to detect certain operations that are illegal while saving */
 extern COREUOBJECT_API bool					GIsSavingPackage;
-/** Imports for EndLoad optimization.									*/
-extern int32						GImportCount;
-/** Forced exports for EndLoad optimization.							*/
-extern int32						GForcedExportCount;
-/** Objects that might need preloading.									*/
-extern TArray<UObject*>			GObjLoaded;
-
-#if !(UE_BUILD_SHIPPING || UE_BUILD_TEST)
-/** Used to verify that the Super::Serialize chain is intact.			*/
-extern TArray<UObject*,TInlineAllocator<16> >		DebugSerialize;
-#endif
-
-
 
 /*-----------------------------------------------------------------------------
 	FObjectDuplicationParameters.
@@ -291,29 +278,29 @@ namespace EAsyncLoadingResult
  * @param	Result		Result of async loading.
  */
 DECLARE_DELEGATE_ThreeParams(FLoadPackageAsyncDelegate, const FName& /*PackageName*/, UPackage* /*LoadedPackage*/, EAsyncLoadingResult::Type /*Result*/)
-struct FAsyncPackage;
+
 /**
  * Asynchronously load a package and all contained objects that match context flags. Non- blocking.
  *
- * @param	PackageName			Name of package to load
- * @param	CompletionDelegate	Delegate called on completion of loading
- * @param	PackageGuid			GUID of the package to load, or NULL for "don't care"
- * @param	PackageType			A type name associated with this package for later use
- * @param	PackageToLoadFrom	If non-null, this is another package name. We load from this package name, into a (probably new) package named PackageName
- * @return	Handle for this async loading request
+ * @param	InName					Name of package to load
+ * @param	InGuid					GUID of the package to load, or NULL for "don't care"
+ * @param	InType					A type name associated with this package for later use
+ * @param	InPackageToLoadFrom		If non-null, this is another package name. We load from this package name, into a (probably new) package named PackageName
+ * @param	InCompletionDelegate	Delegate to be invoked when the packages has finished streaming
+ * @param	InFlags					Package flags
+ * @param	InPIEInstanceID			PIE instance ID
+ * @param	InPackagePriority		Loading priority
  */
-COREUOBJECT_API FAsyncPackage& LoadPackageAsync( const FString& PackageName, FLoadPackageAsyncDelegate CompletionDelegate, const FGuid* RequiredGuid = NULL, FName PackageType = NAME_None, const TCHAR* PackageToLoadFrom = NULL );
+COREUOBJECT_API void LoadPackageAsync(const FString& InName, const FGuid* InGuid = nullptr, FName InType = NAME_None, const TCHAR* InPackageToLoadFrom = nullptr, FLoadPackageAsyncDelegate InCompletionDelegate = FLoadPackageAsyncDelegate(), uint32 InFlags = 0, int32 InPIEInstanceID = INDEX_NONE, uint32 InPackagePriority = 0);
 
 /**
- * Asynchronously load a package and all contained objects that match context flags. Non- blocking.
- *
- * @param	PackageName			Name of package to load
- * @param	PackageGuid			GUID of the package to load, or NULL for "don't care"
- * @param	PackageType			A type name associated with this package for later use
- * @param	PackageToLoadFrom	If non-null, this is another package name. We load from this package name, into a (probably new) package named PackageName
- * @return	Handle for this async loading request
- */
-COREUOBJECT_API FAsyncPackage& LoadPackageAsync( const FString& PackageName, const FGuid* RequiredGuid = NULL, FName PackageType = NAME_None, const TCHAR* PackageToLoadFrom = NULL );
+* Asynchronously load a package and all contained objects that match context flags. Non- blocking.
+*
+* @param	InName					Name of package to load
+* @param	InCompletionDelegate	Delegate to be invoked when the packages has finished streaming
+* @param	InPackagePriority		Loading priority
+*/
+COREUOBJECT_API void LoadPackageAsync(const FString& InName, FLoadPackageAsyncDelegate InCompletionDelegate, uint32 InPackagePriority = 0);
 
 /**
 * Cancels all async package loading requests.
@@ -322,20 +309,33 @@ COREUOBJECT_API void CancelAsyncLoading();
 
 /**
  * Returns the async load percentage for a package in flight with the passed in name or -1 if there isn't one.
+ * THIS IS SLOW. MAY BLOCK ASYNC LOADING.
  *
  * @param	PackageName			Name of package to query load percentage for
  * @return	Async load percentage if package is currently being loaded, -1 otherwise
  */
 COREUOBJECT_API float GetAsyncLoadPercentage( const FName& PackageName );
 
+/**
+* Whether we are inside garbage collection
+*/
+COREUOBJECT_API bool IsGarbageCollecting();
+
 /** 
- * Deletes all unreferenced objects, keeping objects that have any of the passed in KeepFlags set
+ * Deletes all unreferenced objects, keeping objects that have any of the passed in KeepFlags set. Will wait for other threads to unlock GC.
  *
  * @param	KeepFlags			objects with those flags will be kept regardless of being referenced or not
  * @param	bPerformFullPurge	if true, perform a full purge after the mark pass
  */
-COREUOBJECT_API void CollectGarbage( EObjectFlags KeepFlags, bool bPerformFullPurge = true );
-COREUOBJECT_API void SerializeRootSet( FArchive& Ar, EObjectFlags KeepFlags );
+COREUOBJECT_API void CollectGarbage(EObjectFlags KeepFlags, bool bPerformFullPurge = true);
+/**
+* Performs garbage collection only if no other thread holds a lock on GC
+*
+* @param	KeepFlags			objects with those flags will be kept regardless of being referenced or not
+* @param	bPerformFullPurge	if true, perform a full purge after the mark pass
+*/
+COREUOBJECT_API bool TryCollectGarbage(EObjectFlags KeepFlags, bool bPerformFullPurge = true);
+COREUOBJECT_API void SerializeRootSet(FArchive& Ar, EObjectFlags KeepFlags);
 
 /**
  * Returns whether an incremental purge is still pending/ in progress.
@@ -397,12 +397,6 @@ COREUOBJECT_API bool IsReferenced( UObject*& Res, EObjectFlags KeepFlags, bool b
  * @param	ExcludeType		Do not flush packages associated with this specific type name
  */
 COREUOBJECT_API void FlushAsyncLoading( FName ExcludeType = NAME_None );
-/**
- * Returns whether we are currently async loading a package.
- *
- * @return true if we are async loading a package, false otherwise
- */
-COREUOBJECT_API bool IsAsyncLoading();
 
 /**
  * @return number of active async load package requests
@@ -1168,7 +1162,7 @@ inline T* FindObjectChecked( UObject* Outer, const TCHAR* Name, bool ExactClass=
 	return (T*)StaticFindObjectChecked( T::StaticClass(), Outer, Name, ExactClass );
 }
 
-// Find an object without asserting on GIsSavingPackage or GIsGarbageCollecting
+// Find an object without asserting on GIsSavingPackage or IsGarbageCollecting()
 template< class T > 
 inline T* FindObjectSafe( UObject* Outer, const TCHAR* Name, bool ExactClass=false )
 {
diff --git a/Engine/Source/Runtime/CoreUObject/Public/UObject/UObjectHash.h b/Engine/Source/Runtime/CoreUObject/Public/UObject/UObjectHash.h
--- a/Engine/Source/Runtime/CoreUObject/Public/UObject/UObjectHash.h
+++ b/Engine/Source/Runtime/CoreUObject/Public/UObject/UObjectHash.h
@@ -83,5 +83,12 @@ void HashObject(class UObjectBase* Object);
  */
 void UnhashObject(class UObjectBase* Object);
 
+/**
+* Adds a uobject to the global array which is used for uobject iteration
+*
+* @param	Object Object to allocate an index for
+*/
+void AllocateUObjectIndexForCurrentThread(class UObjectBase* Object);
+
 #endif	// __UOBJECTHASH_H__
 
diff --git a/Engine/Source/Runtime/CoreUObject/Public/UObject/UObjectIterator.h b/Engine/Source/Runtime/CoreUObject/Public/UObject/UObjectIterator.h
--- a/Engine/Source/Runtime/CoreUObject/Public/UObject/UObjectIterator.h
+++ b/Engine/Source/Runtime/CoreUObject/Public/UObject/UObjectIterator.h
@@ -19,7 +19,7 @@ public:
 	 * @param	bOnlyGCedObjects	if true, skip all of the permanent objects
 	 */
 	FRawObjectIterator(bool bOnlyGCedObjects = false) :
-		FUObjectArray::TIterator( GUObjectArray, bOnlyGCedObjects )
+		FUObjectArray::TIterator( GetUObjectArray(), bOnlyGCedObjects )
 	{
 	}
 	/**
@@ -59,13 +59,13 @@ public:
 	 * @param	AdditionalExclusionFlags	RF_* flags that should not be included in results
 	 */
 	FObjectIterator( UClass* InClass=UObject::StaticClass(), bool bOnlyGCedObjects = false, EObjectFlags AdditionalExclusionFlags = RF_NoFlags ) :
-		FUObjectArray::TIterator( GUObjectArray, bOnlyGCedObjects ),
+		FUObjectArray::TIterator( GetUObjectArray(), bOnlyGCedObjects ),
 		Class( InClass ),
 		ExclusionFlags(AdditionalExclusionFlags)
 	{
 		// We don't want to return any objects that are currently being background loaded unless we're using the object iterator during async loading.
 		ExclusionFlags = EObjectFlags(ExclusionFlags | RF_Unreachable);
-		if( !GIsAsyncLoading )
+		if (!IsAsyncLoading())
 		{
 			ExclusionFlags = EObjectFlags(ExclusionFlags | RF_AsyncLoading);
 		}
@@ -100,7 +100,7 @@ public:
 	{
 		//@warning: behavior is partially mirrored in UnObjGC.cpp. Make sure to adapt code there as well if you make changes below.
 		// verify that the async loading exclusion flag still matches (i.e. we didn't start/stop async loading within the scope of the iterator)
-		checkSlow(GIsAsyncLoading || (ExclusionFlags & RF_AsyncLoading));
+		checkSlow(IsAsyncLoading() || (ExclusionFlags & RF_AsyncLoading));
 
 		while(Advance())
 		{
@@ -277,7 +277,7 @@ public:
 	void operator++()
 	{
 		// verify that the async loading exclusion flag still matches (i.e. we didn't start/stop async loading within the scope of the iterator)
-		checkSlow(GIsAsyncLoading || (ExclusionFlags & RF_AsyncLoading));
+		checkSlow(IsAsyncLoading() || (ExclusionFlags & RF_AsyncLoading));
 		while(Advance())
 		{
 			if (!(*this)->HasAnyFlags(ExclusionFlags))
diff --git a/Engine/Source/Runtime/Engine/Classes/Components/PrimitiveComponent.h b/Engine/Source/Runtime/Engine/Classes/Components/PrimitiveComponent.h
--- a/Engine/Source/Runtime/Engine/Classes/Components/PrimitiveComponent.h
+++ b/Engine/Source/Runtime/Engine/Classes/Components/PrimitiveComponent.h
@@ -436,7 +436,7 @@ protected:
 	TEnumAsByte<EHasCustomNavigableGeometry::Type> bHasCustomNavigableGeometry;
 
 	/** Next id to be used by a component. */
-	static uint32 NextComponentId;
+	static FThreadSafeCounter NextComponentId;
 
 public:
 	/** 
diff --git a/Engine/Source/Runtime/Engine/Classes/Engine/Texture.h b/Engine/Source/Runtime/Engine/Classes/Engine/Texture.h
--- a/Engine/Source/Runtime/Engine/Classes/Engine/Texture.h
+++ b/Engine/Source/Runtime/Engine/Classes/Engine/Texture.h
@@ -686,6 +686,7 @@ public:
 #if WITH_EDITORONLY_DATA
 	ENGINE_API virtual void GetAssetRegistryTags(TArray<FAssetRegistryTag>& OutTags) const;
 #endif
+	ENGINE_API virtual bool IsPostLoadThreadSafe() const override{ return false; }
 	// End UObject interface.
 
 	/**
diff --git a/Engine/Source/Runtime/Engine/Private/AI/Navigation/NavigationData.cpp b/Engine/Source/Runtime/Engine/Private/AI/Navigation/NavigationData.cpp
--- a/Engine/Source/Runtime/Engine/Private/AI/Navigation/NavigationData.cpp
+++ b/Engine/Source/Runtime/Engine/Private/AI/Navigation/NavigationData.cpp
@@ -117,9 +117,8 @@ ANavigationData::ANavigationData(const FObjectInitializer& ObjectInitializer)
 
 uint16 ANavigationData::GetNextUniqueID()
 {
-	check(IsInGameThread());
-	static uint16 StaticID = INVALID_NAVDATA;
-	return ++StaticID;
+	static FThreadSafeCounter StaticID(INVALID_NAVDATA);
+	return StaticID.Increment();
 }
 
 void ANavigationData::PostInitProperties()
diff --git a/Engine/Source/Runtime/Engine/Private/Actor.cpp b/Engine/Source/Runtime/Engine/Private/Actor.cpp
--- a/Engine/Source/Runtime/Engine/Private/Actor.cpp
+++ b/Engine/Source/Runtime/Engine/Private/Actor.cpp
@@ -506,7 +506,7 @@ void AActor::ProcessEvent(UFunction* Function, void* Parameters)
 	#else
 	const bool bAllowScriptExecution = GAllowActorScriptExecutionInEditor;
 	#endif
-	if( ((GetWorld() && (GetWorld()->AreActorsInitialized() || bAllowScriptExecution)) || HasAnyFlags(RF_ClassDefaultObject)) && !GIsGarbageCollecting )
+	if( ((GetWorld() && (GetWorld()->AreActorsInitialized() || bAllowScriptExecution)) || HasAnyFlags(RF_ClassDefaultObject)) && !IsGarbageCollecting() )
 	{
 #if !UE_BUILD_SHIPPING
 		if (!ProcessEventDelegate.IsBound() || !ProcessEventDelegate.Execute(this, Function, Parameters))
diff --git a/Engine/Source/Runtime/Engine/Private/BlueprintGeneratedClass.cpp b/Engine/Source/Runtime/Engine/Private/BlueprintGeneratedClass.cpp
--- a/Engine/Source/Runtime/Engine/Private/BlueprintGeneratedClass.cpp
+++ b/Engine/Source/Runtime/Engine/Private/BlueprintGeneratedClass.cpp
@@ -96,7 +96,7 @@ void UBlueprintGeneratedClass::GetRequiredPreloadDependencies(TArray<UObject*>&
 	Super::GetRequiredPreloadDependencies(DependenciesOut);
 
 	// the component templates are no longer needed as Preload() dependencies 
-	// (ULinkerLoad now handles these with placeholder export objects instead)...
+	// (FLinkerLoad now handles these with placeholder export objects instead)...
 	// this change was prompted by a cyclic case, where creating the first
 	// component-template tripped the serialization of its class outer, before 
 	// another second component-template could be created (even though the 
diff --git a/Engine/Source/Runtime/Engine/Private/Components/PrimitiveComponent.cpp b/Engine/Source/Runtime/Engine/Private/Components/PrimitiveComponent.cpp
--- a/Engine/Source/Runtime/Engine/Private/Components/PrimitiveComponent.cpp
+++ b/Engine/Source/Runtime/Engine/Private/Components/PrimitiveComponent.cpp
@@ -66,7 +66,7 @@ TAutoConsoleVariable<int32> CVarShowInitialOverlaps(
 int32 UPrimitiveComponent::CurrentTag = 2147483647 / 4;
 
 // 0 is reserved to mean invalid
-uint32 UPrimitiveComponent::NextComponentId = 1;
+FThreadSafeCounter UPrimitiveComponent::NextComponentId;
 
 UPrimitiveComponent::UPrimitiveComponent(const FObjectInitializer& ObjectInitializer /*= FObjectInitializer::Get()*/)
 	: Super(ObjectInitializer)
@@ -99,9 +99,7 @@ UPrimitiveComponent::UPrimitiveComponent(const FObjectInitializer& ObjectInitial
 	VisibilityId = -1;
 	CanBeCharacterBase_DEPRECATED = ECB_Yes;
 	CanCharacterStepUpOn = ECB_Yes;
-	ComponentId.PrimIDValue = NextComponentId;
-	check(IsInGameThread());
-	NextComponentId++;
+	ComponentId.PrimIDValue = NextComponentId.Increment();
 
 	bUseEditorCompositing = false;
 
diff --git a/Engine/Source/Runtime/Engine/Private/GameViewportClient.cpp b/Engine/Source/Runtime/Engine/Private/GameViewportClient.cpp
--- a/Engine/Source/Runtime/Engine/Private/GameViewportClient.cpp
+++ b/Engine/Source/Runtime/Engine/Private/GameViewportClient.cpp
@@ -663,7 +663,7 @@ void UGameViewportClient::SetDropDetail(float DeltaSeconds)
 #if 0 
 		// so where we check to see if we are above some threshold and below 150 ms (any thing above that is usually blocking loading of some sort)
 		// also we don't want to do the auto trace when we are blocking on async loading
-		if( ( 0.070 < FrameTime ) && ( FrameTime < 0.150 ) && GIsAsyncLoading == false && GetWorld()->bRequestedBlockOnAsyncLoading == false && (GetWorld()->GetTimeSeconds() > 30.0f )  )
+		if ((0.070 < FrameTime) && (FrameTime < 0.150) && IsAsyncLoading() == false && GetWorld()->bRequestedBlockOnAsyncLoading == false && (GetWorld()->GetTimeSeconds() > 30.0f))
 		{
 			// now check to see if we have done a trace in the last 30 seconds otherwise we will just trace ourselves to death
 			static float LastTraceTime = -9999.0f;
@@ -2954,7 +2954,7 @@ bool UGameViewportClient::HandleDisplayAllCommand( const TCHAR* Cmd, FOutputDevi
 					// so then we only have to iterate over dynamic things each frame
 					for (TObjectIterator<UObject> It; It; ++It)
 					{
-						if (!GUObjectArray.IsDisregardForGC(*It))
+						if (!GetUObjectArray().IsDisregardForGC(*It))
 						{
 							break;
 						}
@@ -3001,7 +3001,7 @@ bool UGameViewportClient::HandleDisplayAllLocationCommand( const TCHAR* Cmd, FOu
 			// so then we only have to iterate over dynamic things each frame
 			for (TObjectIterator<UObject> It(true); It; ++It)
 			{
-				if (!GUObjectArray.IsDisregardForGC(*It))
+				if (!GetUObjectArray().IsDisregardForGC(*It))
 				{
 					break;
 				}
@@ -3039,7 +3039,7 @@ bool UGameViewportClient::HandleDisplayAllRotationCommand( const TCHAR* Cmd, FOu
 			// so then we only have to iterate over dynamic things each frame
 			for (TObjectIterator<UObject> It(true); It; ++It)
 			{
-				if (!GUObjectArray.IsDisregardForGC(*It))
+				if (!GetUObjectArray().IsDisregardForGC(*It))
 				{
 					break;
 				}
diff --git a/Engine/Source/Runtime/Engine/Private/LevelActor.cpp b/Engine/Source/Runtime/Engine/Private/LevelActor.cpp
--- a/Engine/Source/Runtime/Engine/Private/LevelActor.cpp
+++ b/Engine/Source/Runtime/Engine/Private/LevelActor.cpp
@@ -745,7 +745,7 @@ void UWorld::LoadSecondaryLevels(bool bForce, TSet<FString>* CookedPackages)
 	// Only load secondary levels in the Editor, and not for commandlets.
 	if( (!IsRunningCommandlet() || bForce)
 	// Don't do any work for world info actors that are part of secondary levels being streamed in! 
-	&&	!GIsAsyncLoading )
+	&& !IsAsyncLoading())
 	{
 		for( int32 LevelIndex=0; LevelIndex<StreamingLevels.Num(); LevelIndex++ )
 		{
diff --git a/Engine/Source/Runtime/Engine/Private/LevelStreaming.cpp b/Engine/Source/Runtime/Engine/Private/LevelStreaming.cpp
--- a/Engine/Source/Runtime/Engine/Private/LevelStreaming.cpp
+++ b/Engine/Source/Runtime/Engine/Private/LevelStreaming.cpp
@@ -469,9 +469,7 @@ bool ULevelStreaming::RequestLevel(UWorld* PersistentWorld, bool bAllowLevelLoad
 			{
 				// Load localized part of level first in case it exists. We don't need to worry about GC or completion 
 				// callback as we always kick off another async IO for the level below.
-				LoadPackageAsync(*(GetWorldAssetPackageName() + LOCALIZED_SEEKFREE_SUFFIX), 
-					NULL, NAME_None, *LocalizedPackageName
-					).SetPackageData(PackageFlags, PIEInstanceID);
+				LoadPackageAsync(*(GetWorldAssetPackageName() + LOCALIZED_SEEKFREE_SUFFIX), nullptr, NAME_None, *LocalizedPackageName, FLoadPackageAsyncDelegate(), PackageFlags, PIEInstanceID);
 			}
 		}
 
@@ -483,10 +481,7 @@ bool ULevelStreaming::RequestLevel(UWorld* PersistentWorld, bool bAllowLevelLoad
 			UWorld::WorldTypePreLoadMap.FindOrAdd(DesiredPackageName) = PersistentWorld->WorldType;
 
 			// Kick off async load request.
-			LoadPackageAsync(*DesiredPackageName.ToString(), 
-				FLoadPackageAsyncDelegate::CreateUObject(this, &ULevelStreaming::AsyncLevelLoadComplete), 
-				NULL, NAME_None, *PackageNameToLoadFrom
-				).SetPackageData(PackageFlags, PIEInstanceID);
+			LoadPackageAsync(DesiredPackageName.ToString(), nullptr, NAME_None, *PackageNameToLoadFrom, FLoadPackageAsyncDelegate::CreateUObject(this, &ULevelStreaming::AsyncLevelLoadComplete), PackageFlags, PIEInstanceID);
 
 			// streamingServer: server loads everything?
 			// Editor immediately blocks on load and we also block if background level streaming is disabled.
@@ -598,10 +593,11 @@ void ULevelStreaming::AsyncLevelLoadComplete(const FName& InPackageName, UPackag
 				{
 					// Loading into a new custom package explicitly. Load the destination world directly into the package.
 					// Detach the linker to load from a new file into the same package.
-					ULinkerLoad* PackageLinker = ULinkerLoad::FindExistingLinkerForPackage(LevelPackage);
+					FLinkerLoad* PackageLinker = FLinkerLoad::FindExistingLinkerForPackage(LevelPackage);
 					if (PackageLinker)
 					{
-						PackageLinker->Detach(false);
+						delete PackageLinker;
+						PackageLinker = nullptr;
 					}
 
 					// Make sure the redirector is not in the way of the new world.
diff --git a/Engine/Source/Runtime/Engine/Private/LevelTick.cpp b/Engine/Source/Runtime/Engine/Private/LevelTick.cpp
--- a/Engine/Source/Runtime/Engine/Private/LevelTick.cpp
+++ b/Engine/Source/Runtime/Engine/Private/LevelTick.cpp
@@ -1024,6 +1024,13 @@ public:
 extern bool GCollisionAnalyzerIsRecording;
 #endif // ENABLE_COLLISION_ANALYZER
 
+#if !(UE_BUILD_SHIPPING || UE_BUILD_TEST)
+static TAutoConsoleVariable<int32> CVarStressTestGCWhileStreaming(
+	TEXT("t.StressTestGC"),
+	0,
+	TEXT("If set to 1, the engine will attempt to trigger GC each frame while async loading."));
+#endif
+
 /**
  * Update the level after a variable amount of time, DeltaSeconds, has passed.
  * All child actors are ticked after their owners have been ticked.
@@ -1317,12 +1324,22 @@ void UWorld::Tick( ELevelTick TickType, float DeltaSeconds )
 	bInTick = false;
 	Mark.Pop();
 
-	if ( FullPurgeTriggered )
+	
+#if !(UE_BUILD_SHIPPING || UE_BUILD_TEST)
+	if (CVarStressTestGCWhileStreaming.GetValueOnGameThread() && IsAsyncLoading())
+	{
+		TryCollectGarbage(GARBAGE_COLLECTION_KEEPFLAGS, true);
+	}
+	else 
+#endif
+	if (FullPurgeTriggered)
 	{
-		CollectGarbage( GARBAGE_COLLECTION_KEEPFLAGS, true );
-		CleanupActors();
-		FullPurgeTriggered = false;
-		TimeSinceLastPendingKillPurge = 0.0f;
+		if (TryCollectGarbage(GARBAGE_COLLECTION_KEEPFLAGS, true))
+		{
+			CleanupActors();
+			FullPurgeTriggered = false;
+			TimeSinceLastPendingKillPurge = 0.0f;
+		}
 	}
 	else if( HasBegunPlay() )
 	{
@@ -1430,12 +1447,13 @@ void UWorld::PerformGarbageCollectionAndCleanupActors()
 	if( !IsAsyncLoading() )
 	{
 		// Perform housekeeping.
-		CollectGarbage( GARBAGE_COLLECTION_KEEPFLAGS, false );
-
-		CleanupActors();
+		if (TryCollectGarbage(GARBAGE_COLLECTION_KEEPFLAGS, false))
+		{
+			CleanupActors();
 
-		// Reset counter.
-		TimeSinceLastPendingKillPurge = 0;
+			// Reset counter.
+			TimeSinceLastPendingKillPurge = 0;
+		}
 	}
 }
 
diff --git a/Engine/Source/Runtime/Engine/Private/Materials/MaterialInstance.cpp b/Engine/Source/Runtime/Engine/Private/Materials/MaterialInstance.cpp
--- a/Engine/Source/Runtime/Engine/Private/Materials/MaterialInstance.cpp
+++ b/Engine/Source/Runtime/Engine/Private/Materials/MaterialInstance.cpp
@@ -288,7 +288,7 @@ void UMaterialInstance::PropagateDataToMaterialProxy()
 
 void FMaterialInstanceResource::GameThread_SetParent(UMaterialInterface* ParentMaterialInterface)
 {
-	check(IsInGameThread());
+	check(IsInGameThread() || IsAsyncLoading());
 
 	if( GameThreadParent != ParentMaterialInterface )
 	{
@@ -484,7 +484,7 @@ void UMaterialInstance::InitResources()
 
 const UMaterial* UMaterialInstance::GetMaterial() const
 {
-	check(IsInGameThread());
+	check(IsInGameThread() || IsAsyncLoading());
 	if(ReentrantFlag)
 	{
 		return UMaterial::GetDefaultMaterial(MD_Surface);
@@ -1309,7 +1309,7 @@ void UMaterialInstance::UpdatePermutationAllocations()
 
 void UMaterialInstance::CacheResourceShadersForRendering()
 {
-	check(IsInGameThread());
+	check(IsInGameThread() || IsAsyncLoading());
 
 	// Fix-up the parent lighting guid if it has changed...
 	if (Parent && (Parent->GetLightingGuid() != ParentLightingGuid))
diff --git a/Engine/Source/Runtime/Engine/Private/Materials/MaterialInstanceSupport.h b/Engine/Source/Runtime/Engine/Private/Materials/MaterialInstanceSupport.h
--- a/Engine/Source/Runtime/Engine/Private/Materials/MaterialInstanceSupport.h
+++ b/Engine/Source/Runtime/Engine/Private/Materials/MaterialInstanceSupport.h
@@ -27,7 +27,7 @@ public:
 	FMICReentranceGuard(const UMaterialInstance* InMaterial)
 	: Material(InMaterial)
 	{
-		check(IsInGameThread());
+		check(IsInGameThread() || IsAsyncLoading());
 		// NOTE:  we are switching the check to be a log so we can find out which material is causing the assert
 		//check(!Material->ReentrantFlag);
 		if( Material->ReentrantFlag == 1 )
@@ -40,7 +40,7 @@ public:
 
 	~FMICReentranceGuard()
 	{
-		check(IsInGameThread());
+		check(IsInGameThread() || IsAsyncLoading());
 		const_cast<UMaterialInstance*>(Material)->ReentrantFlag = 0;
 	}
 
diff --git a/Engine/Source/Runtime/Engine/Private/Materials/MaterialShader.cpp b/Engine/Source/Runtime/Engine/Private/Materials/MaterialShader.cpp
--- a/Engine/Source/Runtime/Engine/Private/Materials/MaterialShader.cpp
+++ b/Engine/Source/Runtime/Engine/Private/Materials/MaterialShader.cpp
@@ -1531,13 +1531,13 @@ FMaterialShaderMap::FMaterialShaderMap() :
 	bCompiledSuccessfully(true),
 	bIsPersistent(true)
 {
-	checkSlow(IsInGameThread());
+	checkSlow(IsInGameThread() || IsAsyncLoading());
 	AllMaterialShaderMaps.Add(this);
 }
 
 FMaterialShaderMap::~FMaterialShaderMap()
 { 
-	checkSlow(IsInGameThread());
+	checkSlow(IsInGameThread() || IsAsyncLoading());
 	check(bDeletedThroughDeferredCleanup);
 	check(!bRegistered);
 	AllMaterialShaderMaps.RemoveSwap(this);
diff --git a/Engine/Source/Runtime/Engine/Private/MatineeUtils.cpp b/Engine/Source/Runtime/Engine/Private/MatineeUtils.cpp
--- a/Engine/Source/Runtime/Engine/Private/MatineeUtils.cpp
+++ b/Engine/Source/Runtime/Engine/Private/MatineeUtils.cpp
@@ -352,7 +352,7 @@ namespace FMatineeUtils
 				{
 					// Now look for a component
 					const FName TrialCompName(*CompString);
-					const FName RemappedCompName = ULinkerLoad::FindSubobjectRedirectName(TrialCompName);
+					const FName RemappedCompName = FLinkerLoad::FindSubobjectRedirectName(TrialCompName);
 					const FName CompName = (RemappedCompName != NAME_None) ? RemappedCompName : TrialCompName;
 
 					TArray<UObject*> Components;
diff --git a/Engine/Source/Runtime/Engine/Private/NetConnection.cpp b/Engine/Source/Runtime/Engine/Private/NetConnection.cpp
--- a/Engine/Source/Runtime/Engine/Private/NetConnection.cpp
+++ b/Engine/Source/Runtime/Engine/Private/NetConnection.cpp
@@ -507,7 +507,7 @@ void UNetConnection::FlushNet(bool bIgnoreSimulation)
 #if DO_ENABLE_NET_TEST
 		// if the connection is closing/being destroyed/etc we need to send immediately regardless of settings
 		// because we won't be around to send it delayed
-		if (State == USOCK_Closed || GIsGarbageCollecting || bIgnoreSimulation)
+		if (State == USOCK_Closed || IsGarbageCollecting() || bIgnoreSimulation)
 		{
 			// Checked in FlushNet() so each child class doesn't have to implement this
 			if (Driver->IsNetResourceValid())
diff --git a/Engine/Source/Runtime/Engine/Private/ObjectLibrary.cpp b/Engine/Source/Runtime/Engine/Private/ObjectLibrary.cpp
--- a/Engine/Source/Runtime/Engine/Private/ObjectLibrary.cpp
+++ b/Engine/Source/Runtime/Engine/Private/ObjectLibrary.cpp
@@ -306,7 +306,7 @@ int32 UObjectLibrary::LoadAssetDataFromPaths(const TArray<FString>& Paths, bool
 		ARFilter.ClassNames.Add(ObjectBaseClass->GetFName());
 
 		// Add any old names to the list in case things haven't been resaved
-		TArray<FName> OldNames = ULinkerLoad::FindPreviousNamesForClass(ObjectBaseClass->GetPathName(), false);
+		TArray<FName> OldNames = FLinkerLoad::FindPreviousNamesForClass(ObjectBaseClass->GetPathName(), false);
 
 		ARFilter.ClassNames.Append(OldNames);
 
diff --git a/Engine/Source/Runtime/Engine/Private/Particles/ParticleComponents.cpp b/Engine/Source/Runtime/Engine/Private/Particles/ParticleComponents.cpp
--- a/Engine/Source/Runtime/Engine/Private/Particles/ParticleComponents.cpp
+++ b/Engine/Source/Runtime/Engine/Private/Particles/ParticleComponents.cpp
@@ -1126,7 +1126,7 @@ void UParticleEmitter::PostLoad()
 		{
 			LODLevel->ConditionalPostLoad();
 
-			ULinkerLoad* LODLevelLinker = LODLevel->GetLinker();
+			FLinkerLoad* LODLevelLinker = LODLevel->GetLinker();
 			if (LODLevel->SpawnModule == NULL)
 			{
 				// Force the conversion to SpawnModule
diff --git a/Engine/Source/Runtime/Engine/Private/PlayerController.cpp b/Engine/Source/Runtime/Engine/Private/PlayerController.cpp
--- a/Engine/Source/Runtime/Engine/Private/PlayerController.cpp
+++ b/Engine/Source/Runtime/Engine/Private/PlayerController.cpp
@@ -234,7 +234,7 @@ void APlayerController::ServerUpdateLevelVisibility_Implementation(FName Package
 			// verify that we were passed a valid level name
 			FString Filename;
 			UPackage* TempPkg = FindPackage(NULL, *PackageName.ToString());
-			ULinkerLoad* Linker = ULinkerLoad::FindExistingLinkerForPackage(TempPkg);
+			FLinkerLoad* Linker = FLinkerLoad::FindExistingLinkerForPackage(TempPkg);
 
 			// If we have a linker we know it has been loaded off disk successfully
 			// If we have a file it is fine too
diff --git a/Engine/Source/Runtime/Engine/Private/SkeletalMesh.cpp b/Engine/Source/Runtime/Engine/Private/SkeletalMesh.cpp
--- a/Engine/Source/Runtime/Engine/Private/SkeletalMesh.cpp
+++ b/Engine/Source/Runtime/Engine/Private/SkeletalMesh.cpp
@@ -1155,7 +1155,7 @@ void FStaticLODModel::Serialize( FArchive& Ar, UObject* Owner, int32 Idx )
 	static const auto CVar = IConsoleManager::Get().FindTConsoleVariableDataInt(TEXT("r.FreeSkeletalMeshBuffers"));
 	if(CVar)
 	{
-		bKeepBuffersInCPUMemory = !CVar->GetValueOnGameThread();
+		bKeepBuffersInCPUMemory = !CVar->GetValueOnAnyThread();
 	}
 
 	Ar << Sections;
diff --git a/Engine/Source/Runtime/Engine/Private/Texture2D.cpp b/Engine/Source/Runtime/Engine/Private/Texture2D.cpp
--- a/Engine/Source/Runtime/Engine/Private/Texture2D.cpp
+++ b/Engine/Source/Runtime/Engine/Private/Texture2D.cpp
@@ -758,7 +758,7 @@ void UTexture2D::CalcAllowedMips( int32 MipCount, int32 NumNonStreamingMips, int
 
 FTextureResource* UTexture2D::CreateResource()
 {
-	ULinker* Linker = GetLinker();
+	FLinker* Linker = GetLinker();
 	int32 NumMips = GetNumMips();
 
 	// Determine whether or not this texture can be streamed.
diff --git a/Engine/Source/Runtime/Engine/Private/UnrealEngine.cpp b/Engine/Source/Runtime/Engine/Private/UnrealEngine.cpp
--- a/Engine/Source/Runtime/Engine/Private/UnrealEngine.cpp
+++ b/Engine/Source/Runtime/Engine/Private/UnrealEngine.cpp
@@ -709,7 +709,7 @@ void UEngine::Init(IEngineLoop* InEngineLoop)
 
 	// Subsystems.
 	FURL::StaticInit();
-	ULinkerLoad::StaticInit(UTexture2D::StaticClass());
+	FLinkerLoad::StaticInit(UTexture2D::StaticClass());
 
 #if !UE_BUILD_SHIPPING
 	// Check for overrides to the default map on the command line
@@ -4068,7 +4068,7 @@ bool UEngine::HandleDumpParticleCountsCommand( const TCHAR* Cmd, FOutputDevice&
 bool UEngine::HandleListPreCacheMapPackagesCommand( const TCHAR* Cmd, FOutputDevice& Ar )
 {
 	TArray<FString> Packages;
-	ULinkerLoad::GetListOfPackagesInPackagePrecacheMap( Packages );
+	FLinkerLoad::GetListOfPackagesInPackagePrecacheMap( Packages );
 
 	Packages.Sort();
 
@@ -4705,6 +4705,24 @@ bool UEngine::HandleObjCommand( const TCHAR* Cmd, FOutputDevice& Ar )
 		}
 		return true;
 	}
+	else if (FParse::Command(&Cmd, TEXT("TRYGC")))
+	{
+		// Purge unclaimed objects.		
+		if (TryCollectGarbage(GARBAGE_COLLECTION_KEEPFLAGS))
+		{
+			Ar.Logf(TEXT("Collecting garbage and resetting GC timers on all worlds."));
+			for (TObjectIterator<UWorld> It; It; ++It)
+			{
+				UWorld* CurrentWorld = *It;
+				CurrentWorld->TimeSinceLastPendingKillPurge = 0;
+			}
+		}
+		else
+		{
+			Ar.Logf(TEXT("Garbage collection blocked by other threads."));
+		}
+		return true;
+	}
 	else if (FParse::Command(&Cmd,TEXT("LIST2")))
 	{			
 		UClass* ClassToCheck = NULL;
@@ -4980,7 +4998,7 @@ bool UEngine::HandleObjCommand( const TCHAR* Cmd, FOutputDevice& Ar )
 					continue;
 				}
 
-				if ( bOnlyListGCObjects && GUObjectArray.IsDisregardForGC(*It) )
+				if ( bOnlyListGCObjects && GetUObjectArray().IsDisregardForGC(*It) )
 				{
 					continue;
 				}
@@ -6230,7 +6248,7 @@ FGuid UEngine::GetPackageGuid(FName PackageName)
 	FGuid Result(0,0,0,0);
 
 	BeginLoad();
-	ULinkerLoad* Linker = GetPackageLinker(NULL, *PackageName.ToString(), LOAD_NoWarn | LOAD_NoVerify, NULL, NULL);
+	FLinkerLoad* Linker = GetPackageLinker(NULL, *PackageName.ToString(), LOAD_NoWarn | LOAD_NoVerify, NULL, NULL);
 	if (Linker != NULL && Linker->LinkerRoot != NULL)
 	{
 		Result = Linker->LinkerRoot->GetGuid();
@@ -6544,10 +6562,14 @@ bool UEngine::HandleLogoutStatLevelsCommand( const TCHAR* Cmd, FOutputDevice& Ar
 		{
 			DisplayName += FString::Printf(TEXT(" - %4.1f sec"), LevelPackage->GetLoadTime());
 		}
-		else if( GetAsyncLoadPercentage( LevelStatus.PackageName ) >= 0 )
+		else
 		{
-			const int32 Percentage = FMath::TruncToInt( GetAsyncLoadPercentage( LevelStatus.PackageName ) );
-			DisplayName += FString::Printf(TEXT(" - %3i %%"), Percentage ); 
+			const float LevelLoadPercentage = GetAsyncLoadPercentage(LevelStatus.PackageName);
+			if (LevelLoadPercentage >= 0.0f)
+			{
+				const int32 Percentage = FMath::TruncToInt(LevelLoadPercentage);
+				DisplayName += FString::Printf(TEXT(" - %3i %%"), Percentage);
+			}
 		}
 
 		if ( LevelStatus.bPlayerInside )
@@ -9965,11 +9987,11 @@ bool UEngine::PrepareMapChange(FWorldContext &Context, const TArray<FName>& Leve
 				{
 					// Load localized part of level first in case it exists. We don't need to worry about GC or completion 
 					// callback as we always kick off another async IO for the level below.
-					LoadPackageAsync( *LocalizedPackageName );
+					LoadPackageAsync(LocalizedPackageName);
 				}
 			}
 			
-			LoadPackageAsync( *LevelName.ToString(), 
+			LoadPackageAsync(LevelName.ToString(),
 				FLoadPackageAsyncDelegate::CreateStatic(&AsyncMapChangeLevelLoadCompletionCallback, Context.ContextHandle)
 				);
 		}
@@ -10630,7 +10652,7 @@ void UEngine::CopyPropertiesForUnrelatedObjects(UObject* OldObject, UObject* New
 		NewActor->ResetOwnedComponents();
 	}
 
-	bool bDumpProperties = CVarDumpCopyPropertiesForUnrelatedObjects.GetValueOnGameThread() != 0;
+	bool bDumpProperties = CVarDumpCopyPropertiesForUnrelatedObjects.GetValueOnAnyThread() != 0;
 	// Uncomment the next line to debug CPFUO for a specific object:
 	// bDumpProperties |= (NewObject->GetName().Find(TEXT("SpinTree")) != INDEX_NONE);
 	if (bDumpProperties)
@@ -11224,8 +11246,8 @@ int32 UEngine::RenderStatLevels(UWorld* World, FViewport* Viewport, FCanvas* Can
 		}
 		else
 		{
-			float AsyncLoadPercentage = GetAsyncLoadPercentage(*LevelStatus.PackageName.ToString());
-			if (AsyncLoadPercentage >= 0)
+			const float AsyncLoadPercentage = GetAsyncLoadPercentage(*LevelStatus.PackageName.ToString());
+			if (AsyncLoadPercentage >= 0.0f)
 			{
 				const int32 Percentage = FMath::TruncToInt(AsyncLoadPercentage);
 				DisplayName += FString::Printf(TEXT(" - %3i %%"), Percentage);
diff --git a/Engine/Source/Runtime/Engine/Private/UserDefinedStruct.cpp b/Engine/Source/Runtime/Engine/Private/UserDefinedStruct.cpp
--- a/Engine/Source/Runtime/Engine/Private/UserDefinedStruct.cpp
+++ b/Engine/Source/Runtime/Engine/Private/UserDefinedStruct.cpp
@@ -163,7 +163,7 @@ void UUserDefinedStruct::SerializeTaggedProperties(FArchive& Ar, uint8* Data, US
 
 void UUserDefinedStruct::RecursivelyPreload()
 {
-	ULinkerLoad* Linker = GetLinker();
+	FLinkerLoad* Linker = GetLinker();
 	if( Linker && (NULL == PropertyLink) )
 	{
 		TArray<UObject*> AllChildMembers;
diff --git a/Engine/Source/Runtime/Engine/Private/World.cpp b/Engine/Source/Runtime/Engine/Private/World.cpp
--- a/Engine/Source/Runtime/Engine/Private/World.cpp
+++ b/Engine/Source/Runtime/Engine/Private/World.cpp
@@ -76,6 +76,7 @@
 #include "Engine/Polys.h"
 #include "Engine/LightMapTexture2D.h"
 #include "Engine/GameInstance.h"
+#include "UObject/UObjectThreadContext.h"
 
 DEFINE_LOG_CATEGORY_STATIC(LogWorld, Log, All);
 DEFINE_LOG_CATEGORY(LogSpawn);
@@ -273,7 +274,7 @@ bool UWorld::Rename(const TCHAR* InName, UObject* NewOuter, ERenameFlags Flags)
 	const bool bTestRename = (Flags & REN_Test) != 0;
 
 	// Rename the level script blueprint now, unless we are in PostLoad. ULevel::PostLoad should handle renaming this at load time.
-	if ( !GIsRoutingPostLoad )
+	if (!FUObjectThreadContext::Get().IsRoutingPostLoad)
 	{
 		const bool bDontCreate = true;
 		UBlueprint* LevelScriptBlueprint = PersistentLevel->GetLevelScriptBlueprint(bDontCreate);
@@ -4490,7 +4491,7 @@ void FSeamlessTravelHandler::StartLoadingDestination()
 			{
 				// Load localized part of level first in case it exists. We don't need to worry about GC or completion 
 				// callback as we always kick off another async IO for the level below.
-				LoadPackageAsync(*LocalizedPackageName);
+				LoadPackageAsync(LocalizedPackageName);
 			}
 		}
 
@@ -4516,12 +4517,15 @@ void FSeamlessTravelHandler::StartLoadingDestination()
 			}
 		}
 #endif
-		LoadPackageAsync(URLMapPackageName, 
-			FLoadPackageAsyncDelegate::CreateRaw(this, &FSeamlessTravelHandler::SeamlessTravelLoadCallback), 
+		LoadPackageAsync(
+			URLMapPackageName, 
 			PendingTravelGuid.IsValid() ? &PendingTravelGuid : NULL,
 			NAME_None,
-			*URLMapPackageToLoadFrom
-			).SetPackageData(PackageFlags, PIEInstanceID);
+			*URLMapPackageToLoadFrom,
+			FLoadPackageAsyncDelegate::CreateRaw(this, &FSeamlessTravelHandler::SeamlessTravelLoadCallback), 			
+			PackageFlags,
+			PIEInstanceID
+			);
 	}
 	else
 	{
diff --git a/Engine/Source/Runtime/Engine/Public/LightMap.h b/Engine/Source/Runtime/Engine/Public/LightMap.h
--- a/Engine/Source/Runtime/Engine/Public/LightMap.h
+++ b/Engine/Source/Runtime/Engine/Public/LightMap.h
@@ -84,12 +84,12 @@ public:
 	// Reference counting.
 	void AddRef()
 	{
-		check( IsInGameThread() );
+		check(IsInGameThread() || IsInAsyncLoadingThread());
 		NumRefs++;
 	}
 	void Release()
 	{
-		check( IsInGameThread() );
+		check(IsInGameThread() || IsInAsyncLoadingThread());
 		checkSlow(NumRefs > 0);
 		if(--NumRefs == 0)
 		{
diff --git a/Engine/Source/Runtime/Engine/Public/MaterialShared.h b/Engine/Source/Runtime/Engine/Public/MaterialShared.h
--- a/Engine/Source/Runtime/Engine/Public/MaterialShared.h
+++ b/Engine/Source/Runtime/Engine/Public/MaterialShared.h
@@ -1025,20 +1025,20 @@ public:
 
 	class FMaterialShaderMap* GetGameThreadShaderMap() const 
 	{ 
-		checkSlow(IsInGameThread());
+		checkSlow(IsInGameThread() || IsInAsyncLoadingThread());
 		return GameThreadShaderMap; 
 	}
 
 	/** Note: SetRenderingThreadShaderMap must also be called with the same value, but from the rendering thread. */
 	void SetGameThreadShaderMap(FMaterialShaderMap* InMaterialShaderMap)
 	{
-		checkSlow(IsInGameThread());
+		checkSlow(IsInGameThread() || IsInAsyncLoadingThread());
 		GameThreadShaderMap = InMaterialShaderMap;
 	}
 
 	void SetInlineShaderMap(FMaterialShaderMap* InMaterialShaderMap)
 	{
-		checkSlow(IsInGameThread());
+		checkSlow(IsInGameThread() || IsInAsyncLoadingThread());
 		GameThreadShaderMap = InMaterialShaderMap;
 		bContainsInlineShaders = true;
 	}
diff --git a/Engine/Source/Runtime/Engine/Public/ShadowMap.h b/Engine/Source/Runtime/Engine/Public/ShadowMap.h
--- a/Engine/Source/Runtime/Engine/Public/ShadowMap.h
+++ b/Engine/Source/Runtime/Engine/Public/ShadowMap.h
@@ -246,12 +246,12 @@ public:
 	// Reference counting.
 	void AddRef()
 	{
-		check(IsInGameThread());
+		check(IsInGameThread() || IsInAsyncLoadingThread());
 		NumRefs++;
 	}
 	void Release()
 	{
-		check(IsInGameThread());
+		check(IsInGameThread() || IsInAsyncLoadingThread());
 		checkSlow(NumRefs > 0);
 		if (--NumRefs == 0)
 		{
diff --git a/Engine/Source/Runtime/GameplayTags/Private/GameplayTagsManager.cpp b/Engine/Source/Runtime/GameplayTags/Private/GameplayTagsManager.cpp
--- a/Engine/Source/Runtime/GameplayTags/Private/GameplayTagsManager.cpp
+++ b/Engine/Source/Runtime/GameplayTags/Private/GameplayTagsManager.cpp
@@ -32,7 +32,7 @@ void UGameplayTagsManager::LoadGameplayTagTables(TArray<FString>& TagTableNames)
 			// will not be initialized correctly.
 			if (TagTable && IsLoading())
 			{
-				ULinkerLoad* TagLinker = TagTable->GetLinker();
+				FLinkerLoad* TagLinker = TagTable->GetLinker();
 				if (TagLinker)
 				{
 					TagTable->GetLinker()->Preload(TagTable);
diff --git a/Engine/Source/Runtime/Launch/Private/LaunchEngineLoop.cpp b/Engine/Source/Runtime/Launch/Private/LaunchEngineLoop.cpp
--- a/Engine/Source/Runtime/Launch/Private/LaunchEngineLoop.cpp
+++ b/Engine/Source/Runtime/Launch/Private/LaunchEngineLoop.cpp
@@ -1350,7 +1350,7 @@ int32 FEngineLoop::PreInit( const TCHAR* CmdLine )
 	// Setup GC optimizations
 	if (bIsSeekFreeDedicatedServer || bHasEditorToken)
 	{
-		GUObjectArray.DisableDisregardForGC();
+		GetUObjectArray().DisableDisregardForGC();
 	}
 	
 	SlowTask.EnterProgressFrame(10);
@@ -1430,7 +1430,7 @@ int32 FEngineLoop::PreInit( const TCHAR* CmdLine )
 	}
 
 	MarkObjectsToDisregardForGC(); 
-	GUObjectArray.CloseDisregardForGC();
+	GetUObjectArray().CloseDisregardForGC();
 
 	SetIsServerForOnlineSubsystemsDelegate(FQueryIsRunningServer::CreateStatic(&IsServerDelegateForOSS));
 
diff --git a/Engine/Source/Runtime/RenderCore/Private/RenderingThread.cpp b/Engine/Source/Runtime/RenderCore/Private/RenderingThread.cpp
--- a/Engine/Source/Runtime/RenderCore/Private/RenderingThread.cpp
+++ b/Engine/Source/Runtime/RenderCore/Private/RenderingThread.cpp
@@ -698,7 +698,7 @@ bool FRenderCommandFence::IsFenceComplete() const
 	{
 		return true;
 	}
-	check(IsInGameThread());
+	check(IsInGameThread() || IsInAsyncLoadingThread());
 	CheckRenderingThreadHealth();
 	if (!CompletionEvent.GetReference() || CompletionEvent->IsComplete())
 	{
diff --git a/Engine/Source/Runtime/ShaderCore/Private/Shader.cpp b/Engine/Source/Runtime/ShaderCore/Private/Shader.cpp
--- a/Engine/Source/Runtime/ShaderCore/Private/Shader.cpp
+++ b/Engine/Source/Runtime/ShaderCore/Private/Shader.cpp
@@ -230,9 +230,11 @@ FArchive& operator<<(FArchive& Ar,FShaderType*& Ref)
 }
 
 
-FShader* FShaderType::FindShaderById(const FShaderId& Id) const
+FShader* FShaderType::FindShaderById(const FShaderId& Id)
 {
-	return ShaderIdMap.FindRef(Id);
+	FScopeLock MapLock(&ShaderIdMapCritical);
+	FShader* Result = ShaderIdMap.FindRef(Id);
+	return Result;
 }
 
 FShader* FShaderType::ConstructForDeserialization() const
@@ -289,6 +291,7 @@ void FShaderType::Uninitialize()
 }
 
 TMap<FShaderResourceId, FShaderResource*> FShaderResource::ShaderResourceIdMap;
+FCriticalSection FShaderResource::ShaderResourceIdMapCritical;
 
 FShaderResource::FShaderResource()
 	: SpecificType(NULL)
@@ -314,7 +317,11 @@ FShaderResource::FShaderResource(const FShaderCompilerOutput& Output, FShaderTyp
 	OutputHash = Output.OutputHash;
 	checkSlow(OutputHash != FSHAHash());
 
-	ShaderResourceIdMap.Add(GetId(), this);
+	{
+		FScopeLock ShaderResourceIdMapLock(&ShaderResourceIdMapCritical);
+		ShaderResourceIdMap.Add(GetId(), this);
+	}
+	
 	INC_DWORD_STAT_BY_FName(GetMemoryStatType((EShaderFrequency)Target.Frequency).GetName(), Code.Num());
 	INC_DWORD_STAT_BY(STAT_Shaders_ShaderResourceMemory, GetSizeBytes());
 	INC_DWORD_STAT_BY(STAT_Shaders_NumShaderResourcesLoaded, 1);
@@ -331,6 +338,7 @@ FShaderResource::~FShaderResource()
 
 void FShaderResource::Register()
 {
+	FScopeLock ShaderResourceIdMapLock(&ShaderResourceIdMapCritical);
 	ShaderResourceIdMap.Add(GetId(), this);
 }
 
@@ -365,7 +373,10 @@ void FShaderResource::Release()
 	check(NumRefs != 0);
 	if(--NumRefs == 0)
 	{
-		ShaderResourceIdMap.Remove(GetId());
+		{
+			FScopeLock ShaderResourceIdMapLock(&ShaderResourceIdMapCritical);
+			ShaderResourceIdMap.Remove(GetId());
+		}
 
 		// Send a release message to the rendering thread when the shader loses its last reference.
 		BeginReleaseResource(this);
@@ -377,6 +388,7 @@ void FShaderResource::Release()
 
 FShaderResource* FShaderResource::FindShaderResourceById(const FShaderResourceId& Id)
 {
+	FScopeLock ShaderResourceIdMapLock(&ShaderResourceIdMapCritical);
 	return ShaderResourceIdMap.FindRef(Id);
 }
 
@@ -396,6 +408,7 @@ FShaderResource* FShaderResource::FindOrCreateShaderResource(const FShaderCompil
 
 void FShaderResource::GetAllShaderResourceId(TArray<FShaderResourceId>& Ids)
 {
+	FScopeLock ShaderResourceIdMapLock(&ShaderResourceIdMapCritical);
 	ShaderResourceIdMap.GetKeys(Ids);
 }
 
@@ -939,13 +952,13 @@ void FShader::Register()
 	check(ShaderId.MaterialShaderMapHash != FSHAHash());
 	check(ShaderId.SourceHash != FSHAHash());
 	check(Resource);
-	Type->GetShaderIdMap().Add(ShaderId, this);
+	Type->AddToShaderIdMap(ShaderId, this);
 }
 
 
 void FShader::Deregister()
 {
-	Type->GetShaderIdMap().Remove(GetId());
+	Type->RemoveFromShaderIdMap(GetId());
 }
 
 
@@ -1165,17 +1178,17 @@ void ShaderMapAppendKeyString(EShaderPlatform Platform, FString& KeyString)
 
 	{
 		static const auto CVar = IConsoleManager::Get().FindTConsoleVariableDataInt(TEXT("Compat.UseDXT5NormalMaps"));
-		KeyString += (CVar && CVar->GetValueOnGameThread() != 0) ? TEXT("_DXTN") : TEXT("_BC5N");
+		KeyString += (CVar && CVar->GetValueOnAnyThread() != 0) ? TEXT("_DXTN") : TEXT("_BC5N");
 	}
 
 	{
 		static const auto CVar = IConsoleManager::Get().FindTConsoleVariableDataInt(TEXT("r.CompileShadersForDevelopment"));
-		KeyString += (CVar && CVar->GetValueOnGameThread() != 0) ? TEXT("_DEV") : TEXT("_NoDEV");
+		KeyString += (CVar && CVar->GetValueOnAnyThread() != 0) ? TEXT("_DEV") : TEXT("_NoDEV");
 	}
 
 	{
 		static const auto CVar = IConsoleManager::Get().FindTConsoleVariableDataInt(TEXT("r.AllowStaticLighting"));
-		const bool bValue = CVar ? CVar->GetValueOnGameThread() != 0 : true;
+		const bool bValue = CVar ? CVar->GetValueOnAnyThread() != 0 : true;
 		KeyString += bValue ? TEXT("_SL") : TEXT("_NoSL");
 	}
 
@@ -1189,7 +1202,7 @@ void ShaderMapAppendKeyString(EShaderPlatform Platform, FString& KeyString)
 
 	{
 		static const auto CVar = IConsoleManager::Get().FindTConsoleVariableDataInt(TEXT("r.GBuffer"));
-		if(CVar ? CVar->GetValueOnGameThread() == 0 : false)
+		if (CVar ? CVar->GetValueOnAnyThread() == 0 : false)
 		{
 			KeyString += TEXT("_NoGB");
 		}
@@ -1214,7 +1227,7 @@ void ShaderMapAppendKeyString(EShaderPlatform Platform, FString& KeyString)
 	{
 		{
 			static const auto CVar = IConsoleManager::Get().FindTConsoleVariableDataInt(TEXT("r.PS4MixedModeShaderDebugInfo"));
-			if( CVar && CVar->GetValueOnGameThread() != 0 )
+			if (CVar && CVar->GetValueOnAnyThread() != 0)
 			{
 				KeyString += TEXT("_MMDBG");
 			}
@@ -1222,7 +1235,7 @@ void ShaderMapAppendKeyString(EShaderPlatform Platform, FString& KeyString)
 
 		{
 			static const auto CVar = IConsoleManager::Get().FindTConsoleVariableDataInt(TEXT("r.PS4DumpShaderSDB"));
-			if( CVar && CVar->GetValueOnGameThread() != 0 )
+			if (CVar && CVar->GetValueOnAnyThread() != 0)
 			{
 				KeyString += TEXT("_SDB");
 			}
diff --git a/Engine/Source/Runtime/ShaderCore/Private/ShaderCore.cpp b/Engine/Source/Runtime/ShaderCore/Private/ShaderCore.cpp
--- a/Engine/Source/Runtime/ShaderCore/Private/ShaderCore.cpp
+++ b/Engine/Source/Runtime/ShaderCore/Private/ShaderCore.cpp
@@ -465,7 +465,7 @@ void GetShaderIncludes(const TCHAR* Filename, TArray<FString>& IncludeFilenames,
 const FSHAHash& GetShaderFileHash(const TCHAR* Filename)
 {
 	// Make sure we are only accessing GShaderHashCache from one thread
-	check(IsInGameThread());
+	//check(IsInGameThread() || IsAsyncLoading());
 	STAT(double HashTime = 0);
 	{
 		SCOPE_SECONDS_COUNTER(HashTime);
diff --git a/Engine/Source/Runtime/ShaderCore/Public/Shader.h b/Engine/Source/Runtime/ShaderCore/Public/Shader.h
--- a/Engine/Source/Runtime/ShaderCore/Public/Shader.h
+++ b/Engine/Source/Runtime/ShaderCore/Public/Shader.h
@@ -216,6 +216,8 @@ private:
 
 	/** Tracks loaded shader resources by id. */
 	static TMap<FShaderResourceId, FShaderResource*> ShaderResourceIdMap;
+	/** Critical section for ShaderResourceIdMap. */
+	static FCriticalSection ShaderResourceIdMapCritical;
 };
 
 /** Encapsulates information about a shader's serialization behavior, used to detect when C++ serialization changes to auto-recompile. */
@@ -738,7 +740,7 @@ public:
 	 * Finds a shader of this type by ID.
 	 * @return NULL if no shader with the specified ID was found.
 	 */
-	FShader* FindShaderById(const FShaderId& Id) const;
+	FShader* FindShaderById(const FShaderId& Id);
 
 	/** Constructs a new instance of the shader type for deserialization. */
 	FShader* ConstructForDeserialization() const;
@@ -805,9 +807,16 @@ public:
 		}
 	}
 
-	TMap<FShaderId,FShader*>& GetShaderIdMap()
+	void AddToShaderIdMap(FShaderId Id, FShader* Shader)
 	{
-		return ShaderIdMap;
+		FScopeLock MapLock(&ShaderIdMapCritical);
+		ShaderIdMap.Add(Id, Shader);
+	}
+
+	void RemoveFromShaderIdMap(FShaderId Id)
+	{
+		FScopeLock MapLock(&ShaderIdMapCritical);
+		ShaderIdMap.Remove(Id);
 	}
 
 	bool LimitShaderResourceToThisType()
@@ -833,6 +842,7 @@ private:
 
 	/** A map from shader ID to shader.  A shader will be removed from it when deleted, so this doesn't need to use a TRefCountPtr. */
 	TMap<FShaderId,FShader*> ShaderIdMap;
+	FCriticalSection ShaderIdMapCritical;
 
 	TLinkedList<FShaderType*> GlobalListLink;
 
diff --git a/Engine/Source/Runtime/UMG/Public/Components/Widget.h b/Engine/Source/Runtime/UMG/Public/Components/Widget.h
--- a/Engine/Source/Runtime/UMG/Public/Components/Widget.h
+++ b/Engine/Source/Runtime/UMG/Public/Components/Widget.h
@@ -5,6 +5,7 @@
 #include "SlateWrapperTypes.h"
 #include "WidgetTransform.h"
 #include "DynamicPropertyPath.h"
+#include "UObject/UObjectThreadContext.h"
 
 #include "Widget.generated.h"
 
@@ -373,7 +374,7 @@ public:
 #if WITH_EDITOR
 	FORCEINLINE bool CanSafelyRouteEvent()
 	{
-		return !( IsDesignTime() || GIsRoutingPostLoad || GIntraFrameDebuggingGameThread || HasAnyFlags(RF_Unreachable) );
+		return !(IsDesignTime() || GIntraFrameDebuggingGameThread || HasAnyFlags(RF_Unreachable) || FUObjectThreadContext::Get().IsRoutingPostLoad);
 	}
 #else
 	FORCEINLINE bool CanSafelyRouteEvent() { return true; }
