Fixing CIS on the Mac

diff --git a/Engine/Source/Runtime/CoreUObject/Private/UObject/CoreNative.cpp b/Engine/Source/Runtime/CoreUObject/Private/UObject/CoreNative.cpp
--- a/Engine/Source/Runtime/CoreUObject/Private/UObject/CoreNative.cpp
+++ b/Engine/Source/Runtime/CoreUObject/Private/UObject/CoreNative.cpp
@@ -6,23 +6,24 @@
 #include "StringClassReference.h"
 #include "Linker.h"
 
+void UClassRegisterAllCompiledInClasses();
+bool IsInAsyncLoadingThreadCoreUObjectInternal();
+bool IsAsyncLoadingCoreUObjectInternal();
+
 // CoreUObject module. Handles UObject system pre-init (registers init function with Core callbacks).
 class FCoreUObjectModule : public FDefaultModuleImpl
 {
 public:
 	virtual void StartupModule() override
 	{
-		// Register all classes that have been loaded so far. This is required for CVars to work.
-		void UClassRegisterAllCompiledInClasses();
+		// Register all classes that have been loaded so far. This is required for CVars to work.		
 		UClassRegisterAllCompiledInClasses();
 
 		void InitUObject();
 		FCoreDelegates::OnInit.AddStatic(InitUObject);
 
-		bool IsInAsyncLoadingThreadCoreUObjectInternal();
+		// Substitute Core version of async loading functions with CoreUObject ones.
 		IsInAsyncLoadingThread = &IsInAsyncLoadingThreadCoreUObjectInternal;
-
-		bool IsAsyncLoadingCoreUObjectInternal();
 		IsAsyncLoading = &IsAsyncLoadingCoreUObjectInternal;
 
 		// Make sure that additional content mount points can be registered after CoreUObject loads
diff --git a/Engine/Source/Runtime/CoreUObject/Private/UObject/UObjectArray.cpp b/Engine/Source/Runtime/CoreUObject/Private/UObject/UObjectArray.cpp
--- a/Engine/Source/Runtime/CoreUObject/Private/UObject/UObjectArray.cpp
+++ b/Engine/Source/Runtime/CoreUObject/Private/UObject/UObjectArray.cpp
@@ -107,7 +107,7 @@ void FUObjectArray::FreeUObjectIndex(UObjectBase* Object)
 	// No point in filling this list when doing exit purge. Nothing should be allocated afterwards anyway.
 	if (Index > ObjLastNonGCIndex && !GExitPurge)  
 	{
-		ObjAvailableList.Push((int32*)Index);
+		ObjAvailableList.Push((int32*)(uintptr_t)Index);
 #if WITH_EDITOR
 		ObjAvailableCount.Increment();
 #endif
