Removed unnecessary #pragma once

diff --git a/Engine/Source/Runtime/CoreUObject/Private/UObject/LinkerManager.cpp b/Engine/Source/Runtime/CoreUObject/Private/UObject/LinkerManager.cpp
--- a/Engine/Source/Runtime/CoreUObject/Private/UObject/LinkerManager.cpp
+++ b/Engine/Source/Runtime/CoreUObject/Private/UObject/LinkerManager.cpp
@@ -3,9 +3,6 @@
 /*=============================================================================
 	LinkerManager.h: Unreal object linker manager
 =============================================================================*/
-
-#pragma once
-
 #include "CoreUObjectPrivate.h"
 #include "LinkerManager.h"
 
