Fixed UAssetObjectProperty not being able to import asset references in the form "Texture2D'/Game/Thing'"

UE-13110 - Regression: DataTables: Importing columns of Asset references doesn't work with native FTableRowBase

It was passing the nested object path into UPropertyHelpers::ReadToken without skipping over the opening ' (like UObjectPropertyBase does), which meant that you ended up with an imported object path of something like "Texture2D'" rather than "Texture2D'/Game/Thing'".

This broke data table importing when using a TAssetPtr property. This wasn't an issue when using User Defined Structs as they use a UObjectPropertyBase when referring to other objects.

diff --git a/Engine/Source/Runtime/CoreUObject/Private/UObject/PropertyAssetObject.cpp b/Engine/Source/Runtime/CoreUObject/Private/UObject/PropertyAssetObject.cpp
--- a/Engine/Source/Runtime/CoreUObject/Private/UObject/PropertyAssetObject.cpp
+++ b/Engine/Source/Runtime/CoreUObject/Private/UObject/PropertyAssetObject.cpp
@@ -97,16 +97,20 @@ const TCHAR* UAssetObjectProperty::ImportText_Internal( const TCHAR* InBuffer, v
 	{
 		if( *Buffer == TCHAR('\'') )
 		{
+			// A ' token likely means we're looking at an asset string in the form "Texture2d'/Game/UI/HUD/Actions/Barrel'" and we need to read and append the path part
+			// We have to skip over the first ' as UPropertyHelpers::ReadToken doesn't read single-quoted strings correctly, but does read a path correctly
+			NewPath += *Buffer++; // Append the leading '
 			NewBuffer = UPropertyHelpers::ReadToken( Buffer, NewPath, 1 );
 			if( !NewBuffer )
 			{
 				return NULL;
 			}
 			Buffer = NewBuffer;
-			if( *Buffer++ != TCHAR('\'') )
+			if( *Buffer != TCHAR('\'') )
 			{
 				return NULL;
 			}
+			NewPath += *Buffer++; // Append the trailing '
 		}
 		FStringAssetReference ID(NewPath);
 		AssetPtr = ID;
