Making sure objects that have already been fully loaded don't get loaded again.

diff --git a/Engine/Source/Runtime/CoreUObject/Private/UObject/LinkerLoad.cpp b/Engine/Source/Runtime/CoreUObject/Private/UObject/LinkerLoad.cpp
--- a/Engine/Source/Runtime/CoreUObject/Private/UObject/LinkerLoad.cpp
+++ b/Engine/Source/Runtime/CoreUObject/Private/UObject/LinkerLoad.cpp
@@ -3478,10 +3478,10 @@ UObject* FLinkerLoad::CreateExport( int32 Index )
 					// If this object was allocated but never loaded (components created by a constructor) make sure it gets loaded
 					// Do this for all subobjects created in the native constructor.
 					FUObjectThreadContext::Get().ObjLoaded.AddUnique(Export.Object);
-					if( Export.Object->HasAnyFlags(RF_DefaultSubObject) 
-						|| (ThisParent && ThisParent->HasAnyFlags(RF_ClassDefaultObject)))
+					if (!Export.Object->HasAnyFlags(RF_LoadCompleted) &&
+						(Export.Object->HasAnyFlags(RF_DefaultSubObject) || (ThisParent && ThisParent->HasAnyFlags(RF_ClassDefaultObject))))
 					{
-						Export.Object->SetFlags(RF_NeedLoad|RF_NeedPostLoad|RF_NeedPostLoadSubobjects|RF_WasLoaded);
+						Export.Object->SetFlags(RF_NeedLoad | RF_NeedPostLoad | RF_NeedPostLoadSubobjects | RF_WasLoaded);
 					}
 				}
 				return Export.Object;
