Fixing fallout from CL 2498249 (linkers are no longer UObjects); renaming a missed ULinkerLoad var to FLinkerLoad.

diff --git a/Engine/Source/Runtime/CoreUObject/Private/UObject/LinkerPlaceholderBase.inl b/Engine/Source/Runtime/CoreUObject/Private/UObject/LinkerPlaceholderBase.inl
--- a/Engine/Source/Runtime/CoreUObject/Private/UObject/LinkerPlaceholderBase.inl
+++ b/Engine/Source/Runtime/CoreUObject/Private/UObject/LinkerPlaceholderBase.inl
@@ -55,7 +55,7 @@ void TLinkerImportPlaceholder<PlaceholderType>::AddReferencingProperty(UProperty
 	check(ThisAsObject != nullptr);
 
 	FObjectImport* PlaceholderImport = nullptr;
-	if (ULinkerLoad* PropertyLinker = ReferencingProperty->GetLinker())
+	if (FLinkerLoad* PropertyLinker = ReferencingProperty->GetLinker())
 	{
 		for (FObjectImport& Import : PropertyLinker->ImportMap)
 		{
