[UE-12888] Preventing missing interfaces from triggering a check() crash in deferred dependency loading.

diff --git a/Engine/Source/Runtime/CoreUObject/Private/Blueprint/BlueprintSupport.cpp b/Engine/Source/Runtime/CoreUObject/Private/Blueprint/BlueprintSupport.cpp
--- a/Engine/Source/Runtime/CoreUObject/Private/Blueprint/BlueprintSupport.cpp
+++ b/Engine/Source/Runtime/CoreUObject/Private/Blueprint/BlueprintSupport.cpp
@@ -899,7 +899,7 @@ int32 FLinkerLoad::ResolveDependencyPlaceholder(FLinkerPlaceholderBase* Placehol
 			if (Interface.Class == PlaceholderObj)
 			{
 				++ReplacementCount;
-				Interface.Class = CastChecked<UClass>(RealImportObj);
+				Interface.Class = CastChecked<UClass>(RealImportObj, ECastCheckedType::NullAllowed);
 			}
 		}
 	}
