UE-13183: DDC build fails due to PKG_InMemoryOnly

diff --git a/Engine/Source/Runtime/CoreUObject/Private/UObject/LinkerSave.cpp b/Engine/Source/Runtime/CoreUObject/Private/UObject/LinkerSave.cpp
--- a/Engine/Source/Runtime/CoreUObject/Private/UObject/LinkerSave.cpp
+++ b/Engine/Source/Runtime/CoreUObject/Private/UObject/LinkerSave.cpp
@@ -29,7 +29,7 @@ FLinkerSave::FLinkerSave(UPackage* InParent, const TCHAR* InFilename, bool bForc
 		Summary.SetFileVersions( GPackageFileUE4Version, GPackageFileLicenseeUE4Version, bInSaveUnversioned );
 		Summary.SavedByEngineVersion = GEngineVersion;
 		Summary.CompatibleWithEngineVersion = GCompatibleWithEngineVersion;
-		Summary.PackageFlags  = Package ? Package->PackageFlags : 0;
+		Summary.PackageFlags = Package ? (Package->PackageFlags & ~PKG_NewlyCreated) : 0;
 
 		if (Package)
 		{
@@ -61,7 +61,7 @@ FLinkerSave::FLinkerSave(UPackage* InParent, bool bForceByteSwapping, bool bInSa
 		Summary.SetFileVersions( GPackageFileUE4Version, GPackageFileLicenseeUE4Version, bInSaveUnversioned );
 		Summary.SavedByEngineVersion = GEngineVersion;
 		Summary.CompatibleWithEngineVersion = GCompatibleWithEngineVersion;
-		Summary.PackageFlags  = Package ? Package->PackageFlags : 0;
+		Summary.PackageFlags = Package ? (Package->PackageFlags & ~PKG_NewlyCreated) : 0;
 
 		if (Package)
 		{
diff --git a/Engine/Source/Runtime/CoreUObject/Private/UObject/PackageFileSummary.cpp b/Engine/Source/Runtime/CoreUObject/Private/UObject/PackageFileSummary.cpp
--- a/Engine/Source/Runtime/CoreUObject/Private/UObject/PackageFileSummary.cpp
+++ b/Engine/Source/Runtime/CoreUObject/Private/UObject/PackageFileSummary.cpp
@@ -105,6 +105,15 @@ FArchive& operator<<( FArchive& Ar, FPackageFileSummary& Sum )
 		Ar << Sum.TotalHeaderSize;
 		Ar << Sum.FolderName;
 		Ar << Sum.PackageFlags;
+
+#if WITH_EDITOR
+		if (Ar.IsLoading())
+		{
+			// This flag should never be saved and its reused, so we need to make sure it hasn't been loaded.
+			Sum.PackageFlags &= ~PKG_NewlyCreated;
+		}
+#endif // WITH_EDITOR
+
 		if( Sum.PackageFlags & PKG_FilterEditorOnly )
 		{
 			Ar.SetFilterEditorOnly(true);
diff --git a/Engine/Source/Runtime/CoreUObject/Private/UObject/SavePackage.cpp b/Engine/Source/Runtime/CoreUObject/Private/UObject/SavePackage.cpp
--- a/Engine/Source/Runtime/CoreUObject/Private/UObject/SavePackage.cpp
+++ b/Engine/Source/Runtime/CoreUObject/Private/UObject/SavePackage.cpp
@@ -3557,7 +3557,7 @@ bool UPackage::SavePackage( UPackage* InOuter, UObject* Base, EObjectFlags TopLe
 				Linker->LinkerRoot->ThisRequiresLocalizationGather(Linker->RequiresLocalizationGather());
 				
 				// Update package flags from package, in case serialization has modified package flags.
-				Linker->Summary.PackageFlags  = Linker->LinkerRoot->PackageFlags;
+				Linker->Summary.PackageFlags  = Linker->LinkerRoot->PackageFlags & ~PKG_NewlyCreated;
 
 				Linker->Seek(0);
 				*Linker << Linker->Summary;
