Temporarily enable StaticLoadObject in PostLoad but only when it doesn't happen when async loading on the async loading thread

diff --git a/Engine/Source/Runtime/CoreUObject/Private/UObject/UObjectGlobals.cpp b/Engine/Source/Runtime/CoreUObject/Private/UObject/UObjectGlobals.cpp
--- a/Engine/Source/Runtime/CoreUObject/Private/UObject/UObjectGlobals.cpp
+++ b/Engine/Source/Runtime/CoreUObject/Private/UObject/UObjectGlobals.cpp
@@ -765,7 +765,10 @@ UObject* StaticLoadObjectInternal(UClass* ObjectClass, UObject* InOuter, const T
 
 UObject* StaticLoadObject(UClass* ObjectClass, UObject* InOuter, const TCHAR* InName, const TCHAR* Filename, uint32 LoadFlags, UPackageMap* Sandbox, bool bAllowObjectReconciliation )
 {
-	UE_CLOG(FUObjectThreadContext::Get().IsRoutingPostLoad, LogUObjectGlobals, Fatal, TEXT("Calling StaticLoadObject during PostLoad is forbidden."));
+	// Currently StaticLoadObject will not work if executed in PostLoad on the async loading thread
+	// It should be forbidden anyway but there's legacy code that is still doing it.
+	// Note that IsInAsyncLoadingThread() == true if async loading thread is disabled but we're inside of the Async Loading Tick
+	UE_CLOG(FUObjectThreadContext::Get().IsRoutingPostLoad && IsInAsyncLoadingThread() && !IsInGameThread(), LogUObjectGlobals, Fatal, TEXT("Calling StaticLoadObject during PostLoad is forbidden."));
 
 	UObject* Result = StaticLoadObjectInternal(ObjectClass, InOuter, InName, Filename, LoadFlags, Sandbox, bAllowObjectReconciliation);
 	if (!Result)
