When creating linker for async package, don't attempt to load packages with an empty name. Removes warning spam when trying to convert NAME_None to a package name.

diff --git a/Engine/Source/Runtime/CoreUObject/Private/Serialization/AsyncLoading.cpp b/Engine/Source/Runtime/CoreUObject/Private/Serialization/AsyncLoading.cpp
--- a/Engine/Source/Runtime/CoreUObject/Private/Serialization/AsyncLoading.cpp
+++ b/Engine/Source/Runtime/CoreUObject/Private/Serialization/AsyncLoading.cpp
@@ -870,7 +870,7 @@ EAsyncPackageState::Type FAsyncPackage::CreateLinker()
 		if (!Linker)
 		{
 			FString PackageFileName;
-			if (!FPackageName::DoesPackageExist(Desc.NameToLoad.ToString(), Desc.Guid.IsValid() ? &Desc.Guid : nullptr, &PackageFileName))
+			if (Desc.NameToLoad == NAME_None || !FPackageName::DoesPackageExist(Desc.NameToLoad.ToString(), Desc.Guid.IsValid() ? &Desc.Guid : nullptr, &PackageFileName))
 			{
 				UE_LOG(LogStreaming, Error, TEXT("Couldn't find file for package %s requested by async loading code."), *Desc.Name.ToString());
 				bLoadHasFailed = true;
