Copy pasting collapsed graph nodes with spaces in their name will no longer lose pin connection data.

diff --git a/Engine/Source/Runtime/CoreUObject/Private/UObject/PropertyBaseObject.cpp b/Engine/Source/Runtime/CoreUObject/Private/UObject/PropertyBaseObject.cpp
--- a/Engine/Source/Runtime/CoreUObject/Private/UObject/PropertyBaseObject.cpp
+++ b/Engine/Source/Runtime/CoreUObject/Private/UObject/PropertyBaseObject.cpp
@@ -153,7 +153,9 @@ FString UObjectPropertyBase::GetExportPath(const UObject* Object, const UObject*
 	// Take the path name relative to the stopping point outermost ptr.
 	// This is so that cases like a component referencing a component in another actor work correctly when pasted
 	FString PathName = Object->GetPathName(StopOuter);
-	if ( (PortFlags & PPF_Delimited) && (!Object->GetFName().IsValidXName(INVALID_OBJECTNAME_CHARACTERS)) )
+	int32 ResultIdx = 0;
+	// Object names that contain invalid characters and paths that contain spaces must be put into quotes to be handled correctly
+	if ( (PortFlags & PPF_Delimited) && (!Object->GetFName().IsValidXName(INVALID_OBJECTNAME_CHARACTERS) || PathName.FindChar(' ', ResultIdx)) )
 	{
 		PathName = FString::Printf(TEXT("\"%s\""), *PathName.ReplaceQuotesWithEscapedQuotes());
 	}
