Fixing errors and warnings in async loading:

- Fixed PostLoad order on loaded packages
- Fixed imports not being properly streamed due to bad FAsyncPackage setup

diff --git a/Engine/Source/Runtime/CoreUObject/Private/Serialization/AsyncLoading.cpp b/Engine/Source/Runtime/CoreUObject/Private/Serialization/AsyncLoading.cpp
--- a/Engine/Source/Runtime/CoreUObject/Private/Serialization/AsyncLoading.cpp
+++ b/Engine/Source/Runtime/CoreUObject/Private/Serialization/AsyncLoading.cpp
@@ -300,26 +300,16 @@ EAsyncPackageState::Type FAsyncLoadingThread::ProcessAsyncLoading(int32& OutPack
 			bool bPackageFullyLoaded = false;
 			if (LoadingState == EAsyncPackageState::Complete)
 			{
-				// Only remove packages that are not being referenced by anything (we still need their linkers
-				// when the referencing packages create their imports and exports).
-				if (Package->GetDependencyRefCount() == 0)
+				// We're done, at least on this thread, so we can remove the package now.
+				AddToLoadedPackages(Package);
 				{
-					// We're done so we can remove the package now.
-					AddToLoadedPackages(Package);
-					{
-						FScopeLock LockAsyncPackages(&AsyncPackagesCritical);
-						AsyncPackages.RemoveAt(PackageIndex);
-					}
-					
-					// Need to process this index again as we just removed an item
-					PackageIndex--;
-					bPackageFullyLoaded = true;
-				}
-				else
-				{
-					// This package is being referenced by another package so we can't say streaming is complete yet.
-					LoadingState = EAsyncPackageState::PendingImports;
+					FScopeLock LockAsyncPackages(&AsyncPackagesCritical);
+					AsyncPackages.RemoveAt(PackageIndex);
 				}
+					
+				// Need to process this index again as we just removed an item
+				PackageIndex--;
+				bPackageFullyLoaded = true;
 			}
 			else if (!bUseTimeLimit && !FPlatformProcess::SupportsMultithreading())
 			{
@@ -353,37 +343,43 @@ EAsyncPackageState::Type FAsyncLoadingThread::ProcessLoadedPackages(bool bUseTim
 		LoadedPackages.Empty();
 	}
 		
-	int32 PackageIndex = 0;
-	while (PackageIndex < LoadedPackagesToProcess.Num() && !IsTimeLimitExceeded(TickStartTime, bUseTimeLimit, TimeLimit))
+	for (int32 PackageIndex = 0; PackageIndex < LoadedPackagesToProcess.Num() && !IsTimeLimitExceeded(TickStartTime, bUseTimeLimit, TimeLimit); ++PackageIndex)
 	{
 		auto Package = LoadedPackagesToProcess[PackageIndex];
-		Result = Package->PostLoadDeferredObjects(TickStartTime, bUseTimeLimit, TimeLimit);
-		if (Result == EAsyncPackageState::Complete)
+		if (Package->GetDependencyRefCount() == 0)
 		{
-			if (FPlatformProperties::RequiresCookedData())
+			Result = Package->PostLoadDeferredObjects(TickStartTime, bUseTimeLimit, TimeLimit);
+			if (Result == EAsyncPackageState::Complete)
 			{
-				// Emulates ResetLoaders on the package linker's linkerroot.
-				Package->ResetLoader();
-			}
+				const bool bInternalCallbacks = false;
+				const EAsyncLoadingResult::Type LoadingResult = Package->HasLoadFailed() ? EAsyncLoadingResult::Failed : EAsyncLoadingResult::Succeeded;
+				Package->CallCompletionCallbacks(bInternalCallbacks, LoadingResult);
+
+				if (FPlatformProperties::RequiresCookedData())
+				{
+					// Emulates ResetLoaders on the package linker's linkerroot.
+					Package->ResetLoader();
+				}
 
-			const bool bInternalCallbacks = false;
-			const EAsyncLoadingResult::Type LoadingResult = Package->HasLoadFailed() ? EAsyncLoadingResult::Failed : EAsyncLoadingResult::Succeeded;
-			Package->CallCompletionCallbacks(bInternalCallbacks, LoadingResult);
-			PackageIndex++;
+				// Incremented on the Async Thread, decremented on the Game Thread				
+				AsyncLoadingCounter.Decrement();
+				check(AsyncLoadingCounter.GetValue() >= 0);
 
-			// Incremented on the Async Thread, decremented on the Game Thread				
-			AsyncLoadingCounter.Decrement();
-			check(AsyncLoadingCounter.GetValue() >= 0);
+				delete Package;
+				LoadedPackagesToProcess.RemoveAt(PackageIndex--);
+			}
+			else
+			{
+				break;
+			}
 		}
 		else
 		{
+			Result = EAsyncPackageState::PendingImports;
+			// Break immediately, we want to keep the order of processing when packages get here
 			break;
 		}
 	}
-	if (PackageIndex > 0)
-	{
-		LoadedPackagesToProcess.RemoveAt(0, PackageIndex);
-	}
 
 	return Result;
 }
@@ -952,7 +948,7 @@ void FAsyncPackage::AddImportDependency(int32 CurrentPackageIndex, const FName&
 	int32 ExistingAsyncPackageIndex = FAsyncLoadingThread::Get().FindAsyncPackage(PendingImport);
 	if (ExistingAsyncPackageIndex == INDEX_NONE)
 	{
-		const FAsyncPackageDesc Info(PendingImport, FGuid(), PendingImport);
+		const FAsyncPackageDesc Info(PendingImport);
 		PackageToStream = new FAsyncPackage(Info);
 
 		// If priority of the dependency is not set, inherit from parent.
@@ -972,12 +968,12 @@ void FAsyncPackage::AddImportDependency(int32 CurrentPackageIndex, const FName&
 	{
 		const bool bInternalCallback = true;
 		PackageToStream->AddCompletionCallback(FLoadPackageAsyncDelegate::CreateRaw(this, &FAsyncPackage::ImportFullyLoadedCallback), bInternalCallback);
-		PackageToStream->DependencyRefCount++;
+		PackageToStream->DependencyRefCount.Increment();
 		PendingImportedPackages.Add(PackageToStream);
 	}
 	else
 	{
-		PackageToStream->DependencyRefCount++;
+		PackageToStream->DependencyRefCount.Increment();
 		ReferencedImports.Add(PackageToStream);
 	}
 }
@@ -1086,7 +1082,7 @@ EAsyncPackageState::Type FAsyncPackage::LoadImports()
 				{
 					UE_LOG(LogStreaming, Verbose, TEXT("FAsyncPackage::LoadImports for %s: Linker exists for %s"), *Desc.NameToLoad.ToString(), *ImportPackageFName.ToString());
 					// Only keep a reference to this package so that its linker doesn't go away too soon
-					PendingPackage.DependencyRefCount++;
+					PendingPackage.DependencyRefCount.Increment();
 					ReferencedImports.Add(&PendingPackage);
 					// Check if we need to add its dependencies too.
 					TSet<FAsyncPackage*> SearchedPackages;
@@ -1257,9 +1253,9 @@ void FAsyncPackage::FreeReferencedImports()
 	for (int32 ReferenceIndex = 0; ReferenceIndex < ReferencedImports.Num(); ++ReferenceIndex)
 	{
 		FAsyncPackage& Ref = *ReferencedImports[ReferenceIndex];
-		Ref.DependencyRefCount--;
-		UE_LOG(LogStreaming, Verbose, TEXT("FAsyncPackage::FreeReferencedImports for %s: Releasing %s (%d)"), *Desc.NameToLoad.ToString(), *Ref.GetPackageName().ToString(), Ref.DependencyRefCount);
-		check(Ref.DependencyRefCount >= 0);
+		Ref.DependencyRefCount.Decrement();
+		UE_LOG(LogStreaming, Verbose, TEXT("FAsyncPackage::FreeReferencedImports for %s: Releasing %s (%d)"), *Desc.NameToLoad.ToString(), *Ref.GetPackageName().ToString(), Ref.GetDependencyRefCount());
+		check(Ref.DependencyRefCount.GetValue() >= 0);
 	}
 	ReferencedImports.Empty();
 }
@@ -1543,7 +1539,7 @@ void LoadPackageAsync(const FString& InName, const FGuid* InGuid /*= nullptr*/,
 		UE_LOG(LogStreaming, Fatal, TEXT("Async loading code requires long package names (%s)."), *PackageNameToLoad);
 	}
 
-	FAsyncPackageDesc PackageDesc(*PackageName, InGuid ? *InGuid : FGuid(), InType, *PackageNameToLoad, InCompletionDelegate, InFlags, InPIEInstanceID, InPackagePriority);
+	FAsyncPackageDesc PackageDesc(*PackageName, *PackageNameToLoad, InGuid ? *InGuid : FGuid(), InType, InCompletionDelegate, InFlags, InPIEInstanceID, InPackagePriority);
 	FAsyncLoadingThread::Get().QueuePackage(PackageDesc);
 }
 
diff --git a/Engine/Source/Runtime/CoreUObject/Public/Serialization/AsyncLoading.h b/Engine/Source/Runtime/CoreUObject/Public/Serialization/AsyncLoading.h
--- a/Engine/Source/Runtime/CoreUObject/Public/Serialization/AsyncLoading.h
+++ b/Engine/Source/Runtime/CoreUObject/Public/Serialization/AsyncLoading.h
@@ -70,7 +70,7 @@ struct FAsyncPackage
 	/** Gets the number of references to this package from other packages in the dependency tree. */
 	FORCEINLINE int32 GetDependencyRefCount() const
 	{
-		return DependencyRefCount;
+		return DependencyRefCount.GetValue();
 	}
 
 	/** Returns true if the package has finished loading. */
@@ -124,7 +124,7 @@ private:
 	/** Referenced imports - list of packages we need until we finish loading this package. */
 	TArray<FAsyncPackage*> ReferencedImports;
 	/** Number of references to this package from other packages in the dependency tree. */
-	int32							DependencyRefCount;
+	FThreadSafeCounter	DependencyRefCount;
 	/** Current index into linkers import table used to spread creation over several frames				*/
 	int32							LoadImportIndex;
 	/** Current index into linkers import table used to spread creation over several frames				*/
diff --git a/Engine/Source/Runtime/CoreUObject/Public/Serialization/AsyncPackage.h b/Engine/Source/Runtime/CoreUObject/Public/Serialization/AsyncPackage.h
--- a/Engine/Source/Runtime/CoreUObject/Public/Serialization/AsyncPackage.h
+++ b/Engine/Source/Runtime/CoreUObject/Public/Serialization/AsyncPackage.h
@@ -10,10 +10,10 @@ struct FAsyncPackageDesc
 {
 	/** Name of the UPackage to create. */
 	FName Name;
-	/** An abstract type name associated with this package, for tagging use	*/
-	FName Type;
 	/** Name of the package to load. */
 	FName NameToLoad;
+	/** An abstract type name associated with this package, for tagging use	*/
+	FName Type;
 	/** GUID of the package to load, or the zeroed invalid GUID for "don't care" */
 	FGuid Guid;
 	/** Delegate called on completion of loading */
@@ -28,10 +28,10 @@ struct FAsyncPackageDesc
 #endif
 
 
-	FAsyncPackageDesc(const FName& InName, const FGuid& InGuid = FGuid(), FName InType = NAME_None, FName InPackageToLoadFrom = NAME_None, FLoadPackageAsyncDelegate InCompletionDelegate = FLoadPackageAsyncDelegate(), uint32 InFlags = 0, int32 InPIEInstanceID = INDEX_NONE, uint32 InPriority = 0)
+	FAsyncPackageDesc(const FName& InName, FName InPackageToLoadFrom = NAME_None, const FGuid& InGuid = FGuid(), FName InType = NAME_None, FLoadPackageAsyncDelegate InCompletionDelegate = FLoadPackageAsyncDelegate(), uint32 InFlags = 0, int32 InPIEInstanceID = INDEX_NONE, uint32 InPriority = 0)
 		: Name(InName)
-		, Type(InType)
 		, NameToLoad(InPackageToLoadFrom)
+		, Type(InType)	
 		, Guid(InGuid)
 		, PackageLoadedDelegate(InCompletionDelegate)
 		, Flags(InFlags)
@@ -40,5 +40,9 @@ struct FAsyncPackageDesc
 		, PIEInstanceID(InPIEInstanceID)
 #endif
 	{
+		if (NameToLoad == NAME_None)
+		{
+			NameToLoad = Name;
+		}
 	}
 };
