Fix deprecation warnings in CIS on Mac.

diff --git a/Engine/Source/Runtime/CoreUObject/Private/Serialization/AsyncLoading.cpp b/Engine/Source/Runtime/CoreUObject/Private/Serialization/AsyncLoading.cpp
--- a/Engine/Source/Runtime/CoreUObject/Private/Serialization/AsyncLoading.cpp
+++ b/Engine/Source/Runtime/CoreUObject/Private/Serialization/AsyncLoading.cpp
@@ -423,8 +423,8 @@ EAsyncPackageState::Type FAsyncLoadingThread::TickAsyncLoading(bool bUseTimeLimi
 
 FAsyncLoadingThread::FAsyncLoadingThread()
 {
-	QueuedRequestsEvent = FPlatformProcess::CreateSynchEvent();
-	CancelLoadingEvent = FPlatformProcess::CreateSynchEvent();
+	QueuedRequestsEvent = FPlatformProcess::GetSynchEventFromPool();
+	CancelLoadingEvent = FPlatformProcess::GetSynchEventFromPool();
 	if (FAsyncLoadingThread::IsMultithreaded())
 	{
 		Thread = FRunnableThread::Create(this, TEXT("FAsyncLoadingThread"), 0, TPri_Normal);
@@ -441,9 +441,9 @@ FAsyncLoadingThread::~FAsyncLoadingThread()
 {
 	delete Thread;
 	Thread = nullptr;
-	delete QueuedRequestsEvent;
+	FPlatformProcess::ReturnSynchEventToPool(QueuedRequestsEvent);
 	QueuedRequestsEvent = nullptr;
-	delete CancelLoadingEvent;
+	FPlatformProcess::ReturnSynchEventToPool(CancelLoadingEvent);
 	CancelLoadingEvent = nullptr;
 }
 
