SafeLoadError will no longer take format string and behave like printf,  remove unnecessary and unused parameter.

diff --git a/Engine/Source/Runtime/CoreUObject/Private/UObject/Linker.cpp b/Engine/Source/Runtime/CoreUObject/Private/UObject/Linker.cpp
--- a/Engine/Source/Runtime/CoreUObject/Private/UObject/Linker.cpp
+++ b/Engine/Source/Runtime/CoreUObject/Private/UObject/Linker.cpp
@@ -515,7 +515,7 @@ static void LogGetPackageLinkerError(FArchiveUObject* LinkerArchive, const TCHAR
 		else
 		{
 			// Gracefully handle missing packages
-			SafeLoadError( InOuter, LoadFlags, *InFullErrorMessage.ToString(), *Error.ToString() );
+			SafeLoadError(InOuter, LoadFlags, *Error.ToString());
 		}
 	}
 }
diff --git a/Engine/Source/Runtime/CoreUObject/Private/UObject/UObjectGlobals.cpp b/Engine/Source/Runtime/CoreUObject/Private/UObject/UObjectGlobals.cpp
--- a/Engine/Source/Runtime/CoreUObject/Private/UObject/UObjectGlobals.cpp
+++ b/Engine/Source/Runtime/CoreUObject/Private/UObject/UObjectGlobals.cpp
@@ -401,22 +401,18 @@ void StaticTick( float DeltaTime, bool bUseFullTimeLimit, float AsyncLoadingTime
 //
 // Safe load error-handling.
 //
-void SafeLoadError( UObject* Outer, uint32 LoadFlags, const TCHAR* Error, const TCHAR* Fmt, ... )
+void SafeLoadError( UObject* Outer, uint32 LoadFlags, const TCHAR* ErrorMessage)
 {
-	// Variable arguments setup.
-	TCHAR TempStr[4096];
-	GET_VARARGS( TempStr, ARRAY_COUNT(TempStr), ARRAY_COUNT(TempStr)-1, Fmt, Fmt );
-
 	if( FParse::Param( FCommandLine::Get(), TEXT("TREATLOADWARNINGSASERRORS") ) == true )
 	{
-		UE_LOG(LogUObjectGlobals, Error, TEXT("%s"), TempStr ); 
+		UE_LOG(LogUObjectGlobals, Error, TEXT("%s"), ErrorMessage);
 	}
 	else
 	{
 		// Don't warn here if either quiet or no warn are set
 		if( (LoadFlags & LOAD_Quiet) == 0 && (LoadFlags & LOAD_NoWarn) == 0)
 		{ 
-			UE_LOG(LogUObjectGlobals, Warning, TEXT("%s"), TempStr ); 
+			UE_LOG(LogUObjectGlobals, Warning, TEXT("%s"), ErrorMessage);
 		}
 	}
 }
@@ -782,7 +778,7 @@ UObject* StaticLoadObject(UClass* ObjectClass, UObject* InOuter, const TCHAR* In
 		Arguments.Add(TEXT("OuterName"), InOuter ? FText::FromString(InOuter->GetPathName()) : NSLOCTEXT("Core", "None", "None"));
 		Arguments.Add(TEXT("ObjectName"), FText::FromString(ObjectName));
 		const FString Error = FText::Format(NSLOCTEXT("Core", "ObjectNotFound", "Failed to find object '{ClassName} {OuterName}.{ObjectName}'"), Arguments).ToString();
-		SafeLoadError(InOuter, LoadFlags, *Error, *Error);
+		SafeLoadError(InOuter, LoadFlags, *Error);
 	}
 	return Result;
 }
@@ -801,7 +797,7 @@ UClass* StaticLoadClass( UClass* BaseClass, UObject* InOuter, const TCHAR* InNam
 		Arguments.Add(TEXT("ClassName"), FText::FromString( Class->GetFullName() ));
 		Arguments.Add(TEXT("BaseClassName"), FText::FromString( BaseClass->GetFullName() ));
 		const FString Error = FText::Format( NSLOCTEXT( "Core", "LoadClassMismatch", "{ClassName} is not a child class of {BaseClassName}" ), Arguments ).ToString();
-		SafeLoadError(InOuter, LoadFlags,*Error, *Error);
+		SafeLoadError(InOuter, LoadFlags, *Error);
 
 		// return NULL class due to error
 		Class = NULL;
diff --git a/Engine/Source/Runtime/CoreUObject/Public/UObject/UObjectGlobals.h b/Engine/Source/Runtime/CoreUObject/Public/UObject/UObjectGlobals.h
--- a/Engine/Source/Runtime/CoreUObject/Public/UObject/UObjectGlobals.h
+++ b/Engine/Source/Runtime/CoreUObject/Public/UObject/UObjectGlobals.h
@@ -121,7 +121,7 @@ COREUOBJECT_API TArray<const TCHAR*> ParsePropertyFlags(uint64 Flags);
 COREUOBJECT_API UPackage* GetTransientPackage();
 
 COREUOBJECT_API bool ResolveName( UObject*& Outer, FString& Name, bool Create, bool Throw );
-COREUOBJECT_API void SafeLoadError( UObject* Outer, uint32 LoadFlags, const TCHAR* Error, const TCHAR* Fmt, ... );
+COREUOBJECT_API void SafeLoadError( UObject* Outer, uint32 LoadFlags, const TCHAR* ErrorMessage);
 
 /**
  * Fast version of StaticFindObject that relies on the passed in FName being the object name
