Minor optimization in FArchiveReplaceObjectRef.

diff --git a/Engine/Source/Runtime/CoreUObject/Public/Serialization/ArchiveUObject.h b/Engine/Source/Runtime/CoreUObject/Public/Serialization/ArchiveUObject.h
--- a/Engine/Source/Runtime/CoreUObject/Public/Serialization/ArchiveUObject.h
+++ b/Engine/Source/Runtime/CoreUObject/Public/Serialization/ArchiveUObject.h
@@ -1143,11 +1143,11 @@ public:
 					UE_LOG(LogSerialization, Log,  TEXT("FArchiveReplaceObjectRef: Obj == SearchObject : '%s'"), *ObjName );
 				}
 #endif
-
-				if ( !SerializedObjects.Find(Obj) )
+				bool bAlreadyAdded = false;
+				SerializedObjects.Add(Obj, &bAlreadyAdded);
+				if (!bAlreadyAdded)
 				{
 					// otherwise recurse down into the object if it is contained within the initial search object
-					SerializedObjects.Add(Obj);
 	
 					// serialization for class default objects must be deterministic (since class 
 					// default objects may be serialized during script compilation while the script
