#UE4 Change from Mike.Beach to stabilize the build - [UE-13271] Accounting for rediectors that are loaded to replace deferred dependency placeholders.

diff --git a/Engine/Source/Runtime/CoreUObject/Private/UObject/LinkerPlaceholderBase.inl b/Engine/Source/Runtime/CoreUObject/Private/UObject/LinkerPlaceholderBase.inl
--- a/Engine/Source/Runtime/CoreUObject/Private/UObject/LinkerPlaceholderBase.inl
+++ b/Engine/Source/Runtime/CoreUObject/Private/UObject/LinkerPlaceholderBase.inl
@@ -10,6 +10,32 @@
 template<class PlaceholderType>
 int32 TLinkerImportPlaceholder<PlaceholderType>::ResolveAllPlaceholderReferences(UObject* ReplacementObj)
 {
+	if (UObjectRedirector* ReplacementRedirector = Cast<UObjectRedirector>(ReplacementObj))
+	{
+		if (FLinkerLoad* ReplacementLinker = ReplacementRedirector->GetLinker())
+		{
+			if (!ReplacementRedirector->HasAnyFlags(RF_LoadCompleted))
+			{
+				// we're in the midst of serializing this redirector
+				// somewhere up the stack, in some scenario like this:
+				//
+				//		- ClassA and ClassC both depend on ClassB
+				//		- ClassB has a redirector to ClassB_2
+				//		- ClassB_2 depends on ClassC
+				//
+				// if ClassA is loaded first, it then goes to load ClassB, which
+				// seeks to serialize in its UObjectRedirector; before that's 
+				// set it loads ClassB_2 and subsequently ClassC; ClassC ends up
+				// here, needing to use the ClassB redirector, but we haven't 
+				// returned up the stack for it to be set yet... here we force 
+				// it to finish preloading (like we do in VerifyImport): 
+				ReplacementRedirector->SetFlags(RF_NeedLoad);
+				ReplacementLinker->Preload(ReplacementRedirector);
+			}
+		}
+		ReplacementObj = ReplacementRedirector->DestinationObject;
+	}
+
 	PlaceholderType* TypeCheckedReplacement = CastChecked<PlaceholderType>(ReplacementObj, ECastCheckedType::NullAllowed);
 
 	int32 ReplacementCount = ResolvePropertyReferences(TypeCheckedReplacement);
