Fix per poly skeletal mesh to not mark its cooked data as single use. UE-13215
Fixes cooking crash in qa-game

diff --git a/Engine/Source/Runtime/CoreUObject/Private/Serialization/BulkData.cpp b/Engine/Source/Runtime/CoreUObject/Private/Serialization/BulkData.cpp
--- a/Engine/Source/Runtime/CoreUObject/Private/Serialization/BulkData.cpp
+++ b/Engine/Source/Runtime/CoreUObject/Private/Serialization/BulkData.cpp
@@ -1427,7 +1427,7 @@ void FFloatBulkData::SerializeElement( FArchive& Ar, void* Data, int32 ElementIn
 }
 
 
-void FFormatContainer::Serialize(FArchive& Ar, UObject* Owner, const TArray<FName>* FormatsToSave)
+void FFormatContainer::Serialize(FArchive& Ar, UObject* Owner, const TArray<FName>* FormatsToSave, bool bSingleUse)
 {
 	if (Ar.IsLoading())
 	{
@@ -1463,7 +1463,7 @@ void FFormatContainer::Serialize(FArchive& Ar, UObject* Owner, const TArray<FNam
 				Ar << Name;
 				FByteBulkData* Bulk = It.Value();
 				// Force this kind of bulk data (physics, etc) to be stored inline for streaming
-				Bulk->SetBulkDataFlags(BULKDATA_ForceInlinePayload | BULKDATA_SingleUse);
+				Bulk->SetBulkDataFlags(bSingleUse ? (BULKDATA_ForceInlinePayload | BULKDATA_SingleUse) : BULKDATA_ForceInlinePayload);
 				check(Bulk);
 				Bulk->Serialize(Ar, Owner);
 			}
diff --git a/Engine/Source/Runtime/CoreUObject/Public/Serialization/BulkData.h b/Engine/Source/Runtime/CoreUObject/Public/Serialization/BulkData.h
--- a/Engine/Source/Runtime/CoreUObject/Public/Serialization/BulkData.h
+++ b/Engine/Source/Runtime/CoreUObject/Public/Serialization/BulkData.h
@@ -558,7 +558,7 @@ public:
 		}
 		Formats.Empty();
 	}
-	COREUOBJECT_API void Serialize(FArchive& Ar, UObject* Owner, const TArray<FName>* FormatsToSave = NULL);
+	COREUOBJECT_API void Serialize(FArchive& Ar, UObject* Owner, const TArray<FName>* FormatsToSave = NULL, bool bSingleUse = true);
 };
 
 #endif
diff --git a/Engine/Source/Runtime/Engine/Private/PhysicsEngine/BodySetup.cpp b/Engine/Source/Runtime/Engine/Private/PhysicsEngine/BodySetup.cpp
--- a/Engine/Source/Runtime/Engine/Private/PhysicsEngine/BodySetup.cpp
+++ b/Engine/Source/Runtime/Engine/Private/PhysicsEngine/BodySetup.cpp
@@ -842,7 +842,7 @@ void UBodySetup::Serialize(FArchive& Ar)
 			ActualFormatsToSave.Add(Format);
 
 			Ar << bHasCookedCollisionData;
-			CookedFormatData.Serialize(Ar, this, &ActualFormatsToSave);
+			CookedFormatData.Serialize(Ar, this, &ActualFormatsToSave, !bSharedCookedData);
 		}
 		else
 		{
