AsyncLoading stats improvements (reducing stats spam and hitches):

- removed counters so that async loading stats group can be verbose again
- added a new async loading stats group for game-thread-only stats tracking
- made IsTimeLimitExceeded more aggresive with warnings when the limit is exceeded (reduced the warning threshold from 2.5 to 1.5 * limit)
- made LinkerLoad stats verbose again
- IO blocking will now be reported if it happens on game thread only

diff --git a/Engine/Source/Runtime/CoreUObject/Private/Serialization/AsyncLoading.cpp b/Engine/Source/Runtime/CoreUObject/Private/Serialization/AsyncLoading.cpp
--- a/Engine/Source/Runtime/CoreUObject/Private/Serialization/AsyncLoading.cpp
+++ b/Engine/Source/Runtime/CoreUObject/Private/Serialization/AsyncLoading.cpp
@@ -17,7 +17,7 @@
 
 DECLARE_MEMORY_STAT(TEXT("Streaming Memory Used"),STAT_StreamingAllocSize,STATGROUP_Memory);
 
-DECLARE_STATS_GROUP(TEXT("Async Load"), STATGROUP_AsyncLoad, STATCAT_Advanced);
+DECLARE_STATS_GROUP_VERBOSE(TEXT("Async Load"), STATGROUP_AsyncLoad, STATCAT_Advanced);
 
 DECLARE_CYCLE_STAT(TEXT("Tick AsyncPackage"),STAT_FAsyncPackage_Tick,STATGROUP_AsyncLoad);
 DECLARE_FLOAT_ACCUMULATOR_STAT(TEXT("Tick AsyncPackage Time"), STAT_FAsyncPackage_TickTime, STATGROUP_AsyncLoad);
@@ -29,18 +29,19 @@ DECLARE_CYCLE_STAT(TEXT("FinishTextureAllocations AsyncPackage"),STAT_FAsyncPack
 DECLARE_CYCLE_STAT(TEXT("CreateExports AsyncPackage"),STAT_FAsyncPackage_CreateExports,STATGROUP_AsyncLoad);
 DECLARE_CYCLE_STAT(TEXT("PreLoadObjects AsyncPackage"),STAT_FAsyncPackage_PreLoadObjects,STATGROUP_AsyncLoad);
 DECLARE_CYCLE_STAT(TEXT("PostLoadObjects AsyncPackage"),STAT_FAsyncPackage_PostLoadObjects,STATGROUP_AsyncLoad);
-DECLARE_FLOAT_ACCUMULATOR_STAT(TEXT("PostLoadObjects AsyncPackage Time"), STAT_FAsyncPackage_PostLoadObjectsTime, STATGROUP_AsyncLoad);
-DECLARE_CYCLE_STAT(TEXT("PostLoadObjects GameThread"), STAT_FAsyncPackage_PostLoadObjectsGameThread, STATGROUP_AsyncLoad);
-DECLARE_FLOAT_ACCUMULATOR_STAT(TEXT("PostLoadObjects GameThread Time"), STAT_FAsyncPackage_PostLoadObjectsGameThreadTime, STATGROUP_AsyncLoad);
 DECLARE_CYCLE_STAT(TEXT("FinishObjects AsyncPackage"),STAT_FAsyncPackage_FinishObjects,STATGROUP_AsyncLoad);
-DECLARE_FLOAT_ACCUMULATOR_STAT(TEXT("CreateAsyncPackagesFromQueue Time"), STAT_FAsyncPackage_CreateAsyncPackagesFromQueueTime, STATGROUP_AsyncLoad);
-DECLARE_FLOAT_ACCUMULATOR_STAT(TEXT("AsyncProcess Time"), STAT_FAsyncPackage_AsyncProcessTime, STATGROUP_AsyncLoad);
-DECLARE_FLOAT_ACCUMULATOR_STAT(TEXT("Flush Async Loading Time"), STAT_FAsyncPackage_FlushAsyncLoadingTime, STATGROUP_AsyncLoad);
-
+DECLARE_CYCLE_STAT(TEXT("CreateAsyncPackagesFromQueue"), STAT_FAsyncPackage_CreateAsyncPackagesFromQueue, STATGROUP_AsyncLoad);
+DECLARE_CYCLE_STAT(TEXT("ProcessAsyncLoading AsyncLoadingThread"), STAT_FAsyncLoadingThread_ProcessAsyncLoading, STATGROUP_AsyncLoad);
 DECLARE_CYCLE_STAT(TEXT("Async Loading Time"),STAT_AsyncLoadingTime,STATGROUP_AsyncLoad);
 DECLARE_CYCLE_STAT(TEXT("Async Loading Time Detailed"), STAT_AsyncLoadingTimeDetailed, STATGROUP_AsyncLoad);
 
-DECLARE_FLOAT_ACCUMULATOR_STAT(TEXT("FindAsyncPackage"), STAT_FAsyncPackage_FindAsyncPackage, STATGROUP_AsyncLoad);
+
+DECLARE_STATS_GROUP(TEXT("Async Load Game Thread"), STATGROUP_AsyncLoadGameThread, STATCAT_Advanced);
+
+DECLARE_CYCLE_STAT(TEXT("PostLoadObjects GT"), STAT_FAsyncPackage_PostLoadObjectsGameThread, STATGROUP_AsyncLoad);
+DECLARE_CYCLE_STAT(TEXT("TickAsyncLoading GT"), STAT_FAsyncPackage_TickAsyncLoadingGameThread, STATGROUP_AsyncLoadGameThread);
+DECLARE_CYCLE_STAT(TEXT("Flush Async Loading GT"), STAT_FAsyncPackage_FlushAsyncLoadingGameThread, STATGROUP_AsyncLoad);
+
 
 /** Keeps a reference to all objects created during async load until streaming has finished */
 class FAsyncObjectsReferencer : FGCObject
@@ -119,9 +120,9 @@ static FORCEINLINE bool IsTimeLimitExceeded(double InTickStartTime, bool bUseTim
 		bTimeLimitExceeded = CurrentTime - InTickStartTime > InTimeLimit;
 
 		// Log single operations that take longer than timelimit.
-		if ((CurrentTime - InTickStartTime) > (2.5 * InTimeLimit))
+		if ((CurrentTime - InTickStartTime) > (1.5 * InTimeLimit))
 		{
-			UE_LOG(LogStreaming, Log, TEXT("FAsyncPackage: %s %s took (less than) %5.2f ms"),
+			UE_LOG(LogStreaming, Warning, TEXT("IsTimeLimitExceeded: %s %s took (less than) %5.2f ms"),
 				InLastTypeOfWorkPerformed ? InLastTypeOfWorkPerformed : TEXT("unknown"),
 				InLastObjectWorkWasPerformedOn ? *InLastObjectWorkWasPerformedOn->GetFullName() : TEXT("nullptr"),
 				(CurrentTime - InTickStartTime) * 1000);
@@ -201,60 +202,55 @@ void FAsyncLoadingThread::QueuePackage(const FAsyncPackageDesc& Package)
 
 int32 FAsyncLoadingThread::CreateAsyncPackagesFromQueue()
 {
+	SCOPE_CYCLE_COUNTER(STAT_FAsyncPackage_CreateAsyncPackagesFromQueue);
+
 	int32 NumCreated = 0;
+	checkSlow(IsInAsyncLoadThread());
 
-	STAT(double ThisTime = 0);
+	TArray<FAsyncPackageDesc*> QueueCopy;
 	{
-		SCOPE_SECONDS_COUNTER(ThisTime);
-
-		checkSlow(IsInAsyncLoadThread());
-
-		TArray<FAsyncPackageDesc*> QueueCopy;
+		FScopeLock QueueLock(&QueueCritical);
+		QueueCopy = QueuedPackages;
+		QueuedPackages.Empty();
+	}
+	for (auto PackageRequest : QueueCopy)
+	{
+		FAsyncPackage* Package = nullptr;
+		int ExistingPackageIndex = FindAsyncPackage(PackageRequest->Name);
+		if (ExistingPackageIndex != INDEX_NONE)
 		{
-			FScopeLock QueueLock(&QueueCritical);
-			QueueCopy = QueuedPackages;
-			QueuedPackages.Empty();
+			// This is a package that is currently being loaded.
+			Package = AsyncPackages[ExistingPackageIndex];
 		}
-		for (auto PackageRequest : QueueCopy)
+		else
 		{
-			FAsyncPackage* Package = nullptr;
-			int ExistingPackageIndex = FindAsyncPackage(PackageRequest->Name);
-			if (ExistingPackageIndex != INDEX_NONE)
-			{
-				// This is a package that is currently being loaded.
-				Package = AsyncPackages[ExistingPackageIndex];
-			}
-			else
-			{
-				// [BLOCKING] LoadedPackages are accessed on the main thread too, so lock to be able to add a completion callback
-				FScopeLock LoadedLock(&LoadedPackagesCritical);
-				ExistingPackageIndex = FindLoadedPackage(PackageRequest->Name);
-				if (ExistingPackageIndex != INDEX_NONE && PackageRequest->PackageLoadedDelegate.IsBound())
-				{
-					const bool bInternalCallback = false;
-					LoadedPackages[ExistingPackageIndex]->AddCompletionCallback(PackageRequest->PackageLoadedDelegate, bInternalCallback);
-					continue;
-				}
-			}
-			if (!Package)
-			{
-				// New package that needs to be loaded or a package has already been loaded long time ago
-				Package = new FAsyncPackage(*PackageRequest);
-
-				// Add to queue according to priority.
-				InsertPackage(Package);
-			}
-			if (PackageRequest->PackageLoadedDelegate.IsBound())
+			// [BLOCKING] LoadedPackages are accessed on the main thread too, so lock to be able to add a completion callback
+			FScopeLock LoadedLock(&LoadedPackagesCritical);
+			ExistingPackageIndex = FindLoadedPackage(PackageRequest->Name);
+			if (ExistingPackageIndex != INDEX_NONE && PackageRequest->PackageLoadedDelegate.IsBound())
 			{
 				const bool bInternalCallback = false;
-				Package->AddCompletionCallback(PackageRequest->PackageLoadedDelegate, bInternalCallback);
+				LoadedPackages[ExistingPackageIndex]->AddCompletionCallback(PackageRequest->PackageLoadedDelegate, bInternalCallback);
+				continue;
 			}
-			QueuedPackagesCounter.Decrement();
-			check(QueuedPackagesCounter.GetValue() >= 0);
 		}
-		NumCreated = QueueCopy.Num();
+		if (!Package)
+		{
+			// New package that needs to be loaded or a package has already been loaded long time ago
+			Package = new FAsyncPackage(*PackageRequest);
+
+			// Add to queue according to priority.
+			InsertPackage(Package);
+		}
+		if (PackageRequest->PackageLoadedDelegate.IsBound())
+		{
+			const bool bInternalCallback = false;
+			Package->AddCompletionCallback(PackageRequest->PackageLoadedDelegate, bInternalCallback);
+		}
+		QueuedPackagesCounter.Decrement();
+		check(QueuedPackagesCounter.GetValue() >= 0);
 	}
-	INC_FLOAT_STAT_BY(STAT_FAsyncPackage_CreateAsyncPackagesFromQueueTime, (float)ThisTime);
+	NumCreated = QueueCopy.Num();
 
 	return NumCreated;
 }
@@ -264,69 +260,66 @@ int32 FAsyncLoadingThread::CreateAsyncPackagesFromQueue()
 */
 EAsyncPackageState::Type FAsyncLoadingThread::ProcessAsyncLoading(int32& OutPackagesProcessed, bool bUseTimeLimit /*= false*/, bool bUseFullTimeLimit /*= false*/, float TimeLimit /*= 0.0f*/, FName ExcludeType /*= NAME_None*/)
 {
+	SCOPE_CYCLE_COUNTER(STAT_FAsyncLoadingThread_ProcessAsyncLoading);
+	
 	EAsyncPackageState::Type LoadingState = EAsyncPackageState::Complete;
 	OutPackagesProcessed = 0;
-	STAT(double ThisTime = 0);
-	{
-		SCOPE_SECONDS_COUNTER(ThisTime);
 
-		// We need to loop as the function has to handle finish loading everything given no time limit
-		// like e.g. when called from FlushAsyncLoading.
-		for (int32 PackageIndex = 0; LoadingState != EAsyncPackageState::TimeOut && PackageIndex < AsyncPackages.Num(); ++PackageIndex)
-		{
-			OutPackagesProcessed++;
+	// We need to loop as the function has to handle finish loading everything given no time limit
+	// like e.g. when called from FlushAsyncLoading.
+	for (int32 PackageIndex = 0; LoadingState != EAsyncPackageState::TimeOut && PackageIndex < AsyncPackages.Num(); ++PackageIndex)
+	{
+		OutPackagesProcessed++;
 
-			// Package to be loaded.
-			FAsyncPackage* Package = AsyncPackages[PackageIndex];
+		// Package to be loaded.
+		FAsyncPackage* Package = AsyncPackages[PackageIndex];
 
-			if (ExcludeType != NAME_None && ExcludeType == Package->GetPackageType())
-			{
-				// We should skip packages of this type
-				continue;
-			}
+		if (ExcludeType != NAME_None && ExcludeType == Package->GetPackageType())
+		{
+			// We should skip packages of this type
+			continue;
+		}
 
-			if (Package->HasFinishedLoading() == false)
-			{
-				// Package tick returns EAsyncPackageState::Complete on completion.
-				// We only tick packages that have not yet been loaded.
-				LoadingState = Package->Tick(bUseTimeLimit, bUseFullTimeLimit, TimeLimit);
-			}
-			else
+		if (Package->HasFinishedLoading() == false)
+		{
+			// Package tick returns EAsyncPackageState::Complete on completion.
+			// We only tick packages that have not yet been loaded.
+			LoadingState = Package->Tick(bUseTimeLimit, bUseFullTimeLimit, TimeLimit);
+		}
+		else
+		{
+			// This package has finished loading but some other package is still holding
+			// a reference to it because it has this package in its dependency list.
+			LoadingState = EAsyncPackageState::Complete;
+		}
+		bool bPackageFullyLoaded = false;
+		if (LoadingState == EAsyncPackageState::Complete)
+		{
+			// We're done, at least on this thread, so we can remove the package now.
+			AddToLoadedPackages(Package);
 			{
-				// This package has finished loading but some other package is still holding
-				// a reference to it because it has this package in its dependency list.
-				LoadingState = EAsyncPackageState::Complete;
+				FScopeLock LockAsyncPackages(&AsyncPackagesCritical);
+				AsyncPackages.RemoveAt(PackageIndex);
 			}
-			bool bPackageFullyLoaded = false;
-			if (LoadingState == EAsyncPackageState::Complete)
-			{
-				// We're done, at least on this thread, so we can remove the package now.
-				AddToLoadedPackages(Package);
-				{
-					FScopeLock LockAsyncPackages(&AsyncPackagesCritical);
-					AsyncPackages.RemoveAt(PackageIndex);
-				}
 					
-				// Need to process this index again as we just removed an item
-				PackageIndex--;
-				bPackageFullyLoaded = true;
-			}
-			else if (!bUseTimeLimit && !FPlatformProcess::SupportsMultithreading())
-			{
-				// Tick async loading when multithreading is disabled.
-				FIOSystem::Get().TickSingleThreaded();
-			}
+			// Need to process this index again as we just removed an item
+			PackageIndex--;
+			bPackageFullyLoaded = true;
+		}
+		else if (!bUseTimeLimit && !FPlatformProcess::SupportsMultithreading())
+		{
+			// Tick async loading when multithreading is disabled.
+			FIOSystem::Get().TickSingleThreaded();
+		}
 
-			// Check if there's any new packages in the queue.
-			CreateAsyncPackagesFromQueue();
+		// Check if there's any new packages in the queue.
+		CreateAsyncPackagesFromQueue();
 
-			if (bPackageFullyLoaded)
-			{
-				AsyncPackagesCounter.Decrement();
-			}
+		if (bPackageFullyLoaded)
+		{
+			AsyncPackagesCounter.Decrement();
 		}
 	}
-	INC_FLOAT_STAT_BY(STAT_FAsyncPackage_AsyncProcessTime, (float)ThisTime);
 
 	return LoadingState;
 }
@@ -388,6 +381,8 @@ EAsyncPackageState::Type FAsyncLoadingThread::ProcessLoadedPackages(bool bUseTim
 
 EAsyncPackageState::Type FAsyncLoadingThread::TickAsyncLoading(bool bUseTimeLimit, bool bUseFullTimeLimit, float TimeLimit, FName ExcludeType)
 {	
+	SCOPE_CYCLE_COUNTER(STAT_FAsyncPackage_TickAsyncLoadingGameThread);
+
 	TGuardValue<bool> InAsyncLoadingTickGuard(bIsInAsyncLoadingTick, true);
 
 	static bool bWasAsyncLoading = false;
@@ -724,109 +719,103 @@ EAsyncPackageState::Type FAsyncPackage::Tick(bool InbUseTimeLimit, bool InbUseFu
 	// Whether we should execute the next step.
 	EAsyncPackageState::Type LoadingState = EAsyncPackageState::Complete;
 
-	STAT(double ThisTime = 0);
+	check(LastObjectWorkWasPerformedOn == nullptr);
+	check(LastTypeOfWorkPerformed == nullptr);
+
+	// Set up tick relevant variables.
+	bUseTimeLimit = InbUseTimeLimit;
+	bUseFullTimeLimit = InbUseFullTimeLimit;
+	bTimeLimitExceeded = false;
+	TimeLimit = InOutTimeLimit;
+	TickStartTime = FPlatformTime::Seconds();
+
+	// Keep track of time when we start loading.
+	if (LoadStartTime == 0.0)
 	{
-		SCOPE_SECONDS_COUNTER(ThisTime);
+		LoadStartTime = TickStartTime;
+	}
 
-		check(LastObjectWorkWasPerformedOn == nullptr);
-		check(LastTypeOfWorkPerformed == nullptr);
+	// Make sure we finish our work if there's no time limit. The loop is required as PostLoad
+	// might cause more objects to be loaded in which case we need to Preload them again.
+	do
+	{
+		// Reset value to true at beginning of loop.
+		LoadingState = EAsyncPackageState::Complete;
 
-		// Set up tick relevant variables.
-		bUseTimeLimit = InbUseTimeLimit;
-		bUseFullTimeLimit = InbUseFullTimeLimit;
-		bTimeLimitExceeded = false;
-		TimeLimit = InOutTimeLimit;
-		TickStartTime = FPlatformTime::Seconds();
+		// Begin async loading, simulates BeginLoad
+		BeginAsyncLoad();
 
-		// Keep track of time when we start loading.
-		if (LoadStartTime == 0.0)
+		// Create raw linker. Needs to be async created via ticking before it can be used.
+		if (LoadingState == EAsyncPackageState::Complete)
 		{
-			LoadStartTime = TickStartTime;
+			LoadingState = CreateLinker();
 		}
 
-		// Make sure we finish our work if there's no time limit. The loop is required as PostLoad
-		// might cause more objects to be loaded in which case we need to Preload them again.
-		do
+		// Async create linker.
+		if (LoadingState == EAsyncPackageState::Complete)
 		{
-			// Reset value to true at beginning of loop.
-			LoadingState = EAsyncPackageState::Complete;
-
-			// Begin async loading, simulates BeginLoad
-			BeginAsyncLoad();
-
-			// Create raw linker. Needs to be async created via ticking before it can be used.
-			if (LoadingState == EAsyncPackageState::Complete)
-			{
-				LoadingState = CreateLinker();
-			}
-
-			// Async create linker.
-			if (LoadingState == EAsyncPackageState::Complete)
-			{
-				LoadingState = FinishLinker();
-			}
+			LoadingState = FinishLinker();
+		}
 
-			// Load imports from linker import table asynchronously.
-			if (LoadingState == EAsyncPackageState::Complete)
-			{
-				LoadingState = LoadImports();
-			}
+		// Load imports from linker import table asynchronously.
+		if (LoadingState == EAsyncPackageState::Complete)
+		{
+			LoadingState = LoadImports();
+		}
 
-			// Create imports from linker import table.
-			if (LoadingState == EAsyncPackageState::Complete)
-			{
-				LoadingState = CreateImports();
-			}
+		// Create imports from linker import table.
+		if (LoadingState == EAsyncPackageState::Complete)
+		{
+			LoadingState = CreateImports();
+		}
 
-			// Finish all async texture allocations.
-			if (LoadingState == EAsyncPackageState::Complete)
-			{
-				LoadingState = FinishTextureAllocations();
-			}
+		// Finish all async texture allocations.
+		if (LoadingState == EAsyncPackageState::Complete)
+		{
+			LoadingState = FinishTextureAllocations();
+		}
 
-			// Create exports from linker export table and also preload them.
-			if (LoadingState == EAsyncPackageState::Complete)
-			{
-				LoadingState = CreateExports();
-			}
+		// Create exports from linker export table and also preload them.
+		if (LoadingState == EAsyncPackageState::Complete)
+		{
+			LoadingState = CreateExports();
+		}
 
-			// Call Preload on the linker for all loaded objects which causes actual serialization.
-			if (LoadingState == EAsyncPackageState::Complete)
-			{
-				LoadingState = PreLoadObjects();
-			}
+		// Call Preload on the linker for all loaded objects which causes actual serialization.
+		if (LoadingState == EAsyncPackageState::Complete)
+		{
+			LoadingState = PreLoadObjects();
+		}
 
-			// Call PostLoad on objects, this could cause new objects to be loaded that require
-			// another iteration of the PreLoad loop.
-			if (LoadingState == EAsyncPackageState::Complete)
-			{
-				LoadingState = PostLoadObjects();
-			}
+		// Call PostLoad on objects, this could cause new objects to be loaded that require
+		// another iteration of the PreLoad loop.
+		if (LoadingState == EAsyncPackageState::Complete)
+		{
+			LoadingState = PostLoadObjects();
+		}
 
-			// End async loading, simulates EndLoad
-			EndAsyncLoad();
+		// End async loading, simulates EndLoad
+		EndAsyncLoad();
 
-			// Finish objects (removing RF_AsyncLoading, dissociate imports and forced exports, 
-			// call completion callback, ...
-			// If the load has failed, perform completion callbacks and then quit
-			if (LoadingState == EAsyncPackageState::Complete || bLoadHasFailed)
-			{
-				LoadingState = FinishObjects();
-			}
-		} while (!IsTimeLimitExceeded() && LoadingState == EAsyncPackageState::TimeOut);
+		// Finish objects (removing RF_AsyncLoading, dissociate imports and forced exports, 
+		// call completion callback, ...
+		// If the load has failed, perform completion callbacks and then quit
+		if (LoadingState == EAsyncPackageState::Complete || bLoadHasFailed)
+		{
+			LoadingState = FinishObjects();
+		}
+	} while (!IsTimeLimitExceeded() && LoadingState == EAsyncPackageState::TimeOut);
 
-		check(bUseTimeLimit || LoadingState != EAsyncPackageState::TimeOut);
+	check(bUseTimeLimit || LoadingState != EAsyncPackageState::TimeOut);
 
-		// We can't have a reference to a UObject.
-		LastObjectWorkWasPerformedOn = nullptr;
-		// Reset type of work performed.
-		LastTypeOfWorkPerformed = nullptr;
-		// Mark this package as loaded if everything completed.
-		bLoadHasFinished = (LoadingState == EAsyncPackageState::Complete);
-		// Subtract the time it took to load this package from the global limit.
-		InOutTimeLimit = (float)FMath::Max(0.0, InOutTimeLimit - (FPlatformTime::Seconds() - TickStartTime));
-	}
-	INC_FLOAT_STAT_BY(STAT_FAsyncPackage_TickTime, (float)ThisTime);
+	// We can't have a reference to a UObject.
+	LastObjectWorkWasPerformedOn = nullptr;
+	// Reset type of work performed.
+	LastTypeOfWorkPerformed = nullptr;
+	// Mark this package as loaded if everything completed.
+	bLoadHasFinished = (LoadingState == EAsyncPackageState::Complete);
+	// Subtract the time it took to load this package from the global limit.
+	InOutTimeLimit = (float)FMath::Max(0.0, InOutTimeLimit - (FPlatformTime::Seconds() - TickStartTime));
 
 	// true means that we're done loading this package.
 	return LoadingState;
@@ -1302,37 +1291,32 @@ EAsyncPackageState::Type FAsyncPackage::PostLoadObjects()
 	FGCScopeGuard GCGuard;
 
 	EAsyncPackageState::Type Result = EAsyncPackageState::Complete;
-	STAT(double ThisTime = 0);
-	{
-		SCOPE_SECONDS_COUNTER(ThisTime);
-		TGuardValue<bool> GuardIsRoutingPostLoad(FUObjectThreadContext::Get().IsRoutingPostLoad, true);
+	TGuardValue<bool> GuardIsRoutingPostLoad(FUObjectThreadContext::Get().IsRoutingPostLoad, true);
 
-		auto& ObjLoaded = FUObjectThreadContext::Get().ObjLoaded;
-		// PostLoad objects.
-		while (PostLoadIndex < ObjLoaded.Num() && PostLoadIndex < PreLoadIndex && !IsTimeLimitExceeded())
+	auto& ObjLoaded = FUObjectThreadContext::Get().ObjLoaded;
+	// PostLoad objects.
+	while (PostLoadIndex < ObjLoaded.Num() && PostLoadIndex < PreLoadIndex && !IsTimeLimitExceeded())
+	{
+		UObject* Object = ObjLoaded[PostLoadIndex++];
+		check(Object);
+		if (!FAsyncLoadingThread::Get().IsMultithreaded() || Object->IsPostLoadThreadSafe())
 		{
-			UObject* Object = ObjLoaded[PostLoadIndex++];
-			check(Object);
-			if (!FAsyncLoadingThread::Get().IsMultithreaded() || Object->IsPostLoadThreadSafe())
-			{
-				FScopeCycleCounterUObject ConstructorScope(Object, GET_STATID(STAT_FAsyncPackage_PostLoadObjects));
+			FScopeCycleCounterUObject ConstructorScope(Object, GET_STATID(STAT_FAsyncPackage_PostLoadObjects));
 
-				Object->ConditionalPostLoad();
+			Object->ConditionalPostLoad();
 
-				LastObjectWorkWasPerformedOn = Object;
-				LastTypeOfWorkPerformed = TEXT("postloading_async");
-			}
-			else
-			{
-				DeferredPostLoadObjects.Add(Object);
-			}
+			LastObjectWorkWasPerformedOn = Object;
+			LastTypeOfWorkPerformed = TEXT("postloading_async");
+		}
+		else
+		{
+			DeferredPostLoadObjects.Add(Object);
 		}
-
-		// New objects might have been loaded during PostLoad.
-		Result = (PreLoadIndex == ObjLoaded.Num()) && (PostLoadIndex == ObjLoaded.Num()) ? EAsyncPackageState::Complete : EAsyncPackageState::TimeOut;
-		UE_CLOG(PreLoadIndex < ObjLoaded.Num(), LogStreaming, Fatal, TEXT("Objects have been loading in PostLoad. This is no longer allowed."));
 	}
-	INC_FLOAT_STAT_BY(STAT_FAsyncPackage_PostLoadObjectsTime, (float)ThisTime);
+
+	// New objects might have been loaded during PostLoad.
+	Result = (PreLoadIndex == ObjLoaded.Num()) && (PostLoadIndex == ObjLoaded.Num()) ? EAsyncPackageState::Complete : EAsyncPackageState::TimeOut;
+	UE_CLOG(PreLoadIndex < ObjLoaded.Num(), LogStreaming, Fatal, TEXT("Objects have been loading in PostLoad. This is no longer allowed."));
 
 	return Result;
 }
@@ -1342,46 +1326,42 @@ EAsyncPackageState::Type FAsyncPackage::PostLoadDeferredObjects(double InTickSta
 	SCOPE_CYCLE_COUNTER(STAT_FAsyncPackage_PostLoadObjectsGameThread);
 
 	EAsyncPackageState::Type Result = EAsyncPackageState::Complete;
-	STAT(double ThisTime = 0);
-	{
-		SCOPE_SECONDS_COUNTER(ThisTime);
+	TGuardValue<bool> GuardIsRoutingPostLoad(FUObjectThreadContext::Get().IsRoutingPostLoad, true);
 
-		TGuardValue<bool> GuardIsRoutingPostLoad(FUObjectThreadContext::Get().IsRoutingPostLoad, true);
+	LastObjectWorkWasPerformedOn = nullptr;
+	LastTypeOfWorkPerformed = TEXT("postloading_gamethread");
 
-		while (DeferredPostLoadIndex < DeferredPostLoadObjects.Num() && !::IsTimeLimitExceeded(InTickStartTime, bUseTimeLimit, InOutTimeLimit))
-		{
-			UObject* Object = DeferredPostLoadObjects[DeferredPostLoadIndex++];
-			check(Object);
+	while (DeferredPostLoadIndex < DeferredPostLoadObjects.Num() && !::IsTimeLimitExceeded(InTickStartTime, bUseTimeLimit, InOutTimeLimit, LastTypeOfWorkPerformed, LastObjectWorkWasPerformedOn))
+	{
+		UObject* Object = DeferredPostLoadObjects[DeferredPostLoadIndex++];
+		check(Object);
 
-			FScopeCycleCounterUObject ConstructorScope(Object, GET_STATID(STAT_FAsyncPackage_PostLoadObjectsGameThread));
+		FScopeCycleCounterUObject ConstructorScope(Object, GET_STATID(STAT_FAsyncPackage_PostLoadObjectsGameThread));
 
-			Object->ConditionalPostLoad();
+		Object->ConditionalPostLoad();
 
-			LastObjectWorkWasPerformedOn = Object;
-			LastTypeOfWorkPerformed = TEXT("postloading_gamethread");
+		LastObjectWorkWasPerformedOn = Object;		
 
-			UpdateLoadPercentage();
-		}
+		UpdateLoadPercentage();
+	}
 
-		// New objects might have been loaded during PostLoad.
-		Result = (DeferredPostLoadIndex == DeferredPostLoadObjects.Num()) ? EAsyncPackageState::Complete : EAsyncPackageState::TimeOut;
-		if (Result == EAsyncPackageState::Complete)
+	// New objects might have been loaded during PostLoad.
+	Result = (DeferredPostLoadIndex == DeferredPostLoadObjects.Num()) ? EAsyncPackageState::Complete : EAsyncPackageState::TimeOut;
+	if (Result == EAsyncPackageState::Complete)
+	{
+		// Mark package as having been fully loaded and update load time.
+		if (Linker && !bLoadHasFailed)
 		{
-			// Mark package as having been fully loaded and update load time.
-			if (Linker && !bLoadHasFailed)
+			if (Linker->LinkerRoot)
 			{
-				if (Linker->LinkerRoot)
-				{
-					Linker->LinkerRoot->MarkAsFullyLoaded();
-					Linker->LinkerRoot->SetLoadTime(FPlatformTime::Seconds() - LoadStartTime);
-				}
-
-				// give a hint to the IO system that we are done with this file for now
-				FIOSystem::Get().HintDoneWithFile(Linker->Filename);
+				Linker->LinkerRoot->MarkAsFullyLoaded();
+				Linker->LinkerRoot->SetLoadTime(FPlatformTime::Seconds() - LoadStartTime);
 			}
+
+			// give a hint to the IO system that we are done with this file for now
+			FIOSystem::Get().HintDoneWithFile(Linker->Filename);
 		}
 	}
-	INC_FLOAT_STAT_BY(STAT_FAsyncPackage_PostLoadObjectsGameThreadTime, (float)ThisTime);
 
 	return Result;
 }
@@ -1580,23 +1560,18 @@ void FlushAsyncLoading(FName ExcludeType/*=NAME_None*/)
 {
 	if (IsAsyncLoading())
 	{
+		SCOPE_CYCLE_COUNTER(STAT_FAsyncPackage_FlushAsyncLoadingGameThread);
+
 		// Disallow low priority requests like texture streaming while we are flushing streaming
 		// in order to avoid excessive seeking.
 		FIOSystem::Get().SetMinPriority( AIOP_Normal );
 
 		// Flush async loaders by not using a time limit. Needed for e.g. garbage collection.
 		UE_LOG(LogStreaming, Log,  TEXT("Flushing async loaders.") );
-		STAT(double ThisTime = 0);
+		while (IsAsyncLoading())
 		{
-			SCOPE_SECONDS_COUNTER(ThisTime);
-
-			while (IsAsyncLoading())
-			{
-				FAsyncLoadingThread::Get().TickAsyncLoading(false, false, 0, ExcludeType);
-			}
+			FAsyncLoadingThread::Get().TickAsyncLoading(false, false, 0, ExcludeType);
 		}
-		INC_FLOAT_STAT_BY(STAT_FAsyncPackage_FlushAsyncLoadingTime, (float)ThisTime);
-		UE_LOG(LogStreaming, Log,  TEXT("Flushed async loaders.") );
 
 		if (ExcludeType == NAME_None)
 		{
@@ -2064,7 +2039,7 @@ void FArchiveAsync::Serialize( void* Data, int64 Count )
 
 	// Update stats if we were blocked.
 #if STATS
-	if( bIOBlocked )
+	if( bIOBlocked && IsInGameThread() )
 	{
 		double BlockingTime = FPlatformTime::Seconds() - StartTime;
 		INC_FLOAT_STAT_BY(STAT_AsyncIO_MainThreadBlockTime,(float)BlockingTime);
diff --git a/Engine/Source/Runtime/CoreUObject/Private/UObject/LinkerLoad.cpp b/Engine/Source/Runtime/CoreUObject/Private/UObject/LinkerLoad.cpp
--- a/Engine/Source/Runtime/CoreUObject/Private/UObject/LinkerLoad.cpp
+++ b/Engine/Source/Runtime/CoreUObject/Private/UObject/LinkerLoad.cpp
@@ -14,7 +14,7 @@
 
 #define LOCTEXT_NAMESPACE "LinkerLoad"
 
-DECLARE_STATS_GROUP(TEXT("Linker Load"), STATGROUP_LinkerLoad, STATCAT_Advanced);
+DECLARE_STATS_GROUP_VERBOSE(TEXT("Linker Load"), STATGROUP_LinkerLoad, STATCAT_Advanced);
 
 DECLARE_CYCLE_STAT(TEXT("Linker Preload"),STAT_LinkerPreload,STATGROUP_LinkerLoad);
 DECLARE_CYCLE_STAT(TEXT("Linker Precache"),STAT_LinkerPrecache,STATGROUP_LinkerLoad);
