Adding VerifyAsyncLoadAssumptions console command to check if flags on UObjects get properly reset after async loading has finished.

diff --git a/Engine/Source/Runtime/CoreUObject/Private/Serialization/AsyncLoading.cpp b/Engine/Source/Runtime/CoreUObject/Private/Serialization/AsyncLoading.cpp
--- a/Engine/Source/Runtime/CoreUObject/Private/Serialization/AsyncLoading.cpp
+++ b/Engine/Source/Runtime/CoreUObject/Private/Serialization/AsyncLoading.cpp
@@ -54,6 +54,19 @@ class FAsyncObjectsReferencer : FGCObject
 	/** Critical section for referenced objects list */
 	FCriticalSection ReferencedObjectsCritical;
 
+	FORCEINLINE int32 IndexOf(UObject* InObj)
+	{
+		FScopeLock ReferencedObjectsLock(&ReferencedObjectsCritical);
+		for (int32 Index = 0; Index < ReferencedObjects.Num(); ++Index)
+		{
+			if (ReferencedObjects[Index] == InObj)
+			{
+				return Index;
+			}
+		}
+		return INDEX_NONE;
+	}
+
 public:
 	/** Returns the one and only instance of this object */
 	static FAsyncObjectsReferencer& Get();
@@ -91,11 +104,13 @@ public:
 	/** Removes all referenced objects and markes them for GC */
 	void EmptyReferencedObjectsAndCancelLoading()
 	{
-		const EObjectFlags LoadFlags = RF_AsyncLoading | RF_NeedLoad | RF_NeedPostLoad | RF_NeedPostLoadSubobjects;
+		const EObjectFlags LoadFlags = RF_NeedLoad | RF_NeedPostLoad | RF_NeedPostLoadSubobjects;
+		const EObjectFlags AsyncFlags = RF_Async | RF_AsyncLoading;
 
 		// All of the referenced objects have been created by async loading code and may be in an invalid state so mark them for GC
 		for (auto Object : ReferencedObjects)
 		{
+			Object->ClearFlags(AsyncFlags);
 			if (Object->HasAnyFlags(LoadFlags))
 			{
 				Object->ClearFlags(LoadFlags);
@@ -104,6 +119,24 @@ public:
 		}
 		ReferencedObjects.Empty(ReferencedObjects.Num());
 	}
+
+#if !UE_BUILD_SHIPPING
+	/** Verifies that no object exists that has either RF_AsyncLoading|RF_Async set and is NOT being referenced by FAsyncObjectsReferencer */
+	FORCENOINLINE void VerifyAssumptions()
+	{
+		for (FRawObjectIterator It; It; ++It)
+		{
+			UObject* Obj = *It;
+			if (Obj->HasAnyFlags(RF_AsyncLoading|RF_Async))
+			{
+				if (IndexOf(Obj) == INDEX_NONE)
+				{
+					UE_LOG(LogStreaming, Error, TEXT("%s has RF_AsyncLoading|RF_Async set but is not referenced by FAsyncObjectsReferencer"), *Obj->GetPathName());
+				}
+			}
+		}
+	}
+#endif
 };
 FAsyncObjectsReferencer& FAsyncObjectsReferencer::Get()
 {
@@ -111,6 +144,28 @@ FAsyncObjectsReferencer& FAsyncObjectsReferencer::Get()
 	return Singleton;
 }
 
+#if !UE_BUILD_SHIPPING
+class FAsyncLoadingExec : private FSelfRegisteringExec
+{
+public:
+
+	FAsyncLoadingExec()
+	{}
+
+	/** Console commands **/
+	virtual bool Exec(UWorld* InWorld, const TCHAR* Cmd, FOutputDevice& Ar) override
+	{
+		if (FParse::Command(&Cmd, TEXT("VerifyAsyncLoadAssumptions")))
+		{
+			FAsyncObjectsReferencer::Get().VerifyAssumptions();
+			return true;
+		}
+		return false;
+	}
+};
+static TAutoPtr<FAsyncLoadingExec> GAsyncLoadingExec;
+#endif
+
 static FORCEINLINE bool IsTimeLimitExceeded(double InTickStartTime, bool bUseTimeLimit, float InTimeLimit, const TCHAR* InLastTypeOfWorkPerformed = nullptr, UObject* InLastObjectWorkWasPerformedOn = nullptr)
 {
 	bool bTimeLimitExceeded = false;
@@ -418,6 +473,9 @@ EAsyncPackageState::Type FAsyncLoadingThread::TickAsyncLoading(bool bUseTimeLimi
 
 FAsyncLoadingThread::FAsyncLoadingThread()
 {
+#if !UE_BUILD_SHIPPING
+	GAsyncLoadingExec = new FAsyncLoadingExec();
+#endif
 	QueuedRequestsEvent = FPlatformProcess::GetSynchEventFromPool();
 	CancelLoadingEvent = FPlatformProcess::GetSynchEventFromPool();
 	if (FAsyncLoadingThread::IsMultithreaded())
