Core - Async stats improvements

diff --git a/Engine/Source/Runtime/CoreUObject/Private/Serialization/AsyncLoading.cpp b/Engine/Source/Runtime/CoreUObject/Private/Serialization/AsyncLoading.cpp
--- a/Engine/Source/Runtime/CoreUObject/Private/Serialization/AsyncLoading.cpp
+++ b/Engine/Source/Runtime/CoreUObject/Private/Serialization/AsyncLoading.cpp
@@ -186,7 +186,9 @@ static FORCEINLINE bool IsTimeLimitExceeded(double InTickStartTime, bool bUseTim
 	return bTimeLimitExceeded;
 }
 
-int32 FAsyncLoadingThread::AsyncLoadingThreadID = -1;
+uint32 FAsyncLoadingThread::AsyncLoadingThreadID = 0;
+
+FThreadSafeCounter FAsyncLoadingThread::BlockingCycles = 0;
 
 FAsyncLoadingThread& FAsyncLoadingThread::Get()
 {
@@ -538,6 +540,10 @@ void FAsyncLoadingThread::TickAsyncThread(bool bUseTimeLimit, bool bUseFullTimeL
 		CancelAsyncLoadingInternal();
 		bShouldCancelLoading = false;
 	}
+
+	// Update stats
+	SET_FLOAT_STAT( STAT_AsyncIO_MainThreadBlockTime, FPlatformTime::ToSeconds( BlockingCycles.GetValue() ) );
+	BlockingCycles.Set( 0 );
 }
 
 void FAsyncLoadingThread::Stop()
@@ -2049,7 +2055,7 @@ void FArchiveAsync::Serialize( void* Data, int64 Count )
 	// Ensure we aren't reading beyond the end of the file
 	checkf( CurrentPos + Count <= TotalSize(), TEXT("Seeked past end of file %s (%lld / %lld)"), *FileName, CurrentPos + Count, TotalSize() );
 
-	double	StartTime	= 0;
+	uint32 StartCycles = 0;
 	bool	bIOBlocked	= false;
 
 	// Make sure serialization request fits entirely in already precached region.
@@ -2058,7 +2064,7 @@ void FArchiveAsync::Serialize( void* Data, int64 Count )
 		DECLARE_SCOPE_CYCLE_COUNTER( TEXT( "FArchiveAsync::Serialize.PrecacheBufferContainsRequest" ), STAT_ArchiveAsync_Serialize_PrecacheBufferContainsRequest, STATGROUP_AsyncLoad );
 
 		// Keep track of time we started to block.
-		StartTime	= FPlatformTime::Seconds();
+		StartCycles = FPlatformTime::Cycles();
 		bIOBlocked	= true;
 
 		// Busy wait for region to be precached.
@@ -2085,7 +2091,7 @@ void FArchiveAsync::Serialize( void* Data, int64 Count )
 		if( !bIOBlocked )
 		{
 			// Keep track of time we started to block.
-			StartTime = FPlatformTime::Seconds();
+			StartCycles = FPlatformTime::Cycles();
 			bIOBlocked = true;
 		}
 		if( !FPlatformProcess::SupportsMultithreading() )
@@ -2099,11 +2105,12 @@ void FArchiveAsync::Serialize( void* Data, int64 Count )
 #if STATS
 	if( bIOBlocked && IsInGameThread() )
 	{
-		double BlockingTime = FPlatformTime::Seconds() - StartTime;
-		INC_FLOAT_STAT_BY(STAT_AsyncIO_MainThreadBlockTime,(float)BlockingTime);
+		const int32 BlockingCycles = int32(FPlatformTime::Cycles() - StartCycles);
+		FAsyncLoadingThread::BlockingCycles.Add( BlockingCycles );
+		
 #if LOOKING_FOR_PERF_ISSUES
 		UE_LOG(LogStreaming, Warning, TEXT("FArchiveAsync::Serialize: %5.2fms blocking on read from '%s' (Offset: %lld, Size: %lld)"), 
-			1000 * BlockingTime, 
+			FPlatformTime::ToMilliseconds(BlockingCycles), 
 			*FileName, 
 			CurrentPos, 
 			Count );
diff --git a/Engine/Source/Runtime/CoreUObject/Private/Serialization/AsyncLoadingThread.h b/Engine/Source/Runtime/CoreUObject/Private/Serialization/AsyncLoadingThread.h
--- a/Engine/Source/Runtime/CoreUObject/Private/Serialization/AsyncLoadingThread.h
+++ b/Engine/Source/Runtime/CoreUObject/Private/Serialization/AsyncLoadingThread.h
@@ -10,6 +10,8 @@
  */
 class FAsyncLoadingThread : public FRunnable
 {
+	friend class FArchiveAsync;
+
 	/** Thread to run the worker FRunnable on */
 	FRunnableThread* Thread;
 	/** Stops this thread */
@@ -47,7 +49,10 @@ class FAsyncLoadingThread : public FRunnable
 	FThreadSafeCounter AsyncThreadReady;
 
 	/** Async loading thread ID */
-	static int32 AsyncLoadingThreadID;
+	static uint32 AsyncLoadingThreadID;
+
+	/** Thread safe counter used to accumulate cycles spent on blocking. Using stats may generate to many stats messages. */
+	static FThreadSafeCounter BlockingCycles;
 
 	FAsyncLoadingThread();
 	virtual ~FAsyncLoadingThread();
