CastChecked fixes:
* Now based on DO_CHECK rather than !(UE_BUILD_SHIPPING || UE_BUILD_TEST).
* Didn't work for interface casts.
* Inlined in non-DO_CHECK builds.

Cast implementation moved around a little to avoid some redundant null checks.

diff --git a/Engine/Source/Runtime/CoreUObject/Public/Templates/Casts.h b/Engine/Source/Runtime/CoreUObject/Public/Templates/Casts.h
--- a/Engine/Source/Runtime/CoreUObject/Public/Templates/Casts.h
+++ b/Engine/Source/Runtime/CoreUObject/Public/Templates/Casts.h
@@ -82,13 +82,21 @@ template <typename From, typename To, uint64 CastClass> struct TGetCastType<From
 template <typename To, typename From>
 To* Cast(From* Src);
 
+template <typename To, typename From>
+To* CastChecked(From* Src, ECastCheckedType::Type CheckType = ECastCheckedType::NullChecked);
+
 template <typename From, typename To, ECastType CastType = TGetCastType<From, To>::Value>
 struct TCastImpl
 {
 	// This is the cast flags implementation
 	FORCEINLINE static To* DoCast( UObject* Src )
 	{
-		return Src->GetClass()->HasAnyCastFlag(TCastFlags<To>::Value) ? (To*)Src : nullptr;
+		return Src && Src->GetClass()->HasAnyCastFlag(TCastFlags<To>::Value) ? (To*)Src : nullptr;
+	}
+
+	FORCEINLINE static To* DoCastCheckedWithoutTypeCheck( UObject* Src )
+	{
+		return (To*)Src;
 	}
 };
 
@@ -97,7 +105,12 @@ struct TCastImpl<From, To, ECastType::UObjectToUObject>
 {
 	FORCEINLINE static To* DoCast( UObject* Src )
 	{
-		return Src->IsA<To>() ? (To*)Src : nullptr;
+		return Src && Src->IsA<To>() ? (To*)Src : nullptr;
+	}
+
+	FORCEINLINE static To* DoCastCheckedWithoutTypeCheck( UObject* Src )
+	{
+		return (To*)Src;
 	}
 };
 
@@ -106,7 +119,21 @@ struct TCastImpl<From, To, ECastType::InterfaceToUObject>
 {
 	FORCEINLINE static To* DoCast( From* Src )
 	{
-		return Cast<To>(Src->_getUObject());
+		To* Result = nullptr;
+		if (Src)
+		{
+			UObject* Obj = Src->_getUObject();
+			if (Obj->IsA<To>())
+			{
+				Result = (To*)Obj;
+			}
+		}
+		return Result;
+	}
+
+	FORCEINLINE static To* DoCastCheckedWithoutTypeCheck( From* Src )
+	{
+		return Src ? (To*)Src->_getUObject() : nullptr;
 	}
 };
 
@@ -115,7 +142,12 @@ struct TCastImpl<From, To, ECastType::UObjectToInterface>
 {
 	FORCEINLINE static To* DoCast( UObject* Src )
 	{
-		return (To*)Src->GetInterfaceAddress(To::UClassType::StaticClass());
+		return Src ? (To*)Src->GetInterfaceAddress(To::UClassType::StaticClass()) : nullptr;
+	}
+
+	FORCEINLINE static To* DoCastCheckedWithoutTypeCheck( UObject* Src )
+	{
+		return Src ? (To*)Src->GetInterfaceAddress(To::UClassType::StaticClass()) : nullptr;
 	}
 };
 
@@ -124,7 +156,12 @@ struct TCastImpl<From, To, ECastType::InterfaceToInterface>
 {
 	FORCEINLINE static To* DoCast( From* Src )
 	{
-		return Cast<To>(Src->_getUObject());
+		return Src ? (To*)Src->_getUObject()->GetInterfaceAddress(To::UClassType::StaticClass()) : nullptr;
+	}
+
+	FORCEINLINE static To* DoCastCheckedWithoutTypeCheck( From* Src )
+	{
+		return Src ? (To*)Src->_getUObject()->GetInterfaceAddress(To::UClassType::StaticClass()) : nullptr;
 	}
 };
 
@@ -132,7 +169,7 @@ struct TCastImpl<From, To, ECastType::InterfaceToInterface>
 template <typename To, typename From>
 FORCEINLINE To* Cast(From* Src)
 {
-	return Src ? TCastImpl<From, To>::DoCast(Src) : nullptr;
+	return TCastImpl<From, To>::DoCast(Src);
 }
 
 template< class T >
@@ -141,11 +178,11 @@ FORCEINLINE T* ExactCast( UObject* Src )
 	return Src && (Src->GetClass() == T::StaticClass()) ? (T*)Src : nullptr;
 }
 
-template <typename To, typename From>
-To* CastChecked(From* Src, ECastCheckedType::Type CheckType = ECastCheckedType::NullChecked)
-{
-	#if !(UE_BUILD_SHIPPING || UE_BUILD_TEST)
+#if DO_CHECK
 
+	template <typename To, typename From>
+	To* CastChecked(From* Src, ECastCheckedType::Type CheckType)
+	{
 		if (Src)
 		{
 			To* Result = Cast<To>(Src);
@@ -163,13 +200,17 @@ To* CastChecked(From* Src, ECastCheckedType::Type CheckType = ECastCheckedType::
 		}
 
 		return nullptr;
+	}
 
-	#else
+#else
 
-		return (To*)Src;
+	template <typename To, typename From>
+	FORCEINLINE To* CastChecked(From* Src, ECastCheckedType::Type CheckType)
+	{
+		return TCastImpl<From, To>::DoCastCheckedWithoutTypeCheck(Src);
+	}
 
-	#endif
-}
+#endif
 
 
 template <typename InterfaceType>
