Fixing async loading flags being set on objects not created during async loading (as a result of duplication for example).

UE-13428

diff --git a/Engine/Source/Runtime/CoreUObject/Private/UObject/UObjectGlobals.cpp b/Engine/Source/Runtime/CoreUObject/Private/UObject/UObjectGlobals.cpp
--- a/Engine/Source/Runtime/CoreUObject/Private/UObject/UObjectGlobals.cpp
+++ b/Engine/Source/Runtime/CoreUObject/Private/UObject/UObjectGlobals.cpp
@@ -1875,8 +1875,8 @@ UObject* StaticAllocateObject
 
 	if (!bSubObject)
 	{
-	FMemory::Memzero((void *)Obj, TotalSize);
-	new ((void *)Obj) UObjectBase(InClass,InFlags,InOuter,InName);
+		FMemory::Memzero((void *)Obj, TotalSize);
+		new ((void *)Obj) UObjectBase(InClass,InFlags,InOuter,InName);
 	}
 	else
 	{
@@ -1899,6 +1899,21 @@ UObject* StaticAllocateObject
 	{
 		NotifyConstructedDuringAsyncLoading(Obj, bSubObject);
 	}
+	else
+	{
+		// Sanity checks for async flags.
+		// It's possible to duplicate an object on the game thread that is still being referenced 
+		// by async loading code or has been created on a different thread than the main thread.
+		if (Obj->HasAnyFlags(RF_AsyncLoading))
+		{
+			Obj->ClearFlags(RF_AsyncLoading);
+		}
+		if (Obj->HasAnyFlags(RF_Async) && IsInGameThread())
+		{
+			Obj->ClearFlags(RF_Async);
+		}
+	}
+
 
 	// Let the caller know if a subobject has just been recycled.
 	if (bOutRecycledSubobject)
