Thread safety fixes for multiple classes related to UObjects:

- FUnresolvedStructTracker
- FUniqueObjectGuid
- FLinkerManager
- FStringAssetReference
- FUObjectAnnotationDense

diff --git a/Engine/Source/Runtime/CoreUObject/Private/Blueprint/BlueprintSupport.cpp b/Engine/Source/Runtime/CoreUObject/Private/Blueprint/BlueprintSupport.cpp
--- a/Engine/Source/Runtime/CoreUObject/Private/Blueprint/BlueprintSupport.cpp
+++ b/Engine/Source/Runtime/CoreUObject/Private/Blueprint/BlueprintSupport.cpp
@@ -621,6 +621,7 @@ public:
 		: TrackedStruct(LoadStruct)
 	{
 		DEFERRED_DEPENDENCY_CHECK((LoadStruct != nullptr) && (LoadStruct->GetLinker() != nullptr));
+		FScopeLock UnresolvedStructsLock(&UnresolvedStructsCritical);
 		UnresolvedStructs.Add(LoadStruct);
 	}
 
@@ -631,6 +632,7 @@ public:
 		// we want the most nested one removing it from the set (because this 
 		// means the struct is fully resolved, even if we're still in the middle  
 		// of a ResolveDeferredDependencies() call further up the stack)
+		FScopeLock UnresolvedStructsLock(&UnresolvedStructsCritical);
 		UnresolvedStructs.Remove(TrackedStruct);
 	}
 
@@ -644,6 +646,7 @@ public:
 	 */
 	static bool IsImportStructUnresolved(UObject* ImportObject)
 	{
+		FScopeLock UnresolvedStructsLock(&UnresolvedStructsCritical);
 		return UnresolvedStructs.Contains(ImportObject);
 	}
 
@@ -661,6 +664,7 @@ public:
 	 */
 	static bool IsAssociatedStructUnresolved(const FLinkerLoad* Linker)
 	{
+		FScopeLock UnresolvedStructsLock(&UnresolvedStructsCritical);
 		for (UObject* UnresolvedObj : UnresolvedStructs)
 		{
 			// each unresolved struct should have a linker set on it, because 
@@ -684,6 +688,7 @@ public:
 				ToRemove.Add(UnresolvedObj);
 			}
 		}
+		FScopeLock UnresolvedStructsLock(&UnresolvedStructsCritical);
 		for (UObject* ResetingObj : ToRemove)
 		{
 			UnresolvedStructs.Remove(ResetingObj);
@@ -700,9 +705,12 @@ private:
 	 * we don't have to cast import objects before checking for their presence).
 	 */
 	static TSet<UObject*> UnresolvedStructs;
+	static FCriticalSection UnresolvedStructsCritical;
 };
 /** A global set that tracks structs currently being ran through (and unfinished by) FLinkerLoad::ResolveDeferredDependencies() */
 TSet<UObject*> FUnresolvedStructTracker::UnresolvedStructs;
+FCriticalSection FUnresolvedStructTracker::UnresolvedStructsCritical;
+
 UPackage* LoadPackageInternal(UPackage* InOuter, const TCHAR* InLongPackageName, uint32 LoadFlags, FLinkerLoad* ImportLinker);
 
 void FLinkerLoad::ResolveDeferredDependencies(UStruct* LoadStruct)
diff --git a/Engine/Source/Runtime/CoreUObject/Private/Misc/StringAssetReference.cpp b/Engine/Source/Runtime/CoreUObject/Private/Misc/StringAssetReference.cpp
--- a/Engine/Source/Runtime/CoreUObject/Private/Misc/StringAssetReference.cpp
+++ b/Engine/Source/Runtime/CoreUObject/Private/Misc/StringAssetReference.cpp
@@ -190,5 +190,5 @@ void FStringAssetReference::FixupForPIE()
 	}
 }
 
-int32 FStringAssetReference::CurrentTag = 1;
+FThreadSafeCounter FStringAssetReference::CurrentTag(1);
 TArray<FString> FStringAssetReference::PackageNamesBeingDuplicatedForPIE;
\ No newline at end of file
diff --git a/Engine/Source/Runtime/CoreUObject/Private/Serialization/AsyncLoading.cpp b/Engine/Source/Runtime/CoreUObject/Private/Serialization/AsyncLoading.cpp
--- a/Engine/Source/Runtime/CoreUObject/Private/Serialization/AsyncLoading.cpp
+++ b/Engine/Source/Runtime/CoreUObject/Private/Serialization/AsyncLoading.cpp
@@ -90,19 +90,21 @@ public:
 			InObject->ThisThreadAtomicallyClearedRFUnreachable();
 		}
 	}
-	/** Removes all objects from the list */
-	FORCEINLINE void EmptyReferencedObjects()
+	/** Removes all objects from the list and clears async loading flags */
+	FORCENOINLINE void EmptyReferencedObjects()
 	{
+		const EObjectFlags AsyncFlags = RF_Async | RF_AsyncLoading;
 		FScopeLock ReferencedObjectsLock(&ReferencedObjectsCritical);
-		for (auto Obj : ReferencedObjects)
+		for (UObject* Obj : ReferencedObjects)
 		{
 			check(Obj);
-			Obj->ClearFlags(RF_Async | RF_AsyncLoading);
+			Obj->AtomicallyClearFlags(AsyncFlags);
+			check(!Obj->HasAnyFlags(AsyncFlags))
 		}
 		ReferencedObjects.Empty(ReferencedObjects.Num());
 	}
 	/** Removes all referenced objects and markes them for GC */
-	void EmptyReferencedObjectsAndCancelLoading()
+	FORCENOINLINE void EmptyReferencedObjectsAndCancelLoading()
 	{
 		const EObjectFlags LoadFlags = RF_NeedLoad | RF_NeedPostLoad | RF_NeedPostLoadSubobjects;
 		const EObjectFlags AsyncFlags = RF_Async | RF_AsyncLoading;
@@ -113,9 +115,10 @@ public:
 			Object->ClearFlags(AsyncFlags);
 			if (Object->HasAnyFlags(LoadFlags))
 			{
-				Object->ClearFlags(LoadFlags);
-				Object->SetFlags(RF_PendingKill);
+				Object->AtomicallyClearFlags(LoadFlags);
+				Object->AtomicallySetFlags(RF_PendingKill);
 			}
+			check(!Object->HasAnyFlags(AsyncFlags | LoadFlags));
 		}
 		ReferencedObjects.Empty(ReferencedObjects.Num());
 	}
diff --git a/Engine/Source/Runtime/CoreUObject/Private/UObject/LazyObjectPtr.cpp b/Engine/Source/Runtime/CoreUObject/Private/UObject/LazyObjectPtr.cpp
--- a/Engine/Source/Runtime/CoreUObject/Private/UObject/LazyObjectPtr.cpp
+++ b/Engine/Source/Runtime/CoreUObject/Private/UObject/LazyObjectPtr.cpp
@@ -67,6 +67,7 @@ void FUniqueObjectGuid::FromString(const FString& From)
 FUniqueObjectGuid FUniqueObjectGuid::GetOrCreateIDForObject(const class UObject *Object)
 {
 	check(Object);
+	checkSlow(IsInGameThread());
 	FUniqueObjectGuid ObjectGuid(Object);
 	if (!ObjectGuid.IsValid())
 	{
@@ -77,7 +78,7 @@ FUniqueObjectGuid FUniqueObjectGuid::GetOrCreateIDForObject(const class UObject
 	return ObjectGuid;
 }
 
-int32 FUniqueObjectGuid::CurrentAnnotationTag = 1;
+FThreadSafeCounter FUniqueObjectGuid::CurrentAnnotationTag(1);
 
 
 /*-----------------------------------------------------------------------------------------------------------
diff --git a/Engine/Source/Runtime/CoreUObject/Private/UObject/Linker.cpp b/Engine/Source/Runtime/CoreUObject/Private/UObject/Linker.cpp
--- a/Engine/Source/Runtime/CoreUObject/Private/UObject/Linker.cpp
+++ b/Engine/Source/Runtime/CoreUObject/Private/UObject/Linker.cpp
@@ -332,53 +332,7 @@ void ResetLoaders( UObject* InPkg )
 {
 	// Make sure we're not in the middle of loading something in the background.
 	FlushAsyncLoading();
-
-	// Top level package to reset loaders for.
-	UObject*		TopLevelPackage = InPkg ? InPkg->GetOutermost() : NULL;
-
-	// Find loader/ linker associated with toplevel package. We do this upfront as Detach resets LinkerRoot.
-	if( TopLevelPackage )
-	{
-		// Linker to reset/ detach.
-		auto LinkerToReset = FLinkerLoad::FindExistingLinkerForPackage(CastChecked<UPackage>(TopLevelPackage));
-		if (LinkerToReset)
-		{
-			for (auto Linker : FLinkerManager::Get().GetLoaders())
-			{
-				// Detach LinkerToReset from other linker's import table.
-				if( Linker->LinkerRoot != TopLevelPackage )
-				{
-					for (auto& Import : Linker->ImportMap)
-					{
-						if (Import.SourceLinker == LinkerToReset)
-						{
-							Import.SourceLinker = NULL;
-							Import.SourceIndex = INDEX_NONE;
-						}
-					}
-				}
-				else
-				{
-					check(Linker == LinkerToReset);
-				}
-			}
-			// Detach linker, also removes from array and sets LinkerRoot to NULL.
-			LinkerToReset->LoadAndDetachAllBulkData();
-			delete LinkerToReset;
-			LinkerToReset = nullptr;
-		}
-	}
-	else
-	{
-		auto LinkersToDetach = FLinkerManager::Get().GetLoaders();
-		for (auto Linker : LinkersToDetach)
-		{
-			// Detach linker, also removes from array and sets LinkerRoot to NULL.
-			Linker->LoadAndDetachAllBulkData();
-			delete Linker;
-		}
-	}
-
+	FLinkerManager::Get().ResetLoaders(InPkg);
 }
 
 
@@ -389,46 +343,7 @@ void ResetLoaders( UObject* InPkg )
  */
 void DissociateImportsAndForcedExports()
 {
-	int32& ImportCount = FUObjectThreadContext::Get().ImportCount;
-	if (ImportCount && FLinkerManager::Get().GetLoadersWithNewImports().Num())
-	{
-		for (auto Linker : FLinkerManager::Get().GetLoadersWithNewImports())
-		{
-			for( int32 ImportIndex=0; ImportIndex<Linker->ImportMap.Num(); ImportIndex++ )
-			{
-				FObjectImport& Import = Linker->ImportMap[ImportIndex];
-				if( Import.XObject && !Import.XObject->HasAnyFlags(RF_Native) )
-				{
-					Import.XObject = NULL;
-				}
-				Import.SourceLinker = NULL;
-				// when the SourceLinker is reset, the SourceIndex must also be reset, or recreating
-				// an import that points to a redirector will fail to find the redirector
-				Import.SourceIndex = INDEX_NONE;
-			}
-		}
-	}
-	ImportCount = 0;
-	FLinkerManager::Get().GetLoadersWithNewImports().Empty();
-
-	int32& ForcedExportCount = FUObjectThreadContext::Get().ForcedExportCount;
-	if (ForcedExportCount)
-	{
-		for (auto Linker : FLinkerManager::Get().GetLoaders())
-		{
-			//@todo optimization: only dissociate exports for loaders that had forced exports created
-			//@todo optimization: since the last time this function was called.
-			for (auto& Export : Linker->ExportMap)
-			{
-				if (Export.Object && Export.bForcedExport)
-				{
-					Export.Object->SetLinker(NULL, INDEX_NONE);
-					Export.Object = NULL;
-				}
-			}
-		}
-	}
-	ForcedExportCount = 0;
+	FLinkerManager::Get().DissociateImportsAndForcedExports();
 }
 
 static void LogGetPackageLinkerError(FArchiveUObject* LinkerArchive, const TCHAR* InFilename, const FText& InFullErrorMessage, const FText& InSummaryErrorMessage, UObject* InOuter, uint32 LoadFlags)
diff --git a/Engine/Source/Runtime/CoreUObject/Private/UObject/LinkerLoad.cpp b/Engine/Source/Runtime/CoreUObject/Private/UObject/LinkerLoad.cpp
--- a/Engine/Source/Runtime/CoreUObject/Private/UObject/LinkerLoad.cpp
+++ b/Engine/Source/Runtime/CoreUObject/Private/UObject/LinkerLoad.cpp
@@ -1718,7 +1718,7 @@ FLinkerLoad::ELinkerStatus FLinkerLoad::FinalizeCreation()
 #endif
 
 		// Add this linker to the object manager's linker array.
-		FLinkerManager::Get().GetLoaders().Add(this);
+		FLinkerManager::Get().AddLoader(this);
 
 		// check if the package source matches the package filename's CRC (if it doens't match, a user saved this package)
 		if (Summary.PackageSource != FCrc::StrCrc_DEPRECATED(*FPaths::GetBaseFilename(Filename).ToUpper()))
@@ -2109,7 +2109,7 @@ FLinkerLoad::EVerifyResult FLinkerLoad::VerifyImport(int32 ImportIndex)
 					// now, fake our Import to be what the redirector pointed to
 					Import.XObject = Redir->DestinationObject;
 					FUObjectThreadContext::Get().ImportCount++;
-					FLinkerManager::Get().GetLoadersWithNewImports().Add(this);
+					FLinkerManager::Get().AddLoaderWithNewImports(this);
 				}
 			}
 		}
@@ -2500,7 +2500,7 @@ bool FLinkerLoad::VerifyImportInner(const int32 ImportIndex, FString& WarningSuf
 				// except if we are looking for _the_ package...in which case we are looking for TmpPkg, so we are done
 				Import.XObject = TmpPkg;
 				FUObjectThreadContext::Get().ImportCount++;
-				FLinkerManager::Get().GetLoadersWithNewImports().Add(this);
+				FLinkerManager::Get().AddLoaderWithNewImports(this);
 				return false;
 			}
 		}
@@ -2553,7 +2553,7 @@ bool FLinkerLoad::VerifyImportInner(const int32 ImportIndex, FString& WarningSuf
 				{
 					Import.XObject = FindObject;
 					FUObjectThreadContext::Get().ImportCount++;
-					FLinkerManager::Get().GetLoadersWithNewImports().Add(this);
+					FLinkerManager::Get().AddLoaderWithNewImports(this);
 				}
 				else
 				{
@@ -3500,7 +3500,7 @@ UObject* FLinkerLoad::CreateExport( int32 Index )
 				LoadClass = UObjectRedirector::StaticClass();
 				// Create new import for UObjectRedirector class
 				FObjectImport* RedirectorImport = new( ImportMap )FObjectImport( UObjectRedirector::StaticClass() );
-				FLinkerManager::Get().GetLoadersWithNewImports().Add(this);
+				FLinkerManager::Get().AddLoaderWithNewImports(this);
 				FUObjectThreadContext::Get().ImportCount++;
 				Export.ClassIndex = FPackageIndex::FromImport(ImportMap.Num() - 1);
 				Export.Object = Redirector;
@@ -3799,7 +3799,7 @@ UObject* FLinkerLoad::CreateImport( int32 Index )
 						// Associate import and indicate that we associated an import for later cleanup.
 						Import.XObject = FindObject;
 						FUObjectThreadContext::Get().ImportCount++;
-						FLinkerManager::Get().GetLoadersWithNewImports().Add(this);
+						FLinkerManager::Get().AddLoaderWithNewImports(this);
 					}
 				}
 			}
@@ -3836,7 +3836,7 @@ UObject* FLinkerLoad::CreateImport( int32 Index )
 				}
 #endif
 				FUObjectThreadContext::Get().ImportCount++;
-				FLinkerManager::Get().GetLoadersWithNewImports().Add(this);
+				FLinkerManager::Get().AddLoaderWithNewImports(this);
 			}
 		}
 	}
@@ -3914,8 +3914,8 @@ void FLinkerLoad::Detach()
 	}
 
 	// Remove from object manager, if it has been added.
-	FLinkerManager::Get().GetLoaders().Remove(this);
-	FLinkerManager::Get().GetLoadersWithNewImports().Remove(this);
+	FLinkerManager::Get().RemoveLoader(this);
+	FLinkerManager::Get().RemoveLoaderWithNewImports(this);
 	if (!FPlatformProperties::HasEditorOnlyData())
 	{
 		FUObjectThreadContext::Get().DelayedLinkerClosePackages.Remove(this);
diff --git a/Engine/Source/Runtime/CoreUObject/Private/UObject/LinkerManager.cpp b/Engine/Source/Runtime/CoreUObject/Private/UObject/LinkerManager.cpp
--- a/Engine/Source/Runtime/CoreUObject/Private/UObject/LinkerManager.cpp
+++ b/Engine/Source/Runtime/CoreUObject/Private/UObject/LinkerManager.cpp
@@ -5,6 +5,7 @@
 =============================================================================*/
 #include "CoreUObjectPrivate.h"
 #include "LinkerManager.h"
+#include "UObject/UObjectThreadContext.h"
 
 FLinkerManager& FLinkerManager::Get()
 {
@@ -65,4 +66,108 @@ bool FLinkerManager::Exec(class UWorld* InWorld, const TCHAR* Cmd, FOutputDevice
 		return true;
 	}
 	return false;
+}
+
+void FLinkerManager::ResetLoaders(UObject* InPkg)
+{
+	// Top level package to reset loaders for.
+	UObject*		TopLevelPackage = InPkg ? InPkg->GetOutermost() : NULL;
+
+	// Find loader/ linker associated with toplevel package. We do this upfront as Detach resets LinkerRoot.
+	if (TopLevelPackage)
+	{
+		// Linker to reset/ detach.
+		auto LinkerToReset = FLinkerLoad::FindExistingLinkerForPackage(CastChecked<UPackage>(TopLevelPackage));
+		if (LinkerToReset)
+		{
+			{
+				FScopeLock ObjectLoadersLock(&ObjectLoadersCritical);
+				for (auto Linker : ObjectLoaders)
+				{
+					// Detach LinkerToReset from other linker's import table.
+					if (Linker->LinkerRoot != TopLevelPackage)
+					{
+						for (auto& Import : Linker->ImportMap)
+						{
+							if (Import.SourceLinker == LinkerToReset)
+							{
+								Import.SourceLinker = NULL;
+								Import.SourceIndex = INDEX_NONE;
+							}
+						}
+					}
+					else
+					{
+						check(Linker == LinkerToReset);
+					}
+				}
+			}
+			// Detach linker, also removes from array and sets LinkerRoot to NULL.
+			LinkerToReset->LoadAndDetachAllBulkData();
+			delete LinkerToReset;
+			LinkerToReset = nullptr;
+		}
+	}
+	else
+	{
+		// We just want a copy here
+		TSet<FLinkerLoad*> LinkersToDetach;
+		GetLoaders(LinkersToDetach);
+		for (auto Linker : LinkersToDetach)
+		{
+			// Detach linker, also removes from array and sets LinkerRoot to NULL.
+			Linker->LoadAndDetachAllBulkData();
+			delete Linker;
+		}
+	}
+}
+
+void FLinkerManager::DissociateImportsAndForcedExports()
+{
+	int32& ImportCount = FUObjectThreadContext::Get().ImportCount;
+	if (ImportCount != 0)
+	{
+		TSet<FLinkerLoad*> LocalLoadersWithNewImports;
+		GetLoadersWithNewImportsAndEmpty(LocalLoadersWithNewImports);
+		if (LocalLoadersWithNewImports.Num())
+		{
+			for (auto Linker : LocalLoadersWithNewImports)
+			{
+				for (int32 ImportIndex = 0; ImportIndex < Linker->ImportMap.Num(); ImportIndex++)
+				{
+					FObjectImport& Import = Linker->ImportMap[ImportIndex];
+					if (Import.XObject && !Import.XObject->HasAnyFlags(RF_Native))
+					{
+						Import.XObject = nullptr;
+					}
+					Import.SourceLinker = nullptr;
+					// when the SourceLinker is reset, the SourceIndex must also be reset, or recreating
+					// an import that points to a redirector will fail to find the redirector
+					Import.SourceIndex = INDEX_NONE;
+				}
+			}
+		}
+		ImportCount = 0;
+	}
+
+	int32& ForcedExportCount = FUObjectThreadContext::Get().ForcedExportCount;
+	if (ForcedExportCount)
+	{		
+		TSet<FLinkerLoad*> LocalLoaders;
+		GetLoaders(LocalLoaders);
+		for (auto Linker : LocalLoaders)
+		{
+			//@todo optimization: only dissociate exports for loaders that had forced exports created
+			//@todo optimization: since the last time this function was called.
+			for (auto& Export : Linker->ExportMap)
+			{
+				if (Export.Object && Export.bForcedExport)
+				{
+					Export.Object->SetLinker(nullptr, INDEX_NONE);
+					Export.Object = nullptr;
+				}
+			}
+		}	
+		ForcedExportCount = 0;
+	}	
 }
\ No newline at end of file
diff --git a/Engine/Source/Runtime/CoreUObject/Private/UObject/LinkerManager.h b/Engine/Source/Runtime/CoreUObject/Private/UObject/LinkerManager.h
--- a/Engine/Source/Runtime/CoreUObject/Private/UObject/LinkerManager.h
+++ b/Engine/Source/Runtime/CoreUObject/Private/UObject/LinkerManager.h
@@ -14,14 +14,66 @@ public:
 
 	static FLinkerManager& Get();
 
-	FORCEINLINE TSet<FLinkerLoad*>& GetLoaders()
+	FORCEINLINE void GetLoaders(TSet<FLinkerLoad*>& OutLoaders)
 	{
-		return ObjectLoaders;
+		FScopeLock ObjectLoadersLock(&ObjectLoadersCritical);
+		OutLoaders = ObjectLoaders;
 	}
 
-	FORCEINLINE TSet<FLinkerLoad*>& GetLoadersWithNewImports()
+	FORCEINLINE void GetLoadersAndEmpty(TSet<FLinkerLoad*>& OutLoaders)
 	{
-		return LoadersWithNewImports;
+		FScopeLock ObjectLoadersLock(&ObjectLoadersCritical);
+		OutLoaders = ObjectLoaders;
+		OutLoaders.Empty();
+	}
+
+	FORCEINLINE void AddLoader(FLinkerLoad* LinkerLoad)
+	{
+		FScopeLock ObjectLoadersLock(&ObjectLoadersCritical);
+		ObjectLoaders.Add(LinkerLoad);
+	}
+
+	FORCEINLINE void RemoveLoader(FLinkerLoad* LinkerLoad)
+	{
+		FScopeLock ObjectLoadersLock(&ObjectLoadersCritical);
+		ObjectLoaders.Remove(LinkerLoad);
+	}
+
+	FORCEINLINE void EmptyLoaders(FLinkerLoad* LinkerLoad)
+	{
+		FScopeLock ObjectLoadersLock(&ObjectLoadersCritical);
+		ObjectLoaders.Empty();
+	}
+
+	FORCEINLINE void GetLoadersWithNewImports(TSet<FLinkerLoad*>& OutLoaders)
+	{
+		FScopeLock ObjectLoadersLock(&LoadersWithNewImportsCritical);
+		OutLoaders = LoadersWithNewImports;
+	}
+
+	FORCEINLINE void GetLoadersWithNewImportsAndEmpty(TSet<FLinkerLoad*>& OutLoaders)
+	{
+		FScopeLock ObjectLoadersLock(&LoadersWithNewImportsCritical);
+		OutLoaders = LoadersWithNewImports;
+		OutLoaders.Empty();
+	}
+
+	FORCEINLINE void AddLoaderWithNewImports(FLinkerLoad* LinkerLoad)
+	{
+		FScopeLock ObjectLoadersLock(&LoadersWithNewImportsCritical);
+		LoadersWithNewImports.Add(LinkerLoad);
+	}
+
+	FORCEINLINE void RemoveLoaderWithNewImports(FLinkerLoad* LinkerLoad)
+	{
+		FScopeLock ObjectLoadersLock(&LoadersWithNewImportsCritical);
+		LoadersWithNewImports.Remove(LinkerLoad);
+	}
+
+	FORCEINLINE void EmptyLoadersWithNewImports(FLinkerLoad* LinkerLoad)
+	{
+		FScopeLock ObjectLoadersLock(&LoadersWithNewImportsCritical);
+		LoadersWithNewImports.Empty();
 	}
 
 #if UE_BUILD_DEBUG
@@ -34,14 +86,27 @@ public:
 	// FSelfRegisteringExec interface
 	virtual bool Exec(class UWorld* InWorld, const TCHAR* Cmd, FOutputDevice& Ar) override;
 
+	/** Empty the loaders */
+	void ResetLoaders(UObject* InPkg);
+
+	/**
+	* Dissociates all linker import and forced export object references. This currently needs to
+	* happen as the referred objects might be destroyed at any time.
+	*/
+	void DissociateImportsAndForcedExports();
+
 private:
 
 	/** Map of packages to their open linkers **/
 	TSet<FLinkerLoad*> ObjectLoaders;
+	FCriticalSection ObjectLoadersCritical;
+
 	/** List of loaders that have new imports **/
 	TSet<FLinkerLoad*> LoadersWithNewImports;
+	FCriticalSection LoadersWithNewImportsCritical;
 #if UE_BUILD_DEBUG
 	/** List of all the existing linker loaders **/
 	TArray<FLinkerLoad*> LiveLinkers;
 #endif
+	
 };
diff --git a/Engine/Source/Runtime/CoreUObject/Public/Misc/StringAssetReference.h b/Engine/Source/Runtime/CoreUObject/Public/Misc/StringAssetReference.h
--- a/Engine/Source/Runtime/CoreUObject/Public/Misc/StringAssetReference.h
+++ b/Engine/Source/Runtime/CoreUObject/Public/Misc/StringAssetReference.h
@@ -95,11 +95,11 @@ struct COREUOBJECT_API FStringAssetReference
 
 	static int32 GetCurrentTag()
 	{
-		return CurrentTag;
+		return CurrentTag.GetValue();
 	}
 	static int32 InvalidateTag()
 	{
-		return CurrentTag++;
+		return CurrentTag.Increment();
 	}
 
 	static FStringAssetReference GetOrCreateIDForObject(const UObject *Object);
@@ -116,7 +116,7 @@ private:
 	void FixupForPIE();
 
 	/** Global counter that determines when we need to re-search for GUIDs because more objects have been loaded **/
-	static int32 CurrentTag;
+	static FThreadSafeCounter CurrentTag;
 	static TArray<FString> PackageNamesBeingDuplicatedForPIE;
 };
 
diff --git a/Engine/Source/Runtime/CoreUObject/Public/UObject/LazyObjectPtr.h b/Engine/Source/Runtime/CoreUObject/Public/UObject/LazyObjectPtr.h
--- a/Engine/Source/Runtime/CoreUObject/Public/UObject/LazyObjectPtr.h
+++ b/Engine/Source/Runtime/CoreUObject/Public/UObject/LazyObjectPtr.h
@@ -98,11 +98,11 @@ struct COREUOBJECT_API FUniqueObjectGuid
 
 	static int32 GetCurrentTag()
 	{
-		return CurrentAnnotationTag;
+		return CurrentAnnotationTag.GetValue();
 	}
 	static int32 InvalidateTag()
 	{
-		return CurrentAnnotationTag++;
+		return CurrentAnnotationTag.Increment();
 	}
 
 	static FUniqueObjectGuid GetOrCreateIDForObject(const class UObject *Object);
@@ -111,7 +111,7 @@ private:
 	FGuid Guid;
 
 	/** Global counter that determines when we need to re-search for GUIDs because more objects have been loaded **/
-	static int32 CurrentAnnotationTag;
+	static FThreadSafeCounter CurrentAnnotationTag;
 };
 
 
diff --git a/Engine/Source/Runtime/CoreUObject/Public/UObject/PersistentObjectPtr.h b/Engine/Source/Runtime/CoreUObject/Public/UObject/PersistentObjectPtr.h
--- a/Engine/Source/Runtime/CoreUObject/Public/UObject/PersistentObjectPtr.h
+++ b/Engine/Source/Runtime/CoreUObject/Public/UObject/PersistentObjectPtr.h
@@ -56,7 +56,6 @@ public:
 	 */
 	FORCEINLINE void operator=(const class UObject* Object)
 	{
-		checkSlow(IsInGameThread());
 		if (Object)
 		{
 			ObjectID = TObjectID::GetOrCreateIDForObject(Object);
diff --git a/Engine/Source/Runtime/CoreUObject/Public/UObject/UObjectAnnotation.h b/Engine/Source/Runtime/CoreUObject/Public/UObject/UObjectAnnotation.h
--- a/Engine/Source/Runtime/CoreUObject/Public/UObject/UObjectAnnotation.h
+++ b/Engine/Source/Runtime/CoreUObject/Public/UObject/UObjectAnnotation.h
@@ -498,6 +498,7 @@ public:
 		}
 		else
 		{
+			FScopeLock AnnotationArrayLock(&AnnotationArrayCritical);
 			if (AnnotationArray.Num() == 0)
 			{
 				// we are adding the first one, so if we are auto removing or verifying removal, register now
@@ -538,6 +539,7 @@ public:
 	void RemoveAnnotation(int32 Index)
 	{
 		check(Index >= 0);
+		FScopeLock AnnotationArrayLock(&AnnotationArrayCritical);
 		if (Index <  AnnotationArray.Num())
 		{
 			AnnotationArray[Index] = TAnnotation();
@@ -549,6 +551,7 @@ public:
 	 */
 	void RemoveAllAnnotations()
 	{
+		FScopeLock AnnotationArrayLock(&AnnotationArrayCritical);
 		bool bHadElements = (AnnotationArray.Num() > 0);
 		AnnotationArray.Empty();
 		if (bHadElements)
@@ -582,6 +585,7 @@ public:
 	FORCEINLINE TAnnotation GetAnnotation(int32 Index)
 	{
 		check(Index >= 0);
+		FScopeLock AnnotationArrayLock(&AnnotationArrayCritical);
 		if (Index < AnnotationArray.Num())
 		{
 			return AnnotationArray[Index];
@@ -610,6 +614,7 @@ public:
 	 */
 	FORCEINLINE TAnnotation& GetAnnotationRef(int32 Index)
 	{
+		FScopeLock AnnotationArrayLock(&AnnotationArrayCritical);
 		if (Index >= AnnotationArray.Num())
 		{
 			AddAnnotation(Index, TAnnotation());
@@ -623,6 +628,7 @@ private:
 	 * Map from live objects to an annotation
 	 */
 	TArray<TAnnotation> AnnotationArray;
+	FCriticalSection AnnotationArrayCritical;
 
 };
 
