UE-4679 - Remove CPF_RepRetry

diff --git a/Engine/Source/Developer/BlueprintCompilerCppBackend/Private/BlueprintCompilerCppBackendUtils.h b/Engine/Source/Developer/BlueprintCompilerCppBackend/Private/BlueprintCompilerCppBackendUtils.h
--- a/Engine/Source/Developer/BlueprintCompilerCppBackend/Private/BlueprintCompilerCppBackendUtils.h
+++ b/Engine/Source/Developer/BlueprintCompilerCppBackend/Private/BlueprintCompilerCppBackendUtils.h
@@ -213,7 +213,6 @@ struct FEmitHelper
 			HANDLE_CPF_TAG(TEXT("NoClear"), CPF_NoClear)
 			HANDLE_CPF_TAG(TEXT("EditFixedSize"), CPF_EditFixedSize)
 			HANDLE_CPF_TAG(TEXT("NotReplicated"), CPF_RepSkip)
-			HANDLE_CPF_TAG(TEXT("RepRetry"), CPF_RepRetry)
 			HANDLE_CPF_TAG(TEXT("Interp"), CPF_Edit | CPF_BlueprintVisible | CPF_Interp)
 			HANDLE_CPF_TAG(TEXT("NonTransactional"), CPF_NonTransactional)
 			HANDLE_CPF_TAG(TEXT("BlueprintAssignable"), CPF_BlueprintAssignable)
diff --git a/Engine/Source/Programs/UnrealHeaderTool/Private/HeaderParser.cpp b/Engine/Source/Programs/UnrealHeaderTool/Private/HeaderParser.cpp
--- a/Engine/Source/Programs/UnrealHeaderTool/Private/HeaderParser.cpp
+++ b/Engine/Source/Programs/UnrealHeaderTool/Private/HeaderParser.cpp
@@ -2716,7 +2716,7 @@ void FHeaderParser::GetVarType
 			}
 			else if (Specifier == TEXT("RepRetry"))
 			{
-				Flags |= CPF_RepRetry;
+				FError::Throwf(TEXT("'RepRetry' is deprecated."));
 			}
 			else if (Specifier == TEXT("Interp"))
 			{
@@ -3471,11 +3471,6 @@ void FHeaderParser::GetVarType
 		FError::Throwf(TEXT("Not allowed to use 'config' with object variables"));
 	}
 
-	if ((VarProperty.PropertyFlags & CPF_RepRetry) && VarProperty.Type != CPT_Struct)
-	{
-		FError::Throwf(TEXT("'RepRetry' is only allowed on struct properties"));
-	}
-
 	if ((VarProperty.PropertyFlags & CPF_BlueprintAssignable) && VarProperty.Type != CPT_MulticastDelegate)
 	{
 		FError::Throwf(TEXT("'BlueprintAssignable' is only allowed on multicast delegate properties"));
diff --git a/Engine/Source/Runtime/CoreUObject/Private/UObject/Obj.cpp b/Engine/Source/Runtime/CoreUObject/Private/UObject/Obj.cpp
--- a/Engine/Source/Runtime/CoreUObject/Private/UObject/Obj.cpp
+++ b/Engine/Source/Runtime/CoreUObject/Private/UObject/Obj.cpp
@@ -2444,7 +2444,7 @@ TArray<const TCHAR*> ParsePropertyFlags(uint64 Flags)
 		TEXT("CPF_NonTransactional"),
 		TEXT("CPF_EditorOnly"),
 		TEXT("CPF_NoDestructor"),
-		TEXT("CPF_RepRetry"),
+		TEXT("0x0000002000000000"),
 		TEXT("CPF_AutoWeak"),
 		TEXT("CPF_ContainsInstancedReference"),
 		TEXT("CPF_AssetRegistrySearchable"),
diff --git a/Engine/Source/Runtime/CoreUObject/Public/UObject/ObjectBase.h b/Engine/Source/Runtime/CoreUObject/Public/UObject/ObjectBase.h
--- a/Engine/Source/Runtime/CoreUObject/Public/UObject/ObjectBase.h
+++ b/Engine/Source/Runtime/CoreUObject/Public/UObject/ObjectBase.h
@@ -346,7 +346,7 @@ typedef uint64 EClassCastFlags;
 #define CPF_NonTransactional				DECLARE_UINT64(0x0000000400000000)		// Property isn't transacted
 #define CPF_EditorOnly						DECLARE_UINT64(0x0000000800000000)		// Property should only be loaded in the editor
 #define CPF_NoDestructor					DECLARE_UINT64(0x0000001000000000)		// No destructor
-#define CPF_RepRetry						DECLARE_UINT64(0x0000002000000000)		// retry replication of this property if it fails to be fully sent (e.g. object references not yet available to serialize over the network)
+//#define CPF_								DECLARE_UINT64(0x0000002000000000)		//
 #define CPF_AutoWeak						DECLARE_UINT64(0x0000004000000000)		// Only used for weak pointers, means the export type is autoweak
 #define CPF_ContainsInstancedReference		DECLARE_UINT64(0x0000008000000000)		// Property contains component references.
 #define CPF_AssetRegistrySearchable			DECLARE_UINT64(0x0000010000000000)		// asset instances will add properties with this flag to the asset registry automatically
@@ -803,9 +803,6 @@ namespace UP
 		/// Property is relevant to network replication. Notify actors when a property is replicated (usage: ReplicatedUsing=FunctionName).
 		ReplicatedUsing,
 
-		/// Retry replication of this property if it fails to be fully sent (e.g. object references not yet available to serialize over the network)
-		RepRetry,
-
 		/// Skip replication (only for struct members and parameters in service request functions).
 		NotReplicated,
 
