Don't need to assert on deferred super functions not being fully loaded (when replacing linker placeholder functions)... the regenerate of the super class will result in that function being fully loaded.

diff --git a/Engine/Source/Runtime/CoreUObject/Private/Blueprint/BlueprintSupport.cpp b/Engine/Source/Runtime/CoreUObject/Private/Blueprint/BlueprintSupport.cpp
--- a/Engine/Source/Runtime/CoreUObject/Private/Blueprint/BlueprintSupport.cpp
+++ b/Engine/Source/Runtime/CoreUObject/Private/Blueprint/BlueprintSupport.cpp
@@ -895,8 +895,17 @@ int32 FLinkerLoad::ResolveDependencyPlaceholder(FLinkerPlaceholderBase* Placehol
 		RealImportObj = CreateImport(ImportIndex);
 	}
 
+#if USE_DEFERRED_DEPENDENCY_CHECK_VERIFICATION_TESTS
+	UFunction* AsFunction = Cast<UFunction>(RealImportObj);
+	// it's ok if super functions come in not fully loaded (missing 
+	// RF_LoadCompleted... meaning it's in the middle of serializing in somewhere 
+	// up the stack); the function will be forcefully ran through Preload(), 
+	// when we regenerate the super class (see FRegenerationHelper::ForcedLoadMembers)
+	bool const bIsSuperFunction = (AsFunction != nullptr) && (ReferencingClass != nullptr) && ReferencingClass->IsChildOf(AsFunction->GetOwnerClass());
+
 	DEFERRED_DEPENDENCY_CHECK(RealImportObj != PlaceholderObj);
-	DEFERRED_DEPENDENCY_CHECK(RealImportObj == nullptr || RealImportObj->HasAnyFlags(RF_LoadCompleted));
+	DEFERRED_DEPENDENCY_CHECK(RealImportObj == nullptr || bIsSuperFunction || RealImportObj->HasAnyFlags(RF_LoadCompleted));
+#endif // USE_DEFERRED_DEPENDENCY_CHECK_VERIFICATION_TESTS
 
 	int32 ReplacementCount = 0;
 	if (ReferencingClass != nullptr)
