Brought back the hash efficiency execs we had in UE3 ("obj hash" and "obj hashouter") to debug reports of bad perf

diff --git a/Engine/Source/Runtime/CoreUObject/Private/UObject/UObjectHash.cpp b/Engine/Source/Runtime/CoreUObject/Private/UObject/UObjectHash.cpp
--- a/Engine/Source/Runtime/CoreUObject/Private/UObject/UObjectHash.cpp
+++ b/Engine/Source/Runtime/CoreUObject/Private/UObject/UObjectHash.cpp
@@ -656,9 +656,88 @@ void LockUObjectHashTablesForGC()
 }
 
 /**
-* Releases UObject hash tables lock after GC has finished running
-*/
+ * Releases UObject hash tables lock after GC has finished running
+ */
 void UnlockUObjectHashTablesForGC()
 {
 	FUObjectHashTables::Get().Unlock();
-}
\ No newline at end of file
+}
+
+void LogHashStatisticsInternal(TMultiMap<int32, UObjectBase*>& Hash, FOutputDevice& Ar, const bool bShowHashBucketCollisionInfo)
+{
+	TArray<int32> HashBuckets;
+	// Get the set of keys in use, which is the number of hash buckets
+	int32 SlotsInUse = Hash.GetKeys(HashBuckets);
+
+	int32 TotalCollisions = 0;
+	int32 MinCollisions = OBJECT_HASH_BINS;
+	int32 MaxCollisions = 0;
+	int32 MaxBin = 0;
+
+	// Dump how many slots are in use
+	Ar.Logf(TEXT("Slots in use %d"), SlotsInUse);
+
+	// Work through each slot and figure out how many collisions
+	for (auto HashBucket : HashBuckets)
+	{
+		int32 Collisions = 0;
+
+		for (TMultiMap<int32, UObjectBase*>::TConstKeyIterator HashIt(Hash, HashBucket); HashIt; ++HashIt)
+		{
+			// There's one collision per object in a given bucket
+			Collisions++;
+		}
+
+		// Keep the global stats
+		TotalCollisions += Collisions;
+		if (Collisions > MaxCollisions)
+		{
+			MaxBin = HashBucket;
+		}
+		MaxCollisions = FMath::Max<int32>(Collisions, MaxCollisions);
+		MinCollisions = FMath::Min<int32>(Collisions, MinCollisions);
+
+		if (bShowHashBucketCollisionInfo)
+		{
+			// Now log the output
+			Ar.Logf(TEXT("\tSlot %d has %d collisions"), HashBucket, Collisions);
+		}
+	}
+	Ar.Logf(TEXT(""));
+
+	// Dump the first 30 objects in the worst bin for inspection
+	Ar.Logf(TEXT("Worst hash bucket contains:"));
+	int32 Count = 0;
+	for (TMultiMap<int32, UObjectBase*>::TConstKeyIterator HashIt(Hash, MaxBin); HashIt && Count < 30; ++HashIt)
+	{
+		UObject* Object = (UObject*)HashIt.Value();
+		Ar.Logf(TEXT("\tObject is %s (%s)"), *Object->GetName(), *Object->GetFullName());
+		Count++;
+	}
+	Ar.Logf(TEXT(""));
+
+	// Now dump how efficient the hash is
+	Ar.Logf(TEXT("Collision Stats: Best Case (%d), Average Case (%d), Worst Case (%d)"),
+		MinCollisions,
+		FMath::FloorToInt(((float)TotalCollisions / (float)SlotsInUse)),
+		MaxCollisions);
+}
+
+void LogHashStatistics(FOutputDevice& Ar, const bool bShowHashBucketCollisionInfo)
+{
+	Ar.Logf(TEXT("Hash efficiency statistics for the Object Hash"));
+	Ar.Logf(TEXT("-------------------------------------------------"));
+	Ar.Logf(TEXT(""));
+	LogHashStatisticsInternal(ObjectHash, Ar, bShowHashBucketCollisionInfo);
+	Ar.Logf(TEXT(""));
+}
+
+void LogHashOuterStatistics(FOutputDevice& Ar, const bool bShowHashBucketCollisionInfo)
+{
+	Ar.Logf(TEXT("Hash efficiency statistics for the Outer Object Hash"));
+	Ar.Logf(TEXT("-------------------------------------------------"));
+	Ar.Logf(TEXT(""));
+	LogHashStatisticsInternal(ObjectHashOuter, Ar, bShowHashBucketCollisionInfo);
+	Ar.Logf(TEXT(""));
+}
+
diff --git a/Engine/Source/Runtime/CoreUObject/Public/UObject/UObjectHash.h b/Engine/Source/Runtime/CoreUObject/Public/UObject/UObjectHash.h
--- a/Engine/Source/Runtime/CoreUObject/Public/UObject/UObjectHash.h
+++ b/Engine/Source/Runtime/CoreUObject/Public/UObject/UObjectHash.h
@@ -84,11 +84,28 @@ void HashObject(class UObjectBase* Object);
 void UnhashObject(class UObjectBase* Object);
 
 /**
-* Adds a uobject to the global array which is used for uobject iteration
-*
-* @param	Object Object to allocate an index for
-*/
+ * Logs out information about the object hash for debug purposes
+ *
+ * @param Ar the archive to write the log data to
+ * @param bShowHashBucketCollisionInfo whether to log each bucket's collision count
+ */
+COREUOBJECT_API void LogHashStatistics(FOutputDevice& Ar, const bool bShowHashBucketCollisionInfo);
+
+/**
+ * Logs out information about the outer object hash for debug purposes
+ *
+ * @param Ar the archive to write the log data to
+ * @param bShowHashBucketCollisionInfo whether to log each bucket's collision count
+ */
+COREUOBJECT_API void LogHashOuterStatistics(FOutputDevice& Ar, const bool bShowHashBucketCollisionInfo);
+
+/**
+ * Adds a uobject to the global array which is used for uobject iteration
+ *
+ * @param	Object Object to allocate an index for
+ */
 void AllocateUObjectIndexForCurrentThread(class UObjectBase* Object);
 
+
 #endif	// __UOBJECTHASH_H__
 
diff --git a/Engine/Source/Runtime/Engine/Private/UnrealEngine.cpp b/Engine/Source/Runtime/Engine/Private/UnrealEngine.cpp
--- a/Engine/Source/Runtime/Engine/Private/UnrealEngine.cpp
+++ b/Engine/Source/Runtime/Engine/Private/UnrealEngine.cpp
@@ -5357,6 +5357,18 @@ bool UEngine::HandleObjCommand( const TCHAR* Cmd, FOutputDevice& Ar )
 
 		return true;
 	}
+	else if (FParse::Command(&Cmd, TEXT("HASH")))
+	{
+		const bool bShowHashBucketCollisionInfo = FParse::Param(Cmd, TEXT("SHOWBUCKETCOLLISIONS"));
+		LogHashStatistics(Ar, bShowHashBucketCollisionInfo);
+		return true;
+	}
+	else if (FParse::Command(&Cmd, TEXT("HASHOUTER")))
+	{
+		const bool bShowHashBucketCollisionInfo = FParse::Param(Cmd, TEXT("SHOWBUCKETCOLLISIONS"));
+		LogHashOuterStatistics(Ar, bShowHashBucketCollisionInfo);
+		return true;
+	}
 	else
 	{
 		// OBJ command but not supported here
@@ -7713,17 +7725,7 @@ UNetDriver* FindNamedNetDriver_Local(const TArray<FNamedNetDriver>& ActiveNetDri
 
 UNetDriver* UEngine::FindNamedNetDriver(UWorld * InWorld, FName NetDriverName)
 {
-#if WITH_EDITOR
-	const FWorldContext* WorldContext = GetWorldContextFromWorld(InWorld);
-	if (WorldContext)
-	{
-		return FindNamedNetDriver_Local(WorldContext->ActiveNetDrivers, NetDriverName);
-	}
-		
-	return nullptr;
-#else
 	return FindNamedNetDriver_Local(GetWorldContextFromWorldChecked(InWorld).ActiveNetDrivers, NetDriverName);
-#endif // WITH_EDITOR
 }
 
 UNetDriver* UEngine::FindNamedNetDriver(const UPendingNetGame* InPendingNetGame, FName NetDriverName)
