Fix broken build.

diff --git a/Engine/Source/Runtime/CoreUObject/Private/UObject/UObjectHash.cpp b/Engine/Source/Runtime/CoreUObject/Private/UObject/UObjectHash.cpp
--- a/Engine/Source/Runtime/CoreUObject/Private/UObject/UObjectHash.cpp
+++ b/Engine/Source/Runtime/CoreUObject/Private/UObject/UObjectHash.cpp
@@ -728,7 +728,7 @@ void LogHashStatistics(FOutputDevice& Ar, const bool bShowHashBucketCollisionInf
 	Ar.Logf(TEXT("Hash efficiency statistics for the Object Hash"));
 	Ar.Logf(TEXT("-------------------------------------------------"));
 	Ar.Logf(TEXT(""));
-	LogHashStatisticsInternal(ObjectHash, Ar, bShowHashBucketCollisionInfo);
+	LogHashStatisticsInternal(FUObjectHashTables::Get().Hash, Ar, bShowHashBucketCollisionInfo);
 	Ar.Logf(TEXT(""));
 }
 
@@ -737,7 +737,7 @@ void LogHashOuterStatistics(FOutputDevice& Ar, const bool bShowHashBucketCollisi
 	Ar.Logf(TEXT("Hash efficiency statistics for the Outer Object Hash"));
 	Ar.Logf(TEXT("-------------------------------------------------"));
 	Ar.Logf(TEXT(""));
-	LogHashStatisticsInternal(ObjectHashOuter, Ar, bShowHashBucketCollisionInfo);
+	LogHashStatisticsInternal(FUObjectHashTables::Get().HashOuter, Ar, bShowHashBucketCollisionInfo);
 	Ar.Logf(TEXT(""));
 }
 
