Fixed my missing lock/unlock of the object hashes

diff --git a/Engine/Source/Runtime/CoreUObject/Private/UObject/UObjectHash.cpp b/Engine/Source/Runtime/CoreUObject/Private/UObject/UObjectHash.cpp
--- a/Engine/Source/Runtime/CoreUObject/Private/UObject/UObjectHash.cpp
+++ b/Engine/Source/Runtime/CoreUObject/Private/UObject/UObjectHash.cpp
@@ -728,6 +728,7 @@ void LogHashStatistics(FOutputDevice& Ar, const bool bShowHashBucketCollisionInf
 	Ar.Logf(TEXT("Hash efficiency statistics for the Object Hash"));
 	Ar.Logf(TEXT("-------------------------------------------------"));
 	Ar.Logf(TEXT(""));
+	FHashTableLock HashLock(FUObjectHashTables::Get());
 	LogHashStatisticsInternal(FUObjectHashTables::Get().Hash, Ar, bShowHashBucketCollisionInfo);
 	Ar.Logf(TEXT(""));
 }
@@ -737,6 +738,7 @@ void LogHashOuterStatistics(FOutputDevice& Ar, const bool bShowHashBucketCollisi
 	Ar.Logf(TEXT("Hash efficiency statistics for the Outer Object Hash"));
 	Ar.Logf(TEXT("-------------------------------------------------"));
 	Ar.Logf(TEXT(""));
+	FHashTableLock HashLock(FUObjectHashTables::Get());
 	LogHashStatisticsInternal(FUObjectHashTables::Get().HashOuter, Ar, bShowHashBucketCollisionInfo);
 	Ar.Logf(TEXT(""));
 }
