Only warn about async loading time limit exceeded in cooked builds

diff --git a/Engine/Source/Runtime/CoreUObject/Private/Serialization/AsyncLoading.cpp b/Engine/Source/Runtime/CoreUObject/Private/Serialization/AsyncLoading.cpp
--- a/Engine/Source/Runtime/CoreUObject/Private/Serialization/AsyncLoading.cpp
+++ b/Engine/Source/Runtime/CoreUObject/Private/Serialization/AsyncLoading.cpp
@@ -177,8 +177,8 @@ static FORCEINLINE bool IsTimeLimitExceeded(double InTickStartTime, bool bUseTim
 		double CurrentTime = FPlatformTime::Seconds();
 		bTimeLimitExceeded = CurrentTime - InTickStartTime > InTimeLimit;
 
-		// Log single operations that take longer than timelimit.
-		if ((CurrentTime - InTickStartTime) > (1.5 * InTimeLimit))
+		// Log single operations that take longer than timelimit (but only in cooked builds)
+		if (FPlatformProperties::RequiresCookedData() && (CurrentTime - InTickStartTime) > (1.5 * InTimeLimit))
 		{
 			UE_LOG(LogStreaming, Warning, TEXT("IsTimeLimitExceeded: %s %s took (less than) %5.2f ms"),
 				InLastTypeOfWorkPerformed ? InLastTypeOfWorkPerformed : TEXT("unknown"),
