Avoid shadow variable warnings when a class has a UFUNCTION with a parameter with the same name as a member
Avoid shadow variable warnings when a class has a member named Result

diff --git a/Engine/Plugins/Experimental/StructBox/Source/StructBox/Classes/StructBoxLibrary.h b/Engine/Plugins/Experimental/StructBox/Source/StructBox/Classes/StructBoxLibrary.h
--- a/Engine/Plugins/Experimental/StructBox/Source/StructBox/Classes/StructBoxLibrary.h
+++ b/Engine/Plugins/Experimental/StructBox/Source/StructBox/Classes/StructBoxLibrary.h
@@ -46,7 +46,7 @@ public:
 
 		P_FINISH;
 
-		*(bool*)Result = bResult;
+		*(bool*)RESULT_PARAM = bResult;
 	}
 
 	UFUNCTION(BlueprintCallable, CustomThunk, Category = "StructBox", meta = (CustomStructureParam = "CustomStruct"))
diff --git a/Engine/Source/Programs/UnrealHeaderTool/Private/CodeGenerator.cpp b/Engine/Source/Programs/UnrealHeaderTool/Private/CodeGenerator.cpp
--- a/Engine/Source/Programs/UnrealHeaderTool/Private/CodeGenerator.cpp
+++ b/Engine/Source/Programs/UnrealHeaderTool/Private/CodeGenerator.cpp
@@ -3762,13 +3762,13 @@ void FNativeClassHeaderGenerator::ExportFunctionThunk(FUHTStringBuilder& RPCWrap
 		TypeText = ReplacementText;
 
 		FString DefaultValueText;
-		FString ParamPrefix;
+		FString ParamPrefix = TEXT("Z_Param_");
 
 		// if this property is an out parm, add the REF tag
 		if (Param->PropertyFlags & CPF_OutParm)
 		{
 			EvalModifierText += TEXT("_REF");
-			ParamPrefix = TEXT("Out_");
+			ParamPrefix += TEXT("Out_");
 		}
 
 		// if this property requires a specialization, add a comma to the type name so we can print it out easily
@@ -3895,7 +3895,7 @@ void FNativeClassHeaderGenerator::ExportFunctionThunk(FUHTStringBuilder& RPCWrap
 		ReplacementText += ReturnType;
 		ApplyAlternatePropertyExportText(Return, ReplacementText);
 		ReturnType = ReplacementText;
-		RPCWrappers.Logf(TEXT("*(%s%s*)Result="), *ReturnType, *ReturnExtendedType);
+		RPCWrappers.Logf(TEXT("*(%s%s*)") TEXT(PREPROCESSOR_TO_STRING(RESULT_PARAM)) TEXT("="), *ReturnType, *ReturnExtendedType);
 	}
 
 	// export the call to the C++ version
diff --git a/Engine/Source/Runtime/CoreUObject/Private/UObject/Class.cpp b/Engine/Source/Runtime/CoreUObject/Private/UObject/Class.cpp
--- a/Engine/Source/Runtime/CoreUObject/Private/UObject/Class.cpp
+++ b/Engine/Source/Runtime/CoreUObject/Private/UObject/Class.cpp
@@ -3786,7 +3786,7 @@ void UFunction::Invoke(UObject* Obj, FFrame& Stack, RESULT_DECL)
 	}
 
 	TGuardValue<UFunction*> NativeFuncGuard(Stack.CurrentNativeFunction, this);
-	return (Obj->*Func)(Stack, Result);
+	return (Obj->*Func)(Stack, RESULT_PARAM);
 }
 
 void UFunction::Serialize( FArchive& Ar )
diff --git a/Engine/Source/Runtime/CoreUObject/Private/UObject/ScriptCore.cpp b/Engine/Source/Runtime/CoreUObject/Private/UObject/ScriptCore.cpp
--- a/Engine/Source/Runtime/CoreUObject/Private/UObject/ScriptCore.cpp
+++ b/Engine/Source/Runtime/CoreUObject/Private/UObject/ScriptCore.cpp
@@ -123,7 +123,7 @@ FEditorScriptExecutionGuard::~FEditorScriptExecutionGuard()
 void FFrame::Step(UObject *Context, RESULT_DECL)
 {
 	int32 B = *Code++;
-	(Context->*GNatives[B])(*this,Result);
+	(Context->*GNatives[B])(*this,RESULT_PARAM);
 }
 
 void FFrame::StepExplicitProperty(void*const Result, UProperty* Property)
@@ -372,9 +372,9 @@ void UObject::SkipFunction(FFrame& Stack, RESULT_DECL, UFunction* Function)
 	if (ReturnProp != NULL)
 	{
 		// destroy old value if necessary
-		ReturnProp->DestroyValue(Result);
+		ReturnProp->DestroyValue(RESULT_PARAM);
 		// copy zero value for return property into Result
-		FMemory::Memzero(Result, ReturnProp->ArrayDim * ReturnProp->ElementSize);
+		FMemory::Memzero(RESULT_PARAM, ReturnProp->ArrayDim * ReturnProp->ElementSize);
 	}
 }
 
@@ -410,12 +410,12 @@ void UObject::CallFunction( FFrame& Stack, RESULT_DECL, UFunction* Function )
 			FScopeCycleCounterUObject ContextScope(Stack.Object);
 			FScopeCycleCounterUObject FunctionScope(Function);
 
-			Function->Invoke(this, Stack, Result);
+			Function->Invoke(this, Stack, RESULT_PARAM);
 		}
 		else
 		{
 			// Eat up the remaining parameters in the stream.
-			SkipFunction(Stack, Result, Function);
+			SkipFunction(Stack, RESULT_PARAM, Function);
 		}
 	}
 	else
@@ -447,8 +447,8 @@ void UObject::CallFunction( FFrame& Stack, RESULT_DECL, UFunction* Function )
  					FOutParmRec* RetVal = (FOutParmRec*)FMemory_Alloca(sizeof(FOutParmRec));
  
  					// Our context should be that we're in a variable assignment to the return value, so ensure that we have a valid property to return to
- 					check(Result != NULL);
- 					RetVal->PropAddr = (uint8*)Result;
+ 					check(RESULT_PARAM != NULL);
+ 					RetVal->PropAddr = (uint8*)RESULT_PARAM;
  					RetVal->Property = Property;
 					NewStack.OutParms = RetVal;
  
@@ -531,7 +531,7 @@ void UObject::CallFunction( FFrame& Stack, RESULT_DECL, UFunction* Function )
 		// Execute the code.
 		if( bIsValidFunction )
 		{
-			ProcessInternal( NewStack, Result );
+			ProcessInternal( NewStack, RESULT_PARAM );
 		}
 
 		if (!bUsePersistentFrame)
@@ -556,10 +556,10 @@ void ClearReturnValue(UProperty* ReturnProp, RESULT_DECL)
 		// destroy old value if necessary
 		if (!ReturnProp->HasAllPropertyFlags(CPF_NoDestructor))
 		{
-			ReturnProp->DestroyValue(Result);
+			ReturnProp->DestroyValue(RESULT_PARAM);
 		}
 		// copy zero value for return property into Result
-		FMemory::Memzero(Result, ReturnProp->ArrayDim * ReturnProp->ElementSize);
+		FMemory::Memzero(RESULT_PARAM, ReturnProp->ArrayDim * ReturnProp->ElementSize);
 	}
 }
 
@@ -592,7 +592,7 @@ void UObject::ProcessInternal( FFrame& Stack, RESULT_DECL )
 		{
 			// If we have a return property, return a zeroed value in it, to try and save execution as much as possible
 			UProperty* ReturnProp = ((UFunction*)Stack.Node)->GetReturnProperty();
-			ClearReturnValue(ReturnProp, Result);
+			ClearReturnValue(ReturnProp, RESULT_PARAM);
 			return;
 		}
 		else if (++Recurse == RECURSE_LIMIT)
@@ -602,14 +602,14 @@ void UObject::ProcessInternal( FFrame& Stack, RESULT_DECL )
 
 			// If we have a return property, return a zeroed value in it, to try and save execution as much as possible
 			UProperty* ReturnProp = ((UFunction*)Stack.Node)->GetReturnProperty();
-			ClearReturnValue(ReturnProp, Result);
+			ClearReturnValue(ReturnProp, RESULT_PARAM);
 
 			// Notify anyone who cares that we've had a fatal error, so we can shut down PIE, etc
 			const FString Desc = FString::Printf(TEXT("Infinite script recursion (%i calls) detected"), RECURSE_LIMIT);
 			FBlueprintExceptionInfo InfiniteRecursionExceptionInfo(EBlueprintExceptionType::InfiniteLoop, Desc);
 			FBlueprintCoreDelegates::ThrowScriptException(this, Stack, InfiniteRecursionExceptionInfo);
 
-			// This flag prevents repeated warnings of ininite loop, script exception handler 
+			// This flag prevents repeated warnings of infinite loop, script exception handler 
 			// is expected to have terminated execution appropriately:
 			Ranaway = true;
 
@@ -630,7 +630,7 @@ void UObject::ProcessInternal( FFrame& Stack, RESULT_DECL )
 
 				// If we have a return property, return a zeroed value in it, to try and save execution as much as possible
 				UProperty* ReturnProp = ((UFunction*)Stack.Node)->GetReturnProperty();
-				ClearReturnValue(ReturnProp, Result);
+				ClearReturnValue(ReturnProp, RESULT_PARAM);
 
 				// Notify anyone who cares that we've had a fatal error, so we can shut down PIE, etc
 				const FString Desc = FString::Printf(TEXT("Runaway loop detected (over %i iterations)"), GMaximumScriptLoopIterations );
@@ -653,7 +653,7 @@ void UObject::ProcessInternal( FFrame& Stack, RESULT_DECL )
 
 		if (*Stack.Code != EX_Nothing)
 		{
-			Stack.Step(Stack.Object, Result);
+			Stack.Step(Stack.Object, RESULT_PARAM);
 		}
 		else
 		{
@@ -670,9 +670,9 @@ void UObject::ProcessInternal( FFrame& Stack, RESULT_DECL )
 		if (ReturnProp != NULL)
 		{
 			// destroy old value if necessary
-			ReturnProp->DestroyValue(Result);
+			ReturnProp->DestroyValue(RESULT_PARAM);
 			// copy zero value for return property into Result
-			FMemory::Memzero(Result, ReturnProp->ArrayDim * ReturnProp->ElementSize);
+			FMemory::Memzero(RESULT_PARAM, ReturnProp->ArrayDim * ReturnProp->ElementSize);
 		}
 	}
 }
@@ -1011,9 +1011,9 @@ void UObject::execLocalVariable(FFrame& Stack, RESULT_DECL)
 	UProperty* VarProperty = Stack.ReadProperty();
 	Stack.MostRecentPropertyAddress = VarProperty->ContainerPtrToValuePtr<uint8>(Stack.Locals);
 
-	if (Result)
+	if (RESULT_PARAM)
 	{
-		VarProperty->CopyCompleteValueToScriptVM( Result, Stack.MostRecentPropertyAddress );
+		VarProperty->CopyCompleteValueToScriptVM( RESULT_PARAM, Stack.MostRecentPropertyAddress );
 	}
 }
 IMPLEMENT_VM_FUNCTION( EX_LocalVariable, execLocalVariable );
@@ -1023,9 +1023,9 @@ void UObject::execInstanceVariable(FFrame& Stack, RESULT_DECL)
 	UProperty* VarProperty = Stack.ReadProperty();
 	Stack.MostRecentPropertyAddress = VarProperty->ContainerPtrToValuePtr<uint8>(this);
 
-	if (Result)
+	if (RESULT_PARAM)
 	{
-		VarProperty->CopyCompleteValueToScriptVM(Result, Stack.MostRecentPropertyAddress);
+		VarProperty->CopyCompleteValueToScriptVM(RESULT_PARAM, Stack.MostRecentPropertyAddress);
 	}
 }
 IMPLEMENT_VM_FUNCTION( EX_InstanceVariable, execInstanceVariable );
@@ -1048,9 +1048,9 @@ void UObject::execLocalOutVariable(FFrame& Stack, RESULT_DECL)
 	Stack.MostRecentPropertyAddress = Out->PropAddr;
 
 	// if desired, copy the value in that address to Result
-	if (Result && Result != Stack.MostRecentPropertyAddress)
+	if (RESULT_PARAM && RESULT_PARAM != Stack.MostRecentPropertyAddress)
 	{
-		VarProperty->CopyCompleteValueToScriptVM(Result, Stack.MostRecentPropertyAddress);
+		VarProperty->CopyCompleteValueToScriptVM(RESULT_PARAM, Stack.MostRecentPropertyAddress);
 	}
 }
 IMPLEMENT_VM_FUNCTION(EX_LocalOutVariable, execLocalOutVariable);
@@ -1061,10 +1061,10 @@ void UObject::execInterfaceContext(FFrame& Stack, RESULT_DECL)
 	FScriptInterface InterfaceValue;
 	Stack.Step(this, &InterfaceValue);
 
-	if (Result != NULL)
+	if (RESULT_PARAM != NULL)
 	{
 		// copy the UObject pointer to Result
-		*(UObject**)Result = InterfaceValue.GetObject();
+		*(UObject**)RESULT_PARAM = InterfaceValue.GetObject();
 	}
 }
 IMPLEMENT_VM_FUNCTION( EX_InterfaceContext, execInterfaceContext );
@@ -1467,19 +1467,19 @@ IMPLEMENT_VM_FUNCTION( EX_LetMulticastDelegate, execLetMulticastDelegate );
 void UObject::execSelf( FFrame& Stack, RESULT_DECL )
 {
 	// Get Self actor for this context.
-	*(UObject**)Result = this;
+	*(UObject**)RESULT_PARAM = this;
 }
 IMPLEMENT_VM_FUNCTION( EX_Self, execSelf );
 
 void UObject::execContext( FFrame& Stack, RESULT_DECL )
 {
-	ProcessContextOpcode(Stack, Result, /*bCanFailSilently=*/ false);
+	ProcessContextOpcode(Stack, RESULT_PARAM, /*bCanFailSilently=*/ false);
 }
 IMPLEMENT_VM_FUNCTION( EX_Context, execContext );
 
 void UObject::execContext_FailSilent( FFrame& Stack, RESULT_DECL )
 {
-	ProcessContextOpcode(Stack, Result, /*bCanFailSilently=*/ true);
+	ProcessContextOpcode(Stack, RESULT_PARAM, /*bCanFailSilently=*/ true);
 }
 IMPLEMENT_VM_FUNCTION( EX_Context_FailSilent, execContext_FailSilent );
 
@@ -1495,7 +1495,7 @@ void UObject::ProcessContextOpcode( FFrame& Stack, RESULT_DECL, bool bCanFailSil
 	if (IsValid(NewContext))
 	{
 		Stack.Code += sizeof(CodeSkipSizeType) + sizeof(ScriptPointerType) + sizeof(uint8);
-		Stack.Step( NewContext, Result );
+		Stack.Step( NewContext, RESULT_PARAM );
 	}
 	else
 	{
@@ -1530,9 +1530,9 @@ void UObject::ProcessContextOpcode( FFrame& Stack, RESULT_DECL, bool bCanFailSil
 		Stack.MostRecentPropertyAddress = NULL;
 		Stack.MostRecentProperty = NULL;
 
-		if (Result)
+		if (RESULT_PARAM)
 		{
-			FMemory::Memzero( Result, bSize );
+			FMemory::Memzero( RESULT_PARAM, bSize );
 		}
 	}
 }
@@ -1555,9 +1555,9 @@ void UObject::execStructMemberContext(FFrame& Stack, RESULT_DECL)
 		Stack.MostRecentProperty = StructProperty;
 
 		// Handle variable reads
-		if (Result)
+		if (RESULT_PARAM)
 		{
-			StructProperty->CopyCompleteValueToScriptVM(Result, Stack.MostRecentPropertyAddress);
+			StructProperty->CopyCompleteValueToScriptVM(RESULT_PARAM, Stack.MostRecentPropertyAddress);
 		}
 	}
 	else
@@ -1575,14 +1575,14 @@ IMPLEMENT_VM_FUNCTION( EX_StructMemberContext, execStructMemberContext );
 void UObject::execVirtualFunction( FFrame& Stack, RESULT_DECL )
 {
 	// Call the virtual function.
-	CallFunction( Stack, Result, FindFunctionChecked(Stack.ReadName()) );
+	CallFunction( Stack, RESULT_PARAM, FindFunctionChecked(Stack.ReadName()) );
 }
 IMPLEMENT_VM_FUNCTION( EX_VirtualFunction, execVirtualFunction );
 
 void UObject::execFinalFunction( FFrame& Stack, RESULT_DECL )
 {
 	// Call the final function.
-	CallFunction( Stack, Result, (UFunction*)Stack.ReadObject() );
+	CallFunction( Stack, RESULT_PARAM, (UFunction*)Stack.ReadObject() );
 }
 IMPLEMENT_VM_FUNCTION( EX_FinalFunction, execFinalFunction );
 
@@ -1697,26 +1697,26 @@ IMPLEMENT_VM_FUNCTION( EX_ClearMulticastDelegate, execClearMulticastDelegate );
 
 void UObject::execIntConst( FFrame& Stack, RESULT_DECL )
 {
-	*(int32*)Result = Stack.ReadInt();
+	*(int32*)RESULT_PARAM = Stack.ReadInt();
 }
 IMPLEMENT_VM_FUNCTION( EX_IntConst, execIntConst );
 
 void UObject::execSkipOffsetConst( FFrame& Stack, RESULT_DECL )
 {
 	CodeSkipSizeType Literal = Stack.ReadCodeSkipCount();
-	*(int32*)Result = Literal;
+	*(int32*)RESULT_PARAM = Literal;
 }
 IMPLEMENT_VM_FUNCTION( EX_SkipOffsetConst, execSkipOffsetConst );
 
 void UObject::execFloatConst( FFrame& Stack, RESULT_DECL )
 {
-	*(float*)Result = Stack.ReadFloat();
+	*(float*)RESULT_PARAM = Stack.ReadFloat();
 }
 IMPLEMENT_VM_FUNCTION( EX_FloatConst, execFloatConst );
 
 void UObject::execStringConst( FFrame& Stack, RESULT_DECL )
 {
-	*(FString*)Result = (ANSICHAR*)Stack.Code;
+	*(FString*)RESULT_PARAM = (ANSICHAR*)Stack.Code;
 	while( *Stack.Code )
 		Stack.Code++;
 	Stack.Code++;
@@ -1725,7 +1725,7 @@ IMPLEMENT_VM_FUNCTION( EX_StringConst, execStringConst );
 
 void UObject::execUnicodeStringConst( FFrame& Stack, RESULT_DECL )
 {
- 	*(FString*)Result = FString((UCS2CHAR*)Stack.Code);
+ 	*(FString*)RESULT_PARAM = FString((UCS2CHAR*)Stack.Code);
 
 	while( *(uint16*)Stack.Code )
 	{
@@ -1743,20 +1743,20 @@ void UObject::execTextConst( FFrame& Stack, RESULT_DECL )
 	Stack.Step( Stack.Object, &SourceString);
 	Stack.Step( Stack.Object, &KeyString);
 	Stack.Step( Stack.Object, &Namespace);
-	*(FText*)Result = FInternationalization::ForUseOnlyByLocMacroAndGraphNodeTextLiterals_CreateText(*SourceString, *Namespace, *KeyString);
+	*(FText*)RESULT_PARAM = FInternationalization::ForUseOnlyByLocMacroAndGraphNodeTextLiterals_CreateText(*SourceString, *Namespace, *KeyString);
 }
 IMPLEMENT_VM_FUNCTION( EX_TextConst, execTextConst );
 
 void UObject::execObjectConst( FFrame& Stack, RESULT_DECL )
 {
-	*(UObject**)Result = (UObject*)Stack.ReadObject();
+	*(UObject**)RESULT_PARAM = (UObject*)Stack.ReadObject();
 }
 IMPLEMENT_VM_FUNCTION( EX_ObjectConst, execObjectConst );
 
 void UObject::execInstanceDelegate( FFrame& Stack, RESULT_DECL )
 {
 	FName FunctionName = Stack.ReadName();
-	((FScriptDelegate*)Result)->BindUFunction( (FunctionName == NAME_None) ? NULL : this, FunctionName );
+	((FScriptDelegate*)RESULT_PARAM)->BindUFunction( (FunctionName == NAME_None) ? NULL : this, FunctionName );
 }
 IMPLEMENT_VM_FUNCTION( EX_InstanceDelegate, execInstanceDelegate );
 
@@ -1783,29 +1783,29 @@ IMPLEMENT_VM_FUNCTION( EX_BindDelegate, execBindDelegate );
 
 void UObject::execNameConst( FFrame& Stack, RESULT_DECL )
 {
-	*(FName*)Result = Stack.ReadName();
+	*(FName*)RESULT_PARAM = Stack.ReadName();
 }
 IMPLEMENT_VM_FUNCTION( EX_NameConst, execNameConst );
 
 void UObject::execByteConst( FFrame& Stack, RESULT_DECL )
 {
-	*(uint8*)Result = *Stack.Code++;
+	*(uint8*)RESULT_PARAM = *Stack.Code++;
 }
 IMPLEMENT_VM_FUNCTION( EX_ByteConst, execByteConst );
 
 void UObject::execRotationConst( FFrame& Stack, RESULT_DECL )
 {
-	((FRotator*)Result)->Pitch = Stack.ReadFloat();
-	((FRotator*)Result)->Yaw   = Stack.ReadFloat();
-	((FRotator*)Result)->Roll  = Stack.ReadFloat();
+	((FRotator*)RESULT_PARAM)->Pitch = Stack.ReadFloat();
+	((FRotator*)RESULT_PARAM)->Yaw   = Stack.ReadFloat();
+	((FRotator*)RESULT_PARAM)->Roll  = Stack.ReadFloat();
 }
 IMPLEMENT_VM_FUNCTION( EX_RotationConst, execRotationConst );
 
 void UObject::execVectorConst( FFrame& Stack, RESULT_DECL )
 {
-	((FVector*)Result)->X = Stack.ReadFloat();
-	((FVector*)Result)->Y = Stack.ReadFloat();
-	((FVector*)Result)->Z = Stack.ReadFloat();
+	((FVector*)RESULT_PARAM)->X = Stack.ReadFloat();
+	((FVector*)RESULT_PARAM)->Y = Stack.ReadFloat();
+	((FVector*)RESULT_PARAM)->Z = Stack.ReadFloat();
 }
 IMPLEMENT_VM_FUNCTION( EX_VectorConst, execVectorConst );
 
@@ -1830,7 +1830,7 @@ void UObject::execTransformConst( FFrame& Stack, RESULT_DECL )
 	TmpScale.Y = Stack.ReadFloat();
 	TmpScale.Z = Stack.ReadFloat();
 
-	((FTransform*)Result)->SetComponents(TmpRotation, TmpTranslation, TmpScale);
+	((FTransform*)RESULT_PARAM)->SetComponents(TmpRotation, TmpTranslation, TmpScale);
 }
 IMPLEMENT_VM_FUNCTION( EX_TransformConst, execTransformConst );
 
@@ -1847,7 +1847,7 @@ void UObject::execStructConst( FFrame& Stack, RESULT_DECL )
 	{
 		for (int32 ArrayIter = 0; ArrayIter < StructProp->ArrayDim; ++ArrayIter)
 		{
-			Stack.Step(Stack.Object, StructProp->ContainerPtrToValuePtr<uint8>(Result, ArrayIter));
+			Stack.Step(Stack.Object, StructProp->ContainerPtrToValuePtr<uint8>(RESULT_PARAM, ArrayIter));
 		}
 	}
 
@@ -1882,8 +1882,8 @@ void UObject::execArrayConst(FFrame& Stack, RESULT_DECL)
 {
 	UProperty* InnerProperty = CastChecked<UProperty>(Stack.ReadObject());
 	int32 Num = Stack.ReadInt();
-	check(Result);
-	FScriptArrayHelper ArrayHelper = FScriptArrayHelper::CreateHelperFormInnerProperty(InnerProperty, Result);
+	check(RESULT_PARAM);
+	FScriptArrayHelper ArrayHelper = FScriptArrayHelper::CreateHelperFormInnerProperty(InnerProperty, RESULT_PARAM);
 	ArrayHelper.EmptyValues(Num);
 
 	int32 i = 0;
@@ -1900,44 +1900,44 @@ IMPLEMENT_VM_FUNCTION(EX_ArrayConst, execArrayConst);
 
 void UObject::execIntZero( FFrame& Stack, RESULT_DECL )
 {
-	*(int32*)Result = 0;
+	*(int32*)RESULT_PARAM = 0;
 }
 IMPLEMENT_VM_FUNCTION( EX_IntZero, execIntZero );
 
 void UObject::execIntOne( FFrame& Stack, RESULT_DECL )
 {
-	*(int32*)Result = 1;
+	*(int32*)RESULT_PARAM = 1;
 }
 IMPLEMENT_VM_FUNCTION( EX_IntOne, execIntOne );
 
 void UObject::execTrue( FFrame& Stack, RESULT_DECL )
 {
-	*(bool*)Result = true;
+	*(bool*)RESULT_PARAM = true;
 }
 IMPLEMENT_VM_FUNCTION( EX_True, execTrue );
 
 void UObject::execFalse( FFrame& Stack, RESULT_DECL )
 {
-	*(bool*)Result = false;
+	*(bool*)RESULT_PARAM = false;
 }
 IMPLEMENT_VM_FUNCTION( EX_False, execFalse );
 
 void UObject::execNoObject( FFrame& Stack, RESULT_DECL )
 {
-	*(UObject**)Result = NULL;
+	*(UObject**)RESULT_PARAM = NULL;
 }
 IMPLEMENT_VM_FUNCTION( EX_NoObject, execNoObject );
 
 void UObject::execNullInterface(FFrame& Stack, RESULT_DECL)
 {
-	FScriptInterface& InterfaceValue = *(FScriptInterface*)Result;
+	FScriptInterface& InterfaceValue = *(FScriptInterface*)RESULT_PARAM;
 	InterfaceValue.SetObject(nullptr);
 }
 IMPLEMENT_VM_FUNCTION( EX_NoInterface, execNullInterface );
 
 void UObject::execIntConstByte( FFrame& Stack, RESULT_DECL )
 {
-	*(int32*)Result = *Stack.Code++;
+	*(int32*)RESULT_PARAM = *Stack.Code++;
 }
 IMPLEMENT_VM_FUNCTION( EX_IntConstByte, execIntConstByte );
 
@@ -1950,8 +1950,8 @@ void UObject::execDynamicCast( FFrame& Stack, RESULT_DECL )
 	// Compile object expression.
 	UObject* Castee = NULL;
 	Stack.Step( Stack.Object, &Castee );
-	//*(UObject**)Result = (Castee && Castee->IsA(Class)) ? Castee : NULL;
-	*(UObject**)Result = NULL; // default value
+	//*(UObject**)RESULT_PARAM = (Castee && Castee->IsA(Class)) ? Castee : NULL;
+	*(UObject**)RESULT_PARAM = NULL; // default value
 
 
 	// if we were passed in a null value
@@ -1959,11 +1959,11 @@ void UObject::execDynamicCast( FFrame& Stack, RESULT_DECL )
 	{
 		if( Class->HasAnyClassFlags(CLASS_Interface) )
 		{
-			((FScriptInterface*)Result)->SetObject(NULL);
+			((FScriptInterface*)RESULT_PARAM)->SetObject(NULL);
 		}
 		else
 		{
-			*(UObject**)Result = NULL;
+			*(UObject**)RESULT_PARAM = NULL;
 		}
 		return;
 	}
@@ -1975,14 +1975,14 @@ void UObject::execDynamicCast( FFrame& Stack, RESULT_DECL )
 		if ( Castee->GetClass()->ImplementsInterface(Class) )
 		{
 			// interface property type - convert to FScriptInterface
-			((FScriptInterface*)Result)->SetObject(Castee);
-			((FScriptInterface*)Result)->SetInterface(Castee->GetInterfaceAddress(Class));
+			((FScriptInterface*)RESULT_PARAM)->SetObject(Castee);
+			((FScriptInterface*)RESULT_PARAM)->SetInterface(Castee->GetInterfaceAddress(Class));
 		}
 	}
 	// check to see if the Castee is a castable class
 	else if( Castee->IsA(Class) )
 	{
-		*(UObject**)Result = Castee;
+		*(UObject**)RESULT_PARAM = Castee;
 	}
 }
 IMPLEMENT_VM_FUNCTION( EX_DynamicCast, execDynamicCast );
@@ -1995,20 +1995,20 @@ void UObject::execMetaCast( FFrame& Stack, RESULT_DECL )
 	UObject* Castee = nullptr;
 	Stack.Step( Stack.Object, &Castee );
 	UClass* CasteeClass = dynamic_cast<UClass*>(Castee);
-	*(UObject**)Result = (CasteeClass && CasteeClass->IsChildOf(MetaClass)) ? Castee : nullptr;
+	*(UObject**)RESULT_PARAM = (CasteeClass && CasteeClass->IsChildOf(MetaClass)) ? Castee : nullptr;
 }
 IMPLEMENT_VM_FUNCTION( EX_MetaCast, execMetaCast );
 
 void UObject::execPrimitiveCast( FFrame& Stack, RESULT_DECL )
 {
 	int32 B = *(Stack.Code)++;
-	(Stack.Object->*GCasts[B])( Stack, Result );
+	(Stack.Object->*GCasts[B])( Stack, RESULT_PARAM );
 }
 IMPLEMENT_VM_FUNCTION( EX_PrimitiveCast, execPrimitiveCast );
 
 void UObject::execInterfaceCast( FFrame& Stack, RESULT_DECL )
 {
-	(Stack.Object->*GCasts[CST_ObjectToInterface])(Stack, Result);
+	(Stack.Object->*GCasts[CST_ObjectToInterface])(Stack, RESULT_PARAM);
 }
 IMPLEMENT_VM_FUNCTION( EX_ObjToInterfaceCast, execInterfaceCast );
 
@@ -2016,7 +2016,7 @@ void UObject::execObjectToBool( FFrame& Stack, RESULT_DECL )
 {
 	UObject* Obj=NULL;
 	Stack.Step( Stack.Object, &Obj );
-	*(bool*)Result = Obj != NULL;
+	*(bool*)RESULT_PARAM = Obj != NULL;
 }
 IMPLEMENT_CAST_FUNCTION( UObject, CST_ObjectToBool, execObjectToBool );
 
@@ -2024,13 +2024,13 @@ void UObject::execInterfaceToBool( FFrame& Stack, RESULT_DECL )
 {
 	FScriptInterface Interface;
 	Stack.Step( Stack.Object, &Interface);
-	*(bool*)Result = (Interface.GetObject() != NULL);
+	*(bool*)RESULT_PARAM = (Interface.GetObject() != NULL);
 }
 IMPLEMENT_CAST_FUNCTION( UObject, CST_InterfaceToBool, execInterfaceToBool );
 
 void UObject::execObjectToInterface( FFrame& Stack, RESULT_DECL )
 {
-	FScriptInterface& InterfaceValue = *(FScriptInterface*)Result;
+	FScriptInterface& InterfaceValue = *(FScriptInterface*)RESULT_PARAM;
 
 	// read the interface class off the stack
 	UClass* InterfaceClass = dynamic_cast<UClass*>(Stack.ReadObject());
@@ -2056,7 +2056,7 @@ IMPLEMENT_CAST_FUNCTION( UObject, CST_ObjectToInterface, execObjectToInterface )
 
 void UObject::execInterfaceToInterface( FFrame& Stack, RESULT_DECL )
 {
-	FScriptInterface& CastResult = *(FScriptInterface*)Result;
+	FScriptInterface& CastResult = *(FScriptInterface*)RESULT_PARAM;
 
 	// read the interface class off the stack
 	UClass* ClassToCastTo = dynamic_cast<UClass*>(Stack.ReadObject());
@@ -2095,11 +2095,11 @@ void UObject::execInterfaceToObject(FFrame& Stack, RESULT_DECL)
 	UObject* InputObjWithInterface = InterfaceInput.GetObjectRef();
 	if (InputObjWithInterface && InputObjWithInterface->IsA(ObjClassToCastTo))
 	{
-		*(UObject**)Result = InputObjWithInterface;
+		*(UObject**)RESULT_PARAM = InputObjWithInterface;
 	}
 	else
 	{
-		*(UObject**)Result = nullptr;
+		*(UObject**)RESULT_PARAM = nullptr;
 	}
 }
 IMPLEMENT_VM_FUNCTION( EX_InterfaceToObjCast, execInterfaceToObject );
diff --git a/Engine/Source/Runtime/CoreUObject/Public/UObject/Script.h b/Engine/Source/Runtime/CoreUObject/Public/UObject/Script.h
--- a/Engine/Source/Runtime/CoreUObject/Public/UObject/Script.h
+++ b/Engine/Source/Runtime/CoreUObject/Public/UObject/Script.h
@@ -47,7 +47,8 @@ typedef uint32 CodeSkipSizeType;
 //
 // Blueprint VM intrinsic return value declaration.
 //
-#define RESULT_DECL void*const Result
+#define RESULT_PARAM Z_Param__Result
+#define RESULT_DECL void*const RESULT_PARAM
 
 /** Space where UFunctions are asking to be called */
 namespace FunctionCallspace
diff --git a/Engine/Source/Runtime/Engine/Classes/Kismet/DataTableFunctionLibrary.h b/Engine/Source/Runtime/Engine/Classes/Kismet/DataTableFunctionLibrary.h
--- a/Engine/Source/Runtime/Engine/Classes/Kismet/DataTableFunctionLibrary.h
+++ b/Engine/Source/Runtime/Engine/Classes/Kismet/DataTableFunctionLibrary.h
@@ -48,6 +48,6 @@ class ENGINE_API UDataTableFunctionLibrary : public UBlueprintFunctionLibrary
         
         P_FINISH;
         
-		*(bool*)Result = Generic_GetDataTableRowFromName(Table, RowName, OutRowPtr);
+		*(bool*)RESULT_PARAM = Generic_GetDataTableRowFromName(Table, RowName, OutRowPtr);
     }
 };
diff --git a/Engine/Source/Runtime/Engine/Classes/Kismet/KismetArrayLibrary.h b/Engine/Source/Runtime/Engine/Classes/Kismet/KismetArrayLibrary.h
--- a/Engine/Source/Runtime/Engine/Classes/Kismet/KismetArrayLibrary.h
+++ b/Engine/Source/Runtime/Engine/Classes/Kismet/KismetArrayLibrary.h
@@ -214,7 +214,7 @@ public:
  
  		P_FINISH;
  
-		*(int32*)Result = GenericArray_Add(ArrayAddr, ArrayProperty, NewItemPtr);
+		*(int32*)RESULT_PARAM = GenericArray_Add(ArrayAddr, ArrayProperty, NewItemPtr);
  
 		InnerProp->DestroyValue(StorageSpace);
 	}
@@ -238,7 +238,7 @@ public:
 
 		P_FINISH;
 
-		*(int32*)Result = GenericArray_AddUnique(ArrayAddr, ArrayProperty, NewItemPtr);
+		*(int32*)RESULT_PARAM = GenericArray_AddUnique(ArrayAddr, ArrayProperty, NewItemPtr);
 
 		InnerProp->DestroyValue(StorageSpace);
 	}
@@ -333,7 +333,7 @@ public:
 		GenericArray_HandleBool(InnerProp, ItemPtr);
 		  
 		bool WasRemoved = GenericArray_RemoveItem(ArrayAddr, ArrayProperty, ItemPtr);
-		*(bool*)Result = WasRemoved; 
+		*(bool*)RESULT_PARAM = WasRemoved; 
 
 		InnerProp->DestroyValue(StorageSpace);
 	}
@@ -369,7 +369,7 @@ public:
 		P_GET_OBJECT(UArrayProperty, ArrayProperty);
 		P_FINISH;
 
-		*(int32*)Result = GenericArray_Length(ArrayAddr, ArrayProperty);
+		*(int32*)RESULT_PARAM = GenericArray_Length(ArrayAddr, ArrayProperty);
 	}
 
 	DECLARE_FUNCTION(execArray_LastIndex)
@@ -380,7 +380,7 @@ public:
 		P_GET_OBJECT(UArrayProperty, ArrayProperty);
 		P_FINISH;
 
-		*(int32*)Result = GenericArray_LastIndex(ArrayAddr, ArrayProperty);
+		*(int32*)RESULT_PARAM = GenericArray_LastIndex(ArrayAddr, ArrayProperty);
 	}
 
 	DECLARE_FUNCTION(execArray_Get)
@@ -459,7 +459,7 @@ public:
 
 		// Perform the search
 		int32 FoundIndex = GenericArray_Find(ArrayAddr, ArrayProperty, ItemToFindPtr);
-		*(int32*)Result = FoundIndex;
+		*(int32*)RESULT_PARAM = FoundIndex;
 
 		InnerProp->DestroyValue(StorageSpace);
 	}
@@ -488,7 +488,7 @@ public:
 
 		// Perform the search
 		int32 FoundIndex = GenericArray_Find(ArrayAddr, ArrayProperty, ItemToFindPtr);
-		*(bool*)Result = (FoundIndex >= 0);
+		*(bool*)RESULT_PARAM = (FoundIndex >= 0);
 
 		InnerProp->DestroyValue(StorageSpace);
 	}
