Core - Better UObject hash stats

diff --git a/Engine/Source/Runtime/CoreUObject/Private/UObject/UObjectHash.cpp b/Engine/Source/Runtime/CoreUObject/Private/UObject/UObjectHash.cpp
--- a/Engine/Source/Runtime/CoreUObject/Private/UObject/UObjectHash.cpp
+++ b/Engine/Source/Runtime/CoreUObject/Private/UObject/UObjectHash.cpp
@@ -9,13 +9,13 @@
 DEFINE_LOG_CATEGORY_STATIC(LogUObjectHash, Log, All);
 
 DECLARE_STATS_GROUP(TEXT("UObject Hash"), STATGROUP_UObjectHash, STATCAT_Advanced);
-DECLARE_FLOAT_ACCUMULATOR_STAT(TEXT("StaticFindObjectFastInternal"), STAT_Hash_StaticFindObjectFastInternal, STATGROUP_UObjectHash);
-DECLARE_FLOAT_ACCUMULATOR_STAT(TEXT("StaticFindObjectFastExplicit"), STAT_Hash_StaticFindObjectFastExplicit, STATGROUP_UObjectHash);
-DECLARE_FLOAT_ACCUMULATOR_STAT(TEXT("GetObjectsWithOuter"), STAT_Hash_GetObjectsWithOuter, STATGROUP_UObjectHash);
-DECLARE_FLOAT_ACCUMULATOR_STAT(TEXT("FindObjectWithOuter"), STAT_Hash_FindObjectWithOuter, STATGROUP_UObjectHash);
-DECLARE_FLOAT_ACCUMULATOR_STAT(TEXT("GetObjectsOfClass"), STAT_Hash_GetObjectsOfClass, STATGROUP_UObjectHash);
-DECLARE_FLOAT_ACCUMULATOR_STAT(TEXT("HashObject"), STAT_Hash_HashObject, STATGROUP_UObjectHash);
-DECLARE_FLOAT_ACCUMULATOR_STAT(TEXT("UnhashObject"), STAT_Hash_UnhashObject, STATGROUP_UObjectHash);
+DECLARE_CYCLE_STAT(TEXT("StaticFindObjectFastInternal"), STAT_Hash_StaticFindObjectFastInternal, STATGROUP_UObjectHash);
+DECLARE_CYCLE_STAT( TEXT( "StaticFindObjectFastExplicit" ), STAT_Hash_StaticFindObjectFastExplicit, STATGROUP_UObjectHash );
+DECLARE_CYCLE_STAT( TEXT( "GetObjectsWithOuter" ), STAT_Hash_GetObjectsWithOuter, STATGROUP_UObjectHash );
+DECLARE_CYCLE_STAT( TEXT( "FindObjectWithOuter" ), STAT_Hash_FindObjectWithOuter, STATGROUP_UObjectHash );
+DECLARE_CYCLE_STAT( TEXT( "GetObjectsOfClass" ), STAT_Hash_GetObjectsOfClass, STATGROUP_UObjectHash );
+DECLARE_CYCLE_STAT( TEXT( "HashObject" ), STAT_Hash_HashObject, STATGROUP_UObjectHash );
+DECLARE_CYCLE_STAT( TEXT( "UnhashObject" ), STAT_Hash_UnhashObject, STATGROUP_UObjectHash );
 
 /**
  * This implementation will use more space than the UE3 implementation. The goal was to make UObjects smaller to save L2 cache space. 
@@ -164,17 +164,14 @@ UObject* StaticFindObjectFastExplicitThreadSafe(FUObjectHashTables& ThreadHash,
  */
 UObject* StaticFindObjectFastExplicit( UClass* ObjectClass, FName ObjectName, const FString& ObjectPathName, bool bExactClass, EObjectFlags ExcludeFlags/*=0*/ )
 {
+	SCOPE_CYCLE_COUNTER( STAT_Hash_StaticFindObjectFastExplicit );
 	checkSlow(FPackageName::IsShortPackageName(ObjectName)); //@Package name transition, we aren't checking the name here because we know this is only used for texture
+
 	// Find an object with the specified name and (optional) class, in any package; if bAnyPackage is false, only matches top-level packages
-	int32 Hash = GetObjectHash( ObjectName );
-	UObject* Result = nullptr;
-	STAT(double ThisTime = 0);
-	{
-		SCOPE_SECONDS_COUNTER(ThisTime);
-		auto& ThreadHash = FUObjectHashTables::Get();
-		Result = StaticFindObjectFastExplicitThreadSafe(ThreadHash, ObjectClass, ObjectName, ObjectPathName, bExactClass, ExcludeFlags);
-	}
-	INC_FLOAT_STAT_BY(STAT_Hash_StaticFindObjectFastExplicit, (float)ThisTime);
+	const int32 Hash = GetObjectHash( ObjectName );
+	auto& ThreadHash = FUObjectHashTables::Get();
+	UObject* Result = StaticFindObjectFastExplicitThreadSafe( ThreadHash, ObjectClass, ObjectName, ObjectPathName, bExactClass, ExcludeFlags );
+
 	return Result;
 }
 
@@ -272,18 +269,13 @@ UObject* StaticFindObjectFastInternalThreadSafe(FUObjectHashTables& ThreadHash,
 
 UObject* StaticFindObjectFastInternal( UClass* ObjectClass, UObject* ObjectPackage, FName ObjectName, bool bExactClass, bool bAnyPackage, EObjectFlags ExcludeFlags )
 {
+	SCOPE_CYCLE_COUNTER( STAT_Hash_StaticFindObjectFastInternal );
 	INC_DWORD_STAT(STAT_FindObjectFast);
+
 	check(ObjectPackage != ANY_PACKAGE); // this could never have returned anything but NULL
 	// If they specified an outer use that during the hashing
-
-	UObject* Result = nullptr;
-	STAT(double ThisTime = 0);
-	{
-		SCOPE_SECONDS_COUNTER(ThisTime);
-		auto& ThreadHash = FUObjectHashTables::Get();
-		Result = StaticFindObjectFastInternalThreadSafe(ThreadHash, ObjectClass, ObjectPackage, ObjectName, bExactClass, bAnyPackage, ExcludeFlags);
-	}
-	INC_FLOAT_STAT_BY(STAT_Hash_StaticFindObjectFastInternal, (float)ThisTime);
+	auto& ThreadHash = FUObjectHashTables::Get();
+	UObject* Result = StaticFindObjectFastInternalThreadSafe( ThreadHash, ObjectClass, ObjectPackage, ObjectName, bExactClass, bAnyPackage, ExcludeFlags );
 	return Result;
 }
 
@@ -381,102 +373,94 @@ static void RemoveFromClassMap(FUObjectHashTables& ThreadHash, UObjectBase* Obje
 
 void GetObjectsWithOuter(const class UObjectBase* Outer, TArray<UObject *>& Results, bool bIncludeNestedObjects, EObjectFlags ExclusionFlags)
 {
+	SCOPE_CYCLE_COUNTER( STAT_Hash_GetObjectsWithOuter );
+
 	// We don't want to return any objects that are currently being background loaded unless we're using the object iterator during async loading.
-	STAT(double ThisTime = 0);
+	ExclusionFlags |= RF_Unreachable;
+	if (!IsAsyncLoading())
 	{
-		SCOPE_SECONDS_COUNTER(ThisTime);
-
-		ExclusionFlags |= RF_Unreachable;
-		if (!IsAsyncLoading())
-		{
-			ExclusionFlags = EObjectFlags(ExclusionFlags | RF_AsyncLoading);
-		}
-		int32 StartNum = Results.Num();
-		auto& ThreadHash = FUObjectHashTables::Get();
-		FHashTableLock HashLock(ThreadHash);
-		TSet<UObjectBase*> const* Inners = ThreadHash.ObjectOuterMap.Find(Outer);
-		if (Inners)
+		ExclusionFlags = EObjectFlags(ExclusionFlags | RF_AsyncLoading);
+	}
+	int32 StartNum = Results.Num();
+	auto& ThreadHash = FUObjectHashTables::Get();
+	FHashTableLock HashLock(ThreadHash);
+	TSet<UObjectBase*> const* Inners = ThreadHash.ObjectOuterMap.Find(Outer);
+	if (Inners)
+	{
+		for (TSet<UObjectBase*>::TConstIterator It(*Inners); It; ++It)
 		{
-			for (TSet<UObjectBase*>::TConstIterator It(*Inners); It; ++It)
+			UObject *Object = static_cast<UObject *>(*It);
+			if (!Object->HasAnyFlags(ExclusionFlags))
 			{
-				UObject *Object = static_cast<UObject *>(*It);
-				if (!Object->HasAnyFlags(ExclusionFlags))
-				{
-					Results.Add(Object);
-				}
+				Results.Add(Object);
 			}
-			int32 MaxResults = GetUObjectArray().GetObjectArrayNum();
-			while (StartNum != Results.Num() && bIncludeNestedObjects)
+		}
+		int32 MaxResults = GetUObjectArray().GetObjectArrayNum();
+		while (StartNum != Results.Num() && bIncludeNestedObjects)
+		{
+			int32 RangeStart = StartNum;
+			int32 RangeEnd = Results.Num();
+			StartNum = RangeEnd;
+			for (int32 Index = RangeStart; Index < RangeEnd; Index++)
 			{
-				int32 RangeStart = StartNum;
-				int32 RangeEnd = Results.Num();
-				StartNum = RangeEnd;
-				for (int32 Index = RangeStart; Index < RangeEnd; Index++)
+				TSet<UObjectBase*> const* InnerInners = ThreadHash.ObjectOuterMap.Find(Results[Index]);
+				if (InnerInners)
 				{
-					TSet<UObjectBase*> const* InnerInners = ThreadHash.ObjectOuterMap.Find(Results[Index]);
-					if (InnerInners)
+					for (TSet<UObjectBase*>::TConstIterator It(*InnerInners); It; ++It)
 					{
-						for (TSet<UObjectBase*>::TConstIterator It(*InnerInners); It; ++It)
+						UObject *Object = static_cast<UObject *>(*It);
+						if (!Object->HasAnyFlags(ExclusionFlags))
 						{
-							UObject *Object = static_cast<UObject *>(*It);
-							if (!Object->HasAnyFlags(ExclusionFlags))
-							{
-								Results.Add(Object);
-							}
+							Results.Add(Object);
 						}
 					}
 				}
-				check(Results.Num() <= MaxResults); // otherwise we have a cycle in the outer chain, which should not be possible
 			}
+			check(Results.Num() <= MaxResults); // otherwise we have a cycle in the outer chain, which should not be possible
 		}
 	}
-	INC_FLOAT_STAT_BY(STAT_Hash_GetObjectsWithOuter, (float)ThisTime);
 }
 
 UObjectBase* FindObjectWithOuter(class UObjectBase* Outer, class UClass* ClassToLookFor, FName NameToLookFor)
 {
+	SCOPE_CYCLE_COUNTER( STAT_Hash_FindObjectWithOuter );
+
 	UObject* Result = nullptr;
-	STAT(double ThisTime = 0);
+	check( Outer && ClassToLookFor );
+	// We don't want to return any objects that are currently being background loaded unless we're using the object iterator during async loading.
+	EObjectFlags ExclusionFlags = RF_Unreachable;
+	if( !IsAsyncLoading() )
 	{
-		SCOPE_SECONDS_COUNTER(ThisTime);
-
-		check(Outer && ClassToLookFor);
-		// We don't want to return any objects that are currently being background loaded unless we're using the object iterator during async loading.
-		EObjectFlags ExclusionFlags = RF_Unreachable;
-		if (!IsAsyncLoading())
-		{
-			ExclusionFlags = EObjectFlags(ExclusionFlags | RF_AsyncLoading);
-		}
+		ExclusionFlags = EObjectFlags( ExclusionFlags | RF_AsyncLoading );
+	}
 
-		if (NameToLookFor != NAME_None)
-		{
-			Result = StaticFindObjectFastInternal(ClassToLookFor, static_cast<UObject*>(Outer), NameToLookFor, false, false, ExclusionFlags);
-		}
-		else
+	if( NameToLookFor != NAME_None )
+	{
+		Result = StaticFindObjectFastInternal( ClassToLookFor, static_cast<UObject*>(Outer), NameToLookFor, false, false, ExclusionFlags );
+	}
+	else
+	{
+		auto& ThreadHash = FUObjectHashTables::Get();
+		FHashTableLock HashLock( ThreadHash );
+		TSet<UObjectBase*> const* Inners = ThreadHash.ObjectOuterMap.Find( Outer );
+		if( Inners )
 		{
-			auto& ThreadHash = FUObjectHashTables::Get();
-			FHashTableLock HashLock(ThreadHash);
-			TSet<UObjectBase*> const* Inners = ThreadHash.ObjectOuterMap.Find(Outer);
-			if (Inners)
+			for( TSet<UObjectBase*>::TConstIterator It( *Inners ); It; ++It )
 			{
-				for (TSet<UObjectBase*>::TConstIterator It(*Inners); It; ++It)
+				UObject *Object = static_cast<UObject *>(*It);
+				if( Object->HasAnyFlags( ExclusionFlags ) )
 				{
-					UObject *Object = static_cast<UObject *>(*It);
-					if (Object->HasAnyFlags(ExclusionFlags))
-					{
-						continue;
-					}
-					if (!Object->IsA(ClassToLookFor))
-					{
-						continue;
-					}
-					Result = Object;
-					break;
+					continue;
 				}
+				if( !Object->IsA( ClassToLookFor ) )
+				{
+					continue;
+				}
+				Result = Object;
+				break;
 			}
 		}
 	}
-	INC_FLOAT_STAT_BY(STAT_Hash_FindObjectWithOuter, (float)ThisTime);
 	return Result;
 }
 
@@ -519,37 +503,28 @@ static void GetObjectsOfClassThreadSafe(FUObjectHashTables& ThreadHash, TSet<UCl
 }
 void GetObjectsOfClass(UClass* ClassToLookFor, TArray<UObject *>& Results, bool bIncludeDerivedClasses, EObjectFlags AdditionalExcludeFlags)
 {
-	STAT(double ThisTime = 0);
-	{
-		SCOPE_SECONDS_COUNTER(ThisTime);
-
-		// We don't want to return any objects that are currently being background loaded unless we're using the object iterator during async loading.
-		EObjectFlags ExclusionFlags = RF_Unreachable;
-		if (!IsAsyncLoading())
-		{
-			ExclusionFlags |= RF_AsyncLoading;
-		}
-		ExclusionFlags |= AdditionalExcludeFlags;
-
-		TSet<UClass*> ClassesToSearch;
-		ClassesToSearch.Add(ClassToLookFor);
-		if (bIncludeDerivedClasses)
-		{
-			auto& ThreadHash = FUObjectHashTables::Get();
-			FHashTableLock HashLock(ThreadHash);
-			RecursivelyPopulateDerivedClasses(ThreadHash, ClassToLookFor, ClassesToSearch);
-		}
-
-		GetObjectsOfClassThreadSafe(FUObjectHashTables::Get(), ClassesToSearch, Results, ExclusionFlags);
+	SCOPE_CYCLE_COUNTER( STAT_Hash_GetObjectsOfClass );
 
-		check(Results.Num() <= GetUObjectArray().GetObjectArrayNum()); // otherwise we have a cycle in the outer chain, which should not be possible
+	// We don't want to return any objects that are currently being background loaded unless we're using the object iterator during async loading.
+	EObjectFlags ExclusionFlags = RF_Unreachable;
+	if( !IsAsyncLoading() )
+	{
+		ExclusionFlags |= RF_AsyncLoading;
 	}
-#if STATS
-	if (FThreadStats::IsCollectingData())
+	ExclusionFlags |= AdditionalExcludeFlags;
+
+	TSet<UClass*> ClassesToSearch;
+	ClassesToSearch.Add( ClassToLookFor );
+	if( bIncludeDerivedClasses )
 	{
-		INC_FLOAT_STAT_BY(STAT_Hash_GetObjectsOfClass, (float)ThisTime);
+		auto& ThreadHash = FUObjectHashTables::Get();
+		FHashTableLock HashLock( ThreadHash );
+		RecursivelyPopulateDerivedClasses( ThreadHash, ClassToLookFor, ClassesToSearch );
 	}
-#endif
+
+	GetObjectsOfClassThreadSafe( FUObjectHashTables::Get(), ClassesToSearch, Results, ExclusionFlags );
+
+	check( Results.Num() <= GetUObjectArray().GetObjectArrayNum() ); // otherwise we have a cycle in the outer chain, which should not be possible
 }
 
 void GetDerivedClasses(UClass* ClassToLookFor, TArray<UClass *>& Results, bool bRecursive)
@@ -581,34 +556,26 @@ void AllocateUObjectIndexForCurrentThread(UObjectBase* Object)
 
 void HashObject(UObjectBase* Object)
 {
-	STAT(double ThisTime = 0);
+	SCOPE_CYCLE_COUNTER( STAT_Hash_HashObject );
+
+	FName Name = Object->GetFName();
+	if( Name != NAME_None )
 	{
-		SCOPE_SECONDS_COUNTER(ThisTime);
-		FName Name = Object->GetFName();
-		if (Name != NAME_None)
-		{
-			int32 Hash = GetObjectHash(Name);
-			auto& ThreadHash = FUObjectHashTables::Get();
+		int32 Hash = GetObjectHash( Name );
+		auto& ThreadHash = FUObjectHashTables::Get();
 
-			FHashTableLock HashLock(ThreadHash);
-			checkSlow(!ThreadHash.Hash.FindPair(Hash, Object));  // if it already exists, something is wrong with the external code
+		FHashTableLock HashLock( ThreadHash );
+		checkSlow( !ThreadHash.Hash.FindPair( Hash, Object ) );  // if it already exists, something is wrong with the external code
 
-			ThreadHash.Hash.Add(Hash, Object);
+		ThreadHash.Hash.Add( Hash, Object );
 
-			Hash = GetObjectOuterHash(Name, (PTRINT)Object->GetOuter());
-			checkSlow(!ThreadHash.HashOuter.FindPair(Hash, Object));  // if it already exists, something is wrong with the external code
-			ThreadHash.HashOuter.Add(Hash, Object);
+		Hash = GetObjectOuterHash( Name, (PTRINT)Object->GetOuter() );
+		checkSlow( !ThreadHash.HashOuter.FindPair( Hash, Object ) );  // if it already exists, something is wrong with the external code
+		ThreadHash.HashOuter.Add( Hash, Object );
 
-			AddToOuterMap(ThreadHash, Object);
-			AddToClassMap(ThreadHash, Object);
-		}
-	}
-#if STATS
-	if (FThreadStats::IsCollectingData())
-	{
-		INC_FLOAT_STAT_BY(STAT_Hash_HashObject, (float)ThisTime);
+		AddToOuterMap( ThreadHash, Object );
+		AddToClassMap( ThreadHash, Object );
 	}
-#endif
 }
 
 /**
@@ -618,33 +585,24 @@ void HashObject(UObjectBase* Object)
  */
 void UnhashObject(UObjectBase* Object)
 {
-	STAT(double ThisTime = 0);
+	SCOPE_CYCLE_COUNTER( STAT_Hash_UnhashObject );
+	FName Name = Object->GetFName();
+	if( Name != NAME_None )
 	{
-		SCOPE_SECONDS_COUNTER(ThisTime);
-		FName Name = Object->GetFName();
-		if (Name != NAME_None)
-		{
 
-			int32 Hash = GetObjectHash(Name);
-			auto& ThreadHash = FUObjectHashTables::Get();
-			FHashTableLock LockHash(ThreadHash);
-			int32 NumRemoved = ThreadHash.Hash.RemoveSingle(Hash, Object);
-			check(NumRemoved == 1); // must have existed, else something is wrong with the external code
+		int32 Hash = GetObjectHash( Name );
+		auto& ThreadHash = FUObjectHashTables::Get();
+		FHashTableLock LockHash( ThreadHash );
+		int32 NumRemoved = ThreadHash.Hash.RemoveSingle( Hash, Object );
+		check( NumRemoved == 1 ); // must have existed, else something is wrong with the external code
 
-			Hash = GetObjectOuterHash(Name, (PTRINT)Object->GetOuter());
-			NumRemoved = ThreadHash.HashOuter.RemoveSingle(Hash, Object);
-			check(NumRemoved == 1); // must have existed, else something is wrong with the external code
+		Hash = GetObjectOuterHash( Name, (PTRINT)Object->GetOuter() );
+		NumRemoved = ThreadHash.HashOuter.RemoveSingle( Hash, Object );
+		check( NumRemoved == 1 ); // must have existed, else something is wrong with the external code
 
-			RemoveFromOuterMap(ThreadHash, Object);
-			RemoveFromClassMap(ThreadHash, Object);
-		}
+		RemoveFromOuterMap( ThreadHash, Object );
+		RemoveFromClassMap( ThreadHash, Object );
 	}
-#if STATS
-	if (FThreadStats::IsCollectingData())
-	{
-		INC_FLOAT_STAT_BY(STAT_Hash_UnhashObject, (float)ThisTime);
-	}
-#endif
 }
 
 /**
