Stats - Refactored to use dynamic stats

diff --git a/Engine/Source/Runtime/CoreUObject/Private/UObject/UObjectBase.cpp b/Engine/Source/Runtime/CoreUObject/Private/UObject/UObjectBase.cpp
--- a/Engine/Source/Runtime/CoreUObject/Private/UObject/UObjectBase.cpp
+++ b/Engine/Source/Runtime/CoreUObject/Private/UObject/UObjectBase.cpp
@@ -117,18 +117,7 @@ void UObjectBase::CreateStatID() const
 		LongName = GetClass()->GetFName().GetPlainNameString() / LongName;
 	}
 
-	const FName StatName = FName( *LongName );
-	FStartupMessages::Get().AddMetadata( StatName, *LongName, 
-		STAT_GROUP_TO_FStatGroup( STATGROUP_UObjects )::GetGroupName(), 
-		STAT_GROUP_TO_FStatGroup( STATGROUP_UObjects )::GetGroupCategory(), 
-		STAT_GROUP_TO_FStatGroup( STATGROUP_UObjects )::GetDescription(),
-		true, EStatDataType::ST_int64, true );
-
-	StatID = IStatGroupEnableManager::Get().GetHighPerformanceEnableForStat(StatName, 
-		STAT_GROUP_TO_FStatGroup(STATGROUP_UObjects)::GetGroupName(), 
-		STAT_GROUP_TO_FStatGroup(STATGROUP_UObjects)::GetGroupCategory(), 
-		STAT_GROUP_TO_FStatGroup(STATGROUP_UObjects)::DefaultEnable,
-		true, EStatDataType::ST_int64, *LongName, true);
+	StatID = FDynamicStats::CreateStatId<FStatGroup_STATGROUP_UObjects>( LongName );
 }
 #endif
 
