Core - Dynamic stats for LoadPackage, disabled by default

diff --git a/Engine/Source/Runtime/CoreUObject/Private/UObject/UObjectGlobals.cpp b/Engine/Source/Runtime/CoreUObject/Private/UObject/UObjectGlobals.cpp
--- a/Engine/Source/Runtime/CoreUObject/Private/UObject/UObjectGlobals.cpp
+++ b/Engine/Source/Runtime/CoreUObject/Private/UObject/UObjectGlobals.cpp
@@ -1022,6 +1022,14 @@ UPackage* LoadPackageInternal(UPackage* InOuter, const TCHAR* InLongPackageName,
 
 UPackage* LoadPackage(UPackage* InOuter, const TCHAR* InLongPackageName, uint32 LoadFlags)
 {
+	// Change to 1 if you want more detailed stats for loading packages, but at the cost of adding dynamic stats.
+#if	STATS && 0
+	static FString Package = TEXT( "Package" );
+	const FString LongName = Package / InLongPackageName;
+	const TStatId StatId = FDynamicStats::CreateStatId<FStatGroup_STATGROUP_UObjects>( LongName );
+	FScopeCycleCounter CycleCounter( StatId );
+#endif // STATS
+
 	return LoadPackageInternal(InOuter, InLongPackageName, LoadFlags, /*ImportLinker =*/ nullptr);
 }
 
@@ -2000,6 +2008,8 @@ FObjectInitializer::FObjectInitializer(UObject* InObj, UObject* InObjectArchetyp
 	FTlsObjectInitializers::Push(this);
 }
 
+COREUOBJECT_API bool IgnoreFObjectInitializer = false;
+
 /**
  * Destructor for internal class to finalize UObject creation (initialize properties) after the real C++ constructor is called.
  **/
@@ -2012,6 +2022,11 @@ FObjectInitializer::~FObjectInitializer()
 	check(ThreadContext.IsInConstructor >= 0);
 	ThreadContext.ConstructedObject = LastConstructedObject;
 
+	if( IgnoreFObjectInitializer )
+	{
+		return;
+	}
+
 	SCOPE_CYCLE_COUNTER(STAT_PostConstructInitializeProperties);
 	check(Obj);
 	const bool bIsCDO = Obj->HasAnyFlags(RF_ClassDefaultObject);
