[UE-13906] Make sure packages that skipped LoadAllObjects (as part of deferred dependency loading) have the chance to complete the process before PostLoad()

diff --git a/Engine/Source/Runtime/CoreUObject/Private/Blueprint/BlueprintSupport.cpp b/Engine/Source/Runtime/CoreUObject/Private/Blueprint/BlueprintSupport.cpp
--- a/Engine/Source/Runtime/CoreUObject/Private/Blueprint/BlueprintSupport.cpp
+++ b/Engine/Source/Runtime/CoreUObject/Private/Blueprint/BlueprintSupport.cpp
@@ -755,6 +755,20 @@ void FLinkerLoad::ResolveDeferredDependencies(UStruct* LoadStruct)
 		{
 			FObjectImport& Import = ImportMap[ImportIndex];
 
+			const FLinkerLoad* SourceLinker = (Import.SourceLinker != nullptr) ? Import.SourceLinker :
+				(Import.XObject != nullptr) ? Import.XObject->GetLinker() : nullptr;
+
+			const UPackage* SourcePackage = (SourceLinker != nullptr) ? SourceLinker->LinkerRoot : nullptr;
+			// this package may not have introduced any (possible) cyclic 
+			// dependencies, but it still could have been deferred (kept from
+			// fully loading... we need to make sure metadata gets loaded, etc.)
+			if ((SourcePackage != nullptr) && !SourcePackage->HasAnyFlags(RF_WasLoaded))
+			{
+				uint32 InternalLoadFlags = LoadFlags & (LOAD_NoVerify | LOAD_NoWarn | LOAD_Quiet);
+				// make sure LoadAllObjects() is called for this package
+				LoadPackageInternal(/*Outer =*/nullptr, *SourceLinker->Filename, InternalLoadFlags, this);
+			}
+
 			if (ULinkerPlaceholderClass* PlaceholderClass = Cast<ULinkerPlaceholderClass>(Import.XObject))
 			{
 				DEFERRED_DEPENDENCY_CHECK(PlaceholderClass->PackageIndex.ToImport() == ImportIndex);
@@ -781,13 +795,6 @@ void FLinkerLoad::ResolveDeferredDependencies(UStruct* LoadStruct)
 				// deferred dependencies in the struct 
 				if (Import.SourceLinker != nullptr)
 				{
-					// Make sure meta data is loaded first, so that we have structure field names available to any blueprints that load:
-					if (!Import.SourceLinker->LinkerRoot->HasAnyFlags(RF_WasLoaded))
-					{
-						uint32 InternalLoadFlags = LoadFlags & (LOAD_NoVerify | LOAD_NoWarn | LOAD_Quiet);
-						LoadPackageInternal(nullptr, *(Import.SourceLinker->Filename), InternalLoadFlags, this);
-					}
-
 					Import.SourceLinker->ResolveDeferredDependencies(StructObj);
 				}
 			}
