Relocated PropertyPortFlags to its own header in Core from UnrealTypes.h in CoreUObject. This allows the PortFlags on FArchive to be checked without a dependency on CoreUObject.

diff --git a/Engine/Source/Runtime/Core/Public/UObject/PropertyPortFlags.h b/Engine/Source/Runtime/Core/Public/UObject/PropertyPortFlags.h
new file mode 100644
--- /dev/null
+++ b/Engine/Source/Runtime/Core/Public/UObject/PropertyPortFlags.h
@@ -0,0 +1,96 @@
+// Copyright 1998-2015 Epic Games, Inc. All Rights Reserved.
+
+/*=============================================================================
+	PropertyPortFlags.h: Property import/export flags.
+=============================================================================*/
+
+#pragma once
+
+// Property exporting flags.
+enum EPropertyPortFlags
+{
+	/** No special property exporint flags */
+	PPF_None						= 0x00000000,
+
+	/** Indicates that property data should be treated as text */
+	PPF_Localized					= 0x00000001,
+
+	/** Indicates that property data should be wrapped in quotes (for some types of properties) */
+	PPF_Delimited					= 0x00000002,
+
+	/** Indicates that the object reference should be verified */
+	PPF_CheckReferences				= 0x00000004, 
+	
+	PPF_ExportsNotFullyQualified	= 0x00000008,
+	
+	PPF_AttemptNonQualifiedSearch	= 0x00000010,
+	
+	/** Indicates that importing values for config or localized properties is disallowed */
+	PPF_RestrictImportTypes			= 0x00000020,
+
+	//								= 0x00000040,
+	//								= 0x00000080,
+
+	/** only include properties which are marked CPF_InstancedReference */
+	PPF_SubobjectsOnly				= 0x00000100,
+
+	/**
+	 * Only applicable to component properties (for now)
+	 * Indicates that two object should be considered identical
+	 * if the property values for both objects are all identical
+	 */
+	PPF_DeepComparison				= 0x00000200,
+
+	/**
+	 * Similar to PPF_DeepComparison, except that template components are always compared using standard object
+	 * property comparison logic (basically if the pointers are different, then the property isn't identical)
+	 */
+	PPF_DeepCompareInstances		= 0x00000400,
+
+	/**
+	 * Set if this operation is copying in memory (for copy/paste) instead of exporting to a file. There are
+	 * some subtle differences between the two
+	 */
+	PPF_Copy						= 0x00000800,
+
+	/** Set when duplicating objects via serialization */
+	PPF_Duplicate					= 0x00001000,
+
+	/** Indicates that object property values should be exported without the package or class information */
+	PPF_SimpleObjectText			= 0x00002000,
+
+	/** parsing default properties - allow text for transient properties to be imported - also modifies ObjectProperty importing slightly for subobjects */
+	PPF_ParsingDefaultProperties	= 0x00008000,
+
+	/** indicates that non-categorized transient properties should be exported (by default, they would not be) */
+	PPF_IncludeTransient			= 0x00020000,
+
+	/** modifies behavior of UProperty::Identical - indicates that the comparison is between an object and its archetype */
+	PPF_DeltaComparison				= 0x00040000,
+
+	/** indicates that we're exporting properties for display in the property window. - used to hide EditHide items in collapsed structs */
+	PPF_PropertyWindow				= 0x00080000,
+
+	PPF_NoInternalArcheType			= 0x00100000,
+
+	/** Force fully qualified object names (for debug dumping) */
+	PPF_DebugDump					= 0x00200000,
+
+	/** Set when duplicating objects for PIE */
+	PPF_DuplicateForPIE				= 0x00400000,
+
+	/** Set when exporting just an object declaration, to be followed by another call with PPF_SeparateDefine */
+	PPF_SeparateDeclare				= 0x00800000,
+
+	/** Set when exporting just an object definition, preceded by another call with PPF_SeparateDeclare */
+	PPF_SeparateDefine				= 0x01000000,
+
+	/** Used by 'watch value' while blueprint debugging*/
+	PPF_BlueprintDebugView			= 0x02000000,
+
+	/** Exporting properties for console variables. */
+	PPF_ConsoleVariable				= 0x04000000,
+
+	/** Ignores CPF_Deprecated flag */
+	PPF_UseDeprecatedProperties		= 0x08000000,
+};
\ No newline at end of file
diff --git a/Engine/Source/Runtime/CoreUObject/Public/UObject/UnrealType.h b/Engine/Source/Runtime/CoreUObject/Public/UObject/UnrealType.h
--- a/Engine/Source/Runtime/CoreUObject/Public/UObject/UnrealType.h
+++ b/Engine/Source/Runtime/CoreUObject/Public/UObject/UnrealType.h
@@ -7,6 +7,7 @@
 #pragma once
 
 #include "ObjectBase.h"
+#include "PropertyPortFlags.h"
 
 COREUOBJECT_API DECLARE_LOG_CATEGORY_EXTERN(LogType, Log, All);
 
@@ -14,95 +15,6 @@ COREUOBJECT_API DECLARE_LOG_CATEGORY_EXTERN(LogType, Log, All);
 	UProperty.
 -----------------------------------------------------------------------------*/
 
-// Property exporting flags.
-enum EPropertyPortFlags
-{
-	/** No special property exporint flags */
-	PPF_None						= 0x00000000,
-
-	/** Indicates that property data should be treated as text */
-	PPF_Localized					= 0x00000001,
-
-	/** Indicates that property data should be wrapped in quotes (for some types of properties) */
-	PPF_Delimited					= 0x00000002,
-
-	/** Indicates that the object reference should be verified */
-	PPF_CheckReferences				= 0x00000004, 
-	
-	PPF_ExportsNotFullyQualified	= 0x00000008,
-	
-	PPF_AttemptNonQualifiedSearch	= 0x00000010,
-	
-	/** Indicates that importing values for config or localized properties is disallowed */
-	PPF_RestrictImportTypes			= 0x00000020,
-
-	//								= 0x00000040,
-	//								= 0x00000080,
-
-	/** only include properties which are marked CPF_InstancedReference */
-	PPF_SubobjectsOnly				= 0x00000100,
-
-	/**
-	 * Only applicable to component properties (for now)
-	 * Indicates that two object should be considered identical
-	 * if the property values for both objects are all identical
-	 */
-	PPF_DeepComparison				= 0x00000200,
-
-	/**
-	 * Similar to PPF_DeepComparison, except that template components are always compared using standard object
-	 * property comparison logic (basically if the pointers are different, then the property isn't identical)
-	 */
-	PPF_DeepCompareInstances		= 0x00000400,
-
-	/**
-	 * Set if this operation is copying in memory (for copy/paste) instead of exporting to a file. There are
-	 * some subtle differences between the two
-	 */
-	PPF_Copy						= 0x00000800,
-
-	/** Set when duplicating objects via serialization */
-	PPF_Duplicate					= 0x00001000,
-
-	/** Indicates that object property values should be exported without the package or class information */
-	PPF_SimpleObjectText			= 0x00002000,
-
-	/** parsing default properties - allow text for transient properties to be imported - also modifies ObjectProperty importing slightly for subobjects */
-	PPF_ParsingDefaultProperties	= 0x00008000,
-
-	/** indicates that non-categorized transient properties should be exported (by default, they would not be) */
-	PPF_IncludeTransient			= 0x00020000,
-
-	/** modifies behavior of UProperty::Identical - indicates that the comparison is between an object and its archetype */
-	PPF_DeltaComparison				= 0x00040000,
-
-	/** indicates that we're exporting properties for display in the property window. - used to hide EditHide items in collapsed structs */
-	PPF_PropertyWindow				= 0x00080000,
-
-	PPF_NoInternalArcheType			= 0x00100000,
-
-	/** Force fully qualified object names (for debug dumping) */
-	PPF_DebugDump					= 0x00200000,
-
-	/** Set when duplicating objects for PIE */
-	PPF_DuplicateForPIE				= 0x00400000,
-
-	/** Set when exporting just an object declaration, to be followed by another call with PPF_SeparateDefine */
-	PPF_SeparateDeclare				= 0x00800000,
-
-	/** Set when exporting just an object definition, preceded by another call with PPF_SeparateDeclare */
-	PPF_SeparateDefine				= 0x01000000,
-
-	/** Used by 'watch value' while blueprint debugging*/
-	PPF_BlueprintDebugView			= 0x02000000,
-
-	/** Exporting properties for console variables. */
-	PPF_ConsoleVariable				= 0x04000000,
-
-	/** Ignores CPF_Deprecated flag */
-	PPF_UseDeprecatedProperties		= 0x08000000,
-};
-
 enum EPropertyExportCPPFlags
 {
 	/** Indicates that there are no special C++ export flags */
