[UE-13339] All default subobjects on a CDO are now transactional in editor, this makes undo/redo work on native components

diff --git a/Engine/Source/Runtime/CoreUObject/Private/UObject/UObjectGlobals.cpp b/Engine/Source/Runtime/CoreUObject/Private/UObject/UObjectGlobals.cpp
--- a/Engine/Source/Runtime/CoreUObject/Private/UObject/UObjectGlobals.cpp
+++ b/Engine/Source/Runtime/CoreUObject/Private/UObject/UObjectGlobals.cpp
@@ -2774,6 +2774,11 @@ UObject* FObjectInitializer::CreateDefaultSubobject(UObject* Outer, FName Subobj
 			}
 			if (Outer->HasAnyFlags(RF_ClassDefaultObject) && Outer->GetClass()->GetSuperClass())
 			{
+#if WITH_EDITOR
+				// Default subobjects on the CDO should be transactional, so that we can undo/redo changes made to those objects.
+				// One current example of this is editing natively defined components in the Blueprint Editor.
+				Result->SetFlags(RF_Transactional);
+#endif
 				Outer->GetClass()->AddDefaultSubobject(Result, ReturnType);
 			}
 			Result->SetFlags(RF_DefaultSubObject);
