Disabled text key preservation until it can be properly handled in most cases.

diff --git a/Engine/Source/Runtime/CoreUObject/Private/UObject/UTextProperty.cpp b/Engine/Source/Runtime/CoreUObject/Private/UObject/UTextProperty.cpp
--- a/Engine/Source/Runtime/CoreUObject/Private/UObject/UTextProperty.cpp
+++ b/Engine/Source/Runtime/CoreUObject/Private/UObject/UTextProperty.cpp
@@ -106,7 +106,6 @@ const TCHAR* UTextProperty::ImportText_Internal( const TCHAR* Buffer, void* Data
 		BufferAsString = FString(Temp);
 	}
 
-	FScopedTextIdentityPreserver Preserver(*Text);
 	*Text = FText::FromString(FString(BufferAsString));
 
 	return Buffer;
