Fix enum duplication warning in Asset Creation and Duplication test.

diff --git a/Engine/Source/Runtime/CoreUObject/Private/UObject/Enum.cpp b/Engine/Source/Runtime/CoreUObject/Private/UObject/Enum.cpp
--- a/Engine/Source/Runtime/CoreUObject/Private/UObject/Enum.cpp
+++ b/Engine/Source/Runtime/CoreUObject/Private/UObject/Enum.cpp
@@ -36,11 +36,51 @@ void UEnum::Serialize( FArchive& Ar )
 	}
 
 	if (Ar.IsLoading() || Ar.IsSaving())
-	{	
+	{
+		// We're duplicating this enum.
+		if ((Ar.GetPortFlags() & PPF_Duplicate)
+			// and we're loading it from already serialized base.
+			&& Ar.IsLoading())
+		{
+			// Rename enum names to reflect new class.
+			RenameNamesAfterDuplication();
+		}
 		AddNamesToMasterList();
 	}
 }
 
+FString UEnum::GetBaseEnumNameOnDuplication() const
+{
+	// Last name is always fully qualified, in form EnumName::Prefix_MAX.
+	FString BaseEnumName = Names[Names.Num() - 1].ToString();
+
+	// Double check we have a fully qualified name.
+	auto DoubleColonPos = BaseEnumName.Find(TEXT("::"));
+	check(DoubleColonPos != INDEX_NONE);
+
+	// Get actual base name.
+	BaseEnumName = BaseEnumName.LeftChop(BaseEnumName.Len() - DoubleColonPos);
+
+	return BaseEnumName;
+}
+
+void UEnum::RenameNamesAfterDuplication()
+{
+	// Get name of base enum, from which we're duplicating.
+	FString BaseEnumName = GetBaseEnumNameOnDuplication();
+
+	// Get name of duplicated enum.
+	auto ThisName = GetName();
+
+	// Replace all usages of base class name to the duplicated one.
+	for (auto& Name : Names)
+	{
+		auto NameString = Name.ToString();
+		NameString.ReplaceInline(*BaseEnumName, *ThisName);
+		Name = FName(*NameString);
+	}
+}
+
 int32 UEnum::ResolveEnumerator(FArchive& Ar, int32 EnumeratorIndex) const
 {
 	return EnumeratorIndex;
diff --git a/Engine/Source/Runtime/CoreUObject/Public/UObject/Class.h b/Engine/Source/Runtime/CoreUObject/Public/UObject/Class.h
--- a/Engine/Source/Runtime/CoreUObject/Public/UObject/Class.h
+++ b/Engine/Source/Runtime/CoreUObject/Public/UObject/Class.h
@@ -1681,6 +1681,24 @@ private:
 		UE_CLOG(!EnumClass, LogClass, Fatal, TEXT("Couldn't find enum '%s'"), EnumPath);
 		return EnumClass->GetEnumText(Value);
 	}
+
+	/**
+	 *
+	 * Renames enum values to use duplicated enum name instead of base one, e.g.:
+	 * 
+	 * MyEnum::MyVal
+	 * MyEnum::MyEnum_MAX
+	 * 
+	 * becomes
+	 * 
+	 * MyDuplicatedEnum::MyVal
+	 * MyDuplicatedEnum::MyDuplicatedEnum_MAX
+	 * 
+	 **/
+	void RenameNamesAfterDuplication();
+
+	/** Gets name of enum "this" is duplicate of. If we're not duplicating, just returns "this" name. */
+	FString GetBaseEnumNameOnDuplication() const;
 };
 
 /*-----------------------------------------------------------------------------
