Fix for flushing async loading.

diff --git a/Engine/Source/Runtime/CoreUObject/Private/UObject/UObjectGlobals.cpp b/Engine/Source/Runtime/CoreUObject/Private/UObject/UObjectGlobals.cpp
--- a/Engine/Source/Runtime/CoreUObject/Private/UObject/UObjectGlobals.cpp
+++ b/Engine/Source/Runtime/CoreUObject/Private/UObject/UObjectGlobals.cpp
@@ -1051,7 +1051,7 @@ bool IsLoading()
 void BeginLoad()
 {
 	auto& ThreadContext = FUObjectThreadContext::Get();
-	if (++ThreadContext.ObjBeginLoadCount == 1 && !IsInAsyncLoadingThread())
+	if (ThreadContext.ObjBeginLoadCount == 0 && !IsInAsyncLoadingThread())
 	{
 		// Make sure we're finishing up all pending async loads, and trigger texture streaming next tick if necessary.
 		FlushAsyncLoading( NAME_None );
@@ -1059,6 +1059,8 @@ void BeginLoad()
 		// Validate clean load state.
 		check(ThreadContext.ObjLoaded.Num() == 0);
 	}
+
+	++ThreadContext.ObjBeginLoadCount;
 }
 
 // Sort objects by linker name and file offset
