Faster unhashing of UObject FNames in case of conflicts (like with subobjects which have hardcoded names). Added memory stats to 'obj hash' and 'obj hashouter'

diff --git a/Engine/Source/Runtime/CoreUObject/Private/UObject/UObjectHash.cpp b/Engine/Source/Runtime/CoreUObject/Private/UObject/UObjectHash.cpp
--- a/Engine/Source/Runtime/CoreUObject/Private/UObject/UObjectHash.cpp
+++ b/Engine/Source/Runtime/CoreUObject/Private/UObject/UObjectHash.cpp
@@ -31,6 +31,182 @@ DECLARE_CYCLE_STAT( TEXT( "UnhashObject" ), STAT_Hash_UnhashObject, STATGROUP_UO
 */
 #define OBJECT_HASH_BINS (1024*1024)
 
+/*
+ * Special hash bucket to conserve memory.
+ * Contains a pointer to head element and an optional list of items if more than one element exists in the bucket.
+ * The item list is only allocated if needed.
+ */
+struct FHashBucket
+{
+	friend struct FHashBucketIterator;
+
+	/** This always empty set is used to get an iterator if the bucket doesn't use a TSet (has only 1 element) */
+	static TSet<UObjectBase*> EmptyBucket;
+	/** Used when the bucket has only one item, does not necessarily represent the first item if Items TSet is non-null */
+	UObjectBase* Head;
+	/** If more than one items belongs to this bucket this is non-null and contains all elements of this bucker (including Head element) */
+	TSet<UObjectBase*>* Items;
+
+	/** Constructor */
+	FHashBucket()
+		: Head(nullptr)
+		, Items(nullptr)
+	{}
+	~FHashBucket()
+	{
+		delete Items;
+	}
+	/** Adds an Object to the bucket */
+	FORCEINLINE void Add(UObjectBase* Object)
+	{
+		if (Items)
+		{
+			Items->Add(Object);
+		}
+		else if (Head)
+		{
+			Items = new TSet<UObjectBase*>();
+			Items->Add(Head);
+			Items->Add(Object);
+		}
+		else
+		{
+			Head = Object;
+		}
+	}
+	/** Removes an Object from the bucket */
+	FORCEINLINE int32 Remove(UObjectBase* Object)
+	{
+		int32 Result = 0;
+		if (Items)
+		{
+			Result = Items->Remove(Object);
+			if (Items->Num() <= 1)
+			{
+				if (Items->Num() == 1)
+				{
+					Head = *TSet<UObjectBase*>::TIterator(*Items);
+				}
+				else
+				{
+					// This is probably a bug, we should be able to delete Items set when it has 1 element.
+					Head = nullptr;
+				}
+				delete Items;
+				Items = nullptr;
+			}
+			else if (Head == Object)
+			{
+				Head = *TSet<UObjectBase*>::TIterator(*Items);
+			}
+		}
+		else
+		{
+			check(Head == Object);
+			Head = nullptr;
+			Result = 1;
+		}
+		return Result;
+	}
+	/** Checks if an Object exists in this bucket */
+	FORCEINLINE bool Contains(UObjectBase* Object) const
+	{
+		bool bContains = (Head == Object);
+		if (!bContains && Items)
+		{
+			bContains = Items->Contains(Object);
+		}
+		return bContains;
+	}
+	/** Returns the number of Objects in this bucket */
+	FORCEINLINE int32 Num() const
+	{
+		int32 Result = 0;
+		if (Items)
+		{
+			Result = Items->Num();
+		}
+		else if (Head)
+		{
+			Result = 1;
+		}
+		return Result;
+	}
+	/** Returns the amount of memory allocated for and by Items TSet */
+	FORCEINLINE uint32 GetItemsSize() const
+	{
+		uint32 AllocatedSize = 0;
+		if (Items)
+		{
+			AllocatedSize = (uint32)sizeof(*Items) + Items->GetAllocatedSize();
+		}
+		return AllocatedSize;
+	}
+private:
+	/** Gets an iterator for the TSet in this bucket or for the EmptyBucker if Items is null */
+	FORCEINLINE TSet<UObjectBase*>::TIterator GetIteratorForSet()
+	{
+		return Items ? Items->CreateIterator() : EmptyBucket.CreateIterator();
+	}
+};
+TSet<UObjectBase*> FHashBucket::EmptyBucket;
+
+/** Hash Bucket Iterator. Iterates over all Objects in the bucket */
+struct FHashBucketIterator
+{
+	FHashBucket& Bucket;
+	TSet<UObjectBase*>::TIterator SetIterator;
+	bool bReachedEndNoItems;
+
+	FORCEINLINE FHashBucketIterator(FHashBucket& InBucket)
+		: Bucket(InBucket)
+		, SetIterator(InBucket.GetIteratorForSet())
+		, bReachedEndNoItems(InBucket.Head == nullptr)
+	{
+	}
+	/** Advances the iterator to the next element. */
+	FORCEINLINE FHashBucketIterator& operator++()
+	{
+		if (Bucket.Items)
+		{
+			++SetIterator;
+		}
+		else
+		{
+			bReachedEndNoItems = true;
+		}
+		return *this;
+	}
+	/** conversion to "bool" returning true if the iterator is valid. */
+	FORCEINLINE_EXPLICIT_OPERATOR_BOOL() const
+	{
+		if (Bucket.Items)
+		{
+			return (bool)SetIterator;
+		}
+		else
+		{
+			return !bReachedEndNoItems;
+		}
+	}
+	/** inverse of the "bool" operator */
+	FORCEINLINE bool operator !() const
+	{
+		return !(bool)*this;
+	}
+	FORCEINLINE UObjectBase*& operator*()
+	{
+		if (Bucket.Items)
+		{
+			return *SetIterator;
+		}
+		else
+		{
+			return Bucket.Head;
+		}
+	}
+};
+
 class FUObjectHashTables
 {
 	FCriticalSection CriticalSection;
@@ -38,7 +214,7 @@ class FUObjectHashTables
 public:
 
 	/** Hash sets */
-	TMultiMap<int32, class UObjectBase*> Hash;
+	TMap<int32, FHashBucket> Hash;
 	TMultiMap<int32, class UObjectBase*> HashOuter;
 
 	/** Map of object to their outers, used to avoid an object iterator to find such things. **/
@@ -50,6 +226,39 @@ public:
 	{
 	}
 
+	/** Checks if the Hash/Object pair exists in the FName hash table */
+	FORCEINLINE bool PairExistsInHash(int32 InHash, UObjectBase* Object)
+	{
+		bool bResult = false;
+		FHashBucket* Bucket = Hash.Find(InHash);
+		if (Bucket)
+		{
+			bResult = Bucket->Contains(Object);
+		}
+		return bResult;
+	}
+	/** Adds the Hash/Object pair to the FName hash table */
+	FORCEINLINE void AddToHash(int32 InHash, UObjectBase* Object)
+	{
+		FHashBucket& Bucket = Hash.FindOrAdd(InHash);
+		Bucket.Add(Object);
+	}
+	/** Removes the Hash/Object pair from the FName hash table */
+	FORCEINLINE int32 RemoveFromHash(int32 InHash, UObjectBase* Object)
+	{
+		int32 NumRemoved = 0;
+		FHashBucket* Bucket = Hash.Find(InHash);
+		if (Bucket)
+		{
+			NumRemoved = Bucket->Remove(Object);
+			if (Bucket->Num() == 0)
+			{
+				Hash.Remove(InHash);
+			}
+		}
+		return NumRemoved;
+	}
+
 	FORCEINLINE void Lock()
 	{
 		CriticalSection.Lock();
@@ -125,31 +334,35 @@ UObject* StaticFindObjectFastExplicitThreadSafe(FUObjectHashTables& ThreadHash,
 	// Find an object with the specified name and (optional) class, in any package; if bAnyPackage is false, only matches top-level packages
 	int32 Hash = GetObjectHash(ObjectName);
 	FHashTableLock HashLock(ThreadHash);
-	for (TMultiMap<int32, class UObjectBase*>::TConstKeyIterator HashIt(ThreadHash.Hash, Hash); HashIt; ++HashIt)
+	FHashBucket* Bucket = ThreadHash.Hash.Find(Hash);
+	if (Bucket)
 	{
-		UObject *Object = (UObject *)HashIt.Value();
-		if
-			((Object->GetFName() == ObjectName)
+		for (FHashBucketIterator It(*Bucket); It; ++It)
+		{
+			UObject* Object = (UObject*)*It;
+			if
+				((Object->GetFName() == ObjectName)
 
-			/* Don't return objects that have any of the exclusive flags set */
-			&& !Object->HasAnyFlags(ExcludeFlags)
+				/* Don't return objects that have any of the exclusive flags set */
+				&& !Object->HasAnyFlags(ExcludeFlags)
 
-			/** If a class was specified, check that the object is of the correct class */
-			&& (ObjectClass == NULL || (bExactClass ? Object->GetClass() == ObjectClass : Object->IsA(ObjectClass)))
-			)
+				/** If a class was specified, check that the object is of the correct class */
+				&& (ObjectClass == nullptr || (bExactClass ? Object->GetClass() == ObjectClass : Object->IsA(ObjectClass)))
+				)
 
-		{
-			FString ObjectPath = Object->GetPathName();
-			/** Finally check the explicit path */
-			if (ObjectPath == ObjectPathName)
 			{
-				checkf(!Object->HasAnyFlags(RF_Unreachable), TEXT("%s"), *Object->GetFullName());
-				return Object;
+				FString ObjectPath = Object->GetPathName();
+				/** Finally check the explicit path */
+				if (ObjectPath == ObjectPathName)
+				{
+					checkf(!Object->HasAnyFlags(RF_Unreachable), TEXT("%s"), *Object->GetFullName());
+					return Object;
+				}
 			}
 		}
 	}
 
-	return NULL;
+	return nullptr;
 }
 
 /**
@@ -160,7 +373,7 @@ UObject* StaticFindObjectFastExplicitThreadSafe(FUObjectHashTables& ThreadHash,
  * @param	ObjectPathName	Full path name for the object to search for
  * @param	ExactClass		Whether to require an exact match with the passed in class
  * @param	ExclusiveFlags	Ignores objects that contain any of the specified exclusive flags
- * @return	Returns a pointer to the found object or NULL if none could be found
+ * @return	Returns a pointer to the found object or nullptr if none could be found
  */
 UObject* StaticFindObjectFastExplicit( UClass* ObjectClass, FName ObjectName, const FString& ObjectPathName, bool bExactClass, EObjectFlags ExcludeFlags/*=0*/ )
 {
@@ -178,8 +391,8 @@ UObject* StaticFindObjectFastExplicit( UClass* ObjectClass, FName ObjectName, co
 UObject* StaticFindObjectFastInternalThreadSafe(FUObjectHashTables& ThreadHash, UClass* ObjectClass, UObject* ObjectPackage, FName ObjectName, bool bExactClass, bool bAnyPackage, EObjectFlags ExcludeFlags)
 {
 	// If they specified an outer use that during the hashing
-	UObject* Result = NULL;
-	if (ObjectPackage != NULL)
+	UObject* Result = nullptr;
+	if (ObjectPackage != nullptr)
 	{
 		int32 Hash = GetObjectOuterHash(ObjectName, (PTRINT)ObjectPackage);
 		FHashTableLock HashLock(ThreadHash);
@@ -197,7 +410,7 @@ UObject* StaticFindObjectFastInternalThreadSafe(FUObjectHashTables& ThreadHash,
 				&& Object->GetOuter() == ObjectPackage
 
 				/** If a class was specified, check that the object is of the correct class */
-				&& (ObjectClass == NULL || (bExactClass ? Object->GetClass() == ObjectClass : Object->IsA(ObjectClass))))
+				&& (ObjectClass == nullptr || (bExactClass ? Object->GetClass() == ObjectClass : Object->IsA(ObjectClass))))
 			{
 				checkf(!Object->HasAnyFlags(RF_Unreachable), TEXT("%s"), *Object->GetFullName());
 				if (Result)
@@ -227,39 +440,44 @@ UObject* StaticFindObjectFastInternalThreadSafe(FUObjectHashTables& ThreadHash,
 		}
 		const int32 Hash = GetObjectHash(ActualObjectName);
 		FHashTableLock HashLock(ThreadHash);
-		for (TMultiMap<int32, class UObjectBase*>::TConstKeyIterator HashIt(ThreadHash.Hash, Hash); HashIt; ++HashIt)
+
+		FHashBucket* Bucket = ThreadHash.Hash.Find(Hash);
+		if (Bucket)
 		{
-			UObject *Object = (UObject *)HashIt.Value();
-			if
-				((Object->GetFName() == ActualObjectName)
+			for (FHashBucketIterator It(*Bucket); It; ++It)
+			{
+				UObject* Object = (UObject*)*It;
+				if
+					((Object->GetFName() == ActualObjectName)
 
-				/* Don't return objects that have any of the exclusive flags set */
-				&& !Object->HasAnyFlags(ExcludeFlags)
+					/* Don't return objects that have any of the exclusive flags set */
+					&& !Object->HasAnyFlags(ExcludeFlags)
 
-				/*If there is no package (no InObjectPackage specified, and InName's package is "")
-				and the caller specified any_package, then accept it, regardless of its package.
-				Or, if the object is a top-level package then accept it immediately.*/
-				&& (bAnyPackage || !Object->GetOuter())
+					/*If there is no package (no InObjectPackage specified, and InName's package is "")
+					and the caller specified any_package, then accept it, regardless of its package.
+					Or, if the object is a top-level package then accept it immediately.*/
+					&& (bAnyPackage || !Object->GetOuter())
 
 
-				/** If a class was specified, check that the object is of the correct class */
-				&& (ObjectClass == NULL || (bExactClass ? Object->GetClass() == ObjectClass : Object->IsA(ObjectClass)))
+					/** If a class was specified, check that the object is of the correct class */
+					&& (ObjectClass == nullptr || (bExactClass ? Object->GetClass() == ObjectClass : Object->IsA(ObjectClass)))
 
-				/** Ensure that the partial path provided matches the object found */
-				&& (Object->GetPathName().EndsWith(ObjectNameString)))
-			{
-				checkf(!Object->HasAnyFlags(RF_Unreachable), TEXT("%s"), *Object->GetFullName());
-				if (Result)
-				{
-					UE_LOG(LogUObjectHash, Warning, TEXT("Ambiguous search, could be %s or %s"), *GetFullNameSafe(Result), *GetFullNameSafe(Object));
-				}
-				else
+					/** Ensure that the partial path provided matches the object found */
+					&& (Object->GetPathName().EndsWith(ObjectNameString)))
 				{
-					Result = Object;
-				}
+					checkf(!Object->HasAnyFlags(RF_Unreachable), TEXT("%s"), *Object->GetFullName());
+					if (Result)
+					{
+						UE_LOG(LogUObjectHash, Warning, TEXT("Ambiguous search, could be %s or %s"), *GetFullNameSafe(Result), *GetFullNameSafe(Object));
+					}
+					else
+					{
+						Result = Object;
+					}
 #if (UE_BUILD_SHIPPING || UE_BUILD_TEST)
-				break;
+					break;
 #endif
+				}
 			}
 		}
 	}
@@ -272,7 +490,7 @@ UObject* StaticFindObjectFastInternal( UClass* ObjectClass, UObject* ObjectPacka
 	SCOPE_CYCLE_COUNTER( STAT_Hash_StaticFindObjectFastInternal );
 	INC_DWORD_STAT(STAT_FindObjectFast);
 
-	check(ObjectPackage != ANY_PACKAGE); // this could never have returned anything but NULL
+	check(ObjectPackage != ANY_PACKAGE); // this could never have returned anything but nullptr
 	// If they specified an outer use that during the hashing
 	auto& ThreadHash = FUObjectHashTables::Get();
 	UObject* Result = StaticFindObjectFastInternalThreadSafe( ThreadHash, ObjectClass, ObjectPackage, ObjectName, bExactClass, bAnyPackage, ExcludeFlags );
@@ -559,15 +777,16 @@ void HashObject(UObjectBase* Object)
 	SCOPE_CYCLE_COUNTER( STAT_Hash_HashObject );
 
 	FName Name = Object->GetFName();
-	if( Name != NAME_None )
+	if (Name != NAME_None)
 	{
-		int32 Hash = GetObjectHash( Name );
-		auto& ThreadHash = FUObjectHashTables::Get();
+		int32 Hash = 0;
 
-		FHashTableLock HashLock( ThreadHash );
-		checkSlow( !ThreadHash.Hash.FindPair( Hash, Object ) );  // if it already exists, something is wrong with the external code
+		auto& ThreadHash = FUObjectHashTables::Get();
+		FHashTableLock HashLock(ThreadHash);
 
-		ThreadHash.Hash.Add( Hash, Object );
+		Hash = GetObjectHash(Name);				
+		checkSlow(!ThreadHash.PairExistsInHash(Hash, Object));  // if it already exists, something is wrong with the external code
+		ThreadHash.AddToHash(Hash, Object);
 
 		Hash = GetObjectOuterHash( Name, (PTRINT)Object->GetOuter() );
 		checkSlow( !ThreadHash.HashOuter.FindPair( Hash, Object ) );  // if it already exists, something is wrong with the external code
@@ -585,16 +804,18 @@ void HashObject(UObjectBase* Object)
  */
 void UnhashObject(UObjectBase* Object)
 {
-	SCOPE_CYCLE_COUNTER( STAT_Hash_UnhashObject );
 	FName Name = Object->GetFName();
-	if( Name != NAME_None )
+	if (Name != NAME_None)
 	{
+		int32 Hash = 0;
+		int32 NumRemoved = 0;
 
-		int32 Hash = GetObjectHash( Name );
 		auto& ThreadHash = FUObjectHashTables::Get();
-		FHashTableLock LockHash( ThreadHash );
-		int32 NumRemoved = ThreadHash.Hash.RemoveSingle( Hash, Object );
-		check( NumRemoved == 1 ); // must have existed, else something is wrong with the external code
+		FHashTableLock LockHash(ThreadHash);
+
+		Hash = GetObjectHash(Name);
+		NumRemoved = ThreadHash.RemoveFromHash(Hash, Object);
+		check(NumRemoved == 1); // must have existed, else something is wrong with the external code
 
 		Hash = GetObjectOuterHash( Name, (PTRINT)Object->GetOuter() );
 		NumRemoved = ThreadHash.HashOuter.RemoveSingle( Hash, Object );
@@ -679,6 +900,82 @@ void LogHashStatisticsInternal(TMultiMap<int32, UObjectBase*>& Hash, FOutputDevi
 		MinCollisions,
 		FMath::FloorToInt(((float)TotalCollisions / (float)SlotsInUse)),
 		MaxCollisions);
+
+	// Calculate Hashtable size
+	const uint32 HashtableAllocatedSize = Hash.GetAllocatedSize();
+	Ar.Logf(TEXT("Total memory allocated for Object Outer Hash: %u bytes."), HashtableAllocatedSize);
+}
+
+void LogHashStatisticsInternal(TMap<int32, FHashBucket>& Hash, FOutputDevice& Ar, const bool bShowHashBucketCollisionInfo)
+{
+	TArray<int32> HashBuckets;
+	// Get the set of keys in use, which is the number of hash buckets
+	int32 SlotsInUse = Hash.Num();
+
+	int32 TotalCollisions = 0;
+	int32 MinCollisions = OBJECT_HASH_BINS;
+	int32 MaxCollisions = 0;
+	int32 MaxBin = 0;
+	int32 NumBucketsWithMoreThanOneItem = 0;
+
+	// Dump how many slots are in use
+	Ar.Logf(TEXT("Slots in use %d"), SlotsInUse);
+
+	// Work through each slot and figure out how many collisions
+	for (auto& HashPair : Hash)
+	{
+		int32 Collisions = HashPair.Value.Num();
+		check(Collisions >= 0);
+		if (Collisions > 1)
+		{
+			NumBucketsWithMoreThanOneItem++;
+		}
+
+		// Keep the global stats
+		TotalCollisions += Collisions;
+		if (Collisions > MaxCollisions)
+		{
+			MaxBin = HashPair.Key;
+		}
+		MaxCollisions = FMath::Max<int32>(Collisions, MaxCollisions);
+		MinCollisions = FMath::Min<int32>(Collisions, MinCollisions);
+
+		if (bShowHashBucketCollisionInfo)
+		{
+			// Now log the output
+			Ar.Logf(TEXT("\tSlot %d has %d collisions"), HashPair.Key, Collisions);
+		}
+	}
+	Ar.Logf(TEXT(""));
+
+	// Dump the first 30 objects in the worst bin for inspection
+	Ar.Logf(TEXT("Worst hash bucket contains:"));
+	int32 Count = 0;
+	FHashBucket& WorstBucket = Hash.FindChecked(MaxBin);
+	for (FHashBucketIterator It(WorstBucket); It; ++It)
+	{
+		UObject* Object = (UObject*)*It;
+		Ar.Logf(TEXT("\tObject is %s (%s)"), *Object->GetName(), *Object->GetFullName());
+		Count++;
+	}
+	Ar.Logf(TEXT(""));
+
+	// Now dump how efficient the hash is
+	Ar.Logf(TEXT("Collision Stats: Best Case (%d), Average Case (%d), Worst Case (%d), Number of buckets with more than one item (%d/%d)"),
+		MinCollisions,
+		FMath::FloorToInt(((float)TotalCollisions / (float)SlotsInUse)),
+		MaxCollisions,
+		NumBucketsWithMoreThanOneItem,
+		SlotsInUse);
+
+	// Calculate Hashtable size
+	uint32 HashtableAllocatedSize = Hash.GetAllocatedSize();
+	// Calculate the size of a all Allocations inside of the buckets (TSet Items)
+	for (auto& Pair : Hash)
+	{
+		HashtableAllocatedSize += Pair.Value.GetItemsSize();
+	}
+	Ar.Logf(TEXT("Total memory allocated for and by Object Hash: %u bytes."), HashtableAllocatedSize);
 }
 
 void LogHashStatistics(FOutputDevice& Ar, const bool bShowHashBucketCollisionInfo)
