Async loading improvements:

- changed an assert to a warning when calling LoadObject in PostLoad
- added an opt-out ability for IsTimeLimitExceeded warning
- added descriptions of the current work performed during async loading for IsTimeLimitExceeded warning
- made sure we don't check IsTimeLimitExceeded when we already know we timed out.

diff --git a/Engine/Config/BaseEngine.ini b/Engine/Config/BaseEngine.ini
--- a/Engine/Config/BaseEngine.ini
+++ b/Engine/Config/BaseEngine.ini
@@ -1308,6 +1308,7 @@ UseStrictEngineVersioning=True
 FlushStreamingOnGC=False
 MinBulkDataSizeForAsyncLoading=131072
 AsyncLoadingThreadEnabled=True
+WarnIfTimeLimitExeeded=True
 
 [Internationalization]
 +LocalizationPaths=../../../Engine/Content/Localization/Engine
diff --git a/Engine/Source/Runtime/CoreUObject/Private/Serialization/AsyncLoading.cpp b/Engine/Source/Runtime/CoreUObject/Private/Serialization/AsyncLoading.cpp
--- a/Engine/Source/Runtime/CoreUObject/Private/Serialization/AsyncLoading.cpp
+++ b/Engine/Source/Runtime/CoreUObject/Private/Serialization/AsyncLoading.cpp
@@ -177,9 +177,22 @@ static FORCEINLINE bool IsTimeLimitExceeded(double InTickStartTime, bool bUseTim
 		double CurrentTime = FPlatformTime::Seconds();
 		bTimeLimitExceeded = CurrentTime - InTickStartTime > InTimeLimit;
 
-		// Log single operations that take longer than timelimit (but only in cooked builds)
-		if (FPlatformProperties::RequiresCookedData() && (CurrentTime - InTickStartTime) > (1.5 * InTimeLimit))
+		// One time init for ini override if we should log a warning when time limit is exceeded
+		static struct FWarnIfTimeLimitExeeded
 		{
+			bool Value;
+			FWarnIfTimeLimitExeeded()
+			{
+				check(GConfig);
+				bool bConfigValue = true; // Default to true
+				GConfig->GetBool(TEXT("Core.System"), TEXT("WarnIfTimeLimitExeeded"), bConfigValue, GEngineIni);
+				Value = bConfigValue;
+			}
+		} WarnIfTimeLimitExeeded;
+
+		// Log single operations that take longer than time limit (but only in cooked builds)
+		if (FPlatformProperties::RequiresCookedData() && WarnIfTimeLimitExeeded.Value && (CurrentTime - InTickStartTime) > (1.5 * InTimeLimit))
+		{			
 			UE_LOG(LogStreaming, Warning, TEXT("IsTimeLimitExceeded: %s %s took (less than) %5.2f ms"),
 				InLastTypeOfWorkPerformed ? InLastTypeOfWorkPerformed : TEXT("unknown"),
 				InLastObjectWorkWasPerformedOn ? *InLastObjectWorkWasPerformedOn->GetFullName() : TEXT("nullptr"),
@@ -398,7 +411,7 @@ EAsyncPackageState::Type FAsyncLoadingThread::ProcessLoadedPackages(bool bUseTim
 		
 	for (int32 PackageIndex = 0; PackageIndex < LoadedPackagesToProcess.Num(); ++PackageIndex)
 	{
-		if (PackageIndex % 20 == 0 && IsTimeLimitExceeded(TickStartTime, bUseTimeLimit, TimeLimit))
+		if (PackageIndex % 20 == 0 && IsTimeLimitExceeded(TickStartTime, bUseTimeLimit, TimeLimit, TEXT("ProcessLoadedPackages")))
 		{
 			break;
 		}
@@ -461,13 +474,13 @@ EAsyncPackageState::Type FAsyncLoadingThread::TickAsyncLoading(bool bUseTimeLimi
 		TimeLimitUsedForPreloading = FPlatformTime::Seconds() - TickStartTime;
 	}
 
-	if (!IsTimeLimitExceeded(TickStartTime, bUseTimeLimit, TimeLimit))
+	if (!IsTimeLimitExceeded(TickStartTime, bUseTimeLimit, TimeLimit, TEXT("Pre-ProcessLoadedPackages")))
 	{
 		double RemainingTimeLimit = FMath::Max(0.0, TimeLimit - TimeLimitUsedForPreloading);
 		Result = ProcessLoadedPackages(bUseTimeLimit, bUseFullTimeLimit, RemainingTimeLimit, ExcludeType);
 	}
 
-	if (!IsTimeLimitExceeded(TickStartTime, bUseTimeLimit, TimeLimit))
+	if (Result != EAsyncPackageState::TimeOut && !IsTimeLimitExceeded(TickStartTime, bUseTimeLimit, TimeLimit, TEXT("Pre-EmptyReferencedObjects")))
 	{
 		FScopeLock QueueLock(&QueueCritical);
 		FScopeLock LoadedLock(&LoadedPackagesCritical);
@@ -1388,7 +1401,6 @@ EAsyncPackageState::Type FAsyncPackage::PostLoadObjects()
 
 	// New objects might have been loaded during PostLoad.
 	Result = (PreLoadIndex == ObjLoaded.Num()) && (PostLoadIndex == ObjLoaded.Num()) ? EAsyncPackageState::Complete : EAsyncPackageState::TimeOut;
-	UE_CLOG(PreLoadIndex < ObjLoaded.Num(), LogStreaming, Fatal, TEXT("Objects have been loading in PostLoad. This is no longer allowed."));
 
 	return Result;
 }
diff --git a/Engine/Source/Runtime/CoreUObject/Private/UObject/UObjectGlobals.cpp b/Engine/Source/Runtime/CoreUObject/Private/UObject/UObjectGlobals.cpp
--- a/Engine/Source/Runtime/CoreUObject/Private/UObject/UObjectGlobals.cpp
+++ b/Engine/Source/Runtime/CoreUObject/Private/UObject/UObjectGlobals.cpp
@@ -761,10 +761,7 @@ UObject* StaticLoadObjectInternal(UClass* ObjectClass, UObject* InOuter, const T
 
 UObject* StaticLoadObject(UClass* ObjectClass, UObject* InOuter, const TCHAR* InName, const TCHAR* Filename, uint32 LoadFlags, UPackageMap* Sandbox, bool bAllowObjectReconciliation )
 {
-	// Currently StaticLoadObject will not work if executed in PostLoad on the async loading thread
-	// It should be forbidden anyway but there's legacy code that is still doing it.
-	// Note that IsInAsyncLoadingThread() == true if async loading thread is disabled but we're inside of the Async Loading Tick
-	UE_CLOG(FUObjectThreadContext::Get().IsRoutingPostLoad && IsInAsyncLoadingThread() && !IsInGameThread(), LogUObjectGlobals, Fatal, TEXT("Calling StaticLoadObject during PostLoad is forbidden."));
+	UE_CLOG(FUObjectThreadContext::Get().IsRoutingPostLoad && IsInAsyncLoadingThread(), LogUObjectGlobals, Warning, TEXT("Calling StaticLoadObject during PostLoad may result in hitches during streaming."));
 
 	UObject* Result = StaticLoadObjectInternal(ObjectClass, InOuter, InName, Filename, LoadFlags, Sandbox, bAllowObjectReconciliation);
 	if (!Result)
