Fixing typo in WarnIfTimeLimitExceeded code

diff --git a/Engine/Config/BaseEngine.ini b/Engine/Config/BaseEngine.ini
--- a/Engine/Config/BaseEngine.ini
+++ b/Engine/Config/BaseEngine.ini
@@ -1308,7 +1308,7 @@ UseStrictEngineVersioning=True
 FlushStreamingOnGC=False
 MinBulkDataSizeForAsyncLoading=131072
 AsyncLoadingThreadEnabled=True
-WarnIfTimeLimitExeeded=True
+WarnIfTimeLimitExceeded=True
 
 [Internationalization]
 +LocalizationPaths=../../../Engine/Content/Localization/Engine
diff --git a/Engine/Source/Runtime/CoreUObject/Private/Serialization/AsyncLoading.cpp b/Engine/Source/Runtime/CoreUObject/Private/Serialization/AsyncLoading.cpp
--- a/Engine/Source/Runtime/CoreUObject/Private/Serialization/AsyncLoading.cpp
+++ b/Engine/Source/Runtime/CoreUObject/Private/Serialization/AsyncLoading.cpp
@@ -178,20 +178,20 @@ static FORCEINLINE bool IsTimeLimitExceeded(double InTickStartTime, bool bUseTim
 		bTimeLimitExceeded = CurrentTime - InTickStartTime > InTimeLimit;
 
 		// One time init for ini override if we should log a warning when time limit is exceeded
-		static struct FWarnIfTimeLimitExeeded
+		static struct FWarnIfTimeLimitExceeded
 		{
 			bool Value;
-			FWarnIfTimeLimitExeeded()
+			FWarnIfTimeLimitExceeded()
 			{
 				check(GConfig);
 				bool bConfigValue = true; // Default to true
-				GConfig->GetBool(TEXT("Core.System"), TEXT("WarnIfTimeLimitExeeded"), bConfigValue, GEngineIni);
+				GConfig->GetBool(TEXT("Core.System"), TEXT("WarnIfTimeLimitExceeded"), bConfigValue, GEngineIni);
 				Value = bConfigValue;
 			}
-		} WarnIfTimeLimitExeeded;
+		} WarnIfTimeLimitExceeded;
 
 		// Log single operations that take longer than time limit (but only in cooked builds)
-		if (FPlatformProperties::RequiresCookedData() && WarnIfTimeLimitExeeded.Value && (CurrentTime - InTickStartTime) > (1.5 * InTimeLimit))
+		if (FPlatformProperties::RequiresCookedData() && WarnIfTimeLimitExceeded.Value && (CurrentTime - InTickStartTime) > (1.5 * InTimeLimit))
 		{			
 			UE_LOG(LogStreaming, Warning, TEXT("IsTimeLimitExceeded: %s %s took (less than) %5.2f ms"),
 				InLastTypeOfWorkPerformed ? InLastTypeOfWorkPerformed : TEXT("unknown"),
