Fix for UE-5191, only gather parameters once when calling a native multicast RPC from blueprints.
Also, fixed a potential memory stomp in this case if function parameters were extremely large (over 1024 bytes).

diff --git a/Engine/Source/Runtime/CoreUObject/Private/UObject/ScriptCore.cpp b/Engine/Source/Runtime/CoreUObject/Private/UObject/ScriptCore.cpp
--- a/Engine/Source/Runtime/CoreUObject/Private/UObject/ScriptCore.cpp
+++ b/Engine/Source/Runtime/CoreUObject/Private/UObject/ScriptCore.cpp
@@ -389,13 +389,38 @@ void UObject::CallFunction( FFrame& Stack, RESULT_DECL, UFunction* Function )
 
 	if (Function->FunctionFlags & FUNC_Native)
 	{
-		// Call native networkable function.
-		uint8 Buffer[1024];
+		uint8* Buffer = (uint8*)FMemory_Alloca(Function->ParmsSize);
 		int32 FunctionCallspace = GetFunctionCallspace( Function, Buffer, &Stack );
 		uint8* SavedCode = NULL;
 		if (FunctionCallspace & FunctionCallspace::Remote)
 		{
+			// Call native networkable function.
+
 			SavedCode = Stack.Code; // Since this is native, we need to rollback the stack if we are calling both remotely and locally
+
+			FMemory::Memzero( Buffer, Function->ParmsSize );
+
+			// Form the RPC parameters.
+			for( TFieldIterator<UProperty> It(Function); It && (It->PropertyFlags & (CPF_Parm|CPF_ReturnParm))==CPF_Parm; ++It )
+			{
+				uint8* CurrentPropAddr = It->ContainerPtrToValuePtr<uint8>(Buffer);
+				if ( Cast<UBoolProperty>(*It) && It->ArrayDim == 1 )
+				{
+					// we're going to get '1' returned for bools that are set, so we need to manually mask it in to the proper place
+					bool bValue = false;
+					Stack.Step(Stack.Object, &bValue);
+					if (bValue)
+					{
+						((UBoolProperty*)*It)->SetPropertyValue( CurrentPropAddr, true );
+					}
+				}
+				else
+				{
+					Stack.Step(Stack.Object, CurrentPropAddr);
+				}
+			}
+			checkSlow(*Stack.Code==EX_EndFunctionParms);
+
 			CallRemoteFunction(Function, Buffer, Stack.OutParms, &Stack);
 		}
 
diff --git a/Engine/Source/Runtime/Engine/Private/NetworkDriver.cpp b/Engine/Source/Runtime/Engine/Private/NetworkDriver.cpp
--- a/Engine/Source/Runtime/Engine/Private/NetworkDriver.cpp
+++ b/Engine/Source/Runtime/Engine/Private/NetworkDriver.cpp
@@ -813,35 +813,7 @@ void UNetDriver::InternalProcessRemoteFunction
 
 	TArray< UProperty * > LocalOutParms;
 
-	// Form the RPC parameters.
-	if( Stack )
-	{
-		// this only happens for native replicated functions called from script
-		// because in that case, the C++ function itself handles evaluating the parameters
-		// so we cannot do that before calling CallRemoteFunction() as we do with all other cases
-		FMemory::Memzero( Parms, Function->ParmsSize );
-
-		for( TFieldIterator<UProperty> It(Function); It && (It->PropertyFlags & (CPF_Parm|CPF_ReturnParm))==CPF_Parm; ++It )
-		{
-			uint8* CurrentPropAddr = It->ContainerPtrToValuePtr<uint8>(Parms);
-			if ( Cast<UBoolProperty>(*It) && It->ArrayDim == 1 )
-			{
-				// we're going to get '1' returned for bools that are set, so we need to manually mask it in to the proper place
-				bool bValue = false;
-				Stack->Step(Stack->Object, &bValue);
-				if (bValue)
-				{
-					((UBoolProperty*)*It)->SetPropertyValue( CurrentPropAddr, true );
-				}
-			}
-			else
-			{
-				Stack->Step(Stack->Object, CurrentPropAddr);
-			}
-		}
-		checkSlow(*Stack->Code==EX_EndFunctionParms);
-	}
-	else
+	if( Stack == nullptr )
 	{
 		// Look for CPF_OutParm's, we'll need to copy these into the local parameter memory manually
 		// The receiving side will pull these back out when needed
