Support for retrieval of localized metadata.

Added UField::GetMetaDataText which takes a key and namespace for the localization manager to find the metadata's localized text. Defaults to calling UField::GetMetaData if there is no localized text available.

diff --git a/Engine/Source/Runtime/CoreUObject/Private/UObject/Class.cpp b/Engine/Source/Runtime/CoreUObject/Private/UObject/Class.cpp
--- a/Engine/Source/Runtime/CoreUObject/Private/UObject/Class.cpp
+++ b/Engine/Source/Runtime/CoreUObject/Private/UObject/Class.cpp
@@ -323,6 +323,37 @@ const FString& UField::GetMetaData(const FName& Key) const
 	return MetaDataString;
 }
 
+const FText UField::GetMetaDataText(const TCHAR* MetaDataKey, const FString LocalizationNamespace, const FString LocalizationKey) const
+{
+	FText LocalizedMetaData;
+	if ( !( FText::FindText( LocalizationNamespace, LocalizationKey, /*OUT*/LocalizedMetaData ) ) )
+	{
+		FString DefaultMetaData;
+		if( HasMetaData( MetaDataKey ))
+		{
+			DefaultMetaData = GetMetaData(MetaDataKey);
+			LocalizedMetaData = FText::FromString(DefaultMetaData);
+		}
+	}
+
+	return LocalizedMetaData;
+}
+
+const FText UField::GetMetaDataText(const FName& MetaDataKey, const FString LocalizationNamespace, const FString LocalizationKey) const
+{
+	FText LocalizedMetaData;
+	if ( !( FText::FindText( LocalizationNamespace, LocalizationKey, /*OUT*/LocalizedMetaData ) ) )
+	{
+		FString DefaultMetaData;
+		if( HasMetaData( MetaDataKey ))
+		{
+			DefaultMetaData = GetMetaData(MetaDataKey);
+			LocalizedMetaData = FText::FromString(DefaultMetaData);
+		}
+	}
+	return LocalizedMetaData;
+}
+
 /**
  * Sets the metadata value associated with the key
  * 
diff --git a/Engine/Source/Runtime/CoreUObject/Public/UObject/Class.h b/Engine/Source/Runtime/CoreUObject/Public/UObject/Class.h
--- a/Engine/Source/Runtime/CoreUObject/Public/UObject/Class.h
+++ b/Engine/Source/Runtime/CoreUObject/Public/UObject/Class.h
@@ -115,6 +115,26 @@ class COREUOBJECT_API UField : public UObject
 	/** Find the metadata value associated with the key */
 	const FString& GetMetaData(const FName& Key) const;
 
+	/**
+	 * Find the metadata value associated with the key and localization namespace and key
+	 *
+	 * @param Key						The key to lookup in the metadata
+	 * @param LocalizationNamespace		Namespace to lookup in the localization manager
+	 * @param LocalizationKey			Key to lookup in the localization manager
+	 * @return							Localized metadata if available, defaults to whatever is provided via GetMetaData
+	 */
+	const FText UField::GetMetaDataText(const TCHAR* MetaDataKey, const FString LocalizationNamespace = FString(), const FString LocalizationKey = FString()) const;
+
+	/**
+	 * Find the metadata value associated with the key and localization namespace and key
+	 *
+	 * @param Key						The key to lookup in the metadata
+	 * @param LocalizationNamespace		Namespace to lookup in the localization manager
+	 * @param LocalizationKey			Key to lookup in the localization manager
+	 * @return							Localized metadata if available, defaults to whatever is provided via GetMetaData
+	 */
+	const FText UField::GetMetaDataText(const FName& MetaDataKey, const FString LocalizationNamespace = FString(), const FString LocalizationKey = FString()) const;
+
 	/**
 	 * Sets the metadata value associated with the key
 	 * 
