Fix build failures on Mac.

diff --git a/Engine/Source/Runtime/CoreUObject/Public/UObject/Class.h b/Engine/Source/Runtime/CoreUObject/Public/UObject/Class.h
--- a/Engine/Source/Runtime/CoreUObject/Public/UObject/Class.h
+++ b/Engine/Source/Runtime/CoreUObject/Public/UObject/Class.h
@@ -123,7 +123,7 @@ class COREUOBJECT_API UField : public UObject
 	 * @param LocalizationKey			Key to lookup in the localization manager
 	 * @return							Localized metadata if available, defaults to whatever is provided via GetMetaData
 	 */
-	const FText UField::GetMetaDataText(const TCHAR* MetaDataKey, const FString LocalizationNamespace = FString(), const FString LocalizationKey = FString()) const;
+	const FText GetMetaDataText(const TCHAR* MetaDataKey, const FString LocalizationNamespace = FString(), const FString LocalizationKey = FString()) const;
 
 	/**
 	 * Find the metadata value associated with the key and localization namespace and key
@@ -133,7 +133,7 @@ class COREUOBJECT_API UField : public UObject
 	 * @param LocalizationKey			Key to lookup in the localization manager
 	 * @return							Localized metadata if available, defaults to whatever is provided via GetMetaData
 	 */
-	const FText UField::GetMetaDataText(const FName& MetaDataKey, const FString LocalizationNamespace = FString(), const FString LocalizationKey = FString()) const;
+	const FText GetMetaDataText(const FName& MetaDataKey, const FString LocalizationNamespace = FString(), const FString LocalizationKey = FString()) const;
 
 	/**
 	 * Sets the metadata value associated with the key
