custom alignment for bulk data

diff --git a/Engine/Source/Runtime/CoreUObject/Private/Serialization/BulkData.cpp b/Engine/Source/Runtime/CoreUObject/Private/Serialization/BulkData.cpp
--- a/Engine/Source/Runtime/CoreUObject/Private/Serialization/BulkData.cpp
+++ b/Engine/Source/Runtime/CoreUObject/Private/Serialization/BulkData.cpp
@@ -96,12 +96,13 @@ FUntypedBulkData::FUntypedBulkData( const FUntypedBulkData& Other )
 	DECLARE_SCOPE_CYCLE_COUNTER(TEXT("FUntypedBulkData::FUntypedBulkData"), STAT_UBD_Constructor, STATGROUP_Memory);
 
 	InitializeMemberVariables();
+	BulkDataAlignment = Other.BulkDataAlignment;
 
 	// Prepare bulk data pointer. Can't call any functions that would call virtual GetElementSize on "this" as
 	// we're in the constructor of the base class and would hence call a pure virtual.
 	ElementCount	= Other.ElementCount;
 	check(bShouldFreeOnEmpty);
-	BulkData		= FMemory::Realloc( BulkData, Other.GetBulkDataSize() );
+	BulkData		= FMemory::Realloc( BulkData, Other.GetBulkDataSize(), BulkDataAlignment );
 
 	// Copy data over.
 	Copy( Other );
@@ -160,6 +161,8 @@ FUntypedBulkData& FUntypedBulkData::operator=( const FUntypedBulkData& Other )
 	// Remove bulk data, avoiding potential load in Lock call.
 	RemoveBulkData();
 	
+	BulkDataAlignment = Other.BulkDataAlignment;
+
 	// Reallocate to size of src.
 	Lock(LOCK_READ_WRITE);
 	Realloc(Other.GetElementCount());
@@ -430,7 +433,7 @@ void FUntypedBulkData::GetCopy( void** Dest, bool bDiscardInternalCopy )
 			else
 			{
 				// Allocate enough memory for data...
-				*Dest = FMemory::Malloc( GetBulkDataSize() );
+				*Dest = FMemory::Malloc( GetBulkDataSize(), BulkDataAlignment );
 				// ... and copy it into memory now pointed to by out parameter.
 				FMemory::Memcpy( *Dest, BulkData, GetBulkDataSize() );
 			}
@@ -439,7 +442,7 @@ void FUntypedBulkData::GetCopy( void** Dest, bool bDiscardInternalCopy )
 		else
 		{
 			// Allocate enougn memory for data...
-			*Dest = FMemory::Malloc( GetBulkDataSize() );
+			*Dest = FMemory::Malloc( GetBulkDataSize(), BulkDataAlignment );
 			// ... and directly load into it.
 			LoadDataIntoMemory( *Dest );
 		}
@@ -484,7 +487,6 @@ void* FUntypedBulkData::Lock( uint32 LockFlags )
 		UE_LOG(LogSerialization, Fatal,TEXT("Unknown lock flag %i"),LockFlags);
 	}
 
-	check( BulkData );
 	return BulkData;
 }
 
@@ -517,7 +519,7 @@ void* FUntypedBulkData::Realloc( int32 InElementCount )
 	// Progate element count and reallocate data based on new size.
 	ElementCount	= InElementCount;
 	check(bShouldFreeOnEmpty);
-	BulkData		= FMemory::Realloc( BulkData, GetBulkDataSize() );
+	BulkData		= FMemory::Realloc( BulkData, GetBulkDataSize(), BulkDataAlignment );
 	return BulkData;
 }
 
@@ -611,6 +613,26 @@ uint32 FUntypedBulkData::GetBulkDataFlags() const
 	return BulkDataFlags;
 }
 
+/**
+ * Sets the passed in bulk data alignment.
+ *
+ * @param BulkDataAlignmentToSet	Bulk data alignment to set
+ */
+void FUntypedBulkData::SetBulkDataAlignment( uint32 BulkDataAlignmentToSet )
+{
+	BulkDataAlignment = BulkDataAlignmentToSet;
+}
+
+/**
+* Gets the current bulk data alignment.
+*
+* @return Bulk data alignment currently set
+*/
+uint32 FUntypedBulkData::GetBulkDataAlignment() const
+{
+	return BulkDataAlignment;
+}
+
 /**
  * Clears the passed in bulk data flags.
  *
@@ -647,7 +669,7 @@ void FUntypedBulkData::StartSerializingBulkData(FArchive& Ar, UObject* Owner, in
 		BulkDataAsync = GetBulkDataResourceMemory(Owner, Idx);
 		if (!BulkDataAsync)
 		{
-			BulkDataAsync = FMemory::Realloc(BulkDataAsync, GetBulkDataSize());
+			BulkDataAsync = FMemory::Realloc(BulkDataAsync, GetBulkDataSize(), BulkDataAlignment);
 		}
 
 		FArchive* FileReaderAr = IFileManager::Get().CreateFileReader(*Filename, FILEREAD_Silent);
@@ -722,7 +744,7 @@ void FUntypedBulkData::Serialize( FArchive& Ar, UObject* Owner, int32 Idx )
 
 				// Allocate bulk data.
 				check(bShouldFreeOnEmpty);
-				BulkData = FMemory::Realloc( BulkData, GetBulkDataSize() );
+				BulkData = FMemory::Realloc( BulkData, GetBulkDataSize(), BulkDataAlignment );
 
 				// Deserialize bulk data.
 				SerializeBulkData( Ar, BulkData );
@@ -818,7 +840,7 @@ void FUntypedBulkData::Serialize( FArchive& Ar, UObject* Owner, int32 Idx )
 						BulkData = GetBulkDataResourceMemory(Owner, Idx);
 						if (!BulkData)
 						{
-							BulkData = FMemory::Realloc(BulkData, GetBulkDataSize());
+							BulkData = FMemory::Realloc(BulkData, GetBulkDataSize(), BulkDataAlignment);
 						}
 						// if the payload is stored inline, just serialize it
 						SerializeBulkData(Ar, BulkData);
@@ -842,7 +864,7 @@ void FUntypedBulkData::Serialize( FArchive& Ar, UObject* Owner, int32 Idx )
 				  BulkData = GetBulkDataResourceMemory(Owner,Idx);
 				  if( !BulkData )
 				  {
-					  BulkData = FMemory::Realloc( BulkData, GetBulkDataSize() );
+					  BulkData = FMemory::Realloc( BulkData, GetBulkDataSize(), BulkDataAlignment );
 				  }
 				  
 				  if (bPayloadInline)
@@ -1076,6 +1098,7 @@ void FUntypedBulkData::InitializeMemberVariables()
 	ElementCount = 0;
 	BulkDataOffsetInFile = INDEX_NONE;
 	BulkDataSizeOnDisk = INDEX_NONE;
+	BulkDataAlignment = DEFAULT_ALIGNMENT;
 	BulkData = nullptr;
 	BulkDataAsync = nullptr;
 	LockStatus = LOCKSTATUS_Unlocked;
@@ -1214,12 +1237,13 @@ void FUntypedBulkData::MakeSureBulkDataIsLoaded()
 		}
 		else
 		{
+			const int32 BytesNeeded = GetBulkDataSize();
 			// Allocate memory for bulk data.
-			BulkData = FMemory::Malloc(GetBulkDataSize());
+			BulkData = FMemory::Malloc(BytesNeeded, BulkDataAlignment);
 
 			// Only load if there is something to load. E.g. we might have just created the bulk data array
 			// in which case it starts out with a size of zero.
-			if (GetBulkDataSize() > 0)
+			if (BytesNeeded > 0)
 			{
 				LoadDataIntoMemory(BulkData);
 			}
@@ -1426,8 +1450,7 @@ void FFloatBulkData::SerializeElement( FArchive& Ar, void* Data, int32 ElementIn
 	Ar << FloatData;
 }
 
-
-void FFormatContainer::Serialize(FArchive& Ar, UObject* Owner, const TArray<FName>* FormatsToSave, bool bSingleUse)
+void FFormatContainer::Serialize(FArchive& Ar, UObject* Owner, const TArray<FName>* FormatsToSave, bool bSingleUse, uint32 Alignment)
 {
 	if (Ar.IsLoading())
 	{
@@ -1438,6 +1461,7 @@ void FFormatContainer::Serialize(FArchive& Ar, UObject* Owner, const TArray<FNam
 			FName Name;
 			Ar << Name;
 			FByteBulkData& Bulk = GetFormat(Name);
+			Bulk.SetBulkDataAlignment(Alignment);
 			Bulk.Serialize(Ar, Owner);
 		}
 	}
diff --git a/Engine/Source/Runtime/CoreUObject/Public/Serialization/BulkData.h b/Engine/Source/Runtime/CoreUObject/Public/Serialization/BulkData.h
--- a/Engine/Source/Runtime/CoreUObject/Public/Serialization/BulkData.h
+++ b/Engine/Source/Runtime/CoreUObject/Public/Serialization/BulkData.h
@@ -207,6 +207,20 @@ struct COREUOBJECT_API FUntypedBulkData
 	*/
 	uint32 GetBulkDataFlags() const;
 
+	/**
+	 * Sets the passed in bulk data alignment.
+	 *
+	 * @param BulkDataAlignmentToSet	Bulk data alignment to set
+	 */
+	void SetBulkDataAlignment( uint32 BulkDataAlignmentToSet );
+
+	/**
+	* Gets the current bulk data alignment.
+	*
+	* @return Bulk data alignment currently set
+	*/
+	uint32 GetBulkDataAlignment() const;
+
 	/**
 	 * Clears the passed in bulk data flags.
 	 *
@@ -409,6 +423,8 @@ private:
 	int64					BulkDataOffsetInFile;
 	/** Size of bulk data on disk or INDEX_NONE if no association														*/
 	int32					BulkDataSizeOnDisk;
+	/** Alignment of bulk data																							*/
+	int32					BulkDataAlignment;
 
 	/** Pointer to cached bulk data																						*/
 	void*				BulkData;
@@ -531,6 +547,7 @@ struct COREUOBJECT_API FFloatBulkData : public FUntypedBulkData
 class FFormatContainer
 {
 	TMap<FName, FByteBulkData*> Formats;
+	uint32 Alignment;
 public:
 	~FFormatContainer()
 	{
@@ -558,7 +575,7 @@ public:
 		}
 		Formats.Empty();
 	}
-	COREUOBJECT_API void Serialize(FArchive& Ar, UObject* Owner, const TArray<FName>* FormatsToSave = NULL, bool bSingleUse = true);
+	COREUOBJECT_API void Serialize(FArchive& Ar, UObject* Owner, const TArray<FName>* FormatsToSave = NULL, bool bSingleUse = true, uint32 Alignment = DEFAULT_ALIGNMENT);
 };
 
 #endif
