Fix ordering warnings after struct layout changes.

diff --git a/Engine/Source/Runtime/CoreUObject/Public/UObject/UObject.h b/Engine/Source/Runtime/CoreUObject/Public/UObject/UObject.h
--- a/Engine/Source/Runtime/CoreUObject/Public/UObject/UObject.h
+++ b/Engine/Source/Runtime/CoreUObject/Public/UObject/UObject.h
@@ -540,7 +540,7 @@ public:
 		FString Value;
 
 		FAssetRegistryTag(FName InName, const FString& InValue, ETagType InType)
-			: Name(InName), Value(InValue), Type(InType) {}
+			: Name(InName), Type(InType), Value(InValue) {}
 
 		/** Gathers a list of asset registry searchable tags from given objects properties */
 		COREUOBJECT_API static void GetAssetRegistryTagsFromSearchableProperties(const UObject* Object, TArray<FAssetRegistryTag>& OutTags);
