UE-8954: Update commonly-subclassed actor and component classes to take a default FObjectInitializer constructor argument.

diff --git a/Engine/Plugins/2D/Paper2D/Source/Paper2DEditor/Private/ExtractSprites/PaperExtractSpritesSettings.h b/Engine/Plugins/2D/Paper2D/Source/Paper2DEditor/Private/ExtractSprites/PaperExtractSpritesSettings.h
--- a/Engine/Plugins/2D/Paper2D/Source/Paper2DEditor/Private/ExtractSprites/PaperExtractSpritesSettings.h
+++ b/Engine/Plugins/2D/Paper2D/Source/Paper2DEditor/Private/ExtractSprites/PaperExtractSpritesSettings.h
@@ -103,5 +103,5 @@ class UPaperExtractSpriteGridSettings : public UObject
 	UPROPERTY(Category = Grid, EditAnywhere, meta = (UIMin = 0, ClampMin = 0))
 	int32 SpacingY;
 
-	UPaperExtractSpriteGridSettings(const FObjectInitializer& ObjectInitializer);
+	UPaperExtractSpriteGridSettings(const FObjectInitializer& ObjectInitializer = FObjectInitializer::Get());
 };
diff --git a/Engine/Plugins/Experimental/LiveEditor/Source/LiveEditor/Classes/LiveEditorTypes.h b/Engine/Plugins/Experimental/LiveEditor/Source/LiveEditor/Classes/LiveEditorTypes.h
--- a/Engine/Plugins/Experimental/LiveEditor/Source/LiveEditor/Classes/LiveEditorTypes.h
+++ b/Engine/Plugins/Experimental/LiveEditor/Source/LiveEditor/Classes/LiveEditorTypes.h
@@ -70,5 +70,5 @@ class ULiveEditorTypes : public UObject
 {
 	GENERATED_UCLASS_BODY()
 
-	ULiveEditorTypes(const FObjectInitializer& ObjectInitializer) : Super(ObjectInitializer) {}
+	ULiveEditorTypes(const FObjectInitializer& ObjectInitializer = FObjectInitializer::Get()) : Super(ObjectInitializer) {}
 };
diff --git a/Engine/Source/Editor/Matinee/Classes/MatineeTransBuffer.h b/Engine/Source/Editor/Matinee/Classes/MatineeTransBuffer.h
--- a/Engine/Source/Editor/Matinee/Classes/MatineeTransBuffer.h
+++ b/Engine/Source/Editor/Matinee/Classes/MatineeTransBuffer.h
@@ -14,9 +14,6 @@ class UMatineeTransBuffer : public UTransBuffer
 {
 public:
 	GENERATED_BODY()
-	UMatineeTransBuffer(const FObjectInitializer& ObjectInitializer)
-		:	UTransBuffer(ObjectInitializer)
-	{ }
 
 	/**  
 	 * Begin a Matinee specific transaction
diff --git a/Engine/Source/Editor/UnrealEd/Classes/Factories/TextureFactory.h b/Engine/Source/Editor/UnrealEd/Classes/Factories/TextureFactory.h
--- a/Engine/Source/Editor/UnrealEd/Classes/Factories/TextureFactory.h
+++ b/Engine/Source/Editor/UnrealEd/Classes/Factories/TextureFactory.h
@@ -90,7 +90,7 @@ class UNREALED_API UTextureFactory : public UFactory
 	uint32 bUsingExistingSettings:1;
 
 public:
-	UTextureFactory(const FObjectInitializer& ObjectInitializer);
+	UTextureFactory(const FObjectInitializer& ObjectInitializer = FObjectInitializer::Get());
 
 	// Begin UObject Interface
 	virtual void PostInitProperties() override;
diff --git a/Engine/Source/Programs/UnrealHeaderTool/Private/CodeGenerator.cpp b/Engine/Source/Programs/UnrealHeaderTool/Private/CodeGenerator.cpp
--- a/Engine/Source/Programs/UnrealHeaderTool/Private/CodeGenerator.cpp
+++ b/Engine/Source/Programs/UnrealHeaderTool/Private/CodeGenerator.cpp
@@ -2162,7 +2162,11 @@ void ExportStandardConstructorsMacro(FUHTStringBuilder& Out, FClass* Class, cons
 	if (!Class->HasAnyClassFlags(CLASS_CustomConstructor))
 	{
 		Out.Logf(TEXT("\t/** Standard constructor, called after all reflected properties have been initialized */\r\n"));
-		Out.Logf(TEXT("\t%s_API %s(const FObjectInitializer& ObjectInitializer);\r\n"), *API, NameLookupCPP.GetNameCPP(Class));
+
+		auto* ClassData = GScriptHelper.FindClassData(Class);
+
+		Out.Logf(TEXT("\t%s_API %s(const FObjectInitializer& ObjectInitializer%s);\r\n"), *API, NameLookupCPP.GetNameCPP(Class),
+			ClassData->bDefaultConstructorDeclared ? TEXT("") : TEXT(" = FObjectInitializer::Get()"));
 	}
 	Out.Logf(TEXT("\tDEFINE_DEFAULT_OBJECT_INITIALIZER_CONSTRUCTOR_CALL(%s)\r\n"), NameLookupCPP.GetNameCPP(Class));
 
@@ -2201,7 +2205,7 @@ void ExportConstructorDefinition(FUHTStringBuilder& Out, FClass* Class, const FS
 		}
 		if (bSuperClassObjectInitializerConstructorDeclared)
 		{
-			Out.Logf(TEXT("\t%s_API %s(const class FObjectInitializer& ObjectInitializer = FObjectInitializer::Get()) : Super(ObjectInitializer) { };\r\n"), *API, NameLookupCPP.GetNameCPP(Class));
+			Out.Logf(TEXT("\t%s_API %s(const FObjectInitializer& ObjectInitializer = FObjectInitializer::Get()) : Super(ObjectInitializer) { };\r\n"), *API, NameLookupCPP.GetNameCPP(Class));
 			ClassData->bObjectInitializerConstructorDeclared = true;
 		}
 		else
diff --git a/Engine/Source/Runtime/AIModule/Classes/AIController.h b/Engine/Source/Runtime/AIModule/Classes/AIController.h
--- a/Engine/Source/Runtime/AIModule/Classes/AIController.h
+++ b/Engine/Source/Runtime/AIModule/Classes/AIController.h
@@ -121,7 +121,7 @@ private_subobject:
 
 public:
 
-	AAIController(const FObjectInitializer& ObjectInitializer);
+	AAIController(const FObjectInitializer& ObjectInitializer = FObjectInitializer::Get());
 
 	/** Event called when PossessedPawn is possessed by this controller. */
 	UFUNCTION(BlueprintImplementableEvent)
diff --git a/Engine/Source/Runtime/AIModule/Classes/BehaviorTree/BlackboardComponent.h b/Engine/Source/Runtime/AIModule/Classes/BehaviorTree/BlackboardComponent.h
--- a/Engine/Source/Runtime/AIModule/Classes/BehaviorTree/BlackboardComponent.h
+++ b/Engine/Source/Runtime/AIModule/Classes/BehaviorTree/BlackboardComponent.h
@@ -41,7 +41,7 @@ class AIMODULE_API UBlackboardComponent : public UActorComponent
 	GENERATED_BODY()
 
 public:
-	UBlackboardComponent(const FObjectInitializer& ObjectInitializer);
+	UBlackboardComponent(const FObjectInitializer& ObjectInitializer = FObjectInitializer::Get());
 
 	/** BEGIN UActorComponent overrides */
 	virtual void InitializeComponent() override;
diff --git a/Engine/Source/Runtime/AIModule/Classes/BehaviorTree/Decorators/BTDecorator_BlueprintBase.h b/Engine/Source/Runtime/AIModule/Classes/BehaviorTree/Decorators/BTDecorator_BlueprintBase.h
--- a/Engine/Source/Runtime/AIModule/Classes/BehaviorTree/Decorators/BTDecorator_BlueprintBase.h
+++ b/Engine/Source/Runtime/AIModule/Classes/BehaviorTree/Decorators/BTDecorator_BlueprintBase.h
@@ -29,7 +29,7 @@ class AIMODULE_API UBTDecorator_BlueprintBase : public UBTDecorator
 	GENERATED_BODY()
 
 public:
-	UBTDecorator_BlueprintBase(const FObjectInitializer& ObjectInitializer);
+	UBTDecorator_BlueprintBase(const FObjectInitializer& ObjectInitializer = FObjectInitializer::Get());
 
 	/** initialize data about blueprint defined properties */
 	void InitializeProperties();
diff --git a/Engine/Source/Runtime/AIModule/Classes/BehaviorTree/Decorators/BTDecorator_IsBBEntryOfClass.h b/Engine/Source/Runtime/AIModule/Classes/BehaviorTree/Decorators/BTDecorator_IsBBEntryOfClass.h
--- a/Engine/Source/Runtime/AIModule/Classes/BehaviorTree/Decorators/BTDecorator_IsBBEntryOfClass.h
+++ b/Engine/Source/Runtime/AIModule/Classes/BehaviorTree/Decorators/BTDecorator_IsBBEntryOfClass.h
@@ -11,7 +11,7 @@ class AIMODULE_API UBTDecorator_IsBBEntryOfClass : public UBTDecorator_Blackboar
 	GENERATED_BODY()
 		
 public:
-	UBTDecorator_IsBBEntryOfClass(const FObjectInitializer& ObjectInitializer);
+	UBTDecorator_IsBBEntryOfClass(const FObjectInitializer& ObjectInitializer = FObjectInitializer::Get());
 
 protected:
 
diff --git a/Engine/Source/Runtime/AIModule/Classes/BehaviorTree/Services/BTService_DefaultFocus.h b/Engine/Source/Runtime/AIModule/Classes/BehaviorTree/Services/BTService_DefaultFocus.h
--- a/Engine/Source/Runtime/AIModule/Classes/BehaviorTree/Services/BTService_DefaultFocus.h
+++ b/Engine/Source/Runtime/AIModule/Classes/BehaviorTree/Services/BTService_DefaultFocus.h
@@ -31,7 +31,7 @@ class AIMODULE_API UBTService_DefaultFocus : public UBTService_BlackboardBase
 	UPROPERTY()
 	uint8 FocusPriority;
 
-	UBTService_DefaultFocus(const FObjectInitializer& ObjectInitializer);
+	UBTService_DefaultFocus(const FObjectInitializer& ObjectInitializer = FObjectInitializer::Get());
 
 	virtual uint16 GetInstanceMemorySize() const override { return sizeof(FBTFocusMemory); }
 	virtual void OnBecomeRelevant(UBehaviorTreeComponent& OwnerComp, uint8* NodeMemory) override;
diff --git a/Engine/Source/Runtime/AIModule/Classes/DetourCrowdAIController.h b/Engine/Source/Runtime/AIModule/Classes/DetourCrowdAIController.h
--- a/Engine/Source/Runtime/AIModule/Classes/DetourCrowdAIController.h
+++ b/Engine/Source/Runtime/AIModule/Classes/DetourCrowdAIController.h
@@ -10,5 +10,5 @@ class ADetourCrowdAIController : public AAIController
 {
 	GENERATED_BODY()
 public:
-	ADetourCrowdAIController(const FObjectInitializer& ObjectInitializer);
+	ADetourCrowdAIController(const FObjectInitializer& ObjectInitializer = FObjectInitializer::Get());
 };
\ No newline at end of file
diff --git a/Engine/Source/Runtime/AIModule/Classes/Perception/AIPerceptionSystem.h b/Engine/Source/Runtime/AIModule/Classes/Perception/AIPerceptionSystem.h
--- a/Engine/Source/Runtime/AIModule/Classes/Perception/AIPerceptionSystem.h
+++ b/Engine/Source/Runtime/AIModule/Classes/Perception/AIPerceptionSystem.h
@@ -18,7 +18,7 @@ class AIMODULE_API UAIPerceptionSystem : public UObject, public FTickableGameObj
 	
 public:
 
-	UAIPerceptionSystem(const FObjectInitializer& ObjectInitializer);
+	UAIPerceptionSystem(const FObjectInitializer& ObjectInitializer = FObjectInitializer::Get());
 
 	// We need to implement GetWorld() so that any EQS-related blueprints (such as blueprint contexts) can implement
 	// GetWorld() and so provide access to blueprint nodes using hidden WorldContextObject parameters.
diff --git a/Engine/Source/Runtime/AIModule/Classes/Perception/AISenseConfig.h b/Engine/Source/Runtime/AIModule/Classes/Perception/AISenseConfig.h
--- a/Engine/Source/Runtime/AIModule/Classes/Perception/AISenseConfig.h
+++ b/Engine/Source/Runtime/AIModule/Classes/Perception/AISenseConfig.h
@@ -24,7 +24,7 @@ protected:
 	uint32 bStartsEnabled : 1;
 	
 public:
-	UAISenseConfig(const FObjectInitializer& ObjectInitializer);
+	UAISenseConfig(const FObjectInitializer& ObjectInitializer = FObjectInitializer::Get());
 
 	virtual TSubclassOf<UAISense> GetSenseImplementation() const { return UAISense::StaticClass(); }
 	FAISenseID GetSenseID() const;
diff --git a/Engine/Source/Runtime/AIModule/Classes/Perception/AISenseConfig_Hearing.h b/Engine/Source/Runtime/AIModule/Classes/Perception/AISenseConfig_Hearing.h
--- a/Engine/Source/Runtime/AIModule/Classes/Perception/AISenseConfig_Hearing.h
+++ b/Engine/Source/Runtime/AIModule/Classes/Perception/AISenseConfig_Hearing.h
@@ -29,7 +29,7 @@ public:
 	UPROPERTY(EditDefaultsOnly, BlueprintReadOnly, Category = "Sense", config)
 	FAISenseAffiliationFilter DetectionByAffiliation;
 
-	UAISenseConfig_Hearing(const FObjectInitializer& ObjectInitializer);
+	UAISenseConfig_Hearing(const FObjectInitializer& ObjectInitializer = FObjectInitializer::Get());
 
 	virtual TSubclassOf<UAISense> GetSenseImplementation() const override;
 #if !UE_BUILD_SHIPPING
diff --git a/Engine/Source/Runtime/AIModule/Classes/Perception/AISenseEvent_Hearing.h b/Engine/Source/Runtime/AIModule/Classes/Perception/AISenseEvent_Hearing.h
--- a/Engine/Source/Runtime/AIModule/Classes/Perception/AISenseEvent_Hearing.h
+++ b/Engine/Source/Runtime/AIModule/Classes/Perception/AISenseEvent_Hearing.h
@@ -15,7 +15,7 @@ protected:
 	FAINoiseEvent Event;
 
 public:
-	UAISenseEvent_Hearing(const FObjectInitializer& ObjectInitializer);
+	UAISenseEvent_Hearing(const FObjectInitializer& ObjectInitializer = FObjectInitializer::Get());
 	virtual FAISenseID GetSenseID() const override;
 	
 	FORCEINLINE FAINoiseEvent GetNoiseEvent()
diff --git a/Engine/Source/Runtime/AIModule/Classes/Perception/AISense_Blueprint.h b/Engine/Source/Runtime/AIModule/Classes/Perception/AISense_Blueprint.h
--- a/Engine/Source/Runtime/AIModule/Classes/Perception/AISense_Blueprint.h
+++ b/Engine/Source/Runtime/AIModule/Classes/Perception/AISense_Blueprint.h
@@ -28,7 +28,7 @@ protected:
 	TArray<UAISenseEvent*> UnprocessedEvents;
 
 public:
-	UAISense_Blueprint(const FObjectInitializer& ObjectInitializer);
+	UAISense_Blueprint(const FObjectInitializer& ObjectInitializer = FObjectInitializer::Get());
 
 	//void RegisterEvent(const FAITeamStimulusEvent& Event);	
 
diff --git a/Engine/Source/Runtime/CoreUObject/Classes/Object.h b/Engine/Source/Runtime/CoreUObject/Classes/Object.h
--- a/Engine/Source/Runtime/CoreUObject/Classes/Object.h
+++ b/Engine/Source/Runtime/CoreUObject/Classes/Object.h
@@ -881,7 +881,7 @@ public:
 	/**
 	 * Default UObject constructor.
 	 */
-	UObject(const FObjectInitializer& ObjectInitializer);
+	UObject(const FObjectInitializer& ObjectInitializer = FObjectInitializer::Get());
 #if WITH_HOT_RELOAD_CTORS
 	/** DO NOT USE. This constructor is for internal usage only for hot-reload purposes. */
 	UObject(FVTableHelper& Helper);
diff --git a/Engine/Source/Runtime/CoreUObject/Private/Misc/TextBuffer.cpp b/Engine/Source/Runtime/CoreUObject/Private/Misc/TextBuffer.cpp
--- a/Engine/Source/Runtime/CoreUObject/Private/Misc/TextBuffer.cpp
+++ b/Engine/Source/Runtime/CoreUObject/Private/Misc/TextBuffer.cpp
@@ -13,6 +13,11 @@ IMPLEMENT_CORE_INTRINSIC_CLASS(UTextBuffer, UObject, { });
 /* UTextBuffer constructors
  *****************************************************************************/
 
+UTextBuffer::UTextBuffer(const TCHAR* InText)
+	: UObject(FObjectInitializer::Get())
+	, Text(InText)
+{}
+
 UTextBuffer::UTextBuffer (const FObjectInitializer& ObjectInitializer, const TCHAR* InText)
 	: UObject(ObjectInitializer)
 	, Text(InText)
diff --git a/Engine/Source/Runtime/CoreUObject/Private/UObject/Class.cpp b/Engine/Source/Runtime/CoreUObject/Private/UObject/Class.cpp
--- a/Engine/Source/Runtime/CoreUObject/Private/UObject/Class.cpp
+++ b/Engine/Source/Runtime/CoreUObject/Private/UObject/Class.cpp
@@ -430,6 +430,19 @@ UStruct::UStruct( EStaticConstructor, int32 InSize, EObjectFlags InFlags )
 {
 }
 
+UStruct::UStruct(UStruct* InSuperStruct, SIZE_T ParamsSize, SIZE_T Alignment)
+	: UField(FObjectInitializer::Get())
+	, SuperStruct(InSuperStruct)
+	, Children(NULL)
+	, PropertiesSize(ParamsSize ? ParamsSize : (InSuperStruct ? InSuperStruct->GetPropertiesSize() : 0))
+	, MinAlignment(Alignment ? Alignment : (FMath::Max(InSuperStruct ? InSuperStruct->GetMinAlignment() : 1, 1)))
+	, PropertyLink(NULL)
+	, RefLink(NULL)
+	, DestructorLink(NULL)
+	, PostConstructLink(NULL)
+{
+}
+
 UStruct::UStruct(const FObjectInitializer& ObjectInitializer, UStruct* InSuperStruct, SIZE_T ParamsSize, SIZE_T Alignment )
 :	UField			(ObjectInitializer)
 ,   SuperStruct		( InSuperStruct )
@@ -3788,6 +3801,17 @@ UFunction::UFunction(const FObjectInitializer& ObjectInitializer, UFunction* InS
 {
 }
 
+UFunction::UFunction(UFunction* InSuperFunction, uint32 InFunctionFlags, uint16 InRepOffset, SIZE_T ParamsSize)
+	: UStruct(InSuperFunction, ParamsSize)
+	, FunctionFlags(InFunctionFlags)
+	, RepOffset(InRepOffset)
+	, RPCId(0)
+	, RPCResponseId(0)
+	, FirstPropertyToInit(NULL)
+{
+}
+
+
 void UFunction::InitializeDerivedMembers()
 {
 	NumParms = 0;
@@ -3994,6 +4018,12 @@ UDelegateFunction::UDelegateFunction(const FObjectInitializer& ObjectInitializer
 
 }
 
+UDelegateFunction::UDelegateFunction(UFunction* InSuperFunction, uint32 InFunctionFlags, uint16 InRepOffset, SIZE_T ParamsSize)
+	: UFunction(InSuperFunction, InFunctionFlags, InRepOffset, ParamsSize)
+{
+
+}
+
 IMPLEMENT_CORE_INTRINSIC_CLASS(UDelegateFunction, UFunction,
 	{
 	}
diff --git a/Engine/Source/Runtime/CoreUObject/Private/UObject/LinkerPlaceholderClass.h b/Engine/Source/Runtime/CoreUObject/Private/UObject/LinkerPlaceholderClass.h
--- a/Engine/Source/Runtime/CoreUObject/Private/UObject/LinkerPlaceholderClass.h
+++ b/Engine/Source/Runtime/CoreUObject/Private/UObject/LinkerPlaceholderClass.h
@@ -19,7 +19,7 @@ class ULinkerPlaceholderClass : public UClass, public TLinkerImportPlaceholder<U
 public:
 	DECLARE_CASTED_CLASS_INTRINSIC_NO_CTOR(ULinkerPlaceholderClass, UClass, /*TStaticFlags =*/0, CoreUObject, /*TStaticCastFlags =*/0, NO_API)
 
-	ULinkerPlaceholderClass(const FObjectInitializer& ObjectInitializer);
+	ULinkerPlaceholderClass(const FObjectInitializer& ObjectInitializer = FObjectInitializer::Get());
 
 	// UObject interface.
 	static  void AddReferencedObjects(UObject* InThis, FReferenceCollector& Collector);
diff --git a/Engine/Source/Runtime/CoreUObject/Private/UObject/LinkerPlaceholderExportObject.h b/Engine/Source/Runtime/CoreUObject/Private/UObject/LinkerPlaceholderExportObject.h
--- a/Engine/Source/Runtime/CoreUObject/Private/UObject/LinkerPlaceholderExportObject.h
+++ b/Engine/Source/Runtime/CoreUObject/Private/UObject/LinkerPlaceholderExportObject.h
@@ -18,7 +18,7 @@ class ULinkerPlaceholderExportObject : public UObject, public FLinkerPlaceholder
 public:
 	DECLARE_CASTED_CLASS_INTRINSIC_NO_CTOR(ULinkerPlaceholderExportObject, UObject, /*TStaticFlags =*/0, CoreUObject, /*TStaticCastFlags =*/0, NO_API)
 
-	ULinkerPlaceholderExportObject(const FObjectInitializer& ObjectInitializer);
+	ULinkerPlaceholderExportObject(const FObjectInitializer& ObjectInitializer = FObjectInitializer::Get());
 
 	// UObject interface
 	virtual void BeginDestroy() override;
diff --git a/Engine/Source/Runtime/CoreUObject/Private/UObject/Property.cpp b/Engine/Source/Runtime/CoreUObject/Private/UObject/Property.cpp
--- a/Engine/Source/Runtime/CoreUObject/Private/UObject/Property.cpp
+++ b/Engine/Source/Runtime/CoreUObject/Private/UObject/Property.cpp
@@ -335,20 +335,35 @@ UProperty::UProperty(const FObjectInitializer& ObjectInitializer)
 , ArrayDim(1)
 {
 }
+
+UProperty::UProperty(ECppProperty, int32 InOffset, uint64 InFlags)
+	: UField(FObjectInitializer::Get())
+	, ArrayDim(1)
+	, PropertyFlags(InFlags)
+	, Offset_Internal(InOffset)
+{
+	Init();
+}
+
 UProperty::UProperty(const FObjectInitializer& ObjectInitializer, ECppProperty, int32 InOffset, uint64 InFlags )
 : UField(ObjectInitializer)	
 , ArrayDim(1)
 , PropertyFlags(InFlags)
 , Offset_Internal(InOffset)
+{
+	Init();
+}
+
+void UProperty::Init()
 {
 	// properties created in C++ should always be marked RF_Transient so that when the package containing
 	// this property is saved, it doesn't try to save this UProperty into the ExportMap
-	SetFlags(RF_Transient|RF_Native);
+	SetFlags(RF_Transient | RF_Native);
 #if !WITH_EDITORONLY_DATA
 	[REDACTED_EMAIL]/PACKAGER: Until we have a cooker/packager step, this can fire when WITH_EDITORONLY_DATA is not defined!
-//	checkSlow(!HasAnyPropertyFlags(CPF_EditorOnly));
+	//	checkSlow(!HasAnyPropertyFlags(CPF_EditorOnly));
 #endif // WITH_EDITORONLY_DATA
-	checkSlow(GetOuterUField()->HasAllFlags(RF_Native|RF_Transient));
+	checkSlow(GetOuterUField()->HasAllFlags(RF_Native | RF_Transient));
 
 	GetOuterUField()->AddCppProperty(this);
 }
diff --git a/Engine/Source/Runtime/CoreUObject/Private/UObject/PropertyBool.cpp b/Engine/Source/Runtime/CoreUObject/Private/UObject/PropertyBool.cpp
--- a/Engine/Source/Runtime/CoreUObject/Private/UObject/PropertyBool.cpp
+++ b/Engine/Source/Runtime/CoreUObject/Private/UObject/PropertyBool.cpp
@@ -17,6 +17,16 @@ UBoolProperty::UBoolProperty( const FObjectInitializer& ObjectInitializer )
 	SetBoolSize( 1, false, 1 );
 }
 
+UBoolProperty::UBoolProperty(ECppProperty, int32 InOffset, uint64 InFlags, uint32 InBitMask, uint32 InElementSize, bool bIsNativeBool)
+	: UProperty(FObjectInitializer::Get(), EC_CppProperty, InOffset, InFlags)
+	, FieldSize(0)
+	, ByteOffset(0)
+	, ByteMask(1)
+	, FieldMask(1)
+{
+	SetBoolSize(InElementSize, bIsNativeBool, InBitMask);
+}
+
 UBoolProperty::UBoolProperty( const FObjectInitializer& ObjectInitializer, ECppProperty, int32 InOffset, uint64 InFlags, uint32 InBitMask, uint32 InElementSize, bool bIsNativeBool )
 : UProperty( ObjectInitializer, EC_CppProperty, InOffset, InFlags )
 , FieldSize(0)
diff --git a/Engine/Source/Runtime/CoreUObject/Private/UObject/PropertyStruct.cpp b/Engine/Source/Runtime/CoreUObject/Private/UObject/PropertyStruct.cpp
--- a/Engine/Source/Runtime/CoreUObject/Private/UObject/PropertyStruct.cpp
+++ b/Engine/Source/Runtime/CoreUObject/Private/UObject/PropertyStruct.cpp
@@ -30,6 +30,13 @@ static inline void PreloadInnerStructMembers(UStructProperty* StructProperty)
 	UStructProperty.
 -----------------------------------------------------------------------------*/
 
+UStructProperty::UStructProperty(ECppProperty, int32 InOffset, uint64 InFlags, UScriptStruct* InStruct)
+	: UProperty(FObjectInitializer::Get(), EC_CppProperty, InOffset, InFlags)
+	, Struct(InStruct)
+{
+	ElementSize = Struct->PropertiesSize;
+}
+
 UStructProperty::UStructProperty( const FObjectInitializer& ObjectInitializer, ECppProperty, int32 InOffset, uint64 InFlags, UScriptStruct* InStruct )
 	:	UProperty( ObjectInitializer, EC_CppProperty, InOffset, InFlags )
 	,	Struct( InStruct )
diff --git a/Engine/Source/Runtime/CoreUObject/Public/Misc/TextBuffer.h b/Engine/Source/Runtime/CoreUObject/Public/Misc/TextBuffer.h
--- a/Engine/Source/Runtime/CoreUObject/Public/Misc/TextBuffer.h
+++ b/Engine/Source/Runtime/CoreUObject/Public/Misc/TextBuffer.h
@@ -26,6 +26,13 @@ public:
 	 */
 	COREUOBJECT_API UTextBuffer (const FObjectInitializer& ObjectInitializer, const TCHAR* InText);
 
+	/**
+	 * Creates and initializes a new text buffer.
+	 *
+	 * @param InText - The initial text.
+	 */
+	COREUOBJECT_API UTextBuffer(const TCHAR* InText);
+
 public:
 
 	/**
diff --git a/Engine/Source/Runtime/CoreUObject/Public/UObject/Class.h b/Engine/Source/Runtime/CoreUObject/Public/UObject/Class.h
--- a/Engine/Source/Runtime/CoreUObject/Public/UObject/Class.h
+++ b/Engine/Source/Runtime/CoreUObject/Public/UObject/Class.h
@@ -278,6 +278,7 @@ public:
 public:
 	// Constructors.
 	UStruct( EStaticConstructor, int32 InSize, EObjectFlags InFlags );
+	explicit UStruct(UStruct* InSuperStruct, SIZE_T ParamsSize = 0, SIZE_T Alignment = 0);
 	explicit UStruct(const FObjectInitializer& ObjectInitializer, UStruct* InSuperStruct, SIZE_T ParamsSize = 0, SIZE_T Alignment = 0 );
 
 	// UObject interface.
@@ -993,7 +994,7 @@ public:
 
 	COREUOBJECT_API UScriptStruct( EStaticConstructor, int32 InSize, EObjectFlags InFlags );
 	COREUOBJECT_API explicit UScriptStruct(const FObjectInitializer& ObjectInitializer, UScriptStruct* InSuperStruct, ICppStructOps* InCppStructOps = NULL, EStructFlags InStructFlags = STRUCT_NoFlags, SIZE_T ExplicitSize = 0, SIZE_T ExplicitAlignment = 0);
-	COREUOBJECT_API explicit UScriptStruct(const FObjectInitializer& ObjectInitializer);
+	COREUOBJECT_API explicit UScriptStruct(const FObjectInitializer& ObjectInitializer = FObjectInitializer::Get());
 
 public:
 	EStructFlags StructFlags;
@@ -1239,6 +1240,7 @@ public:
 
 	// Constructors.
 	explicit UFunction(const FObjectInitializer& ObjectInitializer, UFunction* InSuperFunction, uint32 InFunctionFlags = 0, uint16 InRepOffset = 0, SIZE_T ParamsSize = 0 );
+	explicit UFunction(UFunction* InSuperFunction, uint32 InFunctionFlags = 0, uint16 InRepOffset = 0, SIZE_T ParamsSize = 0);
 
 	void InitializeDerivedMembers();
 
@@ -1325,6 +1327,7 @@ class COREUOBJECT_API UDelegateFunction : public UFunction
 	DECLARE_WITHIN(UObject)
 public:
 	explicit UDelegateFunction(const FObjectInitializer& ObjectInitializer, UFunction* InSuperFunction, uint32 InFunctionFlags = 0, uint16 InRepOffset = 0, SIZE_T ParamsSize = 0);
+	explicit UDelegateFunction(UFunction* InSuperFunction, uint32 InFunctionFlags = 0, uint16 InRepOffset = 0, SIZE_T ParamsSize = 0);
 };
 
 /*-----------------------------------------------------------------------------
@@ -1895,7 +1898,7 @@ public:
 
 public:
 	// Constructors
-	UClass(const FObjectInitializer& ObjectInitializer);
+	UClass(const FObjectInitializer& ObjectInitializer = FObjectInitializer::Get());
 	explicit UClass(const FObjectInitializer& ObjectInitializer, UClass* InSuperClass);
 	UClass( EStaticConstructor, FName InName, uint32 InSize, uint32 InClassFlags, EClassCastFlags InClassCastFlags,
 		const TCHAR* InClassConfigName, EObjectFlags InFlags, ClassConstructorType InClassConstructor,
diff --git a/Engine/Source/Runtime/CoreUObject/Public/UObject/ObjectBase.h b/Engine/Source/Runtime/CoreUObject/Public/UObject/ObjectBase.h
--- a/Engine/Source/Runtime/CoreUObject/Public/UObject/ObjectBase.h
+++ b/Engine/Source/Runtime/CoreUObject/Public/UObject/ObjectBase.h
@@ -1169,7 +1169,7 @@ namespace UM
 	};
 }
 
-#define RELAY_CONSTRUCTOR(TClass, TSuperClass) TClass(const FObjectInitializer& ObjectInitializer) : TSuperClass(ObjectInitializer) {}
+#define RELAY_CONSTRUCTOR(TClass, TSuperClass) TClass(const FObjectInitializer& ObjectInitializer = FObjectInitializer::Get()) : TSuperClass(ObjectInitializer) {}
 
 #if !USE_COMPILED_IN_NATIVES
 #define COMPILED_IN_FLAGS(TStaticFlags) (TStaticFlags& ~(CLASS_Intrinsic))
diff --git a/Engine/Source/Runtime/CoreUObject/Public/UObject/UTextProperty.h b/Engine/Source/Runtime/CoreUObject/Public/UObject/UTextProperty.h
--- a/Engine/Source/Runtime/CoreUObject/Public/UObject/UTextProperty.h
+++ b/Engine/Source/Runtime/CoreUObject/Public/UObject/UTextProperty.h
@@ -15,6 +15,11 @@ public:
 	typedef UTextProperty_Super::TTypeFundamentals TTypeFundamentals;
 	typedef TTypeFundamentals::TCppType TCppType;
 
+	UTextProperty(ECppProperty, int32 InOffset, uint64 InFlags)
+		: TProperty(FObjectInitializer::Get(), EC_CppProperty, InOffset, InFlags)
+	{
+	}
+
 	UTextProperty( const FObjectInitializer& ObjectInitializer, ECppProperty, int32 InOffset, uint64 InFlags )
 		:	TProperty( ObjectInitializer, EC_CppProperty, InOffset, InFlags)
 	{
diff --git a/Engine/Source/Runtime/CoreUObject/Public/UObject/UnrealType.h b/Engine/Source/Runtime/CoreUObject/Public/UObject/UnrealType.h
--- a/Engine/Source/Runtime/CoreUObject/Public/UObject/UnrealType.h
+++ b/Engine/Source/Runtime/CoreUObject/Public/UObject/UnrealType.h
@@ -73,7 +73,9 @@ public:
 
 public:
 	// Constructors.
-	UProperty(const FObjectInitializer& ObjectInitializer);
+	UProperty(const FObjectInitializer& ObjectInitializer = FObjectInitializer::Get());
+
+	UProperty(ECppProperty, int32 InOffset, uint64 InFlags);
 	UProperty( const FObjectInitializer& ObjectInitializer, ECppProperty, int32 InOffset, uint64 InFlags );
 
 	// UObject interface
@@ -124,6 +126,11 @@ private:
 	*/
 	int32 SetupOffset();
 
+	/**
+	 * Initializes internal state.
+	 */
+	void Init();
+
 public:
 	/** Return offset of property from container base. */
 	FORCEINLINE int32 GetOffset_ForDebug() const
@@ -841,8 +848,13 @@ public:
 	typedef TInPropertyBaseClass Super;
 	typedef TPropertyTypeFundamentals<InTCppType> TTypeFundamentals;
 
-	TProperty( const FObjectInitializer& ObjectInitializer )
-		:	Super( ObjectInitializer)
+	TProperty(const FObjectInitializer& ObjectInitializer = FObjectInitializer::Get())
+		: Super(ObjectInitializer)
+	{
+		SetElementSize();
+	}
+	TProperty(ECppProperty, int32 InOffset, uint64 InFlags)
+		: Super(FObjectInitializer::Get(), EC_CppProperty, InOffset, InFlags | TTypeFundamentals::GetComputedFlagsPropertyFlags())
 	{
 		SetElementSize();
 	}
@@ -946,8 +958,12 @@ public:
 	typedef InTCppType TCppType;
 	typedef typename Super::TTypeFundamentals TTypeFundamentals;
 
-	TProperty_WithEqualityAndSerializer( const FObjectInitializer& ObjectInitializer )
-		:	Super( ObjectInitializer)
+	TProperty_WithEqualityAndSerializer(const FObjectInitializer& ObjectInitializer = FObjectInitializer::Get())
+		: Super(ObjectInitializer)
+	{
+	}
+	TProperty_WithEqualityAndSerializer(ECppProperty, int32 InOffset, uint64 InFlags)
+		: Super(FObjectInitializer::Get(), EC_CppProperty, InOffset, InFlags)
 	{
 	}
 	TProperty_WithEqualityAndSerializer( const FObjectInitializer& ObjectInitializer, ECppProperty, int32 InOffset, uint64 InFlags )
@@ -977,6 +993,10 @@ class COREUOBJECT_API UNumericProperty : public UProperty
 {
 	DECLARE_CASTED_CLASS_INTRINSIC(UNumericProperty,UProperty,CLASS_Abstract,CoreUObject,CASTCLASS_UNumericProperty)
 
+	UNumericProperty(ECppProperty, int32 InOffset, uint64 InFlags)
+		: UProperty(FObjectInitializer::Get(), EC_CppProperty, InOffset, InFlags)
+	{}
+
 	UNumericProperty( const FObjectInitializer& ObjectInitializer, ECppProperty, int32 InOffset, uint64 InFlags )
 		:	UProperty( ObjectInitializer, EC_CppProperty, InOffset, InFlags )
 	{}
@@ -1109,10 +1129,16 @@ public:
 	typedef TProperty_WithEqualityAndSerializer<InTCppType, UNumericProperty> Super;
 	typedef InTCppType TCppType;
 	typedef typename Super::TTypeFundamentals TTypeFundamentals;
-	TProperty_Numeric( const FObjectInitializer& ObjectInitializer )
-		:	Super( ObjectInitializer)
+	TProperty_Numeric(const FObjectInitializer& ObjectInitializer = FObjectInitializer::Get())
+		: Super(ObjectInitializer)
 	{
 	}
+
+	TProperty_Numeric(ECppProperty, int32 InOffset, uint64 InFlags)
+		: Super(FObjectInitializer::Get(), EC_CppProperty, InOffset, InFlags)
+	{
+	}
+
 	TProperty_Numeric( const FObjectInitializer& ObjectInitializer, ECppProperty, int32 InOffset, uint64 InFlags )
 		:	Super( ObjectInitializer, EC_CppProperty, InOffset, InFlags)
 	{
@@ -1194,7 +1220,13 @@ class COREUOBJECT_API UByteProperty : public TProperty_Numeric<uint8>
 	// Variables.
 	UEnum* Enum;
 
-	UByteProperty( const FObjectInitializer& ObjectInitializer, ECppProperty, int32 InOffset, uint64 InFlags, UEnum* InEnum=NULL )
+	UByteProperty(ECppProperty, int32 InOffset, uint64 InFlags, UEnum* InEnum = nullptr)
+		: TProperty_Numeric(FObjectInitializer::Get(), EC_CppProperty, InOffset, InFlags)
+		, Enum(InEnum)
+	{
+	}
+
+	UByteProperty(const FObjectInitializer& ObjectInitializer, ECppProperty, int32 InOffset, uint64 InFlags, UEnum* InEnum = nullptr)
 	:	TProperty_Numeric( ObjectInitializer, EC_CppProperty, InOffset, InFlags)
 	,	Enum( InEnum )
 	{
@@ -1236,6 +1268,11 @@ class COREUOBJECT_API UInt8Property : public TProperty_Numeric<int8>
 {
 	DECLARE_CASTED_CLASS_INTRINSIC(UInt8Property,TProperty_Numeric<int8>,0,CoreUObject,CASTCLASS_UInt8Property)
 
+	UInt8Property(ECppProperty, int32 InOffset, uint64 InFlags)
+		: TProperty_Numeric(FObjectInitializer::Get(), EC_CppProperty, InOffset, InFlags)
+	{
+	}
+
 	UInt8Property( const FObjectInitializer& ObjectInitializer, ECppProperty, int32 InOffset, uint64 InFlags )
 		:	TProperty_Numeric( ObjectInitializer, EC_CppProperty, InOffset, InFlags)
 	{
@@ -1253,6 +1290,11 @@ class COREUOBJECT_API UInt16Property : public TProperty_Numeric<int16>
 {
 	DECLARE_CASTED_CLASS_INTRINSIC(UInt16Property,TProperty_Numeric<int16>,0,CoreUObject,CASTCLASS_UInt16Property)
 
+	UInt16Property(ECppProperty, int32 InOffset, uint64 InFlags)
+		: TProperty_Numeric(FObjectInitializer::Get(), EC_CppProperty, InOffset, InFlags)
+	{
+	}
+
 	UInt16Property( const FObjectInitializer& ObjectInitializer, ECppProperty, int32 InOffset, uint64 InFlags )
 	:	TProperty_Numeric( ObjectInitializer, EC_CppProperty, InOffset, InFlags)
 	{
@@ -1271,6 +1313,11 @@ class COREUOBJECT_API UIntProperty : public TProperty_Numeric<int32>
 {
 	DECLARE_CASTED_CLASS_INTRINSIC(UIntProperty,TProperty_Numeric<int32>,0,CoreUObject,CASTCLASS_UIntProperty)
 
+	UIntProperty(ECppProperty, int32 InOffset, uint64 InFlags)
+		: TProperty_Numeric(FObjectInitializer::Get(), EC_CppProperty, InOffset, InFlags)
+	{
+	}
+
 	UIntProperty( const FObjectInitializer& ObjectInitializer, ECppProperty, int32 InOffset, uint64 InFlags )
 	:	TProperty_Numeric( ObjectInitializer, EC_CppProperty, InOffset, InFlags)
 	{
@@ -1288,6 +1335,11 @@ class COREUOBJECT_API UInt64Property : public TProperty_Numeric<int64>
 {
 	DECLARE_CASTED_CLASS_INTRINSIC(UInt64Property,TProperty_Numeric<int64>,0,CoreUObject,CASTCLASS_UInt64Property)
 
+	UInt64Property(ECppProperty, int32 InOffset, uint64 InFlags)
+		: TProperty_Numeric(FObjectInitializer::Get(), EC_CppProperty, InOffset, InFlags)
+	{
+	}
+
 	UInt64Property( const FObjectInitializer& ObjectInitializer, ECppProperty, int32 InOffset, uint64 InFlags )
 	:	TProperty_Numeric( ObjectInitializer, EC_CppProperty, InOffset, InFlags )
 	{
@@ -1305,6 +1357,11 @@ class COREUOBJECT_API UUInt16Property : public TProperty_Numeric<uint16>
 {
 	DECLARE_CASTED_CLASS_INTRINSIC(UUInt16Property,TProperty_Numeric<uint16>,0,CoreUObject,CASTCLASS_UUInt16Property)
 
+	UUInt16Property(ECppProperty, int32 InOffset, uint64 InFlags)
+		: TProperty_Numeric(FObjectInitializer::Get(), EC_CppProperty, InOffset, InFlags)
+	{
+	}
+
 	UUInt16Property( const FObjectInitializer& ObjectInitializer, ECppProperty, int32 InOffset, uint64 InFlags )
 	:	TProperty_Numeric( ObjectInitializer, EC_CppProperty, InOffset, InFlags )
 	{
@@ -1339,6 +1396,11 @@ class COREUOBJECT_API UUInt64Property : public TProperty_Numeric<uint64>
 {
 	DECLARE_CASTED_CLASS_INTRINSIC(UUInt64Property,TProperty_Numeric<uint64>,0,CoreUObject,CASTCLASS_UUInt64Property)
 
+	UUInt64Property(ECppProperty, int32 InOffset, uint64 InFlags)
+		: TProperty_Numeric(FObjectInitializer::Get(), EC_CppProperty, InOffset, InFlags)
+	{
+	}
+
 	UUInt64Property( const FObjectInitializer& ObjectInitializer, ECppProperty, int32 InOffset, uint64 InFlags )
 	:	TProperty_Numeric( ObjectInitializer, EC_CppProperty, InOffset, InFlags )
 	{
@@ -1356,6 +1418,11 @@ class COREUOBJECT_API UFloatProperty : public TProperty_Numeric<float>
 {
 	DECLARE_CASTED_CLASS_INTRINSIC(UFloatProperty,TProperty_Numeric<float>,0,CoreUObject,CASTCLASS_UFloatProperty)
 
+	UFloatProperty(ECppProperty, int32 InOffset, uint64 InFlags)
+		: TProperty_Numeric(FObjectInitializer::Get(), EC_CppProperty, InOffset, InFlags)
+	{
+	}
+
 	UFloatProperty( const FObjectInitializer& ObjectInitializer, ECppProperty, int32 InOffset, uint64 InFlags )
 		:	TProperty_Numeric( ObjectInitializer, EC_CppProperty, InOffset, InFlags )
 	{
@@ -1373,8 +1440,13 @@ class COREUOBJECT_API UDoubleProperty : public TProperty_Numeric<double>
 {
 	DECLARE_CASTED_CLASS_INTRINSIC(UDoubleProperty,TProperty_Numeric<double>,0,CoreUObject,CASTCLASS_UDoubleProperty)
 
-	UDoubleProperty( const FObjectInitializer& ObjectInitializer, ECppProperty, int32 InOffset, uint64 InFlags )
-		:	TProperty_Numeric( ObjectInitializer, EC_CppProperty, InOffset, InFlags )
+	UDoubleProperty(ECppProperty, int32 InOffset, uint64 InFlags)
+		: TProperty_Numeric(FObjectInitializer::Get(), EC_CppProperty, InOffset, InFlags)
+	{
+	}
+
+	UDoubleProperty(const FObjectInitializer& ObjectInitializer, ECppProperty, int32 InOffset, uint64 InFlags)
+		: TProperty_Numeric(ObjectInitializer, EC_CppProperty, InOffset, InFlags)
 	{
 	}
 };
@@ -1406,8 +1478,21 @@ private:
 
 public:
 
-	UBoolProperty( const FObjectInitializer& ObjectInitializer );
+	UBoolProperty(const FObjectInitializer& ObjectInitializer = FObjectInitializer::Get());
 	
+	/**
+	 * Constructor.
+	 *
+	 * @param ECppProperty Unused.
+	 * @param InOffset Offset of the property.
+	 * @param InCategory Category of the property.
+	 * @param InFlags Property flags.
+	 * @param InBitMask Bitmask of the bitfield this property represents.
+	 * @param InElementSize Sizeof of the boolean type this property represents.
+	 * @param bIsNativeBool true if this property represents C++ bool type.
+	 */
+	UBoolProperty(ECppProperty, int32 InOffset, uint64 InFlags, uint32 InBitMask, uint32 InElementSize, bool bIsNativeBool);
+
 	/**
 	 * Constructor.
 	 *
@@ -1514,6 +1599,11 @@ class COREUOBJECT_API UObjectPropertyBase : public UProperty
 	// Variables.
 	class UClass* PropertyClass;
 
+	UObjectPropertyBase(ECppProperty, int32 InOffset, uint64 InFlags, UClass* InClass = NULL)
+		: UProperty(FObjectInitializer::Get(), EC_CppProperty, InOffset, InFlags)
+		, PropertyClass(InClass)
+	{}
+
 	UObjectPropertyBase( const FObjectInitializer& ObjectInitializer, ECppProperty, int32 InOffset, uint64 InFlags, UClass* InClass=NULL )
 	:	UProperty( ObjectInitializer, EC_CppProperty, InOffset, InFlags )
 	,	PropertyClass( InClass )
@@ -1671,10 +1761,17 @@ public:
 	typedef InTCppType TCppType;
 	typedef typename Super::TTypeFundamentals TTypeFundamentals;
 
-	TUObjectPropertyBase( const FObjectInitializer& ObjectInitializer )
-		:	Super( ObjectInitializer)
+	TUObjectPropertyBase(const FObjectInitializer& ObjectInitializer = FObjectInitializer::Get())
+		: Super(ObjectInitializer)
+	{
+	}
+
+	TUObjectPropertyBase(ECppProperty, int32 InOffset, uint64 InFlags, UClass* InClass)
+		: Super(FObjectInitializer::Get(), EC_CppProperty, InOffset, InFlags)
 	{
+		this->PropertyClass = InClass;
 	}
+
 	TUObjectPropertyBase( const FObjectInitializer& ObjectInitializer, ECppProperty, int32 InOffset, uint64 InFlags, UClass* InClass )
 		:	Super( ObjectInitializer, EC_CppProperty, InOffset, InFlags)
 	{
@@ -1706,6 +1803,11 @@ class COREUOBJECT_API UObjectProperty : public TUObjectPropertyBase<UObject*>
 {
 	DECLARE_CASTED_CLASS_INTRINSIC(UObjectProperty,TUObjectPropertyBase<UObject*>,0,CoreUObject,CASTCLASS_UObjectProperty)
 
+	UObjectProperty(ECppProperty, int32 InOffset, uint64 InFlags, UClass* InClass)
+		: TUObjectPropertyBase(FObjectInitializer::Get(), EC_CppProperty, InOffset, InFlags, InClass)
+	{
+	}
+
 	UObjectProperty( const FObjectInitializer& ObjectInitializer, ECppProperty, int32 InOffset, uint64 InFlags, UClass* InClass )
 	:	TUObjectPropertyBase( ObjectInitializer, EC_CppProperty, InOffset, InFlags, InClass )
 	{
@@ -1737,6 +1839,11 @@ class COREUOBJECT_API UWeakObjectProperty : public TUObjectPropertyBase<FWeakObj
 {
 	DECLARE_CASTED_CLASS_INTRINSIC(UWeakObjectProperty,TUObjectPropertyBase<FWeakObjectPtr>,0,CoreUObject,CASTCLASS_UWeakObjectProperty)
 
+	UWeakObjectProperty(ECppProperty, int32 InOffset, uint64 InFlags, UClass* InClass)
+		: TUObjectPropertyBase(FObjectInitializer::Get(), EC_CppProperty, InOffset, InFlags, InClass)
+	{
+	}
+
 	UWeakObjectProperty( const FObjectInitializer& ObjectInitializer, ECppProperty, int32 InOffset, uint64 InFlags, UClass* InClass )
 	:	TUObjectPropertyBase( ObjectInitializer, EC_CppProperty, InOffset, InFlags, InClass )
 	{
@@ -1771,6 +1878,11 @@ class COREUOBJECT_API ULazyObjectProperty : public TUObjectPropertyBase<FLazyObj
 {
 	DECLARE_CASTED_CLASS_INTRINSIC(ULazyObjectProperty,TUObjectPropertyBase<FLazyObjectPtr>,0,CoreUObject,CASTCLASS_ULazyObjectProperty)
 
+	ULazyObjectProperty(ECppProperty, int32 InOffset, uint64 InFlags, UClass* InClass)
+		: TUObjectPropertyBase(FObjectInitializer::Get(), EC_CppProperty, InOffset, InFlags, InClass)
+	{
+	}
+
 	ULazyObjectProperty( const FObjectInitializer& ObjectInitializer, ECppProperty, int32 InOffset, uint64 InFlags, UClass* InClass )
 		:	TUObjectPropertyBase( ObjectInitializer, EC_CppProperty, InOffset, InFlags, InClass )
 	{
@@ -1810,6 +1922,10 @@ class COREUOBJECT_API UAssetObjectProperty : public TUObjectPropertyBase<FAssetP
 {
 	DECLARE_CASTED_CLASS_INTRINSIC(UAssetObjectProperty,TUObjectPropertyBase<FAssetPtr>,0,CoreUObject,CASTCLASS_UAssetObjectProperty)
 
+	UAssetObjectProperty(ECppProperty, int32 InOffset, uint64 InFlags, UClass* InClass)
+		: TUObjectPropertyBase(FObjectInitializer::Get(), EC_CppProperty, InOffset, InFlags, InClass)
+	{}
+
 	UAssetObjectProperty( const FObjectInitializer& ObjectInitializer, ECppProperty, int32 InOffset, uint64 InFlags, UClass* InClass )
 		:	TUObjectPropertyBase( ObjectInitializer, EC_CppProperty, InOffset, InFlags, InClass )
 	{}
@@ -1857,6 +1973,12 @@ class COREUOBJECT_API UClassProperty : public UObjectProperty
 	// Variables.
 	class UClass* MetaClass;
 public:
+	UClassProperty(ECppProperty, int32 InOffset, uint64 InFlags, UClass* InMetaClass)
+		: UObjectProperty(FObjectInitializer::Get(), EC_CppProperty, InOffset, InFlags, UClass::StaticClass())
+		, MetaClass(InMetaClass)
+	{
+	}
+
 	UClassProperty( const FObjectInitializer& ObjectInitializer, ECppProperty, int32 InOffset, uint64 InFlags, UClass* InMetaClass )
 	:	UObjectProperty( ObjectInitializer, EC_CppProperty, InOffset, InFlags, UClass::StaticClass() )
 	,	MetaClass( InMetaClass )
@@ -1913,6 +2035,11 @@ class COREUOBJECT_API UAssetClassProperty : public UAssetObjectProperty
 	// Variables.
 	class UClass* MetaClass;
 public:
+	UAssetClassProperty(ECppProperty, int32 InOffset, uint64 InFlags, UClass* InMetaClass)
+		: Super(FObjectInitializer::Get(), EC_CppProperty, InOffset, InFlags, UClass::StaticClass())
+		, MetaClass(InMetaClass)
+	{}
+
 	UAssetClassProperty( const FObjectInitializer& ObjectInitializer, ECppProperty, int32 InOffset, uint64 InFlags, UClass* InMetaClass )
 		:	Super(ObjectInitializer, EC_CppProperty, InOffset, InFlags, UClass::StaticClass() )
 		,	MetaClass( InMetaClass )
@@ -1955,6 +2082,12 @@ public:
 	typedef UInterfaceProperty_Super::TTypeFundamentals TTypeFundamentals;
 	typedef TTypeFundamentals::TCppType TCppType;
 
+	UInterfaceProperty(ECppProperty, int32 InOffset, uint64 InFlags, UClass* InInterfaceClass)
+		: UInterfaceProperty_Super(FObjectInitializer::Get(), EC_CppProperty, InOffset, (InFlags & ~CPF_InterfaceClearMask))
+		, InterfaceClass(InInterfaceClass)
+	{
+	}
+
 	UInterfaceProperty( const FObjectInitializer& ObjectInitializer, ECppProperty, int32 InOffset, uint64 InFlags, UClass* InInterfaceClass )
 		:	UInterfaceProperty_Super( ObjectInitializer, EC_CppProperty, InOffset, (InFlags & ~CPF_InterfaceClearMask) )
 		,	InterfaceClass( InInterfaceClass )
@@ -2018,6 +2151,11 @@ public:
 	typedef UNameProperty_Super::TTypeFundamentals TTypeFundamentals;
 	typedef TTypeFundamentals::TCppType TCppType;
 
+	UNameProperty(ECppProperty, int32 InOffset, uint64 InFlags)
+		: UNameProperty_Super(FObjectInitializer::Get(), EC_CppProperty, InOffset, InFlags)
+	{
+	}
+
 	UNameProperty( const FObjectInitializer& ObjectInitializer, ECppProperty, int32 InOffset, uint64 InFlags )
 	:	UNameProperty_Super( ObjectInitializer, EC_CppProperty, InOffset, InFlags )
 	{
@@ -2051,6 +2189,11 @@ public:
 	typedef UStrProperty_Super::TTypeFundamentals TTypeFundamentals;
 	typedef TTypeFundamentals::TCppType TCppType;
 
+	UStrProperty(ECppProperty, int32 InOffset, uint64 InFlags)
+		: UStrProperty_Super(FObjectInitializer::Get(), EC_CppProperty, InOffset, InFlags)
+	{
+	}
+
 	UStrProperty( const FObjectInitializer& ObjectInitializer, ECppProperty, int32 InOffset, uint64 InFlags )
 	:	UStrProperty_Super( ObjectInitializer, EC_CppProperty, InOffset, InFlags)
 	{
@@ -2088,6 +2231,11 @@ public:
 	typedef UArrayProperty_Super::TTypeFundamentals TTypeFundamentals;
 	typedef TTypeFundamentals::TCppType TCppType;
 
+	UArrayProperty(ECppProperty, int32 InOffset, uint64 InFlags)
+		: UArrayProperty_Super(FObjectInitializer::Get(), EC_CppProperty, InOffset, InFlags)
+	{
+	}
+
 	UArrayProperty( const FObjectInitializer& ObjectInitializer, ECppProperty, int32 InOffset, uint64 InFlags )
 	:	UArrayProperty_Super( ObjectInitializer, EC_CppProperty, InOffset, InFlags)
 	{
@@ -2477,6 +2625,7 @@ class COREUOBJECT_API UStructProperty : public UProperty
 	// Variables.
 	class UScriptStruct* Struct;
 public:
+	UStructProperty(ECppProperty, int32 InOffset, uint64 InFlags, UScriptStruct* InStruct);
 	UStructProperty( const FObjectInitializer& ObjectInitializer, ECppProperty, int32 InOffset, uint64 InFlags, UScriptStruct* InStruct );
 
 	// UObject interface
@@ -2549,6 +2698,12 @@ public:
 	typedef UDelegateProperty_Super::TTypeFundamentals TTypeFundamentals;
 	typedef TTypeFundamentals::TCppType TCppType;
 
+	UDelegateProperty(ECppProperty, int32 InOffset, uint64 InFlags, UFunction* InSignatureFunction = NULL)
+		: UDelegateProperty_Super(FObjectInitializer::Get(), EC_CppProperty, InOffset, InFlags)
+		, SignatureFunction(InSignatureFunction)
+	{
+	}
+
 	UDelegateProperty( const FObjectInitializer& ObjectInitializer, ECppProperty, int32 InOffset, uint64 InFlags, UFunction* InSignatureFunction = NULL )
 		: UDelegateProperty_Super( ObjectInitializer, EC_CppProperty, InOffset, InFlags)
 		, SignatureFunction(InSignatureFunction)
@@ -2596,6 +2751,12 @@ public:
 	typedef UMulticastDelegateProperty_Super::TTypeFundamentals TTypeFundamentals;
 	typedef TTypeFundamentals::TCppType TCppType;
 
+	UMulticastDelegateProperty(ECppProperty, int32 InOffset, uint64 InFlags, UFunction* InSignatureFunction = NULL)
+		: TProperty(FObjectInitializer::Get(), EC_CppProperty, InOffset, InFlags)
+		, SignatureFunction(InSignatureFunction)
+	{
+	}
+
 	UMulticastDelegateProperty( const FObjectInitializer& ObjectInitializer, ECppProperty, int32 InOffset, uint64 InFlags, UFunction* InSignatureFunction = NULL )
 		: TProperty( ObjectInitializer, EC_CppProperty, InOffset, InFlags )
 		, SignatureFunction(InSignatureFunction)
diff --git a/Engine/Source/Runtime/Engine/Classes/AI/Navigation/NavigationInvokerComponent.h b/Engine/Source/Runtime/Engine/Classes/AI/Navigation/NavigationInvokerComponent.h
--- a/Engine/Source/Runtime/Engine/Classes/AI/Navigation/NavigationInvokerComponent.h
+++ b/Engine/Source/Runtime/Engine/Classes/AI/Navigation/NavigationInvokerComponent.h
@@ -18,7 +18,7 @@ protected:
 	float TileRemovalRadius;
 
 public:
-	UNavigationInvokerComponent(const FObjectInitializer& ObjectInitializer);
+	UNavigationInvokerComponent(const FObjectInitializer& ObjectInitializer = FObjectInitializer::Get());
 
 	void RegisterWithNavigationSystem(UNavigationSystem& NavSys);
 
diff --git a/Engine/Source/Runtime/Engine/Classes/AI/Navigation/NavigationTypes.h b/Engine/Source/Runtime/Engine/Classes/AI/Navigation/NavigationTypes.h
--- a/Engine/Source/Runtime/Engine/Classes/AI/Navigation/NavigationTypes.h
+++ b/Engine/Source/Runtime/Engine/Classes/AI/Navigation/NavigationTypes.h
@@ -532,7 +532,7 @@ class UNavigationTypes : public UObject
 {
 	GENERATED_UCLASS_BODY()
 
-	UNavigationTypes(const FObjectInitializer& ObjectInitializer) : Super(ObjectInitializer) { }
+	UNavigationTypes(const FObjectInitializer& ObjectInitializer = FObjectInitializer::Get()) : Super(ObjectInitializer) { }
 };
 
 //////////////////////////////////////////////////////////////////////////
diff --git a/Engine/Source/Runtime/Engine/Classes/Distributions/DistributionFloat.h b/Engine/Source/Runtime/Engine/Classes/Distributions/DistributionFloat.h
--- a/Engine/Source/Runtime/Engine/Classes/Distributions/DistributionFloat.h
+++ b/Engine/Source/Runtime/Engine/Classes/Distributions/DistributionFloat.h
@@ -87,7 +87,7 @@ class UDistributionFloat : public UDistribution
 	virtual float GetFloatValue(float F = 0);
 
 
-	UDistributionFloat(const FObjectInitializer& ObjectInitializer)
+	UDistributionFloat(const FObjectInitializer& ObjectInitializer = FObjectInitializer::Get())
 	:	Super(ObjectInitializer)
 	,   bCanBeBaked(true)
 	,   bIsDirty(true) // make sure the FRawDistribution is initialized
diff --git a/Engine/Source/Runtime/Engine/Classes/Distributions/DistributionVector.h b/Engine/Source/Runtime/Engine/Classes/Distributions/DistributionVector.h
--- a/Engine/Source/Runtime/Engine/Classes/Distributions/DistributionVector.h
+++ b/Engine/Source/Runtime/Engine/Classes/Distributions/DistributionVector.h
@@ -118,7 +118,7 @@ class ENGINE_API UDistributionVector : public UDistribution
 	virtual FVector GetVectorValue(float F = 0);
 
 
-	UDistributionVector(const FObjectInitializer& ObjectInitializer)
+	UDistributionVector(const FObjectInitializer& ObjectInitializer = FObjectInitializer::Get())
 	:	Super(ObjectInitializer)
 	,   bCanBeBaked(true)
 	,   bIsDirty(true) // make sure the FRawDistribution is initialized
diff --git a/Engine/Source/Runtime/Engine/Classes/Engine/ActorChannel.h b/Engine/Source/Runtime/Engine/Classes/Engine/ActorChannel.h
--- a/Engine/Source/Runtime/Engine/Classes/Engine/ActorChannel.h
+++ b/Engine/Source/Runtime/Engine/Classes/Engine/ActorChannel.h
@@ -82,7 +82,7 @@ class ENGINE_API UActorChannel
 	/**
 	 * Default constructor
 	 */
-	UActorChannel(const FObjectInitializer& ObjectInitializer)
+	UActorChannel(const FObjectInitializer& ObjectInitializer = FObjectInitializer::Get())
 		: UChannel(ObjectInitializer)
 	{
 		ChannelClasses[CHTYPE_Actor] = GetClass();
diff --git a/Engine/Source/Runtime/Engine/Classes/Engine/ControlChannel.h b/Engine/Source/Runtime/Engine/Classes/Engine/ControlChannel.h
--- a/Engine/Source/Runtime/Engine/Classes/Engine/ControlChannel.h
+++ b/Engine/Source/Runtime/Engine/Classes/Engine/ControlChannel.h
@@ -49,7 +49,7 @@ class ENGINE_API UControlChannel
 	/**
 	 * Default constructor
 	 */
-	UControlChannel(const FObjectInitializer& ObjectInitializer)
+	UControlChannel(const FObjectInitializer& ObjectInitializer = FObjectInitializer::Get())
 		: UChannel(ObjectInitializer)
 	{
 		ChannelClasses[CHTYPE_Control]      = GetClass();
diff --git a/Engine/Source/Runtime/Engine/Classes/Engine/Level.h b/Engine/Source/Runtime/Engine/Classes/Engine/Level.h
--- a/Engine/Source/Runtime/Engine/Classes/Engine/Level.h
+++ b/Engine/Source/Runtime/Engine/Classes/Engine/Level.h
@@ -512,7 +512,7 @@ public:
 
 	// Constructor.
 	ENGINE_API void Initialize(const FURL& InURL);
-	ULevel(const FObjectInitializer& ObjectInitializer);
+	ULevel(const FObjectInitializer& ObjectInitializer = FObjectInitializer::Get());
 
 #if WITH_HOT_RELOAD_CTORS
 	/** DO NOT USE. This constructor is for internal usage only for hot-reload purposes. */
diff --git a/Engine/Source/Runtime/Engine/Classes/Engine/NetConnection.h b/Engine/Source/Runtime/Engine/Classes/Engine/NetConnection.h
--- a/Engine/Source/Runtime/Engine/Classes/Engine/NetConnection.h
+++ b/Engine/Source/Runtime/Engine/Classes/Engine/NetConnection.h
@@ -330,7 +330,7 @@ public:
 	static class UNetConnection* GNetConnectionBeingCleanedUp;
 
 	// Constructors and destructors.
-	ENGINE_API UNetConnection(const FObjectInitializer& ObjectInitializer);
+	ENGINE_API UNetConnection(const FObjectInitializer& ObjectInitializer = FObjectInitializer::Get());
 
 	// Begin UObject interface.
 
diff --git a/Engine/Source/Runtime/Engine/Classes/Engine/NetDriver.h b/Engine/Source/Runtime/Engine/Classes/Engine/NetDriver.h
--- a/Engine/Source/Runtime/Engine/Classes/Engine/NetDriver.h
+++ b/Engine/Source/Runtime/Engine/Classes/Engine/NetDriver.h
@@ -316,7 +316,7 @@ public:
 #endif
 
 	// Constructors.
-	ENGINE_API UNetDriver(const FObjectInitializer& ObjectInitializer);
+	ENGINE_API UNetDriver(const FObjectInitializer& ObjectInitializer = FObjectInitializer::Get());
 
 
 	// Begin UObject interface.
diff --git a/Engine/Source/Runtime/Engine/Classes/Engine/PackageMapClient.h b/Engine/Source/Runtime/Engine/Classes/Engine/PackageMapClient.h
--- a/Engine/Source/Runtime/Engine/Classes/Engine/PackageMapClient.h
+++ b/Engine/Source/Runtime/Engine/Classes/Engine/PackageMapClient.h
@@ -105,7 +105,7 @@ class ENGINE_API UPackageMapClient : public UPackageMap
 public:
 	GENERATED_BODY()
 
-	UPackageMapClient(const FObjectInitializer & ObjectInitializer);
+	UPackageMapClient(const FObjectInitializer & ObjectInitializer = FObjectInitializer::Get());
 
 	void Initialize(UNetConnection * InConnection, TSharedPtr<FNetGUIDCache> InNetGUIDCache)
 	{
diff --git a/Engine/Source/Runtime/Engine/Classes/Engine/PendingNetGame.h b/Engine/Source/Runtime/Engine/Classes/Engine/PendingNetGame.h
--- a/Engine/Source/Runtime/Engine/Classes/Engine/PendingNetGame.h
+++ b/Engine/Source/Runtime/Engine/Classes/Engine/PendingNetGame.h
@@ -99,7 +99,7 @@ public:
 	void Initialize(const FURL& InURL);
 
 	// Constructor.
-	UPendingNetGame( const FObjectInitializer& ObjectInitializer );
+	UPendingNetGame(const FObjectInitializer& ObjectInitializer = FObjectInitializer::Get());
 
 	void	InitNetDriver();
 
diff --git a/Engine/Source/Runtime/Engine/Classes/Engine/Polys.h b/Engine/Source/Runtime/Engine/Classes/Engine/Polys.h
--- a/Engine/Source/Runtime/Engine/Classes/Engine/Polys.h
+++ b/Engine/Source/Runtime/Engine/Classes/Engine/Polys.h
@@ -586,7 +586,7 @@ class UPolys : public UObject
 	TTransArray<FPoly> Element;
 
 	// Constructors.
-	UPolys(const FObjectInitializer& ObjectInitializer)
+	UPolys(const FObjectInitializer& ObjectInitializer = FObjectInitializer::Get())
 	:  UObject(ObjectInitializer)
 	, Element( this )
 	{}
diff --git a/Engine/Source/Runtime/Engine/Classes/Engine/Selection.h b/Engine/Source/Runtime/Engine/Classes/Engine/Selection.h
--- a/Engine/Source/Runtime/Engine/Classes/Engine/Selection.h
+++ b/Engine/Source/Runtime/Engine/Classes/Engine/Selection.h
@@ -35,7 +35,7 @@ public:
 	/** Called to deselect everything */
 	static FSimpleMulticastDelegate SelectNoneEvent;
 
-	USelection(const FObjectInitializer& ObjectInitializer);
+	USelection(const FObjectInitializer& ObjectInitializer = FObjectInitializer::Get());
 
 	typedef ClassArray::TIterator TClassIterator;
 	typedef ClassArray::TConstIterator TClassConstIterator;
diff --git a/Engine/Source/Runtime/Engine/Classes/Engine/StaticMesh.h b/Engine/Source/Runtime/Engine/Classes/Engine/StaticMesh.h
--- a/Engine/Source/Runtime/Engine/Classes/Engine/StaticMesh.h
+++ b/Engine/Source/Runtime/Engine/Classes/Engine/StaticMesh.h
@@ -439,7 +439,7 @@ public:
 	/**
 	 * Default constructor
 	 */
-	ENGINE_API UStaticMesh(const FObjectInitializer& ObjectInitializer);
+	ENGINE_API UStaticMesh(const FObjectInitializer& ObjectInitializer = FObjectInitializer::Get());
 
 	// Begin UObject interface.
 #if WITH_EDITOR
diff --git a/Engine/Source/Runtime/Engine/Classes/Engine/VoiceChannel.h b/Engine/Source/Runtime/Engine/Classes/Engine/VoiceChannel.h
--- a/Engine/Source/Runtime/Engine/Classes/Engine/VoiceChannel.h
+++ b/Engine/Source/Runtime/Engine/Classes/Engine/VoiceChannel.h
@@ -22,7 +22,7 @@ class ENGINE_API UVoiceChannel : public UChannel
 	/**
 	 * Default constructor
 	 */
-	UVoiceChannel(const FObjectInitializer& ObjectInitializer)
+	UVoiceChannel(const FObjectInitializer& ObjectInitializer = FObjectInitializer::Get())
 		: UChannel(ObjectInitializer)
 	{
 		// Register with the network channel system
diff --git a/Engine/Source/Runtime/Engine/Classes/Engine/World.h b/Engine/Source/Runtime/Engine/Classes/Engine/World.h
--- a/Engine/Source/Runtime/Engine/Classes/Engine/World.h
+++ b/Engine/Source/Runtime/Engine/Classes/Engine/World.h
@@ -1066,7 +1066,7 @@ public:
 	/**
 	 * UWorld default constructor
 	 */
-	UWorld( const FObjectInitializer& ObjectInitializer );
+	UWorld(const FObjectInitializer& ObjectInitializer = FObjectInitializer::Get());
 	
 	// LINE TRACE
 
diff --git a/Engine/Source/Runtime/Engine/Classes/GameFramework/CharacterMovementComponent.h b/Engine/Source/Runtime/Engine/Classes/GameFramework/CharacterMovementComponent.h
--- a/Engine/Source/Runtime/Engine/Classes/GameFramework/CharacterMovementComponent.h
+++ b/Engine/Source/Runtime/Engine/Classes/GameFramework/CharacterMovementComponent.h
@@ -127,7 +127,7 @@ public:
 	/**
 	 * Default UObject constructor.
 	 */
-	UCharacterMovementComponent(const FObjectInitializer& ObjectInitializer);
+	UCharacterMovementComponent(const FObjectInitializer& ObjectInitializer = FObjectInitializer::Get());
 
 protected:
 
diff --git a/Engine/Source/Runtime/Engine/Public/Model.h b/Engine/Source/Runtime/Engine/Public/Model.h
--- a/Engine/Source/Runtime/Engine/Public/Model.h
+++ b/Engine/Source/Runtime/Engine/Public/Model.h
@@ -418,7 +418,7 @@ private:
 public:
 
 	// Constructors.
-	UModel(const FObjectInitializer& ObjectInitializer);
+	UModel(const FObjectInitializer& ObjectInitializer = FObjectInitializer::Get());
 	ENGINE_API void Initialize();
 	ENGINE_API void Initialize(ABrush* Owner, bool InRootOutside = true);
 
