TMap properties.

diff --git a/Engine/Source/Editor/PropertyEditor/Private/PropertyNode.cpp b/Engine/Source/Editor/PropertyEditor/Private/PropertyNode.cpp
--- a/Engine/Source/Editor/PropertyEditor/Private/PropertyNode.cpp
+++ b/Engine/Source/Editor/PropertyEditor/Private/PropertyNode.cpp
@@ -9,6 +9,7 @@
 #include "Editor/UnrealEd/Public/Kismet2/StructureEditorUtils.h"
 #include "Editor/UnrealEd/Public/Kismet2/BlueprintEditorUtils.h"
 #include "Engine/UserDefinedStruct.h"
+#include "Misc/ScopeExit.h"
 
 FPropertySettings& FPropertySettings::Get()
 {
@@ -2354,10 +2355,13 @@ void FPropertyNode::PropagatePropertyChange( UObject* ModifiedObject, const TCHA
 		FBlueprintEditorUtils::HandleDisableEditableWhenInherited(ModifiedObject, ArchetypeInstances);
 	}
 
-	FPropertyNode* Parent = GetParentNode();
-	UArrayProperty* ParentArrayProp  = Cast<UArrayProperty>(Parent->GetProperty());	
+	FPropertyNode*  Parent          = GetParentNode();
+	UProperty*      ParentProp      = Parent->GetProperty();
+	UArrayProperty* ParentArrayProp = Cast<UArrayProperty>(ParentProp);
+	UProperty*      Prop            = GetProperty();
+	UMapProperty*   MapProp         = Cast<UMapProperty>(Prop);
 
-	if (ParentArrayProp != NULL && ParentArrayProp->Inner != GetProperty())
+	if (ParentArrayProp != NULL && ParentArrayProp->Inner != Prop)
 	{
 		ParentArrayProp = NULL;
 	}
@@ -2391,20 +2395,57 @@ void FPropertyNode::PropagatePropertyChange( UObject* ModifiedObject, const TCHA
 			uint8* Addr = GetValueBaseAddress( (uint8*)ActualObjToChange );
 			if (Addr != NULL)
 			{
-				if (ParentArrayProp != NULL)
+				if (MapProp != NULL)
 				{
-					uint8* ArrayAddr = ParentNode->GetValueBaseAddress( (uint8*)ActualObjToChange );
-					ParentArrayProp->ExportText_Direct(OrgValue, ArrayAddr, ArrayAddr, NULL, PPF_Localized );
+					// Read previous value back into object
+					uint8* PreviousMap = (uint8*)FMemory::Malloc(MapProp->GetSize(), MapProp->GetMinAlignment());
+					ON_SCOPE_EXIT
+					{
+						FMemory::Free(PreviousMap);
+					};
+
+					MapProp->InitializeValue(PreviousMap);
+					ON_SCOPE_EXIT
+					{
+						MapProp->DestroyValue(PreviousMap);
+					};
+
+					MapProp->ImportText(*PreviousValue, PreviousMap, PPF_Localized, ModifiedObject);
+
+					uint8* ModifiedObjectAddr = GetValueBaseAddress( (uint8*)ModifiedObject );
+
+					auto ModifiedObjectAddrPtr = (TMap<int32, FString>*)ModifiedObjectAddr;
+
+					// Serialize differences from the 'default' (the old object)
+					TArray<uint8> Data;
+					{
+						FMemoryWriter Ar(Data);
+						MapProp->SerializeItem(Ar, Addr, PreviousMap);
+					}
+
+					// Deserialize differences back over the new object
+					{
+						FMemoryReader Ar(Data);
+						MapProp->SerializeItem(Ar, Addr, ModifiedObjectAddr);
+					}
 				}
 				else
 				{
-					GetProperty()->ExportText_Direct(OrgValue, Addr, Addr, NULL, PPF_Localized );
-				}
+					if (ParentArrayProp != NULL)
+					{
+						uint8* ArrayAddr = ParentNode->GetValueBaseAddress( (uint8*)ActualObjToChange );
+						ParentArrayProp->ExportText_Direct(OrgValue, ArrayAddr, ArrayAddr, NULL, PPF_Localized );
+					}
+					else
+					{
+						Prop->ExportText_Direct(OrgValue, Addr, Addr, NULL, PPF_Localized );
+					}
 
-				// Check if the original value was the default value and change it only then
-				if (OrgValue == PreviousValue)
-				{
-					GetProperty()->ImportText( NewValue, Addr, PPF_Localized, ActualObjToChange );
+					// Check if the original value was the default value and change it only then
+					if (OrgValue == PreviousValue)
+					{
+						Prop->ImportText( NewValue, Addr, PPF_Localized, ActualObjToChange );
+					}
 				}
 			}
 		}
diff --git a/Engine/Source/Editor/PropertyEditor/Private/UserInterface/PropertyEditor/SPropertyEditorText.cpp b/Engine/Source/Editor/PropertyEditor/Private/UserInterface/PropertyEditor/SPropertyEditorText.cpp
--- a/Engine/Source/Editor/PropertyEditor/Private/UserInterface/PropertyEditor/SPropertyEditorText.cpp
+++ b/Engine/Source/Editor/PropertyEditor/Private/UserInterface/PropertyEditor/SPropertyEditorText.cpp
@@ -92,6 +92,7 @@ bool SPropertyEditorText::Supports( const TSharedRef< FPropertyEditor >& InPrope
 		||	Property->IsA(UTextProperty::StaticClass())
 		||	(Property->IsA(UObjectPropertyBase::StaticClass()) && !Property->HasAnyPropertyFlags(CPF_InstancedReference))
 		||	Property->IsA(UInterfaceProperty::StaticClass())
+		||	Property->IsA(UMapProperty::StaticClass())
 		) )
 	{
 		return true;
diff --git a/Engine/Source/Programs/UnrealHeaderTool/Private/CodeGenerator.cpp b/Engine/Source/Programs/UnrealHeaderTool/Private/CodeGenerator.cpp
--- a/Engine/Source/Programs/UnrealHeaderTool/Private/CodeGenerator.cpp
+++ b/Engine/Source/Programs/UnrealHeaderTool/Private/CodeGenerator.cpp
@@ -1001,13 +1001,21 @@ void FNativeClassHeaderGenerator::OutputProperty(FString& Meta, FOutputDevice& O
 		}
 		OutputDevice.Log(*PropertyNew(Meta, Prop, OuterString, PropMacroOuterClass, TEXT(""), Spaces, *SourceStruct));
 	}
-	UArrayProperty* ArrayProperty = Cast<UArrayProperty>(Prop);
-	if (ArrayProperty)
+
+	if (UArrayProperty* ArrayProperty = Cast<UArrayProperty>(Prop))
 	{
 		FString InnerOuterString = FString::Printf(TEXT("NewProp_%s"), *Prop->GetName());
 		FString PropMacroOuterArray = TEXT("FObjectInitializer(), EC_CppProperty, 0");
 		OutputDevice.Log(*PropertyNew(Meta, ArrayProperty->Inner, InnerOuterString, PropMacroOuterArray, TEXT("_Inner"), Spaces));
 	}
+
+	if (UMapProperty* MapProperty = Cast<UMapProperty>(Prop))
+	{
+		FString InnerOuterString = FString::Printf(TEXT("NewProp_%s"), *Prop->GetName());
+		FString PropMacroOuterMap = TEXT("FObjectInitializer(), EC_CppProperty, ");
+		OutputDevice.Log(*PropertyNew(Meta, MapProperty->KeyProp,   InnerOuterString, PropMacroOuterMap + TEXT("0"), TEXT("_KeyProp"),   Spaces));
+		OutputDevice.Log(*PropertyNew(Meta, MapProperty->ValueProp, InnerOuterString, PropMacroOuterMap + TEXT("1"), TEXT("_ValueProp"), Spaces));
+	}
 }
 
 static FString BackSlashAndIndent(const FString& Text)
@@ -1056,17 +1064,29 @@ static void FindNoExportStructs(TArray<UScriptStruct*>& Structs, UStruct* Start)
 			}
 		}
 
-		for ( TFieldIterator<UProperty> ItInner(Start, EFieldIteratorFlags::ExcludeSuper); ItInner; ++ItInner )
+		for (UProperty* Prop : TFieldRange<UProperty>(Start, EFieldIteratorFlags::ExcludeSuper))
 		{
-			UStructProperty* InnerStruct = Cast<UStructProperty>(*ItInner);
-			UArrayProperty* InnerArray = Cast<UArrayProperty>(*ItInner);
-			if (InnerArray)
+			if (UStructProperty* StructProp = Cast<UStructProperty>(Prop))
 			{
-				InnerStruct = Cast<UStructProperty>(InnerArray->Inner);
+				FindNoExportStructs(Structs, StructProp->Struct);
 			}
-			if (InnerStruct)
+			else if (UArrayProperty* ArrayProp = Cast<UArrayProperty>(Prop))
 			{
-				FindNoExportStructs(Structs, InnerStruct->Struct);
+				if (UStructProperty* InnerStructProp = Cast<UStructProperty>(ArrayProp->Inner))
+				{
+					FindNoExportStructs(Structs, InnerStructProp->Struct);
+				}
+			}
+			else if (UMapProperty* MapProp = Cast<UMapProperty>(Prop))
+			{
+				if (UStructProperty* KeyStructProp = Cast<UStructProperty>(MapProp->KeyProp))
+				{
+					FindNoExportStructs(Structs, KeyStructProp->Struct);
+				}
+				if (UStructProperty* ValueStructProp = Cast<UStructProperty>(MapProp->ValueProp))
+				{
+					FindNoExportStructs(Structs, ValueStructProp->Struct);
+				}
 			}
 		}
 		Start = Start->GetSuperStruct();
@@ -3052,8 +3072,7 @@ void ExportProtoDeclaration(FOutputDevice& Out, const FString& MessageName, TFie
 		if (PropClass != UInterfaceProperty::StaticClass() && PropClass != UObjectProperty::StaticClass())
 		{
 			bool bIsRepeated = false;
-			UArrayProperty* ArrayProperty = Cast<UArrayProperty>(Property);
-			if (ArrayProperty != NULL)
+			if (UArrayProperty* ArrayProperty = Cast<UArrayProperty>(Property))
 			{
 				UClass* InnerPropClass = ArrayProperty->Inner->GetClass();
 				if (InnerPropClass != UInterfaceProperty::StaticClass() && InnerPropClass != UObjectProperty::StaticClass())
@@ -3070,6 +3089,10 @@ void ExportProtoDeclaration(FOutputDevice& Out, const FString& MessageName, TFie
 					FError::Throwf(TEXT("ExportProtoDeclaration - Unhandled property type '%s': %s"), *PropClass->GetName(), *Property->GetPathName());
 				}
 			}
+			else if (UMapProperty* MapProperty = Cast<UMapProperty>(Property))
+			{
+				FError::Throwf(TEXT("ExportProtoDeclaration - Map properties not yet supported '%s': %s"), *PropClass->GetName(), *Property->GetPathName());
+			}
 			else if(Property->ArrayDim != 1)
 			{
 				bIsRepeated = true;
@@ -3247,6 +3270,12 @@ void ExportMCPDeclaration(FOutputDevice& Out, const FString& MessageName, TField
 				// skip array generation for Java, this should result in a List<TYPE> declaration, but we can do this by hand for now.
 				continue;
 			}
+			// TODO Implement maps
+			UMapProperty* MapProperty = Cast<UMapProperty>(Property);
+			if (MapProperty != NULL)
+			{
+				continue;
+			}
 			/*bool bIsRepeated = false;
 			UArrayProperty* ArrayProperty = Cast<UArrayProperty>(Property);
 			if (ArrayProperty != NULL)
@@ -3476,6 +3505,7 @@ void FNativeClassHeaderGenerator::ExportEventParm(UFunction* Function, FUHTStrin
 			Cast<UStrProperty>(Prop) ||
 			Cast<UTextProperty>(Prop) ||
 			Cast<UArrayProperty>(Prop) ||
+			Cast<UMapProperty>(Prop) ||
 			Cast<UInterfaceProperty>(Prop)
 			)
 		{
@@ -3550,6 +3580,7 @@ FString FNativeClassHeaderGenerator::GetNullParameterValue( UProperty* Prop, boo
 		return TEXT("FText::GetEmpty()");
 	}
 	else if ( PropClass == UArrayProperty::StaticClass()
+		||    PropClass == UMapProperty::StaticClass()
 		||    PropClass == UDelegateProperty::StaticClass()
 		||    PropClass == UMulticastDelegateProperty::StaticClass() )
 	{
diff --git a/Engine/Source/Programs/UnrealHeaderTool/Private/HeaderParser.cpp b/Engine/Source/Programs/UnrealHeaderTool/Private/HeaderParser.cpp
--- a/Engine/Source/Programs/UnrealHeaderTool/Private/HeaderParser.cpp
+++ b/Engine/Source/Programs/UnrealHeaderTool/Private/HeaderParser.cpp
@@ -2940,9 +2940,9 @@ void FHeaderParser::GetVarType
 		VarType.PropertyFlags |= Flags;
 
 		GetVarType(AllClasses, Scope, VarProperty, Disallow, &VarType, EPropertyDeclarationStyle::None, VariableCategory);
-		if (VarProperty.ArrayType != EArrayType::None)
+		if (VarProperty.ArrayType != EArrayType::None || VarProperty.MapKeyProp.IsValid())
 		{
-			FError::Throwf(TEXT("Arrays within arrays not supported.") );
+			FError::Throwf(TEXT("Nested containers are not supported.") );
 		}
 
 		OriginalVarTypeFlags |= VarProperty.PropertyFlags & (CPF_ContainsInstancedReference | CPF_InstancedReference); // propagate these to the array, we will fix them later
@@ -2973,6 +2973,69 @@ void FHeaderParser::GetVarType
 			FError::Throwf(TEXT("Found '%s' - explicit allocators are not supported in TArray properties."), AllocatorToken.Identifier);
 		}
 	}
+	else if ( VarType.Matches(TEXT("TMap")) )
+	{
+		RequireSymbol( TEXT("<"), TEXT("'tmap'") );
+
+		// GetVarType() clears the property flags of the array var, so use dummy 
+		// flags when getting the inner property
+		uint64 OriginalVarTypeFlags = VarType.PropertyFlags;
+		VarType.PropertyFlags |= Flags;
+
+		FToken MapKeyType;
+		GetVarType(AllClasses, Scope, MapKeyType, Disallow, &VarType, EPropertyDeclarationStyle::None, VariableCategory);
+		if (VarProperty.ArrayType != EArrayType::None || VarProperty.MapKeyProp.IsValid())
+		{
+			FError::Throwf(TEXT("Nested containers are not supported.") );
+		}
+
+		if (MapKeyType.Type == CPT_Struct)
+		{
+			FError::Throwf(TEXT("USTRUCTs are not currently supported as key types."));
+		}
+
+		FToken CommaToken;
+		if (!GetToken(CommaToken, /*bNoConsts=*/ true) || CommaToken.TokenType != TOKEN_Symbol || FCString::Stricmp(CommaToken.Identifier, TEXT(",")))
+		{
+			FError::Throwf(TEXT("Missing value type while parsing TMap."));
+		}
+
+		GetVarType(AllClasses, Scope, VarProperty, Disallow, &VarType, EPropertyDeclarationStyle::None, VariableCategory);
+		if (VarProperty.ArrayType != EArrayType::None || VarProperty.MapKeyProp.IsValid())
+		{
+			FError::Throwf(TEXT("Nested containers are not supported.") );
+		}
+
+		OriginalVarTypeFlags |= VarProperty.PropertyFlags & (CPF_ContainsInstancedReference | CPF_InstancedReference); // propagate these to the array, we will fix them later
+		OriginalVarTypeFlags |= MapKeyType .PropertyFlags & (CPF_ContainsInstancedReference | CPF_InstancedReference); // propagate these to the array, we will fix them later
+		VarType.PropertyFlags = OriginalVarTypeFlags;
+		VarProperty.MapKeyProp = MakeShareable<FToken>(new FToken(MapKeyType));
+		VarProperty.MapKeyProp->PropertyFlags = OriginalVarTypeFlags;
+
+		FToken CloseTemplateToken;
+		if (!GetToken(CloseTemplateToken, /*bNoConsts=*/ true, ESymbolParseOption::CloseTemplateBracket))
+		{
+			FError::Throwf(TEXT("Missing token while parsing TMap."));
+		}
+
+		if (CloseTemplateToken.TokenType != TOKEN_Symbol || FCString::Stricmp(CloseTemplateToken.Identifier, TEXT(">")))
+		{
+			// If we didn't find a comma, report it
+			if (FCString::Stricmp(CloseTemplateToken.Identifier, TEXT(",")))
+			{
+				FError::Throwf(TEXT("Expected '>' but found '%s'"), CloseTemplateToken.Identifier);
+			}
+
+			// If we found a comma, read the next thing, assume it's an allocator, and report that
+			FToken AllocatorToken;
+			if (!GetToken(AllocatorToken, /*bNoConsts=*/ true, ESymbolParseOption::CloseTemplateBracket))
+			{
+				FError::Throwf(TEXT("Expected '>' but found '%s'"), CloseTemplateToken.Identifier);
+			}
+
+			FError::Throwf(TEXT("Found '%s' - explicit allocators are not supported in TMap properties."), AllocatorToken.Identifier);
+		}
+	}
 	else if ( VarType.Matches(TEXT("FString")) )
 	{
 		VarProperty = FPropertyBase(CPT_String);
@@ -3486,6 +3549,12 @@ void FHeaderParser::GetVarType
 		FError::Throwf(TEXT("Specified type modifiers not allowed here") );
 	}
 
+	// For now, copy the flags that a TMap value has to the key
+	if (FPropertyBase* KeyProp = VarProperty.MapKeyProp.Get())
+	{
+		KeyProp->PropertyFlags = VarProperty.PropertyFlags;
+	}
+
 	VarProperty.MetaData = MetaDataFromNewStyle;
 
 	if (ParsedVarIndexRange)
@@ -3630,9 +3699,9 @@ UProperty* FHeaderParser::GetVarNameAndDim
 			FError::Throwf(TEXT("Arrays aren't allowed in this context") );
 		}
 
-		if (VarProperty.ArrayType == EArrayType::Dynamic)
+		if (VarProperty.ArrayType == EArrayType::Dynamic || VarProperty.MapKeyProp.IsValid())
 		{
-			FError::Throwf(TEXT("Static arrays of dynamic arrays are not allowed"));
+			FError::Throwf(TEXT("Static arrays of containers are not allowed"));
 		}
 
 		if (VarProperty.IsBool())
@@ -3704,12 +3773,14 @@ UProperty* FHeaderParser::GetVarNameAndDim
 			Prev = *It;
 		}
 
-		UProperty* Array    = NULL;
-		UObject*   NewScope = Scope;
-		int32      ArrayDim = 1; // 1 = not a static array, 2 = static array
+		UArrayProperty* Array             = nullptr;
+		UMapProperty*   Map               = nullptr;
+		UProperty*      NewMapKeyProperty = nullptr;
+		UObject*        NewScope          = Scope;
+		int32           ArrayDim          = 1; // 1 = not a static array, 2 = static array
 		if (VarProperty.ArrayType == EArrayType::Dynamic)
 		{
-			Array = new(EC_InternalUseOnlyConstructor, Scope, PropertyName, ObjectFlags)UArrayProperty(FObjectInitializer());
+			Array       = new (EC_InternalUseOnlyConstructor, Scope, PropertyName, ObjectFlags) UArrayProperty(FObjectInitializer());
 			NewScope    = Array;
 			ObjectFlags = RF_Public;
 		}
@@ -3717,31 +3788,55 @@ UProperty* FHeaderParser::GetVarNameAndDim
 		{
 			ArrayDim = 2;
 		}
+		else if (VarProperty.MapKeyProp.IsValid())
+		{
+			Map               = new (EC_InternalUseOnlyConstructor, Scope, PropertyName, ObjectFlags) UMapProperty(FObjectInitializer());
+			NewScope          = Map;
+			ObjectFlags       = RF_Public;
+			NewMapKeyProperty = CreateVariableProperty(*VarProperty.MapKeyProp, NewScope, *(PropertyName.ToString() + TEXT("_Key")), ObjectFlags, VariableCategory);
+		}
 
 		NewProperty = CreateVariableProperty(VarProperty, NewScope, PropertyName, ObjectFlags, VariableCategory);
 
-		if( Array )
-		{
-			check(NewProperty);
-			CastChecked<UArrayProperty>(Array)->Inner = NewProperty;
-
+		auto PropagateFlags = [](uint64 FlagsToPropagate, FPropertyBase& From, UProperty* To) {
 			// Copy some of the property flags to the inner property.
-			NewProperty->PropertyFlags |= (VarProperty.PropertyFlags&CPF_PropagateToArrayInner);
+			To->PropertyFlags |= (From.PropertyFlags & FlagsToPropagate);
+
 			// Copy some of the property flags to the array property.
-			if (NewProperty->PropertyFlags & (CPF_ContainsInstancedReference | CPF_InstancedReference))
+			if (To->PropertyFlags & (CPF_ContainsInstancedReference | CPF_InstancedReference))
 			{
-				VarProperty.PropertyFlags |= CPF_ContainsInstancedReference;
-				VarProperty.PropertyFlags &= ~(CPF_InstancedReference | CPF_PersistentInstance); //this was propagated to the inner
+				From.PropertyFlags |= CPF_ContainsInstancedReference;
+				From.PropertyFlags &= ~(CPF_InstancedReference | CPF_PersistentInstance); //this was propagated to the inner
 
-				if (0 != (NewProperty->PropertyFlags & CPF_PersistentInstance))
+				if (To->PropertyFlags & CPF_PersistentInstance)
 				{
 					TMap<FName, FString> MetaData;
 					AddEditInlineMetaData(MetaData);
-					AddMetaDataToClassData(NewProperty, VarProperty.MetaData);
+					AddMetaDataToClassData(To, From.MetaData);
 				}
 			}
+		};
+
+		if( Array )
+		{
+			Array->Inner = NewProperty;
+
+			PropagateFlags(CPF_PropagateToArrayInner, VarProperty, NewProperty);
+
 			NewProperty = Array;
 		}
+
+		if (Map)
+		{
+			Map->KeyProp   = NewMapKeyProperty;
+			Map->ValueProp = NewProperty;
+
+			PropagateFlags(CPF_PropagateToMapKey,   *VarProperty.MapKeyProp, NewMapKeyProperty);
+			PropagateFlags(CPF_PropagateToMapValue, VarProperty,             NewProperty);
+
+			NewProperty = Map;
+		}
+
 		NewProperty->ArrayDim = ArrayDim;
 		if (ArrayDim == 2)
 		{
diff --git a/Engine/Source/Programs/UnrealHeaderTool/Private/ParserHelper.cpp b/Engine/Source/Programs/UnrealHeaderTool/Private/ParserHelper.cpp
--- a/Engine/Source/Programs/UnrealHeaderTool/Private/ParserHelper.cpp
+++ b/Engine/Source/Programs/UnrealHeaderTool/Private/ParserHelper.cpp
@@ -164,6 +164,7 @@ const TCHAR* FPropertyBase::GetPropertyTypeText( EPropertyType Type )
 		CASE_TEXT(CPT_AssetObjectReference);
 		CASE_TEXT(CPT_WeakObjectReference);
 		CASE_TEXT(CPT_LazyObjectReference);
+		CASE_TEXT(CPT_Map);
 		CASE_TEXT(CPT_MAX);
 	}
 
diff --git a/Engine/Source/Programs/UnrealHeaderTool/Private/ParserHelper.h b/Engine/Source/Programs/UnrealHeaderTool/Private/ParserHelper.h
--- a/Engine/Source/Programs/UnrealHeaderTool/Private/ParserHelper.h
+++ b/Engine/Source/Programs/UnrealHeaderTool/Private/ParserHelper.h
@@ -74,6 +74,8 @@ public:
 	uint64              PropertyFlags;
 	ERefQualifier::Type RefQualifier; // This is needed because of legacy stuff - FString mangles the flags for reasons that have become lost in time but we need this info for testing for invalid replicated function signatures.
 
+	TSharedPtr<FPropertyBase> MapKeyProp;
+
 	/**
 	 * A mask of EPropertyHeaderExportFlags which are used for modifying how this property is exported to the native class header
 	 */
diff --git a/Engine/Source/Programs/UnrealHeaderTool/Resources/UHTDebugging/TestObject.h b/Engine/Source/Programs/UnrealHeaderTool/Resources/UHTDebugging/TestObject.h
--- a/Engine/Source/Programs/UnrealHeaderTool/Resources/UHTDebugging/TestObject.h
+++ b/Engine/Source/Programs/UnrealHeaderTool/Resources/UHTDebugging/TestObject.h
@@ -21,6 +21,9 @@ public:
 	UPROPERTY()
 	ECppEnum EnumProperty;
 
+	UPROPERTY()
+	TMap<int32, FString> TestMap;
+
 	UFUNCTION()
 	void CodeGenTestForEnumClasses(ECppEnum Val);
 
diff --git a/Engine/Source/Runtime/Core/Public/Containers/BitArray.h b/Engine/Source/Runtime/Core/Public/Containers/BitArray.h
--- a/Engine/Source/Runtime/Core/Public/Containers/BitArray.h
+++ b/Engine/Source/Runtime/Core/Public/Containers/BitArray.h
@@ -29,6 +29,8 @@ class TConstSetBitIterator;
 template<typename Allocator = FDefaultBitArrayAllocator,typename OtherAllocator = FDefaultBitArrayAllocator>
 class TConstDualSetBitIterator;
 
+class FScriptBitArray;
+
 
 /**
  * Serializer (predefined for no friend injection in gcc 411)
@@ -155,9 +157,9 @@ protected:
  * An array of Booleans.  They stored in one bit/Boolean.  There are iterators that efficiently iterate over only set bits.
  */
 template<typename Allocator /*= FDefaultBitArrayAllocator*/>
-class TBitArray : protected Allocator::template ForElementType<uint32>
+class TBitArray
 {
-	typedef typename Allocator::template ForElementType<uint32> Super;
+	friend class FScriptBitArray;
 
 public:
 
@@ -236,7 +238,7 @@ private:
 	template <typename BitArrayType>
 	static FORCEINLINE typename TEnableIf<TContainerTraits<BitArrayType>::MoveWillEmptyContainer>::Type MoveOrCopy(BitArrayType& ToArray, BitArrayType& FromArray)
 	{
-		ToArray.Super::MoveToEmpty((Super&)FromArray);
+		ToArray.AllocatorInstance.MoveToEmpty(FromArray.AllocatorInstance);
 
 		ToArray  .NumBits = FromArray.NumBits;
 		ToArray  .MaxBits = FromArray.MaxBits;
@@ -291,7 +293,7 @@ public:
 		if(bReallocate)
 		{
 			// Allocate memory for the new bits.
-			const uint32 MaxDWORDs = this->CalculateSlack(
+			const uint32 MaxDWORDs = AllocatorInstance.CalculateSlack(
 				(NumBits + NumBitsPerDWORD - 1) / NumBitsPerDWORD,
 				(MaxBits + NumBitsPerDWORD - 1) / NumBitsPerDWORD,
 				sizeof(uint32)
@@ -625,29 +627,32 @@ public:
 
 	FORCEINLINE const uint32* GetData() const
 	{
-		return (uint32*)this->GetAllocation();
+		return (uint32*)AllocatorInstance.GetAllocation();
 	}
 
 	FORCEINLINE uint32* GetData()
 	{
-		return (uint32*)this->GetAllocation();
+		return (uint32*)AllocatorInstance.GetAllocation();
 	}
 
 private:
-	int32 NumBits;
-	int32 MaxBits;
+	typedef typename Allocator::template ForElementType<uint32> AllocatorType;
+
+	AllocatorType AllocatorInstance;
+	int32         NumBits;
+	int32         MaxBits;
 
 	void Realloc(int32 PreviousNumBits)
 	{
 		const int32 PreviousNumDWORDs = (PreviousNumBits + NumBitsPerDWORD - 1) / NumBitsPerDWORD;
 		const int32 MaxDWORDs = (MaxBits + NumBitsPerDWORD - 1) / NumBitsPerDWORD;
 
-		this->ResizeAllocation(PreviousNumDWORDs,MaxDWORDs,sizeof(uint32));
+		AllocatorInstance.ResizeAllocation(PreviousNumDWORDs,MaxDWORDs,sizeof(uint32));
 
 		if(MaxDWORDs)
 		{
 			// Reset the newly allocated slack DWORDs.
-			FMemory::Memzero((uint32*)this->GetAllocation() + PreviousNumDWORDs,(MaxDWORDs - PreviousNumDWORDs) * sizeof(uint32));
+			FMemory::Memzero((uint32*)AllocatorInstance.GetAllocation() + PreviousNumDWORDs,(MaxDWORDs - PreviousNumDWORDs) * sizeof(uint32));
 		}
 	}
 };
@@ -873,3 +878,138 @@ private:
 		CurrentBitIndex = BaseBitIndex + NumBitsPerDWORD - 1 - FMath::CountLeadingZeros(this->Mask);
 	}
 };
+
+
+// Untyped set type for accessing TSet data, like FScriptArray for TArray.
+// Must have the same memory representation as a TSet.
+class FScriptBitArray
+{
+public:
+	/**
+	 * Minimal initialization constructor.
+	 * @param Value - The value to initial the bits to.
+	 * @param InNumBits - The initial number of bits in the array.
+	 */
+	FScriptBitArray()
+		: NumBits(0)
+		, MaxBits(0)
+	{
+	}
+
+	bool IsValidIndex(int32 Index) const
+	{
+		return Index >= 0 && Index < NumBits;
+	}
+
+	FBitReference operator[](int32 Index)
+	{
+		check(IsValidIndex(Index));
+		return FBitReference(GetData()[Index / NumBitsPerDWORD], 1 << (Index & (NumBitsPerDWORD - 1)));
+	}
+
+	FConstBitReference operator[](int32 Index) const
+	{
+		check(IsValidIndex(Index));
+		return FConstBitReference(GetData()[Index / NumBitsPerDWORD], 1 << (Index & (NumBitsPerDWORD - 1)));
+	}
+
+	void Empty(int32 Slack = 0)
+	{
+		NumBits = 0;
+
+		// If the expected number of bits doesn't match the allocated number of bits, reallocate.
+		if (MaxBits != Slack)
+		{
+			MaxBits = Slack;
+			Realloc(0);
+		}
+	}
+
+	int32 Add(const bool Value)
+	{
+		const int32 Index = NumBits;
+		const bool bReallocate = (NumBits + 1) > MaxBits;
+
+		NumBits++;
+
+		if(bReallocate)
+		{
+			// Allocate memory for the new bits.
+			const uint32 MaxDWORDs = AllocatorInstance.CalculateSlack(
+				(NumBits + NumBitsPerDWORD - 1) / NumBitsPerDWORD,
+				(MaxBits + NumBitsPerDWORD - 1) / NumBitsPerDWORD,
+				sizeof(uint32)
+				);
+			MaxBits = MaxDWORDs * NumBitsPerDWORD;
+			Realloc(NumBits - 1);
+		}
+
+		(*this)[Index] = Value;
+
+		return Index;
+	}
+
+private:
+	typedef FDefaultBitArrayAllocator::ForElementType<uint32> AllocatorType;
+
+	AllocatorType AllocatorInstance;
+	int32         NumBits;
+	int32         MaxBits;
+
+	// This function isn't intended to be called, just to be compiled to validate the correctness of the type.
+	static void CheckConstraints()
+	{
+		typedef FScriptBitArray ScriptType;
+		typedef TBitArray<>     RealType;
+
+		// Check that the class footprint is the same
+		static_assert(sizeof (ScriptType) == sizeof (RealType), "FScriptBitArray's size doesn't match TBitArray");
+		static_assert(ALIGNOF(ScriptType) == ALIGNOF(RealType), "FScriptBitArray's alignment doesn't match TBitArray");
+
+		// Check member sizes
+		static_assert(sizeof(DeclVal<ScriptType>().AllocatorInstance) == sizeof(DeclVal<RealType>().AllocatorInstance), "FScriptBitArray's AllocatorInstance member size does not match TBitArray's");
+		static_assert(sizeof(DeclVal<ScriptType>().NumBits)           == sizeof(DeclVal<RealType>().NumBits),           "FScriptBitArray's NumBits member size does not match TBitArray's");
+		static_assert(sizeof(DeclVal<ScriptType>().MaxBits)           == sizeof(DeclVal<RealType>().MaxBits),           "FScriptBitArray's MaxBits member size does not match TBitArray's");
+
+		// Check member offsets
+		static_assert(STRUCT_OFFSET(ScriptType, AllocatorInstance) == STRUCT_OFFSET(RealType, AllocatorInstance), "FScriptBitArray's AllocatorInstance member offset does not match TBitArray's");
+		static_assert(STRUCT_OFFSET(ScriptType, NumBits)           == STRUCT_OFFSET(RealType, NumBits),           "FScriptBitArray's NumBits member offset does not match TBitArray's");
+		static_assert(STRUCT_OFFSET(ScriptType, MaxBits)           == STRUCT_OFFSET(RealType, MaxBits),           "FScriptBitArray's MaxBits member offset does not match TBitArray's");
+	}
+
+	FORCEINLINE uint32* GetData()
+	{
+		return (uint32*)AllocatorInstance.GetAllocation();
+	}
+
+	FORCEINLINE const uint32* GetData() const
+	{
+		return (const uint32*)AllocatorInstance.GetAllocation();
+	}
+
+	void Realloc(int32 PreviousNumBits)
+	{
+		const int32 PreviousNumDWORDs = (PreviousNumBits + NumBitsPerDWORD - 1) / NumBitsPerDWORD;
+		const int32 MaxDWORDs = (MaxBits + NumBitsPerDWORD - 1) / NumBitsPerDWORD;
+
+		AllocatorInstance.ResizeAllocation(PreviousNumDWORDs, MaxDWORDs, sizeof(uint32));
+
+		if (MaxDWORDs)
+		{
+			// Reset the newly allocated slack DWORDs.
+			FMemory::Memzero((uint32*)AllocatorInstance.GetAllocation() + PreviousNumDWORDs, (MaxDWORDs - PreviousNumDWORDs) * sizeof(uint32));
+		}
+	}
+
+public:
+	// These should really be private, because they shouldn't be called, but there's a bunch of code
+	// that needs to be fixed first.
+	FScriptBitArray(const FScriptBitArray&) { check(false); }
+	void operator=(const FScriptBitArray&) { check(false); }
+};
+
+template <>
+struct TIsZeroConstructType<FScriptBitArray>
+{
+	enum { Value = true };
+};
diff --git a/Engine/Source/Runtime/Core/Public/Containers/Map.h b/Engine/Source/Runtime/Core/Public/Containers/Map.h
--- a/Engine/Source/Runtime/Core/Public/Containers/Map.h
+++ b/Engine/Source/Runtime/Core/Public/Containers/Map.h
@@ -852,11 +852,14 @@ private:
 	};
 };
 
+class FScriptMap;
+
 /** A TMapBase specialization that only allows a single value associated with each key.*/
 template<typename KeyType,typename ValueType,typename SetAllocator /*= FDefaultSetAllocator*/,typename KeyFuncs /*= TDefaultMapKeyFuncs<KeyType,ValueType,false>*/>
 class TMap : public TSortableMapBase<KeyType, ValueType, SetAllocator, KeyFuncs>
 {
 	friend struct TContainerTraits<TMap>;
+	friend class  FScriptMap;
 
 	static_assert(!KeyFuncs::bAllowDuplicateKeys, "TMap cannot be instantiated with a KeyFuncs which allows duplicate keys");
 
@@ -1177,6 +1180,121 @@ public:
 	}
 };
 
+struct FScriptMapLayout
+{
+	int32 KeyOffset;
+	int32 ValueOffset;
+
+	FScriptSetLayout SetLayout;
+};
+
+// Untyped map type for accessing TMap data, like FScriptArray for TArray.
+// Must have the same memory representation as a TMap.
+class FScriptMap
+{
+public:
+	static FScriptMapLayout GetScriptLayout(int32 KeySize, int32 KeyAlignment, int32 ValueSize, int32 ValueAlignment)
+	{
+		FScriptMapLayout Result;
+
+		// TPair<Key, Value>
+		FStructBuilder PairStruct;
+		Result.KeyOffset   = PairStruct.AddMember(KeySize,   KeyAlignment);
+		Result.ValueOffset = PairStruct.AddMember(ValueSize, ValueAlignment);
+		Result.SetLayout   = FScriptSet::GetScriptLayout(PairStruct.GetSize(),  PairStruct.GetAlignment());
+
+		return Result;
+	}
+
+	FScriptMap()
+	{
+	}
+
+	bool IsValidIndex(int32 Index) const
+	{
+		return Pairs.IsValidIndex(Index);
+	}
+
+	int32 Num() const
+	{
+		return Pairs.Num();
+	}
+
+	int32 GetMaxIndex() const
+	{
+		return Pairs.GetMaxIndex();
+	}
+
+	void* GetData(int32 Index, const FScriptMapLayout& Layout)
+	{
+		return Pairs.GetData(Index, Layout.SetLayout);
+	}
+
+	const void* GetData(int32 Index, const FScriptMapLayout& Layout) const
+	{
+		return Pairs.GetData(Index, Layout.SetLayout);
+	}
+
+	void Empty(int32 Slack, const FScriptMapLayout& Layout)
+	{
+		Pairs.Empty(Slack, Layout.SetLayout);
+	}
+
+	void RemoveAt(int32 Index, const FScriptMapLayout& Layout)
+	{
+		Pairs.RemoveAt(Index, Layout.SetLayout);
+	}
+
+	/**
+	 * Adds an uninitialized object to the map.
+	 * The map will need rehashing at some point after this call to make it valid.
+	 *
+	 * @return  The index of the added element.
+	 */
+	int32 AddUninitialized(const FScriptMapLayout& Layout)
+	{
+		return Pairs.AddUninitialized(Layout.SetLayout);
+	}
+
+	void Rehash(const FScriptMapLayout& Layout, TFunctionRef<uint32 (const void*)> GetKeyHash)
+	{
+		Pairs.Rehash(Layout.SetLayout, GetKeyHash);
+	}
+
+private:
+	FScriptSet Pairs;
+
+	// This function isn't intended to be called, just to be compiled to validate the correctness of the type.
+	static void CheckConstraints()
+	{
+		typedef FScriptMap        ScriptType;
+		typedef TMap<int32, int8> RealType;
+
+		// Check that the class footprint is the same
+		static_assert(sizeof (ScriptType) == sizeof (RealType), "FScriptMap's size doesn't match TMap");
+		static_assert(ALIGNOF(ScriptType) == ALIGNOF(RealType), "FScriptMap's alignment doesn't match TMap");
+
+		// Check member sizes
+		static_assert(sizeof(DeclVal<ScriptType>().Pairs) == sizeof(DeclVal<RealType>().Pairs), "FScriptMap's Pairs member size does not match TMap's");
+
+		// Check member offsets
+		static_assert(STRUCT_OFFSET(ScriptType, Pairs) == STRUCT_OFFSET(RealType, Pairs), "FScriptMap's Pairs member offset does not match TMap's");
+	}
+
+public:
+	// These should really be private, because they shouldn't be called, but there's a bunch of code
+	// that needs to be fixed first.
+	FScriptMap(const FScriptMap&) { check(false); }
+	void operator=(const FScriptMap&) { check(false); }
+};
+
+template <>
+struct TIsZeroConstructType<FScriptMap>
+{
+	enum { Value = true };
+};
+
+
 template <typename KeyType, typename ValueType, typename SetAllocator,typename KeyFuncs>
 struct TContainerTraits<TMap<KeyType, ValueType, SetAllocator, KeyFuncs>> : public TContainerTraitsBase<TMap<KeyType, ValueType, SetAllocator, KeyFuncs>>
 {
diff --git a/Engine/Source/Runtime/Core/Public/Containers/Set.h b/Engine/Source/Runtime/Core/Public/Containers/Set.h
--- a/Engine/Source/Runtime/Core/Public/Containers/Set.h
+++ b/Engine/Source/Runtime/Core/Public/Containers/Set.h
@@ -3,8 +3,11 @@
 #pragma once
 
 #include "Containers/SparseArray.h"
+#include "Misc/StructBuilder.h"
 #include "Templates/Sorting.h"
+#include "Templates/Function.h"
 
+class FScriptSet;
 
 /**
  * The base KeyFuncs type with some useful definitions for all KeyFuncs; meant to be derived from instead of used directly.
@@ -81,6 +84,8 @@ public:
 	template<typename,typename,typename>
 	friend class TSet;
 
+	friend class FScriptSet;
+
 	/** Default constructor. */
 	FORCEINLINE FSetElementId():
 		Index(INDEX_NONE)
@@ -126,10 +131,11 @@ private:
 };
 
 /** An element in the set. */
-template <typename ElementType>
+template <typename InElementType>
 class TSetElement
 {
 public:
+	typedef InElementType ElementType;
 
 	/** The element's value. */
 	ElementType Value;
@@ -188,6 +194,7 @@ template<
 class TSet
 {
 	friend struct TContainerTraits<TSet>;
+	friend class  FScriptSet;
 
 	typedef typename KeyFuncs::KeyInitType     KeyInitType;
 	typedef typename KeyFuncs::ElementInitType ElementInitType;
@@ -291,12 +298,9 @@ public:
 		if(!ConditionalRehash(ExpectedNumElements,true))
 		{
 			// If the hash was already the desired size, clear the references to the elements that have now been removed.
-			if(HashSize)
+			for(int32 HashIndex = 0;HashIndex < HashSize;HashIndex++)
 			{
-				for(int32 HashIndex = 0;HashIndex < HashSize;HashIndex++)
-				{
-					GetTypedHash(HashIndex) = FSetElementId();
-				}
+				GetTypedHash(HashIndex) = FSetElementId();
 			}
 		}
 	}
@@ -1148,3 +1152,225 @@ struct TContainerTraits<TSet<ElementType, KeyFuncs, Allocator> > : public TConta
 		TContainerTraits<typename TSet<ElementType, KeyFuncs, Allocator>::ElementArrayType>::MoveWillEmptyContainer &&
 		TAllocatorTraits<typename Allocator::HashAllocator>::SupportsMove };
 };
+
+struct FScriptSetLayout
+{
+	int32 ElementOffset;
+	int32 HashNextIdOffset;
+	int32 HashIndexOffset;
+	int32 Size;
+
+	FScriptSparseArrayLayout SparseArrayLayout;
+};
+
+// Untyped set type for accessing TSet data, like FScriptArray for TArray.
+// Must have the same memory representation as a TSet.
+class FScriptSet
+{
+public:
+	static FScriptSetLayout GetScriptLayout(int32 ElementSize, int32 ElementAlignment)
+	{
+		FScriptSetLayout Result;
+
+		// TSetElement<TPair<Key, Value>>
+		FStructBuilder SetElementStruct;
+		Result.ElementOffset     = SetElementStruct.AddMember(ElementSize,           ElementAlignment);
+		Result.HashNextIdOffset  = SetElementStruct.AddMember(sizeof(FSetElementId), ALIGNOF(FSetElementId));
+		Result.HashIndexOffset   = SetElementStruct.AddMember(sizeof(int32),         ALIGNOF(int32));
+		Result.Size              = SetElementStruct.GetSize();
+		Result.SparseArrayLayout = FScriptSparseArray::GetScriptLayout(SetElementStruct.GetSize(), SetElementStruct.GetAlignment());
+
+		return Result;
+	}
+
+	FScriptSet()
+		: HashSize(0)
+	{
+	}
+
+	bool IsValidIndex(int32 Index) const
+	{
+		return Elements.IsValidIndex(Index);
+	}
+
+	int32 Num() const
+	{
+		return Elements.Num();
+	}
+
+	int32 GetMaxIndex() const
+	{
+		return Elements.GetMaxIndex();
+	}
+
+	void* GetData(int32 Index, const FScriptSetLayout& Layout)
+	{
+		return Elements.GetData(Index, Layout.SparseArrayLayout);
+	}
+
+	const void* GetData(int32 Index, const FScriptSetLayout& Layout) const
+	{
+		return Elements.GetData(Index, Layout.SparseArrayLayout);
+	}
+
+	void Empty(int32 Slack, const FScriptSetLayout& Layout)
+	{
+		// Empty the elements array, and reallocate it for the expected number of elements.
+		Elements.Empty(Slack, Layout.SparseArrayLayout);
+
+		// Calculate the desired hash size for the specified number of elements.
+		const int32 DesiredHashSize = Allocator::GetNumberOfHashBuckets(Slack);
+
+		// If the hash hasn't been created yet, or is smaller than the desired hash size, rehash.
+		if (Slack != 0 && (HashSize == 0 || HashSize != DesiredHashSize))
+		{
+			HashSize = DesiredHashSize;
+
+			// Free the old hash.
+			Hash.ResizeAllocation(0, HashSize, sizeof(FSetElementId));
+		}
+
+		for (auto* It = (FSetElementId*)Hash.GetAllocation(), *End = It + HashSize; It != End; ++It)
+		{
+			*It = FSetElementId();
+		}
+	}
+
+	void RemoveAt(int32 Index, const FScriptSetLayout& Layout)
+	{
+		check(IsValidIndex(Index));
+
+		auto* ElementBeingRemoved = Elements.GetData(Index, Layout.SparseArrayLayout);
+
+		// Remove the element from the hash.
+		for (FSetElementId* NextElementId = &GetTypedHash(GetHashIndexRef(ElementBeingRemoved, Layout)); NextElementId->IsValidId(); NextElementId = &GetHashNextIdRef(Elements.GetData(NextElementId->AsInteger(), Layout.SparseArrayLayout), Layout))
+		{
+			if (NextElementId->AsInteger() == Index)
+			{
+				*NextElementId = GetHashNextIdRef(ElementBeingRemoved, Layout);
+				break;
+			}
+		}
+
+		// Remove the element from the elements array.
+		Elements.RemoveAtUninitialized(Layout.SparseArrayLayout, Index);
+	}
+
+	/**
+	 * Adds an uninitialized object to the set.
+	 * The set will need rehashing at some point after this call to make it valid.
+	 *
+	 * @return  The index of the added element.
+	 */
+	int32 AddUninitialized(const FScriptSetLayout& Layout)
+	{
+		int32 Result = Elements.AddUninitialized(Layout.SparseArrayLayout);
+		++HashSize;
+		return Result;
+	}
+
+	void Rehash(const FScriptSetLayout& Layout, TFunctionRef<uint32 (const void*)> GetKeyHash)
+	{
+		// Free the old hash.
+		Hash.ResizeAllocation(0,0,sizeof(FSetElementId));
+
+		HashSize = Allocator::GetNumberOfHashBuckets(Elements.Num());
+		if (HashSize)
+		{
+			// Allocate the new hash.
+			checkSlow(FMath::IsPowerOfTwo(HashSize));
+			Hash.ResizeAllocation(0, HashSize, sizeof(FSetElementId));
+			for (int32 HashIndex = 0; HashIndex < HashSize; ++HashIndex)
+			{
+				GetTypedHash(HashIndex) = FSetElementId();
+			}
+
+			int32 NumBytesPerSetElement = Layout.Size;
+
+			// Add the existing elements to the new hash.
+			int32 Index = 0;
+			int32 Count = Elements.Num();
+			while (Count)
+			{
+				if (Elements.IsValidIndex(Index))
+				{
+					FSetElementId ElementId(Index);
+
+					void* Element = (uint8*)Elements.GetData(Index, Layout.SparseArrayLayout);
+
+					// Compute the hash bucket the element goes in.
+					uint32 ElementHash = GetKeyHash(Element);
+					int32  HashIndex   = ElementHash & (HashSize - 1);
+					GetHashIndexRef(Element, Layout) = ElementHash & (HashSize - 1);
+
+					// Link the element into the hash bucket.
+					GetHashNextIdRef(Element, Layout) = GetTypedHash(HashIndex);
+					GetTypedHash(HashIndex) = ElementId;
+
+					--Count;
+				}
+
+				++Index;
+			}
+		}
+	}
+
+private:
+	typedef FDefaultSetAllocator Allocator;
+	typedef Allocator::HashAllocator::ForElementType<FSetElementId> HashType;
+
+	FScriptSparseArray Elements;
+	mutable HashType   Hash;
+	mutable int32      HashSize;
+
+	FORCEINLINE FSetElementId& GetTypedHash(int32 HashIndex) const
+	{
+		return ((FSetElementId*)Hash.GetAllocation())[HashIndex & (HashSize - 1)];
+	}
+
+	static FSetElementId& GetHashNextIdRef(const void* Element, const FScriptSetLayout& Layout)
+	{
+		return *(FSetElementId*)((uint8*)Element + Layout.HashNextIdOffset);
+	}
+
+	static int32& GetHashIndexRef(const void* Element, const FScriptSetLayout& Layout)
+	{
+		return *(int32*)((uint8*)Element + Layout.HashIndexOffset);
+	}
+
+	// This function isn't intended to be called, just to be compiled to validate the correctness of the type.
+	static void CheckConstraints()
+	{
+		typedef FScriptSet  ScriptType;
+		typedef TSet<int32> RealType;
+
+		// Check that the class footprint is the same
+		static_assert(sizeof (ScriptType) == sizeof (RealType), "FScriptSet's size doesn't match TSet");
+		static_assert(ALIGNOF(ScriptType) == ALIGNOF(RealType), "FScriptSet's alignment doesn't match TSet");
+
+		// Check member sizes
+		static_assert(sizeof(DeclVal<ScriptType>().Elements) == sizeof(DeclVal<RealType>().Elements), "FScriptSet's Elements member size does not match TSet's");
+		static_assert(sizeof(DeclVal<ScriptType>().Hash)     == sizeof(DeclVal<RealType>().Hash),     "FScriptSet's Hash member size does not match TSet's");
+		static_assert(sizeof(DeclVal<ScriptType>().HashSize) == sizeof(DeclVal<RealType>().HashSize), "FScriptSet's HashSize member size does not match TSet's");
+
+		// Check member offsets
+		static_assert(STRUCT_OFFSET(ScriptType, Elements) == STRUCT_OFFSET(RealType, Elements), "FScriptSet's Elements member offset does not match TSet's");
+		static_assert(STRUCT_OFFSET(ScriptType, Hash)     == STRUCT_OFFSET(RealType, Hash),     "FScriptSet's Hash member offset does not match TSet's");
+		static_assert(STRUCT_OFFSET(ScriptType, HashSize) == STRUCT_OFFSET(RealType, HashSize), "FScriptSet's FirstFreeIndex member offset does not match TSet's");
+
+		typedef TSet<bool>  ::SetElementType RealSetElementType1;
+		typedef TSet<double>::SetElementType RealSetElementType2;
+	}
+
+public:
+	// These should really be private, because they shouldn't be called, but there's a bunch of code
+	// that needs to be fixed first.
+	FScriptSet(const FScriptSet&) { check(false); }
+	void operator=(const FScriptSet&) { check(false); }
+};
+
+template <>
+struct TIsZeroConstructType<FScriptSet>
+{
+	enum { Value = true };
+};
diff --git a/Engine/Source/Runtime/Core/Public/Containers/SparseArray.h b/Engine/Source/Runtime/Core/Public/Containers/SparseArray.h
--- a/Engine/Source/Runtime/Core/Public/Containers/SparseArray.h
+++ b/Engine/Source/Runtime/Core/Public/Containers/SparseArray.h
@@ -35,6 +35,8 @@ union TSparseArrayElementOrFreeListLink
 	};
 };
 
+class FScriptSparseArray;
+
 /**
  * A dynamically sized array where element indices aren't necessarily contiguous.  Memory is allocated for all 
  * elements in the array's index range, so it doesn't save memory; but it does allow O(1) element removal that 
@@ -46,6 +48,7 @@ template<typename ElementType,typename Allocator /*= FDefaultSparseArrayAllocato
 class TSparseArray
 {
 	friend struct TContainerTraits<TSparseArray>;
+	friend class  FScriptSparseArray;
 
 public:
 
@@ -806,6 +809,184 @@ struct TContainerTraits<TSparseArray<ElementType, Allocator> > : public TContain
 		TContainerTraits<typename TSparseArray<ElementType, Allocator>::AllocationBitArrayType>::MoveWillEmptyContainer };
 };
 
+struct FScriptSparseArrayLayout
+{
+	int32 ElementOffset;
+	int32 Alignment;
+	int32 Size;
+};
+
+// Untyped sparse array type for accessing TSparseArray data, like FScriptArray for TArray.
+// Must have the same memory representation as a TSet.
+class FScriptSparseArray
+{
+public:
+	static FScriptSparseArrayLayout GetScriptLayout(int32 ElementSize, int32 ElementAlignment)
+	{
+		FScriptSparseArrayLayout Result;
+		Result.ElementOffset = 0;
+		Result.Alignment     = FMath::Max(ElementAlignment, (int32)ALIGNOF(FFreeListLink));
+		Result.Size          = FMath::Max(ElementSize,      (int32)sizeof (FFreeListLink));
+
+		return Result;
+	}
+
+	FScriptSparseArray()
+		: FirstFreeIndex(-1)
+		, NumFreeIndices(0)
+	{
+	}
+
+	bool IsValidIndex(int32 Index) const
+	{
+		return AllocationFlags.IsValidIndex(Index) && AllocationFlags[Index];
+	}
+
+	int32 Num() const
+	{
+		return Data.Num() - NumFreeIndices;
+	}
+
+	int32 GetMaxIndex() const
+	{
+		return Data.Num();
+	}
+
+	void* GetData(int32 Index, const FScriptSparseArrayLayout& Layout)
+	{
+		return (uint8*)Data.GetData() + Layout.Size * Index;
+	}
+
+	const void* GetData(int32 Index, const FScriptSparseArrayLayout& Layout) const
+	{
+		return (const uint8*)Data.GetData() + Layout.Size * Index;
+	}
+
+	void Empty(int32 Slack, const FScriptSparseArrayLayout& Layout)
+	{
+		// Free the allocated elements.
+		Data.Empty(Slack, Layout.Size);
+		FirstFreeIndex = -1;
+		NumFreeIndices = 0;
+		AllocationFlags.Empty(Slack);
+	}
+
+	/**
+	 * Adds an uninitialized object to the array.
+	 *
+	 * @return  The index of the added element.
+	 */
+	int32 AddUninitialized(const FScriptSparseArrayLayout& Layout)
+	{
+		int32 Index;
+		if (NumFreeIndices)
+		{
+			// Remove and use the first index from the list of free elements.
+			Index = FirstFreeIndex;
+			FirstFreeIndex = GetFreeListLink(FirstFreeIndex, Layout)->NextFreeIndex;
+			--NumFreeIndices;
+			if(NumFreeIndices)
+			{
+				GetFreeListLink(FirstFreeIndex, Layout)->PrevFreeIndex = -1;
+			}
+		}
+		else
+		{
+			// Add a new element.
+			Index = Data.Add(1, Layout.Size);
+			AllocationFlags.Add(false);
+		}
+
+		AllocationFlags[Index] = true;
+
+		return Index;
+	}
+
+	/** Removes Count elements from the array, starting from Index, without destructing them. */
+	void RemoveAtUninitialized(const FScriptSparseArrayLayout& Layout, int32 Index, int32 Count = 1)
+	{
+		for (; Count; --Count)
+		{
+			check(AllocationFlags[Index]);
+
+			// Mark the element as free and add it to the free element list.
+			if(NumFreeIndices)
+			{
+				GetFreeListLink(FirstFreeIndex, Layout)->PrevFreeIndex = Index;
+			}
+
+			auto* IndexData = GetFreeListLink(Index, Layout);
+			IndexData->PrevFreeIndex = -1;
+			IndexData->NextFreeIndex = NumFreeIndices > 0 ? FirstFreeIndex : INDEX_NONE;
+			FirstFreeIndex = Index;
+			++NumFreeIndices;
+			AllocationFlags[Index] = false;
+
+			++Index;
+		}
+	}
+
+private:
+	FScriptArray    Data;
+	FScriptBitArray AllocationFlags;
+	int32           FirstFreeIndex;
+	int32           NumFreeIndices;
+
+	// This function isn't intended to be called, just to be compiled to validate the correctness of the type.
+	static void CheckConstraints()
+	{
+		typedef FScriptSparseArray  ScriptType;
+		typedef TSparseArray<int32> RealType;
+
+		// Check that the class footprint is the same
+		static_assert(sizeof (ScriptType) == sizeof (RealType), "FScriptSparseArray's size doesn't match TSparseArray");
+		static_assert(ALIGNOF(ScriptType) == ALIGNOF(RealType), "FScriptSparseArray's alignment doesn't match TSparseArray");
+
+		// Check member sizes
+		static_assert(sizeof(DeclVal<ScriptType>().Data)            == sizeof(DeclVal<RealType>().Data),            "FScriptSparseArray's Data member size does not match TSparseArray's");
+		static_assert(sizeof(DeclVal<ScriptType>().AllocationFlags) == sizeof(DeclVal<RealType>().AllocationFlags), "FScriptSparseArray's AllocationFlags member size does not match TSparseArray's");
+		static_assert(sizeof(DeclVal<ScriptType>().FirstFreeIndex)  == sizeof(DeclVal<RealType>().FirstFreeIndex),  "FScriptSparseArray's FirstFreeIndex member size does not match TSparseArray's");
+		static_assert(sizeof(DeclVal<ScriptType>().NumFreeIndices)  == sizeof(DeclVal<RealType>().NumFreeIndices),  "FScriptSparseArray's NumFreeIndices member size does not match TSparseArray's");
+
+		// Check member offsets
+		static_assert(STRUCT_OFFSET(ScriptType, Data)            == STRUCT_OFFSET(RealType, Data),            "FScriptSparseArray's Data member offset does not match TSparseArray's");
+		static_assert(STRUCT_OFFSET(ScriptType, AllocationFlags) == STRUCT_OFFSET(RealType, AllocationFlags), "FScriptSparseArray's AllocationFlags member offset does not match TSparseArray's");
+		static_assert(STRUCT_OFFSET(ScriptType, FirstFreeIndex)  == STRUCT_OFFSET(RealType, FirstFreeIndex),  "FScriptSparseArray's FirstFreeIndex member offset does not match TSparseArray's");
+		static_assert(STRUCT_OFFSET(ScriptType, NumFreeIndices)  == STRUCT_OFFSET(RealType, NumFreeIndices),  "FScriptSparseArray's NumFreeIndices member offset does not match TSparseArray's");
+
+		// Check free index offsets
+		static_assert(STRUCT_OFFSET(ScriptType::FFreeListLink, PrevFreeIndex) == STRUCT_OFFSET(RealType::FElementOrFreeListLink, PrevFreeIndex), "FScriptSparseArray's FFreeListLink's PrevFreeIndex member offset does not match TSparseArray's");
+		static_assert(STRUCT_OFFSET(ScriptType::FFreeListLink, NextFreeIndex) == STRUCT_OFFSET(RealType::FElementOrFreeListLink, NextFreeIndex), "FScriptSparseArray's FFreeListLink's NextFreeIndex member offset does not match TSparseArray's");
+	}
+
+	struct FFreeListLink
+	{
+		/** If the element isn't allocated, this is a link to the previous element in the array's free list. */
+		int32 PrevFreeIndex;
+
+		/** If the element isn't allocated, this is a link to the next element in the array's free list. */
+		int32 NextFreeIndex;
+	};
+
+	/** Accessor for the element or free list data. */
+	FORCEINLINE FFreeListLink* GetFreeListLink(int32 Index, const FScriptSparseArrayLayout& Layout)
+	{
+		return (FFreeListLink*)GetData(Index, Layout);
+	}
+
+public:
+	// These should really be private, because they shouldn't be called, but there's a bunch of code
+	// that needs to be fixed first.
+	FScriptSparseArray(const FScriptSparseArray&) { check(false); }
+	void operator=(const FScriptSparseArray&) { check(false); }
+};
+
+template <>
+struct TIsZeroConstructType<FScriptSparseArray>
+{
+	enum { Value = true };
+};
+
 /**
  * A placement new operator which constructs an element in a sparse array allocation.
  */
diff --git a/Engine/Source/Runtime/Core/Public/Misc/StructBuilder.h b/Engine/Source/Runtime/Core/Public/Misc/StructBuilder.h
new file mode 100644
--- /dev/null
+++ b/Engine/Source/Runtime/Core/Public/Misc/StructBuilder.h
@@ -0,0 +1,60 @@
+// Copyright 1998-2015 Epic Games, Inc. All Rights Reserved.
+
+#pragma once
+
+/**
+ * Lays out up a struct at runtime, in hopefully the same way as a compiler,
+ * and returns the offset of each member as it's added.
+ */
+class FStructBuilder
+{
+public:
+	FStructBuilder()
+		: EndOfLastMember(0)
+		, Alignment      (0)
+	{
+	}
+
+	/**
+	 * Adds a new member to the end of the struct.
+	 *
+	 * @param  MemberSize       The size of the member.
+	 * @param  MemberAlignment  The alignment of the member.
+	 *
+	 * @return  The offset of the member from the start of the struct.
+	 */
+	int32 AddMember(int32 MemberSize, int32 MemberAlignment)
+	{
+		int32 Result    = Align(EndOfLastMember, MemberAlignment);
+		EndOfLastMember = Result + MemberSize;
+		Alignment       = FMath::Max(Alignment, MemberAlignment);
+		return Result;
+	}
+
+	/**
+	 * Returns the current size of the struct.
+	 *
+	 * @return  The size of the struct.
+	 */
+	int32 GetSize() const
+	{
+		return Align(EndOfLastMember, Alignment);
+	}
+
+	/**
+	 * Returns the current alignment of the struct.
+	 *
+	 * @return  The alignment of the struct.
+	 */
+	int32 GetAlignment() const
+	{
+		return Alignment;
+	}
+
+private:
+	// The offset from the start of the struct to the end of the last added member
+	int32 EndOfLastMember;
+
+	// The alignment of the struct
+	int32 Alignment;
+};
diff --git a/Engine/Source/Runtime/Core/Public/Templates/TypeHash.h b/Engine/Source/Runtime/Core/Public/Templates/TypeHash.h
--- a/Engine/Source/Runtime/Core/Public/Templates/TypeHash.h
+++ b/Engine/Source/Runtime/Core/Public/Templates/TypeHash.h
@@ -85,11 +85,16 @@ inline uint32 GetTypeHash( const int64 A )
 	return (uint32)A+((uint32)(A>>32) * 23);
 }
 
-inline uint32 GetTypeHash( const float& Value )
+inline uint32 GetTypeHash( float Value )
 {
 	return *(uint32*)&Value;
 }
 
+inline uint32 GetTypeHash( double Value )
+{
+	return GetTypeHash(*(uint64*)&Value);
+}
+
 inline uint32 GetTypeHash( const TCHAR* S )
 {
 	return FCrc::Strihash_DEPRECATED(S);
@@ -109,4 +114,4 @@ template <typename EnumType>
 FORCEINLINE  typename TEnableIf<TIsEnum<EnumType>::Value, uint32>::Type GetTypeHash(EnumType E)
 {
 	return GetTypeHash((__underlying_type(EnumType))E);
-}
\ No newline at end of file
+}
diff --git a/Engine/Source/Runtime/Core/Public/Templates/UnrealTemplate.h b/Engine/Source/Runtime/Core/Public/Templates/UnrealTemplate.h
--- a/Engine/Source/Runtime/Core/Public/Templates/UnrealTemplate.h
+++ b/Engine/Source/Runtime/Core/Public/Templates/UnrealTemplate.h
@@ -33,7 +33,8 @@ public:
  * @param Alignment		Alignment, must be a power of two
  * @return				Aligned value
  */
-template< class T > inline T Align( const T Ptr, int32 Alignment )
+template <typename T>
+inline CONSTEXPR T Align( const T Ptr, int32 Alignment )
 {
 	return (T)(((PTRINT)Ptr + Alignment - 1) & ~(Alignment-1));
 }
@@ -548,3 +549,12 @@ template <typename LHS, typename RHS>
 struct TOr : TOrValue<LHS::Value, RHS>
 {
 };
+
+
+/**
+ * Equivalent to std::declval.  
+ *
+ * Note that this function is unimplemented, and is only intended to be used in unevaluated contexts, like sizeof and trait expressions.
+ */
+template <typename T>
+T&& DeclVal();
diff --git a/Engine/Source/Runtime/CoreUObject/Private/UObject/GarbageCollection.cpp b/Engine/Source/Runtime/CoreUObject/Private/UObject/GarbageCollection.cpp
--- a/Engine/Source/Runtime/CoreUObject/Private/UObject/GarbageCollection.cpp
+++ b/Engine/Source/Runtime/CoreUObject/Private/UObject/GarbageCollection.cpp
@@ -826,6 +826,14 @@ public:
 						TokenReturnCount = REFERENCE_INFO.ReturnCount;
 						AddReferencedObjects(CurrentObject, ReferenceCollector);
 					}
+					else if( REFERENCE_INFO.Type == GCRT_AddTMapReferencedObjects )
+					{
+						void*         Map         = StackEntryData + REFERENCE_INFO.Offset;
+						UMapProperty* MapProperty = (UMapProperty*)TokenStream->ReadPointer( TokenStreamIndex );
+						TokenReturnCount = REFERENCE_INFO.ReturnCount;
+						FSimpleObjectReferenceCollectorArchive CollectorArchive(CurrentObject, ReferenceCollector);
+						MapProperty->SerializeItem(CollectorArchive, Map, nullptr);
+					}
 					else if( REFERENCE_INFO.Type == GCRT_EndOfStream )
 					{
 						// Break out of loop.
@@ -1420,6 +1428,19 @@ bool UArrayProperty::ContainsObjectReference() const
 	return Inner->ContainsObjectReference();
 }
 
+/**
+ * Returns true if this property, or in the case of e.g. array or struct properties any sub- property, contains a
+ * UObject reference.
+ *
+ * @return true if property (or sub- properties) contain a UObject reference, false otherwise
+ */
+bool UMapProperty::ContainsObjectReference() const
+{
+	check(KeyProp);
+	check(ValueProp);
+	return KeyProp->ContainsObjectReference() || ValueProp->ContainsObjectReference();
+}
+
 /**
  * Returns true if this property, or in the case of e.g. array or struct properties any sub- property, contains a
  * UObject reference.
@@ -1472,6 +1493,14 @@ bool UArrayProperty::ContainsWeakObjectReference() const
 	return Inner->ContainsWeakObjectReference();
 }
 
+// Returns true if this property contains a weak UObject reference.
+bool UMapProperty::ContainsWeakObjectReference() const
+{
+	check(KeyProp);
+	check(ValueProp);
+	return KeyProp->ContainsWeakObjectReference() || ValueProp->ContainsWeakObjectReference();
+}
+
 // Returns true if this property contains a weak UObject reference.
 bool UStructProperty::ContainsWeakObjectReference() const
 {
@@ -1621,6 +1650,20 @@ void UArrayProperty::EmitReferenceInfo(UClass& OwnerClass, int32 BaseOffset)
 }
 
 
+/**
+ * Emits tokens used by realtime garbage collection code to passed in OwnerClass' ReferenceTokenStream. The offset emitted is relative
+ * to the passed in BaseOffset which is used by e.g. arrays of structs.
+ */
+void UMapProperty::EmitReferenceInfo(UClass& OwnerClass, int32 BaseOffset)
+{
+	if (ContainsObjectReference())
+	{
+		OwnerClass.ReferenceTokenStream.EmitReferenceInfo(FGCReferenceInfo(GCRT_AddTMapReferencedObjects, BaseOffset + GetOffset_ForGC()));
+		OwnerClass.ReferenceTokenStream.EmitPointer((const void*)this);
+	}
+}
+
+
 /**
  * Emits tokens used by realtime garbage collection code to passed in ReferenceTokenStream. The offset emitted is relative
  * to the passed in BaseOffset which is used by e.g. arrays of structs.
diff --git a/Engine/Source/Runtime/CoreUObject/Private/UObject/Property.cpp b/Engine/Source/Runtime/CoreUObject/Private/UObject/Property.cpp
--- a/Engine/Source/Runtime/CoreUObject/Private/UObject/Property.cpp
+++ b/Engine/Source/Runtime/CoreUObject/Private/UObject/Property.cpp
@@ -688,6 +688,11 @@ int32 UProperty::SetupOffset()
 	return Offset_Internal + GetSize();
 }
 
+void UProperty::SetOffset_Internal(int32 NewOffset)
+{
+	Offset_Internal = NewOffset;
+}
+
 bool UProperty::SameType(const UProperty* Other) const
 {
 	return Other && (this->GetClass() == Other->GetClass());
diff --git a/Engine/Source/Runtime/CoreUObject/Private/UObject/PropertyArray.cpp b/Engine/Source/Runtime/CoreUObject/Private/UObject/PropertyArray.cpp
--- a/Engine/Source/Runtime/CoreUObject/Private/UObject/PropertyArray.cpp
+++ b/Engine/Source/Runtime/CoreUObject/Private/UObject/PropertyArray.cpp
@@ -24,34 +24,27 @@ bool UArrayProperty::Identical( const void* A, const void* B, uint32 PortFlags )
 	checkSlow(Inner);
 
 	FScriptArrayHelper ArrayHelperA(this, A);
-	FScriptArrayHelper ArrayHelperB(this, B);
 
 	const int32 ArrayNum = ArrayHelperA.Num();
-	if ( ArrayNum != (B ? ArrayHelperB.Num() : 0) )
+	if ( B == NULL )
 	{
-		return false;
+		return ArrayNum == 0;
 	}
 
-	if ( B != NULL )
+	FScriptArrayHelper ArrayHelperB(this, B);
+	if ( ArrayNum != ArrayHelperB.Num() )
 	{
-		for ( int32 ArrayIndex = 0; ArrayIndex < ArrayNum; ArrayIndex++ )
-		{
-			if ( !Inner->Identical( ArrayHelperA.GetRawPtr(ArrayIndex), ArrayHelperB.GetRawPtr(ArrayIndex), PortFlags) )
-			{
-				return false;
-			}
-		}
+		return false;
 	}
-	else
+
+	for ( int32 ArrayIndex = 0; ArrayIndex < ArrayNum; ArrayIndex++ )
 	{
-		for ( int32 ArrayIndex = 0; ArrayIndex < ArrayNum; ArrayIndex++ )
+		if ( !Inner->Identical( ArrayHelperA.GetRawPtr(ArrayIndex), ArrayHelperB.GetRawPtr(ArrayIndex), PortFlags) )
 		{
-			if ( !Inner->Identical( ArrayHelperA.GetRawPtr(ArrayIndex), 0, PortFlags) )
-			{
-				return false;
-			}
+			return false;
 		}
 	}
+
 	return true;
 }
 void UArrayProperty::SerializeItem( FArchive& Ar, void* Value, void const* Defaults ) const
diff --git a/Engine/Source/Runtime/CoreUObject/Private/UObject/PropertyMap.cpp b/Engine/Source/Runtime/CoreUObject/Private/UObject/PropertyMap.cpp
new file mode 100644
--- /dev/null
+++ b/Engine/Source/Runtime/CoreUObject/Private/UObject/PropertyMap.cpp
@@ -0,0 +1,791 @@
+// Copyright 1998-2014 Epic Games, Inc. All Rights Reserved.
+
+#include "CoreUObjectPrivate.h"
+#include "PropertyHelper.h"
+#include "ScopeExit.h"
+
+namespace UE4MapProperty_Private
+{
+	/**
+	 * Checks if any of the pairs in the map compare equal to the one passed.
+	 *
+	 * @param  MapHelper  The map to search through.
+	 * @param  Index      The index in the map to start searching from.
+	 * @param  Num        The number of elements to compare.
+	 */
+	bool AnyEqual(const FScriptMapHelper& MapHelper, int32 Index, int32 Num, const uint8* PairToCompare, uint32 PortFlags)
+	{
+		UProperty* KeyProp   = MapHelper.GetKeyProperty();
+		UProperty* ValueProp = MapHelper.GetValueProperty();
+
+		int32 ValueOffset = MapHelper.MapLayout.ValueOffset;
+
+		for (; Num; --Num)
+		{
+			while (!MapHelper.IsValidIndex(Index))
+			{
+				++Index;
+			}
+
+			if (KeyProp->Identical(MapHelper.GetPairPtr(Index), PairToCompare, PortFlags) && ValueProp->Identical(MapHelper.GetPairPtr(Index) + ValueOffset, PairToCompare + ValueOffset, PortFlags))
+			{
+				return true;
+			}
+
+			++Index;
+			--Num;
+		}
+
+		return false;
+	}
+
+	bool RangesContainSameAmountsOfVal(const FScriptMapHelper& MapHelperA, int32 IndexA, const FScriptMapHelper& MapHelperB, int32 IndexB, int32 Num, const uint8* PairToCompare, uint32 PortFlags)
+	{
+		UProperty* KeyProp   = MapHelperA.GetKeyProperty();
+		UProperty* ValueProp = MapHelperA.GetValueProperty();
+
+		// Ensure that both maps are the same type
+		check(KeyProp   == MapHelperB.GetKeyProperty());
+		check(ValueProp == MapHelperB.GetValueProperty());
+
+		int32 ValueOffset = MapHelperA.MapLayout.ValueOffset;
+
+		int32 CountA = 0;
+		int32 CountB = 0;
+		for (;;)
+		{
+			if (Num == 0)
+			{
+				return CountA == CountB;
+			}
+
+			while (!MapHelperA.IsValidIndex(IndexA))
+			{
+				++IndexA;
+			}
+
+			while (!MapHelperB.IsValidIndex(IndexB))
+			{
+				++IndexB;
+			}
+
+			const uint8* PairA = MapHelperA.GetPairPtr(IndexA);
+			const uint8* PairB = MapHelperB.GetPairPtr(IndexB);
+			if (KeyProp->Identical(PairA, PairToCompare, PortFlags) && ValueProp->Identical(PairA + ValueOffset, PairToCompare + ValueOffset, PortFlags))
+			{
+				++CountA;
+			}
+
+			if (KeyProp->Identical(PairB, PairToCompare, PortFlags) && ValueProp->Identical(PairB + ValueOffset, PairToCompare + ValueOffset, PortFlags))
+			{
+				++CountB;
+			}
+
+			++IndexA;
+			++IndexB;
+			--Num;
+		}
+	}
+
+	bool IsPermutation(const FScriptMapHelper& MapHelperA, const FScriptMapHelper& MapHelperB, uint32 PortFlags)
+	{
+		UProperty* KeyProp   = MapHelperA.GetKeyProperty();
+		UProperty* ValueProp = MapHelperA.GetValueProperty();
+
+		// Ensure that both maps are the same type
+		check(KeyProp   == MapHelperB.GetKeyProperty());
+		check(ValueProp == MapHelperB.GetValueProperty());
+
+		int32 Num = MapHelperA.Num();
+		if (Num != MapHelperB.Num())
+		{
+			return false;
+		}
+
+		int32 ValueOffset = MapHelperA.MapLayout.ValueOffset;
+
+		// Skip over common initial sequence
+		int32 IndexA = 0;
+		int32 IndexB = 0;
+		for (;;)
+		{
+			if (Num == 0)
+			{
+				return true;
+			}
+
+			while (!MapHelperA.IsValidIndex(IndexA))
+			{
+				++IndexA;
+			}
+
+			while (!MapHelperB.IsValidIndex(IndexB))
+			{
+				++IndexB;
+			}
+
+			const uint8* PairA = MapHelperA.GetPairPtr(IndexA);
+			const uint8* PairB = MapHelperB.GetPairPtr(IndexB);
+			if (!KeyProp->Identical(PairA, PairB, PortFlags))
+			{
+				break;
+			}
+
+			if (!ValueProp->Identical(PairA + ValueOffset, PairB + ValueOffset, PortFlags))
+			{
+				break;
+			}
+
+			++IndexA;
+			++IndexB;
+			--Num;
+		}
+
+		int32 FirstIndexA = IndexA;
+		int32 FirstIndexB = IndexB;
+		int32 FirstNum    = Num;
+		for (;;)
+		{
+			const uint8* PairA = MapHelperA.GetPairPtr(IndexA);
+			if (!AnyEqual(MapHelperA, FirstIndexA, FirstNum - Num, PairA, PortFlags) && !RangesContainSameAmountsOfVal(MapHelperA, IndexA, MapHelperB, IndexB, Num, PairA, PortFlags))
+			{
+				return false;
+			}
+
+			--Num;
+			if (Num == 0)
+			{
+				return true;
+			}
+
+			while (!MapHelperA.IsValidIndex(IndexA))
+			{
+				++IndexA;
+			}
+
+			while (!MapHelperB.IsValidIndex(IndexB))
+			{
+				++IndexB;
+			}
+		}
+	}
+}
+
+UMapProperty::UMapProperty(const FObjectInitializer& ObjectInitializer, ECppProperty, int32 InOffset, uint64 InFlags)
+: UMapProperty_Super(ObjectInitializer, EC_CppProperty, InOffset, InFlags)
+{
+	// These are expected to be set post-construction by AddCppProperty
+	KeyProp   = nullptr;
+	ValueProp = nullptr;
+}
+
+void UMapProperty::LinkInternal(FArchive& Ar)
+{
+	check(KeyProp && ValueProp);
+
+	if (auto* MyLinker = GetLinker())
+	{
+		MyLinker->Preload(this);
+	}
+	Ar.Preload(KeyProp);
+	Ar.Preload(ValueProp);
+	KeyProp  ->Link(Ar);
+	ValueProp->Link(Ar);
+
+	int32 KeySize        = KeyProp  ->GetSize();
+	int32 ValueSize      = ValueProp->GetSize();
+	int32 KeyAlignment   = KeyProp  ->GetMinAlignment();
+	int32 ValueAlignment = ValueProp->GetMinAlignment();
+
+	MapLayout = FScriptMap::GetScriptLayout(KeySize, KeyAlignment, ValueSize, ValueAlignment);
+
+	ValueProp->SetOffset_Internal(MapLayout.ValueOffset);
+
+	Super::LinkInternal(Ar);
+}
+
+bool UMapProperty::Identical(const void* A, const void* B, uint32 PortFlags) const
+{
+	checkSlow(KeyProp);
+	checkSlow(ValueProp);
+
+	FScriptMapHelper MapHelperA(this, A);
+
+	int32 ANum = MapHelperA.Num();
+
+	if (!B)
+	{
+		return ANum == 0;
+	}
+
+	FScriptMapHelper MapHelperB(this, B);
+	if (ANum != MapHelperB.Num())
+	{
+		return false;
+	}
+
+	return UE4MapProperty_Private::IsPermutation(MapHelperA, MapHelperB, PortFlags);
+}
+
+void UMapProperty::SerializeItem(FArchive& Ar, void* Value, const void* Defaults) const
+{
+	checkSlow(KeyProp);
+	checkSlow(ValueProp);
+
+	// Ensure that the key/value properties have been loaded before calling SerializeItem() on them
+	Ar.Preload(KeyProp);
+	Ar.Preload(ValueProp);
+
+	FScriptMapHelper MapHelper(this, Value);
+
+	if (Ar.IsLoading())
+	{
+		if (Defaults)
+		{
+			CopyValuesInternal(Value, Defaults, 1);
+		}
+		else
+		{
+			MapHelper.EmptyValues();
+		}
+
+		uint8* TempKeyStorage = nullptr;
+		ON_SCOPE_EXIT
+		{
+			KeyProp->DestroyValue(TempKeyStorage);
+			FMemory::Free(TempKeyStorage);
+		};
+
+		// Delete any explicitly-removed keys
+		int32 NumKeysToRemove = 0;
+		Ar << NumKeysToRemove;
+		if (NumKeysToRemove)
+		{
+			TempKeyStorage = (uint8*)FMemory::Malloc(MapLayout.SetLayout.Size);
+			KeyProp->InitializeValue(TempKeyStorage);
+
+			for (; NumKeysToRemove; --NumKeysToRemove)
+			{
+				// Read key into temporary storage
+				KeyProp->SerializeItem(Ar, TempKeyStorage);
+
+				// If the key is in the map, remove it
+				int32 Found = MapHelper.FindMapIndexWithKey(TempKeyStorage);
+				if (Found != INDEX_NONE)
+				{
+					MapHelper.RemoveAt_NeedsRehash(Found);
+				}
+			}
+		}
+
+		int32 Num = 0;
+		Ar << Num;
+
+		// Allocate temporary key space if we haven't allocated it already above
+		if (Num != 0 && !TempKeyStorage)
+		{
+			TempKeyStorage = (uint8*)FMemory::Malloc(MapLayout.SetLayout.Size);
+			KeyProp->InitializeValue(TempKeyStorage);
+		}
+
+		// Read remaining items into container
+		for (; Num; --Num)
+		{
+			// Read key into temporary storage
+			KeyProp->SerializeItem(Ar, TempKeyStorage);
+
+			// Add a new default value if the key doesn't currently exist in the map
+			int32 NextPairIndex = MapHelper.FindMapIndexWithKey(TempKeyStorage);
+			if (NextPairIndex == INDEX_NONE)
+			{
+				NextPairIndex = MapHelper.AddDefaultValue_Invalid_NeedsRehash();
+			}
+
+			uint8* NextPairPtr = MapHelper.GetPairPtrWithoutCheck(NextPairIndex);
+
+			// Copy over deserialised key from temporary storage
+			KeyProp->CopyCompleteValue_InContainer(NextPairPtr, TempKeyStorage);
+
+			// Deserialize value
+			ValueProp->SerializeItem(Ar, NextPairPtr + MapLayout.ValueOffset);
+		}
+
+		MapHelper.Rehash();
+	}
+	else
+	{
+		FScriptMapHelper DefaultsHelper(this, Defaults);
+
+		// Container for temporarily tracking some indices
+		TSet<int32> Indices;
+
+		// Determine how many keys are missing from the object
+		if (Defaults)
+		{
+			for (int32 Index = 0, Count = DefaultsHelper.Num(); Count; ++Index)
+			{
+				uint8* DefaultPairPtr = DefaultsHelper.GetPairPtrWithoutCheck(Index);
+
+				if (DefaultsHelper.IsValidIndex(Index))
+				{
+					if (!MapHelper.FindMapPairPtrWithKey(DefaultPairPtr))
+					{
+						Indices.Add(Index);
+					}
+
+					--Count;
+				}
+			}
+		}
+
+		// Write out the missing keys
+		int32 MissingKeysNum = Indices.Num();
+		Ar << MissingKeysNum;
+		for (int32 Index : Indices)
+		{
+			KeyProp->SerializeItem(Ar, DefaultsHelper.GetPairPtr(Index));
+		}
+
+		// Write out differences from defaults
+		if (Defaults)
+		{
+			Indices.Empty(Indices.Num());
+			for (int32 Index = 0, Count = MapHelper.Num(); Count; ++Index)
+			{
+				if (MapHelper.IsValidIndex(Index))
+				{
+					uint8* ValuePairPtr   = MapHelper.GetPairPtrWithoutCheck(Index);
+					uint8* DefaultPairPtr = DefaultsHelper.FindMapPairPtrWithKey(ValuePairPtr);
+
+					if (!DefaultPairPtr || !ValueProp->Identical(ValuePairPtr + MapLayout.ValueOffset, DefaultPairPtr + MapLayout.ValueOffset))
+					{
+						Indices.Add(Index);
+					}
+
+					--Count;
+				}
+			}
+
+			// Write out differences from defaults
+			int32 Num = Indices.Num();
+			Ar << Num;
+			for (int32 Index : Indices)
+			{
+				uint8* ValuePairPtr = MapHelper.GetPairPtrWithoutCheck(Index);
+
+				KeyProp  ->SerializeItem(Ar, ValuePairPtr);
+				ValueProp->SerializeItem(Ar, ValuePairPtr + MapLayout.ValueOffset);
+			}
+		}
+		else
+		{
+			int32 Num = MapHelper.Num();
+			Ar << Num;
+			for (int32 Index = 0; Num; ++Index)
+			{
+				if (MapHelper.IsValidIndex(Index))
+				{
+					uint8* ValuePairPtr = MapHelper.GetPairPtrWithoutCheck(Index);
+
+					KeyProp  ->SerializeItem(Ar, ValuePairPtr);
+					ValueProp->SerializeItem(Ar, ValuePairPtr + MapLayout.ValueOffset);
+
+					--Num;
+				}
+			}
+		}
+	}
+}
+
+bool UMapProperty::NetSerializeItem( FArchive& Ar, UPackageMap* Map, void* Data, TArray<uint8> * MetaData ) const
+{
+	UE_LOG( LogProperty, Fatal, TEXT( "Deprecated code path" ) );
+	return 1;
+}
+
+void UMapProperty::Serialize( FArchive& Ar )
+{
+	Super::Serialize( Ar );
+	Ar << KeyProp;
+	Ar << ValueProp;
+}
+
+void UMapProperty::AddReferencedObjects(UObject* InThis, FReferenceCollector& Collector)
+{
+	UMapProperty* This = CastChecked<UMapProperty>(InThis);
+
+	Collector.AddReferencedObject(This->KeyProp,   This);
+	Collector.AddReferencedObject(This->ValueProp, This);
+
+	Super::AddReferencedObjects(This, Collector);
+}
+
+FString UMapProperty::GetCPPType(FString* ExtendedTypeText, uint32 CPPExportFlags) const
+{
+	checkSlow(KeyProp);
+	checkSlow(ValueProp);
+
+	if (ExtendedTypeText)
+	{
+		FString KeyExtendedTypeText;
+		FString KeyTypeText = KeyProp->GetCPPType(&KeyExtendedTypeText, CPPExportFlags & ~CPPF_ArgumentOrReturnValue); // we won't consider map keys to be "arguments or return values"
+
+		FString ValueExtendedTypeText;
+		FString ValueTypeText = ValueProp->GetCPPType(&ValueExtendedTypeText, CPPExportFlags & ~CPPF_ArgumentOrReturnValue); // we won't consider map values to be "arguments or return values"
+
+		*ExtendedTypeText = FString::Printf(TEXT("<%s%s,%s%s>"), *KeyTypeText, *KeyExtendedTypeText, *ValueTypeText, *ValueExtendedTypeText);
+	}
+
+	return TEXT("TMap");
+}
+
+FString UMapProperty::GetCPPMacroType( FString& ExtendedTypeText ) const
+{
+	checkSlow(KeyProp);
+	checkSlow(ValueProp);
+	ExtendedTypeText = FString::Printf(TEXT("%s,%s"), *KeyProp->GetCPPType(), *ValueProp->GetCPPType());
+	return TEXT("TMAP");
+}
+
+void UMapProperty::ExportTextItem(FString& ValueStr, const void* PropertyValue, const void* DefaultValue, UObject* Parent, int32 PortFlags, UObject* ExportRootScope) const
+{
+	checkSlow(KeyProp);
+	checkSlow(ValueProp);
+
+	FScriptMapHelper MapHelper(this, PropertyValue);
+
+	if (MapHelper.Num() == 0)
+	{
+		return;
+	}
+
+	uint8* StructDefaults = nullptr;
+	if (auto* StructValueProp = dynamic_cast<UStructProperty*>(ValueProp))
+	{
+		checkSlow(StructValueProp->Struct);
+
+		StructDefaults = (uint8*)FMemory::Malloc(MapLayout.SetLayout.Size);
+		ValueProp->InitializeValue(StructDefaults + MapLayout.ValueOffset);
+	}
+	ON_SCOPE_EXIT
+	{
+		if (StructDefaults)
+		{
+			ValueProp->DestroyValue(StructDefaults + MapLayout.ValueOffset);
+			FMemory::Free(StructDefaults);
+		}
+	};
+
+	FScriptMapHelper DefaultMapHelper(this, DefaultValue);
+
+	uint8* PropData = MapHelper.GetPairPtrWithoutCheck(0);
+	if (PortFlags & PPF_BlueprintDebugView)
+	{
+		int32 Index  = 0;
+		bool  bFirst = true;
+		for (int32 Count = MapHelper.Num(); Count; PropData += MapLayout.SetLayout.Size, ++Index)
+		{
+			if (MapHelper.IsValidIndex(Index))
+			{
+				if (bFirst)
+				{
+					bFirst = false;
+				}
+				else
+				{
+					ValueStr += TCHAR('\n');
+				}
+
+				ValueStr += TEXT("[");
+				KeyProp->ExportTextItem(ValueStr, PropData, nullptr, Parent, PortFlags | PPF_Delimited, ExportRootScope);
+				ValueStr += TEXT("] ");
+
+				// Always use struct defaults if the inner is a struct, for symmetry with the import of array inner struct defaults
+				uint8* PropDefault = StructDefaults ? StructDefaults : DefaultValue ? DefaultMapHelper.FindMapPairPtrWithKey(PropData) : nullptr;
+
+				ValueProp->ExportTextItem(ValueStr, PropData + MapLayout.ValueOffset, PropDefault + MapLayout.ValueOffset, Parent, PortFlags | PPF_Delimited, ExportRootScope);
+
+				--Count;
+			}
+		}
+	}
+	else
+	{
+		int32 Index  = 0;
+		bool  bFirst = true;
+		for (int32 Count = MapHelper.Num(); Count; PropData += MapLayout.SetLayout.Size, ++Index)
+		{
+			if (MapHelper.IsValidIndex(Index))
+			{
+				if (bFirst)
+				{
+					ValueStr += TCHAR('(');
+					bFirst = false;
+				}
+				else
+				{
+					ValueStr += TCHAR(',');
+				}
+
+				ValueStr += TEXT("(");
+
+				KeyProp->ExportTextItem(ValueStr, PropData, nullptr, Parent, PortFlags | PPF_Delimited, ExportRootScope);
+
+				ValueStr += TEXT(", ");
+
+				// Always use struct defaults if the inner is a struct, for symmetry with the import of array inner struct defaults
+				uint8* PropDefault = StructDefaults ? StructDefaults : DefaultValue ? DefaultMapHelper.FindMapPairPtrWithKey(PropData) : nullptr;
+
+				ValueProp->ExportTextItem(ValueStr, PropData + MapLayout.ValueOffset, PropDefault + MapLayout.ValueOffset, Parent, PortFlags | PPF_Delimited, ExportRootScope);
+
+				ValueStr += TEXT(")");
+
+				--Count;
+			}
+		}
+
+		ValueStr += TEXT(")");
+	}
+}
+
+const TCHAR* UMapProperty::ImportText_Internal(const TCHAR* Buffer, void* Data, int32 PortFlags, UObject* Parent, FOutputDevice* ErrorText) const
+{
+	checkSlow(KeyProp);
+	checkSlow(ValueProp);
+
+	FScriptMapHelper MapHelper(this, Data);
+	MapHelper.EmptyValues();
+
+	// If we export an empty array we export an empty string, so ensure that if we're passed an empty string
+	// we interpret it as an empty array.
+	if (*Buffer++ != TCHAR('('))
+	{
+		return nullptr;
+	}
+
+	SkipWhitespace(Buffer);
+	if (*Buffer == TCHAR(')'))
+	{
+		return Buffer + 1;
+	}
+
+	int32 Index = 0;
+	for (;;)
+	{
+		MapHelper.AddUninitializedValue();
+		MapHelper.ConstructItem(Index);
+		uint8* PairPtr = MapHelper.GetPairPtrWithoutCheck(Index);
+
+		if (*Buffer++ != TCHAR('('))
+		{
+			return nullptr;
+		}
+
+		// Parse the key
+		Buffer = KeyProp->ImportText(Buffer, PairPtr, PortFlags | PPF_Delimited, Parent, ErrorText);
+		if (!Buffer)
+		{
+			return nullptr;
+		}
+
+		SkipWhitespace(Buffer);
+		if (*Buffer++ != TCHAR(','))
+		{
+			return nullptr;
+		}
+
+		// Parse the value
+		SkipWhitespace(Buffer);
+		Buffer = ValueProp->ImportText(Buffer, PairPtr + MapLayout.ValueOffset, PortFlags | PPF_Delimited, Parent, ErrorText);
+		if (!Buffer)
+		{
+			return nullptr;
+		}
+
+		SkipWhitespace(Buffer);
+		if (*Buffer++ != TCHAR(')'))
+		{
+			return nullptr;
+		}
+
+		switch (*Buffer++)
+		{
+			case TCHAR(')'):
+				MapHelper.Rehash();
+				return Buffer;
+
+			case TCHAR(','):
+				break;
+
+			default:
+				return nullptr;
+		}
+
+		++Index;
+	}
+}
+
+void UMapProperty::AddCppProperty( UProperty* Property )
+{
+	check(Property);
+
+	if (Property->GetOffset_ForInternal() == 0)
+	{
+		// If it's at offset zero, assume it's the key
+		check(!KeyProp);
+		KeyProp = Property;
+	}
+	else
+	{
+		// Otherwise assume it's the value
+		check(!ValueProp);
+		ValueProp = Property;
+	}
+}
+
+void UMapProperty::CopyValuesInternal(void* Dest, void const* Src, int32 Count) const
+{
+	check(Count == 1);
+
+	FScriptMapHelper SrcMapHelper (this, Src);
+	FScriptMapHelper DestMapHelper(this, Dest);
+
+	int32 Num = SrcMapHelper.Num();
+	DestMapHelper.EmptyValues(Num);
+
+	if (Num == 0)
+	{
+		return;
+	}
+
+	int32 DestIndex = 0;
+	for (int32 SrcIndex = 0; Num; ++SrcIndex)
+	{
+		if (SrcMapHelper.IsValidIndex(SrcIndex))
+		{
+			int32 DestIndex = DestMapHelper.AddDefaultValue_Invalid_NeedsRehash();
+
+			uint8* SrcData  = SrcMapHelper .GetPairPtrWithoutCheck(SrcIndex);
+			uint8* DestData = DestMapHelper.GetPairPtrWithoutCheck(DestIndex);
+
+			KeyProp  ->CopyCompleteValue_InContainer(DestData, SrcData);
+			ValueProp->CopyCompleteValue_InContainer(DestData, SrcData);
+
+			--Num;
+		}
+	}
+
+	DestMapHelper.Rehash();
+}
+
+void UMapProperty::ClearValueInternal(void* Data) const
+{
+	FScriptMapHelper MapHelper(this, Data);
+	MapHelper.EmptyValues();
+}
+
+void UMapProperty::DestroyValueInternal(void* Data) const
+{
+	FScriptMapHelper MapHelper(this, Data);
+	MapHelper.EmptyValues();
+
+	//@todo UE4 potential double destroy later from this...would be ok for a script map, but still
+	((FScriptMap*)Data)->~FScriptMap();
+}
+
+bool UMapProperty::PassCPPArgsByRef() const
+{
+	return true;
+}
+
+/**
+ * Creates new copies of components
+ * 
+ * @param	Data				pointer to the address of the instanced object referenced by this UComponentProperty
+ * @param	DefaultData			pointer to the address of the default value of the instanced object referenced by this UComponentProperty
+ * @param	Owner				the object that contains this property's data
+ * @param	InstanceGraph		contains the mappings of instanced objects and components to their templates
+ */
+void UMapProperty::InstanceSubobjects(void* Data, void const* DefaultData, UObject* Owner, FObjectInstancingGraph* InstanceGraph)
+{
+	if (!Data)
+	{
+		return;
+	}
+
+	bool bInstancedKey   = KeyProp  ->ContainsInstancedObjectProperty();
+	bool bInstancedValue = ValueProp->ContainsInstancedObjectProperty();
+
+	if (!bInstancedKey && !bInstancedValue)
+	{
+		return;
+	}
+
+	FScriptMapHelper MapHelper(this, Data);
+
+	if (DefaultData)
+	{
+		FScriptMapHelper DefaultMapHelper(this, DefaultData);
+		int32            DefaultNum = DefaultMapHelper.Num();
+
+		for (int32 Index = 0, Num = MapHelper.Num(); Num; ++Index)
+		{
+			if (MapHelper.IsValidIndex(Index))
+			{
+				uint8* PairPtr        = MapHelper.GetPairPtr(Index);
+				uint8* DefaultPairPtr = DefaultMapHelper.FindMapPairPtrWithKey(PairPtr, Index);
+
+				if (bInstancedKey)
+				{
+					KeyProp->InstanceSubobjects(PairPtr, DefaultPairPtr, Owner, InstanceGraph);
+				}
+
+				if (bInstancedValue)
+				{
+					ValueProp->InstanceSubobjects(PairPtr + MapLayout.ValueOffset, DefaultPairPtr + MapLayout.ValueOffset, Owner, InstanceGraph);
+				}
+
+				--Num;
+			}
+		}
+	}
+	else
+	{
+		for (int32 Index = 0, Num = MapHelper.Num(); Num; ++Index)
+		{
+			if (MapHelper.IsValidIndex(Index))
+			{
+				uint8* PairPtr = MapHelper.GetPairPtr(Index);
+
+				if (bInstancedKey)
+				{
+					KeyProp->InstanceSubobjects(PairPtr, nullptr, Owner, InstanceGraph);
+				}
+
+				if (bInstancedValue)
+				{
+					ValueProp->InstanceSubobjects(PairPtr + MapLayout.ValueOffset, nullptr, Owner, InstanceGraph);
+				}
+
+				--Num;
+			}
+		}
+	}
+}
+
+bool UMapProperty::SameType(const UProperty* Other) const
+{
+	auto* MapProp = (UMapProperty*)Other;
+	return Super::SameType(Other) && KeyProp && ValueProp && KeyProp->SameType(MapProp->KeyProp) && ValueProp->SameType(MapProp->ValueProp);
+}
+
+IMPLEMENT_CORE_INTRINSIC_CLASS(UMapProperty, UProperty,
+	{
+		Class->EmitObjectReference(STRUCT_OFFSET(UMapProperty, KeyProp),   TEXT("KeyProp"));
+		Class->EmitObjectReference(STRUCT_OFFSET(UMapProperty, ValueProp), TEXT("ValueProp"));
+
+		// Ensure that TArray and FScriptMap are interchangeable, as FScriptMap will be used to access a native array property
+		// from script that is declared as a TArray in C++.
+		static_assert(sizeof(FScriptMap) == sizeof(TMap<uint32, uint8>), "FScriptMap and TMap<uint32, uint8> must be interchangable.");
+	}
+);
diff --git a/Engine/Source/Runtime/CoreUObject/Public/UObject/GarbageCollection.h b/Engine/Source/Runtime/CoreUObject/Public/UObject/GarbageCollection.h
--- a/Engine/Source/Runtime/CoreUObject/Public/UObject/GarbageCollection.h
+++ b/Engine/Source/Runtime/CoreUObject/Public/UObject/GarbageCollection.h
@@ -25,6 +25,7 @@ enum EGCReferenceType
 	GCRT_FixedArray,
 	GCRT_AddStructReferencedObjects,
 	GCRT_AddReferencedObjects,
+	GCRT_AddTMapReferencedObjects,
 	GCRT_EndOfStream,
 };
 
diff --git a/Engine/Source/Runtime/CoreUObject/Public/UObject/ObjectBase.h b/Engine/Source/Runtime/CoreUObject/Public/UObject/ObjectBase.h
--- a/Engine/Source/Runtime/CoreUObject/Public/UObject/ObjectBase.h
+++ b/Engine/Source/Runtime/CoreUObject/Public/UObject/ObjectBase.h
@@ -299,6 +299,7 @@ typedef uint64 EClassCastFlags;
 #define CASTCLASS_UBlueprint					DECLARE_UINT64(0x0000080000000000)
 #define CASTCLASS_UDelegateFunction				DECLARE_UINT64(0x0000100000000000)
 #define CASTCLASS_UStaticMeshComponent			DECLARE_UINT64(0x0000200000000000)
+#define CASTCLASS_UMapProperty					DECLARE_UINT64(0x0000400000000000)
 
 #define CASTCLASS_AllFlags						DECLARE_UINT64(0xFFFFFFFFFFFFFFFF)
 
@@ -360,6 +361,7 @@ typedef uint64 EClassCastFlags;
 #define CPF_ExposeOnSpawn					DECLARE_UINT64(0x0001000000000000)		// Property is exposed on spawn
 #define CPF_PersistentInstance				DECLARE_UINT64(0x0002000000000000)		// A object referenced by the property is duplicated like a component. (Each actor should have an own instance.)
 #define CPF_UObjectWrapper					DECLARE_UINT64(0x0004000000000000)		// Property was parsed as a wrapper class like TSubobjectOf<T>, FScriptInterface etc., rather than a USomething*
+#define CPF_HasGetValueTypeHash				DECLARE_UINT64(0x0008000000000000)		// This property can generate a meaningful hash value.
 
 #define CPF_NonPIETransient \
 	EMIT_DEPRECATED_WARNING_MESSAGE("CPF_NonPIETransient is deprecated. Please use CPF_NonPIEDuplicateTransient instead.") \
@@ -369,6 +371,8 @@ typedef uint64 EClassCastFlags;
 //@{
 #define CPF_ParmFlags				(CPF_Parm | CPF_OutParm | CPF_ReturnParm | CPF_ReferenceParm | CPF_ConstParm)
 #define CPF_PropagateToArrayInner	(CPF_ExportObject | CPF_PersistentInstance | CPF_InstancedReference | CPF_ContainsInstancedReference | CPF_Config | CPF_EditConst | CPF_Deprecated | CPF_EditorOnly | CPF_AutoWeak | CPF_UObjectWrapper )
+#define CPF_PropagateToMapValue		(CPF_ExportObject | CPF_PersistentInstance | CPF_InstancedReference | CPF_ContainsInstancedReference | CPF_Config | CPF_EditConst | CPF_Deprecated | CPF_EditorOnly | CPF_AutoWeak | CPF_UObjectWrapper | CPF_Edit )
+#define CPF_PropagateToMapKey		(CPF_ExportObject | CPF_PersistentInstance | CPF_InstancedReference | CPF_ContainsInstancedReference | CPF_Config | CPF_EditConst | CPF_Deprecated | CPF_EditorOnly | CPF_AutoWeak | CPF_UObjectWrapper | CPF_Edit )
 
 /** the flags that should never be set on interface properties */
 #define CPF_InterfaceClearMask		(CPF_ExportObject|CPF_InstancedReference|CPF_ContainsInstancedReference)
diff --git a/Engine/Source/Runtime/CoreUObject/Public/UObject/Stack.h b/Engine/Source/Runtime/CoreUObject/Public/UObject/Stack.h
--- a/Engine/Source/Runtime/CoreUObject/Public/UObject/Stack.h
+++ b/Engine/Source/Runtime/CoreUObject/Public/UObject/Stack.h
@@ -45,6 +45,7 @@ enum EPropertyType
 	CPT_LazyObjectReference,
 	CPT_AssetObjectReference,
 	CPT_Double,
+	CPT_Map,
 
 	// when you add new property types, make sure you add the corresponding entry
 	// in the PropertyTypeToNameMap array in ScriptCompiler.cpp!!
@@ -347,6 +348,8 @@ inline VariableSizeType FFrame::ReadVariableSize( UField** ExpressionField/*=NUL
 			break;
 		case CPT_Text: Result = sizeof(FText);
 			break;
+		case CPT_Map: Result = sizeof(FScriptMap);
+			break;
 		default:
 			UE_LOG(LogScriptFrame, Fatal, TEXT("Unhandled property type in FFrame::ReadVariableSize(): %u"), NullPropertyType);
 			break;
diff --git a/Engine/Source/Runtime/CoreUObject/Public/UObject/UnrealType.h b/Engine/Source/Runtime/CoreUObject/Public/UObject/UnrealType.h
--- a/Engine/Source/Runtime/CoreUObject/Public/UObject/UnrealType.h
+++ b/Engine/Source/Runtime/CoreUObject/Public/UObject/UnrealType.h
@@ -126,6 +126,12 @@ private:
 	*/
 	int32 SetupOffset();
 
+protected:
+	friend class UMapProperty;
+
+	/** Set the alignment offset for this property - added for UMapProperty */
+	void SetOffset_Internal(int32 NewOffset);
+
 	/**
 	 * Initializes internal state.
 	 */
@@ -148,6 +154,11 @@ public:
 		return Offset_Internal;
 	}
 	/** Return offset of property from container base. */
+	FORCEINLINE int32 GetOffset_ForInternal() const
+	{
+		return Offset_Internal;
+	}
+	/** Return offset of property from container base. */
 	FORCEINLINE int32 GetOffset_ReplaceWith_ContainerPtrToValuePtr() const
 	{
 		return Offset_Internal;
@@ -260,32 +271,47 @@ public:
 
 private:
 
-	FORCEINLINE void* ContainerPtrToValuePtrInternal(void const* ContainerPtr, int32 ArrayIndex, bool bKnownToBeUObject) const
+	FORCEINLINE void* ContainerVoidPtrToValuePtrInternal(void* ContainerPtr, int32 ArrayIndex) const
 	{
-		if (bKnownToBeUObject)
-		{
-			// in the future, these checks will be tested if the property is supposed be from a UClass
-			// need something for networking, since those are NOT live uobjects, just memory blocks
-			check(((UObject*)ContainerPtr)->IsValidLowLevel()); // Check its a valid UObject that was passed in
-			check(((UObject*)ContainerPtr)->GetClass() != NULL);
-			check(Cast<UClass>(GetOuter()) != nullptr); // Check that the outer of this property is a UClass (not another property)
-
-			// Check that the object we are accessing is of the class that contains this property
-			checkfSlow(((UObject*)ContainerPtr)->IsA((UClass*)GetOuter()), TEXT("'%s' is of class '%s' however property '%s' belongs to class '%s'") 
-																	 , *((UObject*)ContainerPtr)->GetName()
-																	 , *((UObject*)ContainerPtr)->GetClass()->GetName()
-																	 , *GetName()
-																	 , *((UClass*)GetOuter())->GetName());
-		}
+		check(ArrayIndex < ArrayDim);
+		check(ContainerPtr);
+
 		if (0)
 		{
 			// in the future, these checks will be tested if the property is NOT relative to a UClass
 			check(!Cast<UClass>(GetOuter())); // Check we are _not_ calling this on a direct child property of a UClass, you should pass in a UObject* in that case
 		}
+
+		return (uint8*)ContainerPtr + Offset_Internal + ElementSize * ArrayIndex;
+	}
+
+	FORCEINLINE void* ContainerUObjectPtrToValuePtrInternal(UObject* ContainerPtr, int32 ArrayIndex) const
+	{
 		check(ArrayIndex < ArrayDim);
 		check(ContainerPtr);
-		return (void*)(((uint8*)ContainerPtr) + Offset_Internal + (ElementSize * ArrayIndex));
+
+		// in the future, these checks will be tested if the property is supposed be from a UClass
+		// need something for networking, since those are NOT live uobjects, just memory blocks
+		check(((UObject*)ContainerPtr)->IsValidLowLevel()); // Check its a valid UObject that was passed in
+		check(((UObject*)ContainerPtr)->GetClass() != NULL);
+		check(GetOuter()->IsA(UClass::StaticClass())); // Check that the outer of this property is a UClass (not another property)
+
+		// Check that the object we are accessing is of the class that contains this property
+		checkf(((UObject*)ContainerPtr)->IsA((UClass*)GetOuter()), TEXT("'%s' is of class '%s' however property '%s' belongs to class '%s'")
+			, *((UObject*)ContainerPtr)->GetName()
+			, *((UObject*)ContainerPtr)->GetClass()->GetName()
+			, *GetName()
+			, *((UClass*)GetOuter())->GetName());
+
+		if (0)
+		{
+			// in the future, these checks will be tested if the property is NOT relative to a UClass
+			check(!GetOuter()->IsA(UClass::StaticClass())); // Check we are _not_ calling this on a direct child property of a UClass, you should pass in a UObject* in that case
+		}
+
+		return (uint8*)ContainerPtr + Offset_Internal + ElementSize * ArrayIndex;
 	}
+
 public:
 
 	/** 
@@ -298,22 +324,22 @@ public:
 	template<typename ValueType>
 	FORCEINLINE ValueType* ContainerPtrToValuePtr(UObject* ContainerPtr, int32 ArrayIndex = 0) const
 	{
-		return (ValueType*)ContainerPtrToValuePtrInternal(ContainerPtr, ArrayIndex, true);
+		return (ValueType*)ContainerUObjectPtrToValuePtrInternal(ContainerPtr, ArrayIndex);
 	}
 	template<typename ValueType>
 	FORCEINLINE ValueType* ContainerPtrToValuePtr(void* ContainerPtr, int32 ArrayIndex = 0) const
 	{
-		return (ValueType*)ContainerPtrToValuePtrInternal(ContainerPtr, ArrayIndex, false);
+		return (ValueType*)ContainerVoidPtrToValuePtrInternal(ContainerPtr, ArrayIndex);
 	}
 	template<typename ValueType>
 	FORCEINLINE ValueType const* ContainerPtrToValuePtr(UObject const* ContainerPtr, int32 ArrayIndex = 0) const
 	{
-		return (ValueType const*)ContainerPtrToValuePtrInternal(ContainerPtr, ArrayIndex, true);
+		return ContainerPtrToValuePtr<ValueType>((UObject*)ContainerPtr, ArrayIndex);
 	}
 	template<typename ValueType>
 	FORCEINLINE ValueType const* ContainerPtrToValuePtr(void const* ContainerPtr, int32 ArrayIndex = 0) const
 	{
-		return (ValueType const*)ContainerPtrToValuePtrInternal(ContainerPtr, ArrayIndex, false);
+		return ContainerPtrToValuePtr<ValueType>((void*)ContainerPtr, ArrayIndex);
 	}
 
 	// Default variants, these accept and return NULL, and also check the property against the size of the container. 
@@ -323,7 +349,7 @@ public:
 	{
 		if (ContainerPtr && IsInContainer(ContainerClass))
 		{
-			return (ValueType*)ContainerPtrToValuePtrInternal(ContainerPtr, ArrayIndex, true);
+			return ContainerPtrToValuePtr<ValueType>(ContainerPtr, ArrayIndex);
 		}
 		return NULL;
 	}
@@ -332,7 +358,7 @@ public:
 	{
 		if (ContainerPtr && IsInContainer(ContainerClass))
 		{
-			return (ValueType*)ContainerPtrToValuePtrInternal(ContainerPtr, ArrayIndex, false);
+			return ContainerPtrToValuePtr<ValueType>(ContainerPtr, ArrayIndex);
 		}
 		return NULL;
 	}
@@ -341,7 +367,7 @@ public:
 	{
 		if (ContainerPtr && IsInContainer(ContainerClass))
 		{
-			return (ValueType const*)ContainerPtrToValuePtrInternal(ContainerPtr, ArrayIndex, true);
+			return ContainerPtrToValuePtr<ValueType>(ContainerPtr, ArrayIndex);
 		}
 		return NULL;
 	}
@@ -350,19 +376,19 @@ public:
 	{
 		if (ContainerPtr && IsInContainer(ContainerClass))
 		{
-			return (ValueType const*)ContainerPtrToValuePtrInternal(ContainerPtr, ArrayIndex, false);
+			return ContainerPtrToValuePtr<ValueType>(ContainerPtr, ArrayIndex);
 		}
 		return NULL;
 	}
 	/** See if the offset of this property is below the supplied container size */
 	FORCEINLINE bool IsInContainer(int32 ContainerSize) const
 	{
-		return (Offset_Internal + GetSize() <= ContainerSize);
+		return Offset_Internal + GetSize() <= ContainerSize;
 	}
 	/** See if the offset of this property is below the supplied container size */
 	FORCEINLINE bool IsInContainer(UStruct* ContainerClass) const
 	{
-		return (Offset_Internal + GetSize() <= (ContainerClass ? ContainerClass->GetPropertiesSize() : MAX_int32));
+		return Offset_Internal + GetSize() <= (ContainerClass ? ContainerClass->GetPropertiesSize() : MAX_int32);
 	}
 
 	/**
@@ -390,12 +416,29 @@ public:
 			}
 		}
 	}
+
+	/**
+	 * Returns the hash value for an element of this property.
+	 */
+	uint32 GetValueTypeHash(const void* Src) const
+	{
+		check(PropertyFlags & CPF_HasGetValueTypeHash); // make sure the type is hashable
+		check(Src);
+		return GetValueTypeHashInternal(Src);
+	}
+
 protected:
 	virtual void CopyValuesInternal( void* Dest, void const* Src, int32 Count  ) const
 	{
 		check(0); // if you are not memcpyable, then you need to deal with the virtual call
 	}
 
+	virtual uint32 GetValueTypeHashInternal(const void* Src) const
+	{
+		check(false); // you need to deal with the virtual call
+		return 0;
+	}
+
 public:
 	/**
 	 * Copy the value for all elements of this property.
@@ -1140,7 +1183,7 @@ public:
 	}
 
 	TProperty_Numeric( const FObjectInitializer& ObjectInitializer, ECppProperty, int32 InOffset, uint64 InFlags )
-		:	Super( ObjectInitializer, EC_CppProperty, InOffset, InFlags)
+		:	Super( ObjectInitializer, EC_CppProperty, InOffset, InFlags | CPF_HasGetValueTypeHash)
 	{
 	}
 
@@ -1154,8 +1197,14 @@ public:
 		return FString();
 	}
 
-	// UNumericProperty interface.
+	// UProperty interface
+	uint32 GetValueTypeHashInternal(const void* Src) const override
+	{
+		return GetTypeHash(*(const InTCppType*)Src);
+	}
+	// End of UProperty interface
 
+	// UNumericProperty interface.
 	virtual bool IsFloatingPoint() const override
 	{
 		return TIsFloatType<TCppType>::Value;
@@ -1203,7 +1252,6 @@ public:
 		return TTypeFundamentals::GetPropertyValue(Data);
 	}
 	// End of UNumericProperty interface
-
 };
 
 /*-----------------------------------------------------------------------------
@@ -1683,7 +1731,6 @@ class COREUOBJECT_API UObjectPropertyBase : public UProperty
 			SetObjectPropertyValue(((uint8*)Dest) + Index * ElementSize, ((UObject**)Src)[Index]);
 		}
 	}
-
 	// End of UProperty interface
 
 	// UObjectPropertyBase interface
@@ -2157,7 +2204,7 @@ public:
 	}
 
 	UNameProperty( const FObjectInitializer& ObjectInitializer, ECppProperty, int32 InOffset, uint64 InFlags )
-	:	UNameProperty_Super( ObjectInitializer, EC_CppProperty, InOffset, InFlags )
+	:	UNameProperty_Super( ObjectInitializer, EC_CppProperty, InOffset, InFlags | CPF_HasGetValueTypeHash )
 	{
 	}
 
@@ -2168,6 +2215,11 @@ public:
 	{
 		return FString();
 	}
+
+	uint32 GetValueTypeHashInternal(const void* Src) const override
+	{
+		return GetTypeHash(*(const FName*)Src);
+	}
 	// End of UProperty interface
 };
 
@@ -2195,7 +2247,7 @@ public:
 	}
 
 	UStrProperty( const FObjectInitializer& ObjectInitializer, ECppProperty, int32 InOffset, uint64 InFlags )
-	:	UStrProperty_Super( ObjectInitializer, EC_CppProperty, InOffset, InFlags)
+	:	UStrProperty_Super( ObjectInitializer, EC_CppProperty, InOffset, InFlags | CPF_HasGetValueTypeHash)
 	{
 	}
 
@@ -2206,6 +2258,11 @@ public:
 	{
 		return FString();
 	}
+
+	uint32 GetValueTypeHashInternal(const void* Src) const override
+	{
+		return GetTypeHash(*(const FString*)Src);
+	}
 	// End of UProperty interface
 };
 
@@ -2272,6 +2329,54 @@ public:
 	// End of UProperty interface
 };
 
+// need to break this out a different type so that the DECLARE_CASTED_CLASS_INTRINSIC macro can digest the comma
+typedef TProperty<FScriptMap, UProperty> UMapProperty_Super;
+
+class COREUOBJECT_API UMapProperty : public UMapProperty_Super
+{
+	DECLARE_CASTED_CLASS_INTRINSIC(UMapProperty, UMapProperty_Super, 0, CoreUObject, CASTCLASS_UMapProperty)
+
+	// Properties representing the key type and value type of the contained pairs
+	UProperty*       KeyProp;
+	UProperty*       ValueProp;
+	FScriptMapLayout MapLayout;
+
+public:
+	typedef UMapProperty_Super::TTypeFundamentals TTypeFundamentals;
+	typedef TTypeFundamentals::TCppType TCppType;
+
+	UMapProperty(const FObjectInitializer& ObjectInitializer, ECppProperty, int32 InOffset, uint64 InFlags);
+
+	// UObject interface
+	virtual void Serialize(FArchive& Ar) override;
+	static void AddReferencedObjects(UObject* InThis, FReferenceCollector& Collector);
+	// End of UObject interface
+
+	// UField interface
+	virtual void AddCppProperty(UProperty* Property) override;
+	// End of UField interface
+
+	// UProperty interface
+	virtual FString GetCPPMacroType(FString& ExtendedTypeText) const  override;
+	virtual FString GetCPPType(FString* ExtendedTypeText, uint32 CPPExportFlags) const override;
+	virtual void LinkInternal(FArchive& Ar) override;
+	virtual bool Identical(const void* A, const void* B, uint32 PortFlags) const override;
+	virtual void SerializeItem(FArchive& Ar, void* Value, void const* Defaults) const override;
+	virtual bool NetSerializeItem(FArchive& Ar, UPackageMap* Map, void* Data, TArray<uint8> * MetaData = NULL) const override;
+	virtual void ExportTextItem(FString& ValueStr, const void* PropertyValue, const void* DefaultValue, UObject* Parent, int32 PortFlags, UObject* ExportRootScope) const override;
+	virtual const TCHAR* ImportText_Internal(const TCHAR* Buffer, void* Data, int32 PortFlags, UObject* OwnerObject, FOutputDevice* ErrorText) const override;
+	virtual void CopyValuesInternal(void* Dest, void const* Src, int32 Count) const override;
+	virtual void ClearValueInternal(void* Data) const override;
+	virtual void DestroyValueInternal(void* Dest) const override;
+	virtual bool PassCPPArgsByRef() const override;
+	virtual void InstanceSubobjects(void* Data, void const* DefaultData, UObject* Owner, struct FObjectInstancingGraph* InstanceGraph) override;
+	virtual bool ContainsObjectReference() const override;
+	virtual bool ContainsWeakObjectReference() const override;
+	virtual void EmitReferenceInfo(UClass& OwnerClass, int32 BaseOffset) override;
+	virtual bool SameType(const UProperty* Other) const override;
+	// End of UProperty interface
+};
+
 /**
  * FScriptArrayHelper: Pseudo dynamic array. Used to work with array properties in a sensible way.
  **/
@@ -2544,12 +2649,12 @@ private:
 		}
 		else
 		{
-		for (int32 LoopIndex = 0 ; LoopIndex < Count; LoopIndex++, Dest += ElementSize)
-		{
-			InnerProperty->InitializeValue(Dest);
+			for (int32 LoopIndex = 0 ; LoopIndex < Count; LoopIndex++, Dest += ElementSize)
+			{
+				InnerProperty->InitializeValue(Dest);
+			}
 		}
 	}
-	}
 	/**
 	 *	Internal function to call into the property system to destruct elements.
 	 *	@param Index: first item to .
@@ -2559,17 +2664,17 @@ private:
 	{
 		if (!(InnerProperty->PropertyFlags & (CPF_IsPlainOldData | CPF_NoDestructor)))
 		{
-		checkSlow(Count > 0);
-		checkSlow(Index >= 0); 
-		checkSlow(Index < Num());
-		checkSlow(Index + Count <= Num());
-		uint8 *Dest = GetRawPtr(Index);
-		for (int32 LoopIndex = 0 ; LoopIndex < Count; LoopIndex++, Dest += ElementSize)
-		{
-			InnerProperty->DestroyValue(Dest);
+			checkSlow(Count > 0);
+			checkSlow(Index >= 0); 
+			checkSlow(Index < Num());
+			checkSlow(Index + Count <= Num());
+			uint8 *Dest = GetRawPtr(Index);
+			for (int32 LoopIndex = 0 ; LoopIndex < Count; LoopIndex++, Dest += ElementSize)
+			{
+				InnerProperty->DestroyValue(Dest);
+			}
 		}
 	}
-	}
 	/**
 	 *	Internal function to call into the property system to clear elements.
 	 *	@param Index: first item to .
@@ -2588,12 +2693,12 @@ private:
 		}
 		else
 		{
-		for (int32 LoopIndex = 0 ; LoopIndex < Count; LoopIndex++, Dest += ElementSize)
-		{
-			InnerProperty->ClearValue(Dest);
+			for (int32 LoopIndex = 0; LoopIndex < Count; LoopIndex++, Dest += ElementSize)
+			{
+				InnerProperty->ClearValue(Dest);
+			}
 		}
 	}
-	}
 
 	const UProperty* InnerProperty;
 	FScriptArray* Array;
@@ -2610,6 +2715,406 @@ public:
 };
 
 
+/**
+ * FScriptMapHelper: Pseudo dynamic map. Used to work with map properties in a sensible way.
+ */
+class FScriptMapHelper
+{
+	friend class UMapProperty;
+
+public:
+	/**
+	 * Constructor, brings together a property and an instance of the property located in memory
+	 *
+	 * @param  InProperty  The property associated with this memory
+	 * @param  InMap       Pointer to raw memory that corresponds to this map. This can be NULL, and sometimes is, but in that case almost all operations will crash.
+	 */
+	FORCEINLINE FScriptMapHelper(const UMapProperty* InProperty, const void* InMap)
+		: KeyProp         (InProperty->KeyProp)
+		, ValueProp       (InProperty->ValueProp)
+		, Map             ((FScriptMap*)InMap)  //@todo, we are casting away the const here
+		, MapLayout(InProperty->MapLayout)
+	{
+		check(KeyProp && ValueProp);
+	}
+
+	/**
+	 * Index range check
+	 *
+	 * @param  Index  Index to check
+	 *
+	 * @return true if accessing this element is legal.
+	 */
+	FORCEINLINE bool IsValidIndex(int32 Index) const
+	{
+		return Map->IsValidIndex(Index);
+	}
+
+	/**
+	 * Returns the number of elements in the map.
+	 *
+	 * @return The number of elements in the map.
+	 */
+	FORCEINLINE int32 Num() const
+	{
+		int32 Result = Map->Num();
+		checkSlow(Result >= 0); 
+		return Result;
+	}
+
+	/**
+	 * Returns the (non-inclusive) maximum index of elements in the map.
+	 *
+	 * @return The (non-inclusive) maximum index of elements in the map.
+	 */
+	FORCEINLINE int32 GetMaxIndex() const
+	{
+		int32 Result = Map->GetMaxIndex();
+		checkSlow(Result >= Num());
+		return Result;
+	}
+
+	/**
+	 * Static version of Num() used when you don't need to bother to construct a FScriptArrayHelper. Returns the number of elements in the array.
+	 *
+	 * @param  Target  Pointer to the raw memory associated with a FScriptArray
+	 *
+	 * @return The number of elements in the array.
+	 */
+	static FORCEINLINE int32 Num(const void* Target)
+	{
+		int32 Result = ((const FScriptMap*)Target)->Num();
+		checkSlow(Result >= 0); 
+		return Result;
+	}
+
+	/**
+	 * Returns a uint8 pointer to the pair in the array
+	 *
+	 * @param  Index  index of the item to return a pointer to.
+	 *
+	 * @return Pointer to the pair, or nullptr if the array is empty.
+	 */
+	FORCEINLINE uint8* GetPairPtr(int32 Index)
+	{
+		if (Num() == 0)
+		{
+			checkSlow(!Index);
+			return nullptr;
+		}
+
+		checkSlow(IsValidIndex(Index));
+		return (uint8*)Map->GetData(Index, MapLayout);
+	}
+
+	/**
+	 * Returns a uint8 pointer to the pair in the map.
+	 *
+	 * @param  Index  index of the item to return a pointer to.
+	 *
+	 * @return Pointer to the pair, or nullptr if the array is empty.
+	 */
+	FORCEINLINE const uint8* GetPairPtr(int32 Index) const
+	{
+		return const_cast<FScriptMapHelper*>(this)->GetPairPtr(Index);
+	}
+
+	/**
+	 * Add an uninitialized value to the end of the map.
+	 *
+	 * @return  The index of the added element.
+	 */
+	FORCEINLINE int32 AddUninitializedValue()
+	{
+		checkSlow(Num() >= 0);
+
+		return Map->AddUninitialized(MapLayout);
+	}
+
+	/**
+	 *	Remove all values from the map, calling destructors, etc as appropriate.
+	 *	@param Slack: used to presize the array for a subsequent add, to avoid reallocation.
+	**/
+	void EmptyValues(int32 Slack = 0)
+	{
+		checkSlow(Slack >= 0);
+
+		int32 OldNum = Num();
+		if (OldNum)
+		{
+			DestructItems(0, OldNum);
+		}
+		if (OldNum || Slack)
+		{
+			Map->Empty(Slack, MapLayout);
+		}
+	}
+
+	/**
+	 * Adds a blank, constructed value to a given size.
+	 * Note that this will create an invalid map because all the keys will be default constructed, and the map needs rehashing.
+	 *
+	 * @return  The index of the first element added.
+	 **/
+	int32 AddDefaultValue_Invalid_NeedsRehash()
+	{
+		checkSlow(Num() >= 0);
+
+		int32 Result = AddUninitializedValue();
+		ConstructItem(Result);
+
+		return Result;
+	}
+
+	/**
+	 * Returns the property representing the key of the map pair.
+	 *
+	 * @return The property representing the key of the map pair.
+	 */
+	UProperty* GetKeyProperty() const
+	{
+		return KeyProp;
+	}
+
+	/**
+	 * Returns the property representing the value of the map pair.
+	 *
+	 * @return The property representing the value of the map pair.
+	 */
+	UProperty* GetValueProperty() const
+	{
+		return ValueProp;
+	}
+
+	/**
+	 * Removes an element at the specified index, destroying it.
+	 * The map will be invalid until the next Rehash() call.
+	 *
+	 * @param  Index  The index of the element to remove.
+	 */
+	void RemoveAt_NeedsRehash(int32 Index, int32 Count = 1)
+	{
+		check(IsValidIndex(Index));
+
+		DestructItems(Index, Count);
+		for (; Count; ++Index)
+		{
+			if (IsValidIndex(Index))
+			{
+				Map->RemoveAt(Index, MapLayout);
+				--Count;
+			}
+		}
+	}
+
+	/**
+	 * Rehashes the keys in the map.
+	 * This function must be called to create a valid map.
+	 */
+	void Rehash()
+	{
+		Map->Rehash(MapLayout, [=](const void* Src) {
+			return KeyProp->GetValueTypeHash(Src);
+		});
+	}
+
+	/**
+	 * Finds the index of an element in a map which matches the key in another pair.
+	 *
+	 * @param  PairWithKeyToFind  The address of a map pair which contains the key to search for.
+	 * @param  IndexHint          The index to start searching from.
+	 *
+	 * @return The index of an element found in MapHelper, or -1 if none was found.
+	 */
+	int32 FindMapIndexWithKey(const void* PairWithKeyToFind, int32 IndexHint = 0) const
+	{
+		int32 MapMax = GetMaxIndex();
+		if (MapMax == 0)
+		{
+			return INDEX_NONE;
+		}
+
+		check(IndexHint >= 0 && IndexHint < MapMax);
+
+		UProperty* LocalKeyProp = this->KeyProp; // prevent aliasing in loop below
+
+		int32 Index = IndexHint;
+		for (;;)
+		{
+			if (IsValidIndex(Index))
+			{
+				const void* PairToSearch = GetPairPtrWithoutCheck(Index);
+				if (LocalKeyProp->Identical(PairWithKeyToFind, PairToSearch))
+				{
+					return Index;
+				}
+			}
+
+			++Index;
+			if (Index == MapMax)
+			{
+				Index = 0;
+			}
+
+			if (Index == IndexHint)
+			{
+				return INDEX_NONE;
+			}
+		}
+	}
+
+	/**
+	 * Finds the pair in a map which matches the key in another pair.
+	 *
+	 * @param  PairWithKeyToFind  The address of a map pair which contains the key to search for.
+	 * @param  IndexHint          The index to start searching from.
+	 *
+	 * @return A pointer to the found pair, or nullptr if none was found.
+	 */
+	FORCEINLINE uint8* FindMapPairPtrWithKey(const void* PairWithKeyToFind, int32 IndexHint = 0)
+	{
+		int32 Index = FindMapIndexWithKey(PairWithKeyToFind, IndexHint);
+		uint8* Result = (Index >= 0) ? GetPairPtr(Index) : nullptr;
+		return Result;
+	}
+
+private:
+	/**
+	 * Internal function to call into the property system to construct / initialize elements.
+	 *
+	 * @param  Index  First item to construct.
+	 * @param  Count  Number of items to construct.
+	 */
+	void ConstructItem(int32 Index)
+	{
+		check(IsValidIndex(Index));
+
+		bool bZeroKey   = !!(KeyProp  ->PropertyFlags & CPF_ZeroConstructor);
+		bool bZeroValue = !!(ValueProp->PropertyFlags & CPF_ZeroConstructor);
+
+		uint8* Dest = GetPairPtrWithoutCheck(Index);
+
+		if (bZeroKey || bZeroValue)
+		{
+			// If any nested property needs zeroing, just pre-zero the whole space
+			FMemory::Memzero(Dest, MapLayout.SetLayout.Size);
+		}
+
+		if (!bZeroKey)
+		{
+			KeyProp->InitializeValue_InContainer(Dest);
+		}
+
+		if (!bZeroValue)
+		{
+			ValueProp->InitializeValue_InContainer(Dest);
+		}
+	}
+
+	/**
+	 * Internal function to call into the property system to destruct elements.
+	 */
+	void DestructItems(int32 Index, int32 Count)
+	{
+		check(Index >= 0);
+		check(Count >= 0);
+
+		if (Count == 0)
+		{
+			return;
+		}
+
+		bool bDestroyKeys   = !(KeyProp  ->PropertyFlags & (CPF_IsPlainOldData | CPF_NoDestructor));
+		bool bDestroyValues = !(ValueProp->PropertyFlags & (CPF_IsPlainOldData | CPF_NoDestructor));
+
+		if (bDestroyKeys || bDestroyValues)
+		{
+			uint32 Stride  = MapLayout.SetLayout.Size;
+			uint8* PairPtr = GetPairPtr(Index);
+			if (bDestroyKeys)
+			{
+				if (bDestroyValues)
+				{
+					for (; Count; ++Index)
+					{
+						if (IsValidIndex(Index))
+						{
+							KeyProp  ->DestroyValue_InContainer(PairPtr);
+							ValueProp->DestroyValue_InContainer(PairPtr);
+							--Count;
+						}
+						PairPtr += Stride;
+					}
+				}
+				else
+				{
+					for (; Count; ++Index)
+					{
+						if (IsValidIndex(Index))
+						{
+							KeyProp->DestroyValue_InContainer(PairPtr);
+							--Count;
+						}
+						PairPtr += Stride;
+					}
+				}
+			}
+			else
+			{
+				for (; Count; ++Index)
+				{
+					if (IsValidIndex(Index))
+					{
+						ValueProp->DestroyValue_InContainer(PairPtr);
+						--Count;
+					}
+					PairPtr += Stride;
+				}
+			}
+		}
+	}
+
+	/**
+	 * Returns a uint8 pointer to the pair in the array without checking the index.
+	 *
+	 * @param  Index  index of the item to return a pointer to.
+	 *
+	 * @return Pointer to the pair, or nullptr if the array is empty.
+	 */
+	FORCEINLINE uint8* GetPairPtrWithoutCheck(int32 Index)
+	{
+		return (uint8*)Map->GetData(Index, MapLayout);
+	}
+
+	/**
+	 * Returns a uint8 pointer to the pair in the array without checking the index.
+	 *
+	 * @param  Index  index of the item to return a pointer to.
+	 *
+	 * @return Pointer to the pair, or nullptr if the array is empty.
+	 */
+	FORCEINLINE const uint8* GetPairPtrWithoutCheck(int32 Index) const
+	{
+		return const_cast<FScriptMapHelper*>(this)->GetPairPtrWithoutCheck(Index);
+	}
+
+public:
+	UProperty*       KeyProp;
+	UProperty*       ValueProp;
+	FScriptMap*      Map;
+	FScriptMapLayout MapLayout;
+};
+
+class FScriptMapHelper_InContainer : public FScriptMapHelper
+{
+public:
+	FORCEINLINE FScriptMapHelper_InContainer(const UMapProperty* InProperty, const void* InArray, int32 FixedArrayIndex=0)
+		:FScriptMapHelper(InProperty, InProperty->ContainerPtrToValuePtr<void>(InArray, FixedArrayIndex))
+	{
+	}
+};
+
+
 /*-----------------------------------------------------------------------------
 	UStructProperty.
 -----------------------------------------------------------------------------*/
