Setting bShouldIntializePropsFromArchetype in FObjectInitializer default constructor to false by default (it should not attempt to initialize props from archetype for native objects).

diff --git a/Engine/Source/Runtime/CoreUObject/Private/UObject/UObjectGlobals.cpp b/Engine/Source/Runtime/CoreUObject/Private/UObject/UObjectGlobals.cpp
--- a/Engine/Source/Runtime/CoreUObject/Private/UObject/UObjectGlobals.cpp
+++ b/Engine/Source/Runtime/CoreUObject/Private/UObject/UObjectGlobals.cpp
@@ -1976,7 +1976,7 @@ FObjectInitializer::FObjectInitializer() :
 Obj(nullptr),
 	ObjectArchetype(nullptr),
 	bCopyTransientsFromClassDefaults(false),
-	bShouldIntializePropsFromArchetype(true),
+	bShouldIntializePropsFromArchetype(false),
 	bSubobjectClassInitializationAllowed(true),
 	InstanceGraph(nullptr),
 	LastConstructedObject(nullptr)
