Fixing IsInAsyncLoadingThread to be consisent with the old single-threaded path.

diff --git a/Engine/Source/Runtime/CoreUObject/Private/Serialization/AsyncLoading.cpp b/Engine/Source/Runtime/CoreUObject/Private/Serialization/AsyncLoading.cpp
--- a/Engine/Source/Runtime/CoreUObject/Private/Serialization/AsyncLoading.cpp
+++ b/Engine/Source/Runtime/CoreUObject/Private/Serialization/AsyncLoading.cpp
@@ -212,6 +212,29 @@ FAsyncLoadingThread& FAsyncLoadingThread::Get()
 	return GAsyncLoader;
 }
 
+/** Just like TGuardValue for FAsyncLoadingThread::bIsInAsyncLoadingTick but only works for the game thread */
+struct FAsyncLoadingTickScope
+{
+	bool bWasInTick;
+	FAsyncLoadingTickScope()
+		: bWasInTick(false)
+	{
+		if (IsInGameThread())
+		{
+			FAsyncLoadingThread& AsyncLoadingThread = FAsyncLoadingThread::Get();
+			bWasInTick = AsyncLoadingThread.GetIsInAsyncLoadingTick();
+			AsyncLoadingThread.SetIsInAsyncLoadingTick(true);
+		}
+	}
+	~FAsyncLoadingTickScope()
+	{
+		if (IsInGameThread())
+		{
+			FAsyncLoadingThread::Get().SetIsInAsyncLoadingTick(bWasInTick);
+		}
+	}
+};
+
 void FAsyncLoadingThread::InitializeAsyncThread()
 {
 	AsyncThreadReady.Increment();
@@ -277,6 +300,8 @@ int32 FAsyncLoadingThread::CreateAsyncPackagesFromQueue()
 {
 	SCOPE_CYCLE_COUNTER(STAT_FAsyncPackage_CreateAsyncPackagesFromQueue);
 
+	FAsyncLoadingTickScope InAsyncLoadingTick;
+
 	int32 NumCreated = 0;
 	checkSlow(IsInAsyncLoadThread());
 
@@ -459,9 +484,6 @@ EAsyncPackageState::Type FAsyncLoadingThread::ProcessLoadedPackages(bool bUseTim
 
 EAsyncPackageState::Type FAsyncLoadingThread::TickAsyncLoading(bool bUseTimeLimit, bool bUseFullTimeLimit, float TimeLimit, FName ExcludeType)
 {
-	TGuardValue<bool> InAsyncLoadingTickGuard(bIsInAsyncLoadingTick, true);
-
-	static bool bWasAsyncLoading = false;
 	EAsyncPackageState::Type Result = EAsyncPackageState::Complete;	
 	double TickStartTime = FPlatformTime::Seconds();
 	double TimeLimitUsedForPreloading = 0;
@@ -763,6 +785,11 @@ bool FAsyncPackage::GiveUpTimeSlice()
  */
 void FAsyncPackage::BeginAsyncLoad()
 {
+	if (IsInGameThread())
+	{
+		FAsyncLoadingThread::Get().SetIsInAsyncLoadingTick(true);
+	}
+
 	// this won't do much during async loading except increase the load count which causes IsLoading to return true
 	BeginLoad();
 }
@@ -778,6 +805,11 @@ void FAsyncPackage::EndAsyncLoad()
 	// this won't do much during async loading except decrease the load count which causes IsLoading to return false
 	EndLoad();
 
+	if (IsInGameThread())
+	{
+		FAsyncLoadingThread::Get().SetIsInAsyncLoadingTick(false);
+	}
+
 	if (!bLoadHasFailed)
 	{
 		// Mark the package as loaded, if we succeeded
@@ -1409,6 +1441,7 @@ EAsyncPackageState::Type FAsyncPackage::PostLoadDeferredObjects(double InTickSta
 
 	EAsyncPackageState::Type Result = EAsyncPackageState::Complete;
 	TGuardValue<bool> GuardIsRoutingPostLoad(FUObjectThreadContext::Get().IsRoutingPostLoad, true);
+	FAsyncLoadingTickScope InAsyncLoadingTick;
 
 	LastObjectWorkWasPerformedOn = nullptr;
 	LastTypeOfWorkPerformed = TEXT("postloading_gamethread");
diff --git a/Engine/Source/Runtime/CoreUObject/Private/Serialization/AsyncLoadingThread.h b/Engine/Source/Runtime/CoreUObject/Private/Serialization/AsyncLoadingThread.h
--- a/Engine/Source/Runtime/CoreUObject/Private/Serialization/AsyncLoadingThread.h
+++ b/Engine/Source/Runtime/CoreUObject/Private/Serialization/AsyncLoadingThread.h
@@ -97,7 +97,8 @@ public:
 					check(GConfig);
 					bool bConfigValue = true;
 					GConfig->GetBool(TEXT("Core.System"), TEXT("AsyncLoadingThreadEnabled"), bConfigValue, GEngineIni);
-					Value = bConfigValue && FPlatformProcess::SupportsMultithreading();
+					bool bCommandLineNoAsyncThread = false;
+					Value = bConfigValue && FApp::ShouldUseThreadingForPerformance() && !FParse::Param(FCommandLine::Get(), TEXT("NoAsyncLoadingThread"));
 				}
 				else
 				{
@@ -108,6 +109,18 @@ public:
 		return AsyncLoadingThreadEnabled.Value;
 	}
 
+	/** Sets the current state of async loading */
+	FORCEINLINE void SetIsInAsyncLoadingTick(bool InTick)
+	{
+		bIsInAsyncLoadingTick = InTick;
+	}
+
+	/** Gets the current state of async loading */
+	FORCEINLINE bool GetIsInAsyncLoadingTick() const
+	{
+		return bIsInAsyncLoadingTick;
+	}
+
 	/** Returns true if packages are currently being loaded on the async thread */
 	FORCEINLINE bool IsAsyncLoadingPackages()
 	{
@@ -121,7 +134,11 @@ public:
 		bool bResult = false;
 		if (IsMultithreaded())
 		{
-			bResult = FPlatformTLS::GetCurrentThreadId() == AsyncLoadingThreadID;
+			// We still need to report we're in async loading thread even if 
+			// we're on game thread but inside of async loading code (PostLoad mostly)
+			// to make it behave exactly like the non-threaded version
+			bResult = FPlatformTLS::GetCurrentThreadId() == AsyncLoadingThreadID || 
+				(IsInGameThread() && FAsyncLoadingThread::Get().bIsInAsyncLoadingTick);
 		}
 		else
 		{
