Possible link error fix.
Comment fixed.

diff --git a/Engine/Source/Runtime/Core/Public/Containers/BitArray.h b/Engine/Source/Runtime/Core/Public/Containers/BitArray.h
--- a/Engine/Source/Runtime/Core/Public/Containers/BitArray.h
+++ b/Engine/Source/Runtime/Core/Public/Containers/BitArray.h
@@ -880,8 +880,8 @@ private:
 };
 
 
-// Untyped set type for accessing TSet data, like FScriptArray for TArray.
-// Must have the same memory representation as a TSet.
+// Untyped bit array type for accessing TBitArray data, like FScriptArray for TArray.
+// Must have the same memory representation as a TBitArray.
 class FScriptBitArray
 {
 public:
diff --git a/Engine/Source/Runtime/CoreUObject/Private/UObject/Property.cpp b/Engine/Source/Runtime/CoreUObject/Private/UObject/Property.cpp
--- a/Engine/Source/Runtime/CoreUObject/Private/UObject/Property.cpp
+++ b/Engine/Source/Runtime/CoreUObject/Private/UObject/Property.cpp
@@ -1175,6 +1175,17 @@ const TCHAR* UProperty::ImportSingleProperty( const TCHAR* Str, void* DestData,
 	return Str;
 }
 
+/**
+ * Returns the hash value for an element of this property.
+ */
+uint32 UProperty::GetValueTypeHash(const void* Src) const
+{
+	check(PropertyFlags & CPF_HasGetValueTypeHash); // make sure the type is hashable
+	check(Src);
+	return GetValueTypeHashInternal(Src);
+}
+
+
 IMPLEMENT_CORE_INTRINSIC_CLASS(UProperty, UField,
 	{
 	}
diff --git a/Engine/Source/Runtime/CoreUObject/Public/UObject/UnrealType.h b/Engine/Source/Runtime/CoreUObject/Public/UObject/UnrealType.h
--- a/Engine/Source/Runtime/CoreUObject/Public/UObject/UnrealType.h
+++ b/Engine/Source/Runtime/CoreUObject/Public/UObject/UnrealType.h
@@ -420,12 +420,7 @@ public:
 	/**
 	 * Returns the hash value for an element of this property.
 	 */
-	uint32 GetValueTypeHash(const void* Src) const
-	{
-		check(PropertyFlags & CPF_HasGetValueTypeHash); // make sure the type is hashable
-		check(Src);
-		return GetValueTypeHashInternal(Src);
-	}
+	uint32 GetValueTypeHash(const void* Src) const;
 
 protected:
 	virtual void CopyValuesInternal( void* Dest, void const* Src, int32 Count  ) const
