Weird linker fix.

diff --git a/Engine/Source/Runtime/CoreUObject/Private/UObject/PropertyMap.cpp b/Engine/Source/Runtime/CoreUObject/Private/UObject/PropertyMap.cpp
--- a/Engine/Source/Runtime/CoreUObject/Private/UObject/PropertyMap.cpp
+++ b/Engine/Source/Runtime/CoreUObject/Private/UObject/PropertyMap.cpp
@@ -789,3 +789,11 @@ IMPLEMENT_CORE_INTRINSIC_CLASS(UMapProperty, UProperty,
 		static_assert(sizeof(FScriptMap) == sizeof(TMap<uint32, uint8>), "FScriptMap and TMap<uint32, uint8> must be interchangable.");
 	}
 );
+
+void FScriptMapHelper::Rehash()
+{
+	// Moved out-of-line to maybe fix a weird link error
+	Map->Rehash(MapLayout, [=](const void* Src) {
+		return KeyProp->GetValueTypeHash(Src);
+	});
+}
diff --git a/Engine/Source/Runtime/CoreUObject/Public/UObject/UnrealType.h b/Engine/Source/Runtime/CoreUObject/Public/UObject/UnrealType.h
--- a/Engine/Source/Runtime/CoreUObject/Public/UObject/UnrealType.h
+++ b/Engine/Source/Runtime/CoreUObject/Public/UObject/UnrealType.h
@@ -2906,12 +2906,7 @@ public:
 	 * Rehashes the keys in the map.
 	 * This function must be called to create a valid map.
 	 */
-	void Rehash()
-	{
-		Map->Rehash(MapLayout, [=](const void* Src) {
-			return KeyProp->GetValueTypeHash(Src);
-		});
-	}
+	COREUOBJECT_API void Rehash();
 
 	/**
 	 * Finds the index of an element in a map which matches the key in another pair.
