Reworked cooker patch warning so that there is only valid patching warnings.

diff --git a/Engine/Source/Runtime/CoreUObject/Private/UObject/SavePackage.cpp b/Engine/Source/Runtime/CoreUObject/Private/UObject/SavePackage.cpp
--- a/Engine/Source/Runtime/CoreUObject/Private/UObject/SavePackage.cpp
+++ b/Engine/Source/Runtime/CoreUObject/Private/UObject/SavePackage.cpp
@@ -2380,6 +2380,10 @@ EObjectMark UPackage::GetObjectMarksForTargetPlatform( const class ITargetPlatfo
 #define UE_FINISH_LOG_COOK_TIME()
 #endif
 
+#if WITH_EDITOR
+COREUOBJECT_API extern bool GOutputCookingWarnings;
+#endif
+
 bool UPackage::SavePackage( UPackage* InOuter, UObject* Base, EObjectFlags TopLevelFlags, const TCHAR* Filename, 
 	FOutputDevice* Error, FLinkerLoad* Conform, bool bForceByteSwapping, bool bWarnOfLongFilename, uint32 SaveFlags, 
 	const class ITargetPlatform* TargetPlatform, const FDateTime&  FinalTimeStamp, bool bSlowTask)
@@ -3012,6 +3016,25 @@ bool UPackage::SavePackage( UPackage* InOuter, UObject* Base, EObjectFlags TopLe
 				Linker->Summary.NameOffset = Linker->Tell();
 				SavePackageState->UpdateLinkerWithMarkedNames(Linker);
 
+#if WITH_EDITOR
+				if ( GOutputCookingWarnings )
+				{
+					// check the name list for uniqueobjectnamefor cooking
+					static FName NAME_UniqueObjectNameForCooking(TEXT("UniqueObjectNameForCooking"));
+
+					for (const auto& NameInUse : Linker->NameMap)
+					{
+						if (NameInUse.GetComparisonIndex() == NAME_UniqueObjectNameForCooking.GetComparisonIndex())
+						{
+							//UObject *Object = FindObject<UObject>( ANY_PACKAGE, *NameInUse.ToString());
+
+							// error
+							// check(Object);
+							UE_LOG(LogSavePackage, Warning, TEXT("Saving object into cooked package which was created at cook time, Object Name %s"), *NameInUse.ToString());
+						}
+					}
+				}
+#endif
 
 				UE_LOG_COOK_TIME(TEXT("Serialize Summary"));
 
@@ -3076,6 +3099,26 @@ bool UPackage::SavePackage( UPackage* InOuter, UObject* Base, EObjectFlags TopLe
 					}
 				}
 
+
+
+#if WITH_EDITOR
+				if (GOutputCookingWarnings)
+				{
+					// check the name list for uniqueobjectnamefor cooking
+					static FName NAME_UniqueObjectNameForCooking(TEXT("UniqueObjectNameForCooking"));
+
+					for (const auto& Export : Linker->ExportMap)
+					{
+						const auto& NameInUse = Export.ObjectName;
+						if (NameInUse.GetComparisonIndex() == NAME_UniqueObjectNameForCooking.GetComparisonIndex())
+						{
+							UE_LOG(LogSavePackage, Warning, TEXT("Saving object into cooked package which was created at cook time, Object Name %s, Full Path %s, Class %s"), *NameInUse.ToString(), *Export.Object->GetFullName(), *Export.Object->GetClass()->GetName());
+						}
+					}
+				}
+#endif
+
+
 				UE_LOG_COOK_TIME(TEXT("Build Export Map"));
 
 				if ( EndSavingIfCancelled( Linker, TempFilename ) ) { return false; }
diff --git a/Engine/Source/Runtime/CoreUObject/Private/UObject/UObjectGlobals.cpp b/Engine/Source/Runtime/CoreUObject/Private/UObject/UObjectGlobals.cpp
--- a/Engine/Source/Runtime/CoreUObject/Private/UObject/UObjectGlobals.cpp
+++ b/Engine/Source/Runtime/CoreUObject/Private/UObject/UObjectGlobals.cpp
@@ -1772,13 +1772,16 @@ UObject* StaticAllocateObject
 	UObject* Obj = NULL;
 	if(InName == NAME_None)
 	{
-		InName = MakeUniqueObjectName( InOuter, InClass );
 #if WITH_EDITOR
+		static FName NAME_UniqueObjectNameForCooking(TEXT("UniqueObjectNameForCooking"));
 		if ( GOutputCookingWarnings && GetTransientPackage() != InOuter->GetOutermost() )
 		{
-			UE_LOG(LogUObjectGlobals, Warning, TEXT("Shouldn't create dynamic objects whilest cooking, if they are saved, they will be saved with unique names, this will damage patch sizes on release.  Package %s, Object %s."), *InOuter->GetOutermost()->GetName(), *InName.ToString() );
+			InName = MakeUniqueObjectName(InOuter, InClass, NAME_UniqueObjectNameForCooking);
 		}
+		else
 #endif
+		InName = MakeUniqueObjectName(InOuter, InClass);
+
 	}
 	else
 	{
