'Should serialize' expression rearranged for better short-circuit evaluation.

diff --git a/Engine/Source/Runtime/CoreUObject/Private/UObject/Class.cpp b/Engine/Source/Runtime/CoreUObject/Private/UObject/Class.cpp
--- a/Engine/Source/Runtime/CoreUObject/Private/UObject/Class.cpp
+++ b/Engine/Source/Runtime/CoreUObject/Private/UObject/Class.cpp
@@ -1294,8 +1294,7 @@ void UStruct::SerializeTaggedProperties(FArchive& Ar, uint8* Data, UStruct* Defa
 				{
 					uint8* DataPtr      = Property->ContainerPtrToValuePtr           <uint8>(Data, Idx);
 					uint8* DefaultValue = Property->ContainerPtrToValuePtrForDefaults<uint8>(DefaultsStruct, Defaults, Idx);
-					if( (!dynamic_cast<const UClass*>(this) && !Defaults) || !Ar.DoDelta() || 
-						!Property->Identical( DataPtr, DefaultValue, Ar.GetPortFlags()) || Ar.IsTransacting() )
+					if( !Ar.DoDelta() || Ar.IsTransacting() || (!Defaults && !dynamic_cast<const UClass*>(this)) || !Property->Identical( DataPtr, DefaultValue, Ar.GetPortFlags()) )
 					{
 						if (bUseAtomicSerialization)
 						{
