IsA optimisations.

diff --git a/Engine/Source/Runtime/CoreUObject/Private/UObject/Class.cpp b/Engine/Source/Runtime/CoreUObject/Private/UObject/Class.cpp
--- a/Engine/Source/Runtime/CoreUObject/Private/UObject/Class.cpp
+++ b/Engine/Source/Runtime/CoreUObject/Private/UObject/Class.cpp
@@ -420,6 +420,7 @@ IMPLEMENT_CORE_INTRINSIC_CLASS(UField, UObject,
 //
 UStruct::UStruct( EStaticConstructor, int32 InSize, EObjectFlags InFlags )
 :	UField			( EC_StaticConstructor, InFlags )
+,	SuperStruct		( nullptr )
 ,	Children		( NULL )
 ,	PropertiesSize	( InSize )
 ,	MinAlignment	( 1 )
@@ -2958,10 +2959,174 @@ void UClass::Link(FArchive& Ar, bool bRelinkExistingProperties)
 	}
 }
 
+#if UCLASS_FAST_ISA_IMPL & 2
+
+	/**
+	 * Tree for fast IsA implementation.
+	 *
+	 * Structure is:
+	 * - every class is located at index Class->ClassTreeIndex.
+	 * - the Class->ClassTreeNumChildren classes immediately following each class are the children of the class.
+	 */
+	class FFastIndexingClassTree
+	{
+		friend class UClass;
+		friend class FFastIndexingClassTreeRegistrar;
+
+		static void Register(UClass* NewClass);
+		static void Unregister(UClass* NewClass);
+
+		static TArray<UClass*>                Classes;
+		static TMap<UClass*, TArray<UClass*>> Orphans;
+		static FCriticalSection               ClassesCriticalSection;
+	};
+
+	void FFastIndexingClassTree::Register(UClass* Class)
+	{
+		FScopeLock Lock(&ClassesCriticalSection);
+
+		UClass* ParentClass = Class->GetSuperClass();
+
+		// If the parent has previously been orphaned, flag the child as orphaned
+		if (TArray<UClass*>* ParentOrphans = Orphans.Find(ParentClass))
+		{
+			ParentOrphans->Add(Class);
+			return;
+		}
+
+		int32 NewIndex;
+		if (ParentClass)
+		{
+			// Can happen if a child is registered *after* the parent
+			if (!Classes.Contains(ParentClass))
+			{
+				Orphans.Add(ParentClass).Add(Class);
+				return;
+			}
+
+			NewIndex = ParentClass->ClassTreeIndex + ParentClass->ClassTreeNumChildren + 1;
+		}
+		else
+		{
+			NewIndex = Classes.Num();
+		}
+
+		// Increment indices of following classes
+		for (auto Index = NewIndex, LastIndex = Classes.Num(); Index != LastIndex; ++Index)
+		{
+			++Classes[Index]->ClassTreeIndex;
+		}
+
+		// Update children count of all parents
+		for (auto* Parent = ParentClass; Parent; Parent = Parent->GetSuperClass())
+		{
+			++Parent->ClassTreeNumChildren;
+		}
+
+		// Add class
+		Class->ClassTreeIndex       = NewIndex;
+		Class->ClassTreeNumChildren = 0;
+		Classes.Insert(Class, NewIndex);
+
+		// Re-register any children orphaned by a previous Unregister call
+		if (TArray<UClass*>* FoundOrphans = Orphans.Find(Class))
+		{
+			TArray<UClass*> OrphansToReregister = MoveTemp(*FoundOrphans);
+			Orphans.Remove(Class);
+			Orphans.Compact();
+
+			for (UClass* Orphan : OrphansToReregister)
+			{
+				Register(Orphan);
+			}
+		}
+	}
+
+	void FFastIndexingClassTree::Unregister(UClass* Class)
+	{
+		FScopeLock Lock(&ClassesCriticalSection);
+
+		UClass* ParentClass = Class->GetSuperClass();
+
+		// Remove class if it was orphaned
+		if (TArray<UClass*>* Children = Orphans.Find(ParentClass))
+		{
+			// Remove the child, or the entire array if it's the only child left.
+			if (Children->Num() != 1)
+			{
+				int32 ChildIndex = Children->Find(Class);
+				check(ChildIndex != INDEX_NONE);
+				Children->RemoveAt(ChildIndex);
+			}
+			else
+			{
+				check(Children->Last() == Class);
+				Orphans.Remove(ParentClass);
+				Orphans.Compact();
+			}
+
+			return;
+		}
+
+		// Remove it and mark its children as orphaned
+		int32 ClassIndex       = Class->ClassTreeIndex;
+		int32 ClassNumChildren = Class->ClassTreeNumChildren;
+		int32 NumRemoved       = ClassNumChildren + 1;
+
+		// Mark any children as orphaned
+		for (int32 Index = ClassIndex + 1, EndIndex = ClassIndex + NumRemoved; Index != EndIndex; ++Index)
+		{
+			Orphans.FindOrAdd(Classes[Index]->GetSuperClass()).Add(Classes[Index]);
+		}
+
+		// Decrement indices of following classes
+		for (int32 Index = ClassIndex + NumRemoved, IndexEnd = Classes.Num(); Index != IndexEnd; ++Index)
+		{
+			Classes[Index]->ClassTreeIndex -= NumRemoved;
+		}
+
+		// Update children count of all parents
+		for (auto* Parent = ParentClass; Parent; Parent = Parent->GetSuperClass())
+		{
+			Parent->ClassTreeNumChildren -= NumRemoved;
+		}
+
+		Classes.RemoveAt(ClassIndex, NumRemoved, false);
+	}
+
+	TArray<UClass*>                FFastIndexingClassTree::Classes;
+	TMap<UClass*, TArray<UClass*>> FFastIndexingClassTree::Orphans;
+	FCriticalSection               FFastIndexingClassTree::ClassesCriticalSection;
+
+	FFastIndexingClassTreeRegistrar::FFastIndexingClassTreeRegistrar()
+	{
+		ClassTreeIndex = -1;
+		FFastIndexingClassTree::Register((UClass*)this);
+	}
+
+	FFastIndexingClassTreeRegistrar::FFastIndexingClassTreeRegistrar(const FFastIndexingClassTreeRegistrar&)
+	{
+		ClassTreeIndex = -1;
+		FFastIndexingClassTree::Register((UClass*)this);
+	}
+
+	FFastIndexingClassTreeRegistrar::~FFastIndexingClassTreeRegistrar()
+	{
+		FFastIndexingClassTree::Unregister((UClass*)this);
+	}
+
+#endif
+
 void UClass::SetSuperStruct(UStruct* NewSuperStruct)
 {
 	UnhashObject(this);
+#if UCLASS_FAST_ISA_IMPL & 2
+	FFastIndexingClassTree::Unregister(this);
+#endif
 	Super::SetSuperStruct(NewSuperStruct);
+#if UCLASS_FAST_ISA_IMPL & 2
+	FFastIndexingClassTree::Register(this);
+#endif
 	HashObject(this);
 }
 
@@ -2971,12 +3136,18 @@ void UClass::Serialize( FArchive& Ar )
 	{
 		// Rehash since SuperStruct will be serialized in UStruct::Serialize
 		UnhashObject(this);
+#if UCLASS_FAST_ISA_IMPL & 2
+	FFastIndexingClassTree::Unregister(this);
+#endif
 	}
 
 	Super::Serialize( Ar );
 
 	if ( Ar.IsLoading() || Ar.IsModifyingWeakAndStrongReferences() )
 	{
+#if UCLASS_FAST_ISA_IMPL & 2
+	FFastIndexingClassTree::Register(this);
+#endif
 		HashObject(this);
 	}
 
diff --git a/Engine/Source/Runtime/CoreUObject/Private/UObject/UObjectBaseUtility.cpp b/Engine/Source/Runtime/CoreUObject/Private/UObject/UObjectBaseUtility.cpp
--- a/Engine/Source/Runtime/CoreUObject/Private/UObject/UObjectBaseUtility.cpp
+++ b/Engine/Source/Runtime/CoreUObject/Private/UObject/UObjectBaseUtility.cpp
@@ -252,14 +252,31 @@ bool UObjectBaseUtility::IsA( const UClass* SomeBase ) const
 {
 	UE_CLOG(!SomeBase, LogObj, Fatal, TEXT("IsA(NULL) cannot yield meaningful results"));
 
-	for ( UClass* TempClass=GetClass(); TempClass; TempClass=TempClass->GetSuperClass() )
-	{
-		if ( TempClass == SomeBase )
+	#if UCLASS_FAST_ISA_IMPL & 1
+		bool bOldResult = false;
+		for ( UClass* TempClass=GetClass(); TempClass; TempClass=TempClass->GetSuperClass() )
 		{
-			return true;
+			if ( TempClass == SomeBase )
+			{
+				bOldResult = true;
+				break;
+			}
 		}
-	}
-	return false;
+	#endif
+
+	#if UCLASS_FAST_ISA_IMPL & 2
+		bool bNewResult = GetClass()->IsAUsingFastTree(*SomeBase);
+	#endif
+
+	#if (UCLASS_FAST_ISA_IMPL & 1) && (UCLASS_FAST_ISA_IMPL & 2)
+		ensureOnceMsgf(bOldResult == bNewResult, TEXT("New cast code failed"));
+	#endif
+
+	#if UCLASS_FAST_ISA_IMPL & 1
+		return bOldResult;
+	#else
+		return bNewResult;
+	#endif
 }
 
 
diff --git a/Engine/Source/Runtime/CoreUObject/Public/Templates/Casts.h b/Engine/Source/Runtime/CoreUObject/Public/Templates/Casts.h
--- a/Engine/Source/Runtime/CoreUObject/Public/Templates/Casts.h
+++ b/Engine/Source/Runtime/CoreUObject/Public/Templates/Casts.h
@@ -71,7 +71,11 @@ struct TCastFlags
 template <typename From, typename To, bool bFromInterface = TIsIInterface<From>::Value, bool bToInterface = TIsIInterface<To>::Value, uint64 CastClass = TCastFlags<To>::Value>
 struct TGetCastType
 {
+#if UCLASS_FAST_ISA_IMPL == 2
+	static const ECastType Value = ECastType::UObjectToUObject;
+#else
 	static const ECastType Value = ECastType::FromCastFlags;
+#endif
 };
 
 template <typename From, typename To                  > struct TGetCastType<From, To, false, false, CASTCLASS_None> { static const ECastType Value = ECastType::UObjectToUObject;     };
diff --git a/Engine/Source/Runtime/CoreUObject/Public/UObject/Class.h b/Engine/Source/Runtime/CoreUObject/Public/UObject/Class.h
--- a/Engine/Source/Runtime/CoreUObject/Public/UObject/Class.h
+++ b/Engine/Source/Runtime/CoreUObject/Public/UObject/Class.h
@@ -8,6 +8,15 @@
 
 #include "ObjectBase.h"
 
+// 1 = old behavior
+// 2 = new behavior
+// 3 = old behavior with checks against the new behavior
+#if !(UE_BUILD_SHIPPING || UE_BUILD_TEST)
+	#define UCLASS_FAST_ISA_IMPL 3
+#else
+	#define UCLASS_FAST_ISA_IMPL 2
+#endif
+
 /*-----------------------------------------------------------------------------
 	Mirrors of mirror structures in Object.h. These are used by generated code 
 	to facilitate correct offsets and alignments for structures containing these '
@@ -1777,16 +1786,50 @@ namespace EIncludeSuperFlag
 	};
 }
 
+
+#if UCLASS_FAST_ISA_IMPL & 2
+
+	class FFastIndexingClassTreeRegistrar
+	{
+	public:
+		FFastIndexingClassTreeRegistrar();
+		FFastIndexingClassTreeRegistrar(const FFastIndexingClassTreeRegistrar&);
+		~FFastIndexingClassTreeRegistrar();
+
+	private:
+		friend class UClass;
+		friend class UObjectBaseUtility;
+		friend class FFastIndexingClassTree;
+
+		bool IsAUsingFastTree(const FFastIndexingClassTreeRegistrar& Parent) const
+		{
+			return ClassTreeIndex - Parent.ClassTreeIndex <= Parent.ClassTreeNumChildren;
+		}
+
+		FFastIndexingClassTreeRegistrar& operator=(const FFastIndexingClassTreeRegistrar&);
+
+		uint32 ClassTreeIndex;
+		uint32 ClassTreeNumChildren;
+	};
+
+#endif
+
+
 /**
  * An object class.
  */
 class COREUOBJECT_API UClass : public UStruct
+#if UCLASS_FAST_ISA_IMPL & 2
+	, private FFastIndexingClassTreeRegistrar
+#endif
 {
 	DECLARE_CASTED_CLASS_INTRINSIC_NO_CTOR(UClass,UStruct,0,CoreUObject,CASTCLASS_UClass,NO_API)
 	DECLARE_WITHIN(UPackage)
 
 public:
-
+#if UCLASS_FAST_ISA_IMPL & 2
+	friend class FFastIndexingClassTree;
+#endif
 	friend class FRestoreClassInfo;
 
 	typedef void		(*ClassConstructorType)				(const FObjectInitializer&);
@@ -2277,6 +2320,12 @@ public:
 
 	virtual UObject* GetArchetypeForCDO() const;
 private:
+	#if UCLASS_FAST_ISA_IMPL & 2
+		// For UObjectBaseUtility
+		friend class UObjectBaseUtility;
+		using FFastIndexingClassTreeRegistrar::IsAUsingFastTree;
+	#endif
+
 	// This signature intentionally hides the method declared in UObjectBaseUtility to make it private.
 	// Call IsChildOf instead; Hidden because calling IsA on a class almost always indicates an error where the caller should use IsChildOf.
 	bool IsA(const UClass* Parent) const
