Fix for ShooterGame cook crash.

diff --git a/Engine/Source/Runtime/CoreUObject/Private/UObject/Class.cpp b/Engine/Source/Runtime/CoreUObject/Private/UObject/Class.cpp
--- a/Engine/Source/Runtime/CoreUObject/Private/UObject/Class.cpp
+++ b/Engine/Source/Runtime/CoreUObject/Private/UObject/Class.cpp
@@ -1341,7 +1341,7 @@ void UStruct::Serialize( FArchive& Ar )
 {
 	Super::Serialize( Ar );
 
-	Ar << SuperStruct;
+	SerializeSuperStruct(Ar);
 	Ar << Children;
 
 	if (Ar.IsLoading())
@@ -1490,6 +1490,11 @@ void UStruct::SetSuperStruct(UStruct* NewSuperStruct)
 	SuperStruct = NewSuperStruct;
 }
 
+void UStruct::SerializeSuperStruct(FArchive& Ar)
+{
+	Ar << SuperStruct;
+}
+
 #if WITH_EDITOR
 bool UStruct::GetBoolMetaDataHierarchical(const FName& Key) const
 {
@@ -3130,24 +3135,29 @@ void UClass::SetSuperStruct(UStruct* NewSuperStruct)
 	HashObject(this);
 }
 
+void UClass::SerializeSuperStruct(FArchive& Ar)
+{
+#if UCLASS_FAST_ISA_IMPL & 2
+	FFastIndexingClassTree::Unregister(this);
+#endif
+	Super::SerializeSuperStruct(Ar);
+#if UCLASS_FAST_ISA_IMPL & 2
+	FFastIndexingClassTree::Register(this);
+#endif
+}
+
 void UClass::Serialize( FArchive& Ar )
 {
 	if ( Ar.IsLoading() || Ar.IsModifyingWeakAndStrongReferences() )
 	{
 		// Rehash since SuperStruct will be serialized in UStruct::Serialize
 		UnhashObject(this);
-#if UCLASS_FAST_ISA_IMPL & 2
-	FFastIndexingClassTree::Unregister(this);
-#endif
 	}
 
 	Super::Serialize( Ar );
 
 	if ( Ar.IsLoading() || Ar.IsModifyingWeakAndStrongReferences() )
 	{
-#if UCLASS_FAST_ISA_IMPL & 2
-	FFastIndexingClassTree::Register(this);
-#endif
 		HashObject(this);
 	}
 
diff --git a/Engine/Source/Runtime/CoreUObject/Public/UObject/Class.h b/Engine/Source/Runtime/CoreUObject/Public/UObject/Class.h
--- a/Engine/Source/Runtime/CoreUObject/Public/UObject/Class.h
+++ b/Engine/Source/Runtime/CoreUObject/Public/UObject/Class.h
@@ -416,6 +416,11 @@ public:
 	 */
 	virtual void SetSuperStruct(UStruct* NewSuperStruct);
 
+	/**
+	 * Serializes the SuperStruct pointer.
+	 */
+	virtual void SerializeSuperStruct(FArchive& Ar);
+
 	void LinkChild(UField* Child)
 	{
 		Child->Next = Children;
@@ -2043,6 +2048,7 @@ public:
 	// UStruct interface.
 	virtual void Link(FArchive& Ar, bool bRelinkExistingProperties) override;
 	virtual void SetSuperStruct(UStruct* NewSuperStruct) override;
+	virtual void SerializeSuperStruct(FArchive& Ar) override;
 	// End of UStruct interface.
 	
 	/**
