Fix single use bulk data unloading during cook of multiple levels and crashing. UE-13215

diff --git a/Engine/Source/Runtime/CoreUObject/Private/Serialization/BulkData.cpp b/Engine/Source/Runtime/CoreUObject/Private/Serialization/BulkData.cpp
--- a/Engine/Source/Runtime/CoreUObject/Private/Serialization/BulkData.cpp
+++ b/Engine/Source/Runtime/CoreUObject/Private/Serialization/BulkData.cpp
@@ -1487,9 +1487,12 @@ void FFormatContainer::Serialize(FArchive& Ar, UObject* Owner, const TArray<FNam
 				Ar << Name;
 				FByteBulkData* Bulk = It.Value();
 				// Force this kind of bulk data (physics, etc) to be stored inline for streaming
+				uint32 OldBulkDataSettings = Bulk->GetBulkDataFlags();
 				Bulk->SetBulkDataFlags(bSingleUse ? (BULKDATA_ForceInlinePayload | BULKDATA_SingleUse) : BULKDATA_ForceInlinePayload);
 				check(Bulk);
 				Bulk->Serialize(Ar, Owner);
+				Bulk->ClearBulkDataFlags(0xFFFFFFFF);
+				Bulk->SetBulkDataFlags(OldBulkDataSettings);
 			}
 		}
 		check(NumFormats == 0);
