UMulticastDelegateProperty::ImportText_Internal treats empty Multicast Delagate as a valid one.

diff --git a/Engine/Source/Runtime/CoreUObject/Private/UObject/PropertyMulticastDelegate.cpp b/Engine/Source/Runtime/CoreUObject/Private/UObject/PropertyMulticastDelegate.cpp
--- a/Engine/Source/Runtime/CoreUObject/Private/UObject/PropertyMulticastDelegate.cpp
+++ b/Engine/Source/Runtime/CoreUObject/Private/UObject/PropertyMulticastDelegate.cpp
@@ -168,6 +168,16 @@ const TCHAR* UMulticastDelegateProperty::ImportText_Internal( const TCHAR* Buffe
 	// Clear the existing delegate
 	MulticastDelegate.Clear();
 
+	// process opening parenthesis
+	++Buffer;
+	SkipWhitespace(Buffer);
+
+	// Empty Multi-cast delegates is still valid.
+	if (*Buffer == TCHAR(')'))
+	{
+		return Buffer;
+	}
+
 	do
 	{
 		// Parse the delegate
