Fixing BulkData serialization when cooking compressed packages

diff --git a/Engine/Source/Runtime/CoreUObject/Private/Serialization/BulkData.cpp b/Engine/Source/Runtime/CoreUObject/Private/Serialization/BulkData.cpp
--- a/Engine/Source/Runtime/CoreUObject/Private/Serialization/BulkData.cpp
+++ b/Engine/Source/Runtime/CoreUObject/Private/Serialization/BulkData.cpp
@@ -946,6 +946,7 @@ void FUntypedBulkData::Serialize( FArchive& Ar, UObject* Owner, int32 Idx )
 
 				BulkStore.BulkDataOffsetInFilePos = SavedBulkDataOffsetInFilePos;
 				BulkStore.BulkDataSizeOnDiskPos = SavedBulkDataSizeOnDiskPos;
+				BulkStore.BulkDataFlags = BulkDataFlags;
 				BulkStore.BulkData = this;
 				
 				// Serialize bulk data into the storage info
@@ -1486,13 +1487,13 @@ void FFormatContainer::Serialize(FArchive& Ar, UObject* Owner, const TArray<FNam
 				FName Name = It.Key();
 				Ar << Name;
 				FByteBulkData* Bulk = It.Value();
-				// Force this kind of bulk data (physics, etc) to be stored inline for streaming
-				uint32 OldBulkDataSettings = Bulk->GetBulkDataFlags();
-				Bulk->SetBulkDataFlags(bSingleUse ? (BULKDATA_ForceInlinePayload | BULKDATA_SingleUse) : BULKDATA_ForceInlinePayload);
 				check(Bulk);
+				// Force this kind of bulk data (physics, etc) to be stored inline for streaming
+				const uint32 OldBulkDataFlags = Bulk->GetBulkDataFlags();
+				Bulk->SetBulkDataFlags(bSingleUse ? (BULKDATA_ForceInlinePayload | BULKDATA_SingleUse) : BULKDATA_ForceInlinePayload);				
 				Bulk->Serialize(Ar, Owner);
 				Bulk->ClearBulkDataFlags(0xFFFFFFFF);
-				Bulk->SetBulkDataFlags(OldBulkDataSettings);
+				Bulk->SetBulkDataFlags(OldBulkDataFlags);
 			}
 		}
 		check(NumFormats == 0);
diff --git a/Engine/Source/Runtime/CoreUObject/Private/UObject/SavePackage.cpp b/Engine/Source/Runtime/CoreUObject/Private/UObject/SavePackage.cpp
--- a/Engine/Source/Runtime/CoreUObject/Private/UObject/SavePackage.cpp
+++ b/Engine/Source/Runtime/CoreUObject/Private/UObject/SavePackage.cpp
@@ -3492,24 +3492,32 @@ bool UPackage::SavePackage( UPackage* InOuter, UObject* Base, EObjectFlags TopLe
 					{
 						BulkDataFeedback.EnterProgressFrame();
 
-						FLinkerSave::FBulkDataStorageInfo& BulkData = Linker->BulkDataToAppend[i];
+						FLinkerSave::FBulkDataStorageInfo& BulkDataStorageInfo = Linker->BulkDataToAppend[i];
+
+						// Set bulk data flags to what they were during initial serialization (they might have changed after that)
+						const uint32 OldBulkDataFlags = BulkDataStorageInfo.BulkData->GetBulkDataFlags();
+						BulkDataStorageInfo.BulkData->SetBulkDataFlags(BulkDataStorageInfo.BulkDataFlags);
 
 						int64 BulkStartOffset = Linker->Tell();
 						int64 StoredBulkStartOffset = BulkStartOffset - StartOfBulkDataArea;
 
-						BulkData.BulkData->SerializeBulkData(*Linker, BulkData.BulkData->Lock(LOCK_READ_ONLY));
-						BulkData.BulkData->Unlock();
+						BulkDataStorageInfo.BulkData->SerializeBulkData(*Linker, BulkDataStorageInfo.BulkData->Lock(LOCK_READ_ONLY));
+						BulkDataStorageInfo.BulkData->Unlock();
 
 						int64 BulkEndOffset = Linker->Tell();
 						int32 SizeOnDisk = (int32)(BulkEndOffset - BulkStartOffset);
 				
-						Linker->Seek(BulkData.BulkDataOffsetInFilePos);
+						Linker->Seek(BulkDataStorageInfo.BulkDataOffsetInFilePos);
 						*Linker << StoredBulkStartOffset;
 
-						Linker->Seek(BulkData.BulkDataSizeOnDiskPos);
+						Linker->Seek(BulkDataStorageInfo.BulkDataSizeOnDiskPos);
 						*Linker << SizeOnDisk;
 
 						Linker->Seek(BulkEndOffset);
+
+						// Restore BulkData flags to before serialization started
+						BulkDataStorageInfo.BulkData->ClearBulkDataFlags(0xFFFFFFFF);
+						BulkDataStorageInfo.BulkData->SetBulkDataFlags(OldBulkDataFlags);
 					}
 				}
 
diff --git a/Engine/Source/Runtime/CoreUObject/Public/UObject/Linker.h b/Engine/Source/Runtime/CoreUObject/Public/UObject/Linker.h
--- a/Engine/Source/Runtime/CoreUObject/Public/UObject/Linker.h
+++ b/Engine/Source/Runtime/CoreUObject/Public/UObject/Linker.h
@@ -2077,6 +2077,8 @@ public:
 		int64 BulkDataOffsetInFilePos;
 		/** Offset to the location where the payload size is stored */
 		int64 BulkDataSizeOnDiskPos;
+		/** Bulk data flags at the time of serialization */
+		uint32 BulkDataFlags;
 		/** The bulkdata */
 		FUntypedBulkData* BulkData;
 	};
