Only report skipped saved property in the load from disk (persistent) and when working with cooked content

diff --git a/Engine/Source/Runtime/CoreUObject/Private/UObject/Class.cpp b/Engine/Source/Runtime/CoreUObject/Private/UObject/Class.cpp
--- a/Engine/Source/Runtime/CoreUObject/Private/UObject/Class.cpp
+++ b/Engine/Source/Runtime/CoreUObject/Private/UObject/Class.cpp
@@ -1207,7 +1207,7 @@ void UStruct::SerializeTaggedProperties(FArchive& Ar, uint8* Data, UStruct* Defa
 			}
 			else if( !Property->ShouldSerializeValue(Ar) )
 			{
-				UE_LOG(LogClass, Warning, TEXT("Skipping saved property %s of %s since it is no longer serializable for asset:  %s. (Maybe resave asset?)"), *Tag.Name.ToString(), *GetName(), *Ar.GetArchiveName() );
+				UE_CLOG((Ar.IsPersistent() && FPlatformProperties::RequiresCookedData()), LogClass, Warning, TEXT("Skipping saved property %s of %s since it is no longer serializable for asset:  %s. (Maybe resave asset?)"), *Tag.Name.ToString(), *GetName(), *Ar.GetArchiveName() );
 			}
 			else if ( Tag.Type == NAME_ByteProperty && ( (Tag.EnumName == NAME_None && ExactCast<UByteProperty>(Property)->Enum != NULL) || 
 														(Tag.EnumName != NAME_None && ExactCast<UByteProperty>(Property)->Enum == NULL) ))
