UE-14122 Passing FText objects accessed from Access None objects into functions crashes the game

After a contex fail the result is properly cleared.

diff --git a/Engine/Source/Editor/KismetCompiler/Private/KismetCompilerVMBackend.cpp b/Engine/Source/Editor/KismetCompiler/Private/KismetCompilerVMBackend.cpp
--- a/Engine/Source/Editor/KismetCompiler/Private/KismetCompilerVMBackend.cpp
+++ b/Engine/Source/Editor/KismetCompiler/Private/KismetCompilerVMBackend.cpp
@@ -275,12 +275,12 @@ protected:
 			FSkipOffsetEmitter Skipper(ScriptBuilder.Writer.ScriptBuffer);
 			Skipper.Emit();
 
-			// R-Value property
+			// R-Value property, see ReadVariableSize in UObject::ProcessContextOpcode() for usage
 			//@TODO: Not sure what to use for yet
 			UProperty* RValueProperty = RValueTerm ? RValueTerm->AssociatedVarProperty : NULL;
 			Writer << RValueProperty;
 
-			// Property type if needed
+			// Property type if needed, it seems it's not used in ue4
 			//@TODO: Not sure what to use for yet
 			uint8 PropetyType = 0;
 			Writer << PropetyType;
diff --git a/Engine/Source/Runtime/CoreUObject/Private/UObject/ScriptCore.cpp b/Engine/Source/Runtime/CoreUObject/Private/UObject/ScriptCore.cpp
--- a/Engine/Source/Runtime/CoreUObject/Private/UObject/ScriptCore.cpp
+++ b/Engine/Source/Runtime/CoreUObject/Private/UObject/ScriptCore.cpp
@@ -1519,7 +1519,9 @@ void UObject::ProcessContextOpcode( FFrame& Stack, RESULT_DECL, bool bCanFailSil
 	// Execute or skip the following expression in the object's context.
 	if (IsValid(NewContext))
 	{
-		Stack.Code += sizeof(CodeSkipSizeType) + sizeof(ScriptPointerType) + sizeof(uint8);
+		Stack.Code += sizeof(CodeSkipSizeType)	// Code offset for NULL expressions.
+			+ sizeof(ScriptPointerType)			// Property corresponding to the r-value data, in case the l-value needs to be cleared
+			+ sizeof(uint8);					// Property type, in case the r-value is a non-property - in ue4 it seems to be unused
 		Stack.Step( NewContext, RESULT_PARAM );
 	}
 	else
@@ -1549,15 +1551,25 @@ void UObject::ProcessContextOpcode( FFrame& Stack, RESULT_DECL, bool bCanFailSil
 			}
 		}
 
-		CodeSkipSizeType wSkip = Stack.ReadCodeSkipCount();
-		VariableSizeType bSize = Stack.ReadVariableSize();
+		const CodeSkipSizeType wSkip = Stack.ReadCodeSkipCount(); // Code offset for NULL expressions. Code += sizeof(CodeSkipSizeType)
+		UField* RValueField = nullptr;
+		const VariableSizeType bSize = Stack.ReadVariableSize(&RValueField); // Code += sizeof(ScriptPointerType) + sizeof(uint8)
 		Stack.Code += wSkip;
 		Stack.MostRecentPropertyAddress = NULL;
 		Stack.MostRecentProperty = NULL;
 
 		if (RESULT_PARAM)
 		{
-			FMemory::Memzero( RESULT_PARAM, bSize );
+			auto RValueProperty = Cast<const UProperty>(RValueField);
+			ensure(RValueProperty || !RValueField);
+			if (RValueProperty)
+			{
+				RValueProperty->ClearValue(RESULT_PARAM);
+			}
+			else
+			{
+				Stack.KismetExecutionMessage(TEXT("Cannot clear R-value after a context fail. Unknown property."), ELogVerbosity::Warning);
+			}
 		}
 	}
 }
diff --git a/Engine/Source/Runtime/CoreUObject/Public/UObject/Stack.h b/Engine/Source/Runtime/CoreUObject/Public/UObject/Stack.h
--- a/Engine/Source/Runtime/CoreUObject/Public/UObject/Stack.h
+++ b/Engine/Source/Runtime/CoreUObject/Public/UObject/Stack.h
@@ -280,6 +280,7 @@ inline VariableSizeType FFrame::ReadVariableSize( UField** ExpressionField/*=NUL
 
 	if ( Field != NULL )
 	{
+		// currently only Property (or null) seems to be a valid option.
 		if (UProperty* Property = dynamic_cast<UProperty*>(Field))
 		{
 			Result = Property->GetSize();
