UE-11912 Access None errors creating log spam in user's project when editing BP

In Editor.ini [ScriptErrorLog] section and MaxNumOfAccessViolation variable were added.

diff --git a/Engine/Config/BaseEditor.ini b/Engine/Config/BaseEditor.ini
--- a/Engine/Config/BaseEditor.ini
+++ b/Engine/Config/BaseEditor.ini
@@ -344,6 +344,10 @@ bBlueprintIsNotBlueprintType=false
 SkipByteCodeSerialization=true
 BreakSerializationRecursion=true
 
+[ScriptErrorLog]
+; per object
+MaxNumOfAccessViolation=16
+
 [/Script/Engine.ParticleSystemAuditCommandlet]
 HighSpawnRateOrBurstThreshold=35
 FarLODDistanceTheshold=3000
diff --git a/Engine/Source/Runtime/CoreUObject/Private/UObject/ScriptCore.cpp b/Engine/Source/Runtime/CoreUObject/Private/UObject/ScriptCore.cpp
--- a/Engine/Source/Runtime/CoreUObject/Private/UObject/ScriptCore.cpp
+++ b/Engine/Source/Runtime/CoreUObject/Private/UObject/ScriptCore.cpp
@@ -76,7 +76,33 @@ void FBlueprintCoreDelegates::ThrowScriptException(const UObject* ActiveObject,
 	case EBlueprintExceptionType::Tracepoint:
 	case EBlueprintExceptionType::WireTracepoint:
 		break;
+#if WITH_EDITOR
+	case EBlueprintExceptionType::AccessViolation:
+		{
+			struct FIntConfigValueHelper
+			{
+				int32 Value;
 
+				FIntConfigValueHelper() : Value(0)
+				{
+					GConfig->GetInt(TEXT("ScriptErrorLog"), TEXT("MaxNumOfAccessViolation"), Value, GEditorIni);
+				}
+			};
+
+			static const FIntConfigValueHelper MaxNumOfAccessViolation;
+			if (MaxNumOfAccessViolation.Value > 0)
+			{
+				static TMap<FName, int32> DisplayedWarningsMap;
+				const FName ActiveObjectName = ActiveObject ? ActiveObject->GetFName() : FName();
+				int32& Num = DisplayedWarningsMap.FindOrAdd(ActiveObjectName);
+				if (Num > MaxNumOfAccessViolation.Value)
+				{
+					break;
+				}
+				Num++;
+			}
+		}
+#endif // WITH_EDITOR
 	default:
 		UE_SUPPRESS(LogScript, Warning, const_cast<FFrame*>(&StackFrame)->Logf(TEXT("%s"), *(Info.GetDescription())));
 		break;
@@ -1566,10 +1592,6 @@ void UObject::ProcessContextOpcode( FFrame& Stack, RESULT_DECL, bool bCanFailSil
 			{
 				RValueProperty->ClearValue(RESULT_PARAM);
 			}
-			else
-			{
-				Stack.KismetExecutionMessage(TEXT("Cannot clear R-value after a context fail. Unknown property."), ELogVerbosity::Warning);
-			}
 		}
 	}
 }
