Core - Removed debug code

diff --git a/Engine/Source/Runtime/CoreUObject/Private/UObject/UObjectGlobals.cpp b/Engine/Source/Runtime/CoreUObject/Private/UObject/UObjectGlobals.cpp
--- a/Engine/Source/Runtime/CoreUObject/Private/UObject/UObjectGlobals.cpp
+++ b/Engine/Source/Runtime/CoreUObject/Private/UObject/UObjectGlobals.cpp
@@ -2010,8 +2010,6 @@ FObjectInitializer::FObjectInitializer(UObject* InObj, UObject* InObjectArchetyp
 	FTlsObjectInitializers::Push(this);
 }
 
-COREUOBJECT_API bool IgnoreFObjectInitializer = false;
-
 /**
  * Destructor for internal class to finalize UObject creation (initialize properties) after the real C++ constructor is called.
  **/
@@ -2024,11 +2022,6 @@ FObjectInitializer::~FObjectInitializer()
 	check(ThreadContext.IsInConstructor >= 0);
 	ThreadContext.ConstructedObject = LastConstructedObject;
 
-	if( IgnoreFObjectInitializer )
-	{
-		return;
-	}
-
 	SCOPE_CYCLE_COUNTER(STAT_PostConstructInitializeProperties);
 	check(Obj);
 	const bool bIsCDO = Obj->HasAnyFlags(RF_ClassDefaultObject);
