FStringAssetReference::TryLoad will now correctly follow redirectors

diff --git a/Engine/Source/Runtime/CoreUObject/Private/Misc/StringAssetReference.cpp b/Engine/Source/Runtime/CoreUObject/Private/Misc/StringAssetReference.cpp
--- a/Engine/Source/Runtime/CoreUObject/Private/Misc/StringAssetReference.cpp
+++ b/Engine/Source/Runtime/CoreUObject/Private/Misc/StringAssetReference.cpp
@@ -113,9 +113,15 @@ bool FStringAssetReference::SerializeFromMismatchedTag(struct FPropertyTag const
 
 UObject* FStringAssetReference::TryLoad() const
 {
+	UObject* LoadedObject = nullptr;
+
 	if ( IsValid() )
 	{
-		return LoadObject<UObject>(nullptr, *ToString());
+		LoadedObject = LoadObject<UObject>(nullptr, *ToString());
+		while (UObjectRedirector* Redirector = Cast<UObjectRedirector>(LoadedObject))
+		{
+			LoadedObject = Redirector->DestinationObject;
+		}
 	}
 
 	return nullptr;
@@ -127,18 +133,13 @@ UObject* FStringAssetReference::ResolveObject() const
 	// and we usually don't want to force references to weak pointers while saving.
 	if (!IsValid() || GIsSavingPackage)
 	{
-		return NULL;
+		return nullptr;
 	}
-	UObject* FoundObject = NULL;
-	// construct full name
-	FString FullPath = ToString();
-	FoundObject = StaticFindObject( UObject::StaticClass(), NULL, *FullPath);
 
-	UObjectRedirector* Redir = dynamic_cast<UObjectRedirector*>(FoundObject);
-	if (Redir)
+	UObject* FoundObject = FindObject<UObject>(nullptr, *ToString());
+	while (UObjectRedirector* Redirector = Cast<UObjectRedirector>(FoundObject))
 	{
-		// If we found a redirector, follow it
-		FoundObject = Redir->DestinationObject;
+		FoundObject = Redirector->DestinationObject;
 	}
 
 	return FoundObject;
