[UE-13579] Deferring sub-object/property initialization for certain CDOs (until after we can garuntee their super CDO has been initialized).

diff --git a/Engine/Source/Runtime/CoreUObject/Private/Blueprint/BlueprintSupport.cpp b/Engine/Source/Runtime/CoreUObject/Private/Blueprint/BlueprintSupport.cpp
--- a/Engine/Source/Runtime/CoreUObject/Private/Blueprint/BlueprintSupport.cpp
+++ b/Engine/Source/Runtime/CoreUObject/Private/Blueprint/BlueprintSupport.cpp
@@ -1221,21 +1221,13 @@ void FLinkerLoad::ResolveDeferredExports(UClass* LoadClass)
 		// do. Simply moving the call to SerializeDefaultObject next to the call to create the CDO means
 		// the stream order changes...
 		DEFERRED_DEPENDENCY_CHECK(BlueprintCDO->GetClass() == LoadClass);
-		// @TODO: Find a different way of resolving UE-11029, this was causing a
-		//        more subtle bug with UE-12953 (it seems that we need to run 
-		//        all of ~FObjectInitializer() together; we can't pick and 
-		//        choose like we tried here)
-// 		if (BlueprintCDO->HasAnyFlags(RF_NeedLoad) &&
-// 			!LoadClass->GetSuperClass()->HasAnyFlags(RF_Native))
-// 		{
-// 			for (UProperty* P = LoadClass->PropertyLink; P; P = P->PropertyLinkNext)
-// 			{
-// 				if (P->IsInContainer(LoadClass->GetSuperClass()))
-// 				{
-// 					P->CopyCompleteValue_InContainer(BlueprintCDO, LoadClass->GetSuperClass()->GetDefaultObject(false));
-// 				}
-// 			}
-// 		}
+
+		if (FObjectInitializer* DeferredInitializer = FDeferredObjInitializerTracker::Find(LoadClass))
+		{
+			FScriptIntegrationObjectHelper::PostConstructInitObject(*DeferredInitializer);
+			FDeferredObjInitializerTracker::Remove(LoadClass);
+		}
+ 		DEFERRED_DEPENDENCY_CHECK(BlueprintCDO->GetClass() == LoadClass);
 
 		// should load the CDO (ensuring that it has been serialized in by the 
 		// time we get to class regeneration)
@@ -1339,9 +1331,6 @@ void FLinkerLoad::ResetDeferredLoadingState()
 #endif // USE_CIRCULAR_DEPENDENCY_LOAD_DEFERRING
 }
 
-// don't want other files ending up with this internal define
-#undef DEFERRED_DEPENDENCY_CHECK
-
 /*******************************************************************************
  * UObject
  ******************************************************************************/
@@ -1420,3 +1409,44 @@ bool FObjectInitializer::InitNonNativeProperty(UProperty* Property, UObject* Dat
 		return false;
 	}
 }
+
+/*******************************************************************************
+ * FDeferredObjInitializerTracker
+ ******************************************************************************/
+
+FObjectInitializer* FDeferredObjInitializerTracker::Add(const FObjectInitializer& DeferringInitializer)
+{
+	UObject* InitingObj = DeferringInitializer.GetObj();
+	DEFERRED_DEPENDENCY_CHECK((InitingObj != nullptr) && InitingObj->HasAnyFlags(RF_ClassDefaultObject));
+
+	UClass* LoadClass = (InitingObj != nullptr) ? InitingObj->GetClass() : nullptr;
+	if (LoadClass != nullptr)
+	{
+		FDeferredObjInitializerTracker& ThreadInst = FDeferredObjInitializerTracker::Get();
+		auto& InitializersCache = ThreadInst.DeferredInitializers;
+
+		DEFERRED_DEPENDENCY_CHECK(InitializersCache.Find(LoadClass) == nullptr); // did we try to init the CDO twice?
+		DEFERRED_DEPENDENCY_CHECK(LoadClass->GetLinker()->LoadFlags & LOAD_DeferDependencyLoads);
+
+		// NOTE: we copy the FObjectInitializer, because it is most likely being destroyed
+		FObjectInitializer& DeferredCopy = InitializersCache.Add(LoadClass, DeferringInitializer);
+
+		return &DeferredCopy;
+	}
+	return nullptr;
+}
+
+FObjectInitializer* FDeferredObjInitializerTracker::Find(UClass* LoadClass)
+{
+	FDeferredObjInitializerTracker& ThreadInst = FDeferredObjInitializerTracker::Get();
+	return ThreadInst.DeferredInitializers.Find(LoadClass);
+}
+
+void FDeferredObjInitializerTracker::Remove(UClass* LoadClass)
+{
+	FDeferredObjInitializerTracker& ThreadInst = FDeferredObjInitializerTracker::Get();
+	ThreadInst.DeferredInitializers.Remove(LoadClass);
+}
+
+// don't want other files ending up with this internal define
+#undef DEFERRED_DEPENDENCY_CHECK
diff --git a/Engine/Source/Runtime/CoreUObject/Private/UObject/UObjectGlobals.cpp b/Engine/Source/Runtime/CoreUObject/Private/UObject/UObjectGlobals.cpp
--- a/Engine/Source/Runtime/CoreUObject/Private/UObject/UObjectGlobals.cpp
+++ b/Engine/Source/Runtime/CoreUObject/Private/UObject/UObjectGlobals.cpp
@@ -10,6 +10,7 @@
 #include "UObject/UObjectAnnotation.h"
 #include "UObject/TlsObjectInitializers.h"
 #include "UObject/UObjectThreadContext.h"
+#include "BlueprintSupport.h" // for FDeferredObjInitializerTracker
 
 DEFINE_LOG_CATEGORY(LogUObjectGlobals);
 
@@ -1944,7 +1945,7 @@ void UObject::PostInitProperties()
 	FUObjectThreadContext::Get().PostInitPropertiesCheck.Push(this);
 #endif
 #if USE_UBER_GRAPH_PERSISTENT_FRAME
-	GetClass()->CreatePersistentUberGraphFrame(this);
+	GetClass()->CreatePersistentUberGraphFrame(this, true);
 #endif
 }
 
@@ -1975,14 +1976,17 @@ UObject::UObject(const FObjectInitializer& ObjectInitializer)
 
 
 
-FObjectInitializer::FObjectInitializer() :
-Obj(nullptr),
-	ObjectArchetype(nullptr),
-	bCopyTransientsFromClassDefaults(false),
-	bShouldIntializePropsFromArchetype(false),
-	bSubobjectClassInitializationAllowed(true),
-	InstanceGraph(nullptr),
-	LastConstructedObject(nullptr)
+FObjectInitializer::FObjectInitializer()
+	: Obj(nullptr)
+	, ObjectArchetype(nullptr)
+	, bCopyTransientsFromClassDefaults(false)
+	, bShouldIntializePropsFromArchetype(false)
+	, bSubobjectClassInitializationAllowed(true)
+	, InstanceGraph(nullptr)
+	, LastConstructedObject(nullptr)
+#if USE_CIRCULAR_DEPENDENCY_LOAD_DEFERRING
+	, bIsDeferredInitializer(false)
+#endif // USE_CIRCULAR_DEPENDENCY_LOAD_DEFERRING
 {
 	auto& ThreadContext = FUObjectThreadContext::Get();
 	// Mark we're in the constructor now.	
@@ -1992,15 +1996,18 @@ Obj(nullptr),
 	FTlsObjectInitializers::Push(this);
 }	
 
-FObjectInitializer::FObjectInitializer(UObject* InObj, UObject* InObjectArchetype, bool bInCopyTransientsFromClassDefaults, bool bInShouldIntializeProps, struct FObjectInstancingGraph* InInstanceGraph) :
-	Obj(InObj),
-	ObjectArchetype(InObjectArchetype),
-	// if the SubobjectRoot NULL, then we want to copy the transients from the template, otherwise we are doing a duplicate and we want to copy the transients from the class defaults
-	bCopyTransientsFromClassDefaults(bInCopyTransientsFromClassDefaults),
-	bShouldIntializePropsFromArchetype(bInShouldIntializeProps),
-	bSubobjectClassInitializationAllowed(true),
-	InstanceGraph(InInstanceGraph),
-	LastConstructedObject(nullptr)
+FObjectInitializer::FObjectInitializer(UObject* InObj, UObject* InObjectArchetype, bool bInCopyTransientsFromClassDefaults, bool bInShouldIntializeProps, struct FObjectInstancingGraph* InInstanceGraph)
+	: Obj(InObj)
+	, ObjectArchetype(InObjectArchetype)
+	  // if the SubobjectRoot NULL, then we want to copy the transients from the template, otherwise we are doing a duplicate and we want to copy the transients from the class defaults
+	, bCopyTransientsFromClassDefaults(bInCopyTransientsFromClassDefaults)
+	, bShouldIntializePropsFromArchetype(bInShouldIntializeProps)
+	, bSubobjectClassInitializationAllowed(true)
+	, InstanceGraph(InInstanceGraph)
+	, LastConstructedObject(nullptr)
+#if USE_CIRCULAR_DEPENDENCY_LOAD_DEFERRING
+	, bIsDeferredInitializer(false)
+#endif // USE_CIRCULAR_DEPENDENCY_LOAD_DEFERRING
 {
 	auto& ThreadContext = FUObjectThreadContext::Get();
 	// Mark we're in the constructor now.
@@ -2015,15 +2022,35 @@ FObjectInitializer::FObjectInitializer(UObject* InObj, UObject* InObjectArchetyp
  **/
 FObjectInitializer::~FObjectInitializer()
 {
-	FTlsObjectInitializers::Pop();
-	auto& ThreadContext = FUObjectThreadContext::Get();
-	// Let the FObjectFinders know we left the constructor.
-	ThreadContext.IsInConstructor--;
-	check(ThreadContext.IsInConstructor >= 0);
-	ThreadContext.ConstructedObject = LastConstructedObject;
 
-	SCOPE_CYCLE_COUNTER(STAT_PostConstructInitializeProperties);
-	check(Obj);
+#if USE_CIRCULAR_DEPENDENCY_LOAD_DEFERRING
+	// if we're not at the top of FTlsObjectInitializers, then this is most 
+	// likely a deferred FObjectInitializer that's a copy of one that was used 
+	// in a constructor (that has already been popped)
+	if (!bIsDeferredInitializer)
+	{
+#endif // USE_CIRCULAR_DEPENDENCY_LOAD_DEFERRING
+		check(FTlsObjectInitializers::Top() == this);
+
+		FTlsObjectInitializers::Pop();
+		auto& ThreadContext = FUObjectThreadContext::Get();
+		// Let the FObjectFinders know we left the constructor.
+		ThreadContext.IsInConstructor--;
+		check(ThreadContext.IsInConstructor >= 0);
+		ThreadContext.ConstructedObject = LastConstructedObject;
+
+		check(Obj != nullptr);
+
+#if USE_CIRCULAR_DEPENDENCY_LOAD_DEFERRING
+	}
+	else if (Obj == nullptr)
+	{
+		// the deferred initialization has already been ran, we clear Obj once 
+		// PostConstructInit() has been executed
+		return;
+	}
+#endif // USE_CIRCULAR_DEPENDENCY_LOAD_DEFERRING
+
 	const bool bIsCDO = Obj->HasAnyFlags(RF_ClassDefaultObject);
 	UClass* Class = Obj->GetClass();
 
@@ -2038,8 +2065,66 @@ FObjectInitializer::~FObjectInitializer()
 	else if (bIsCDO)
 	{
 		// for the Object CDO, make sure that we do not use an archetype
-		check(!ObjectArchetype);
+		check(ObjectArchetype == nullptr);
+	}
+
+#if USE_CIRCULAR_DEPENDENCY_LOAD_DEFERRING
+	bool bIsPostConstructInitDeferred = false;
+	if (bIsCDO && (ObjectArchetype != nullptr))
+	{
+		UClass* ArchetypeClass = ObjectArchetype->GetClass();
+		// if this is a blueprint CDO that derives from another blueprint, and 
+		// that parent (archetype) CDO isn't fully serialized
+		if ( !Class->HasAnyFlags(RF_Native) && !ArchetypeClass->HasAnyFlags(RF_Native) &&
+			(ObjectArchetype->HasAnyFlags(RF_NeedLoad) || !ObjectArchetype->HasAnyFlags(RF_LoadCompleted)) )
+		{
+			FLinkerLoad* ClassLinker = Class->GetLinker();
+			if ((ClassLinker != nullptr) && (ClassLinker->LoadFlags & LOAD_DeferDependencyLoads) != 0x00)
+			{
+#if USE_DEFERRED_DEPENDENCY_CHECK_VERIFICATION_TESTS
+				// make sure we haven't already deferred this once, if we have
+				// then something is destroying this one prematurely 
+				check(bIsDeferredInitializer == false);
+#endif // USE_DEFERRED_DEPENDENCY_CHECK_VERIFICATION_TESTS
+
+				// makes a copy of this and saves it off, to be ran later
+				if (FObjectInitializer* DeferredCopy = FDeferredObjInitializerTracker::Add(*this))
+				{
+					bIsPostConstructInitDeferred = true;
+					DeferredCopy->bIsDeferredInitializer = true;
+
+					// make sure this wasn't mistakenly pushed into FTlsObjectInitializers
+					// (the copy constructor should have been what was invoked, 
+					// which doesn't push to FTlsObjectInitializers)
+					check(FTlsObjectInitializers::Top() != DeferredCopy);
+				}
+			}
+		}
+	}
+
+	if (!bIsPostConstructInitDeferred)
+#endif // USE_CIRCULAR_DEPENDENCY_LOAD_DEFERRING
+	{
+		PostConstructInit();
 	}
+}
+
+void FObjectInitializer::PostConstructInit()
+{
+	// we clear the Obj pointer at the end of this function, so if it is null 
+	// then it most likely means that this is being ran for a second time
+	if (Obj == nullptr)
+	{
+#if USE_DEFERRED_DEPENDENCY_CHECK_VERIFICATION_TESTS
+		checkf(Obj != nullptr, TEXT("Looks like you're attempting to run FObjectInitializer::PostConstructInit() twice, and that should never happen."));
+#endif // USE_DEFERRED_DEPENDENCY_CHECK_VERIFICATION_TESTS
+		return;
+	}
+
+	SCOPE_CYCLE_COUNTER(STAT_PostConstructInitializeProperties);
+	const bool bIsCDO = Obj->HasAnyFlags(RF_ClassDefaultObject);
+	UClass* Class = Obj->GetClass();
+
 	if (bShouldIntializePropsFromArchetype)
 	{
 		UClass* BaseClass = (bIsCDO && !GIsDuplicatingClassForReinstancing) ? Obj->GetClass()->GetSuperClass() : Class;
@@ -2130,6 +2215,9 @@ FObjectInitializer::~FObjectInitializer()
 	{
 		Obj->CheckDefaultSubobjects();
 	}
+
+	// clear the object pointer so we can guard against runing this function again
+	Obj = nullptr;
 }
 
 bool FObjectInitializer::IsInstancingAllowed() const
diff --git a/Engine/Source/Runtime/CoreUObject/Public/Blueprint/BlueprintSupport.h b/Engine/Source/Runtime/CoreUObject/Public/Blueprint/BlueprintSupport.h
--- a/Engine/Source/Runtime/CoreUObject/Public/Blueprint/BlueprintSupport.h
+++ b/Engine/Source/Runtime/CoreUObject/Public/Blueprint/BlueprintSupport.h
@@ -64,4 +64,30 @@ private:
 
 	FScopedClassDependencyGather();
 };
-#endif //WITH_EDITOR
\ No newline at end of file
+#endif //WITH_EDITOR
+
+/** 
+ * A helper struct for storing FObjectInitializers that were not run on 
+ * Blueprint CDO's post-construction (presumably because that CDO's super had 
+ * not been fully serialized yet). 
+ * 
+ * This was designed to hold onto FObjectInitializers until a later point, when 
+ * they can properly be ran (presumably in FLinkerLoad::ResolveDeferredExports,
+ * after the super has been finalized).
+ */
+struct FDeferredObjInitializerTracker : TThreadSingleton<FDeferredObjInitializerTracker>
+{
+public:
+	/** Stores a copy of the specified FObjectInitializer and returns a pointer to it (could be null if a corresponding class could not be determined). */
+	static FObjectInitializer* Add(const FObjectInitializer& DeferringInitializer);
+
+	/** Looks up a FObjectInitializer that was deferred for the specified class (an FObjectInitializer for that class's CDO). */
+	static FObjectInitializer* Find(UClass* LoadClass);
+
+	/** Destroys any FObjectInitializers that were cached corresponding to the specified class. */
+	static void Remove(UClass* LoadClass);
+
+private:
+	/** A map that tracks the relationsship between Blueprint classes and  FObjectInitializers for their CDOs */
+	TMap<UClass*, FObjectInitializer> DeferredInitializers;
+};
\ No newline at end of file
diff --git a/Engine/Source/Runtime/CoreUObject/Public/UObject/UObjectGlobals.h b/Engine/Source/Runtime/CoreUObject/Public/UObject/UObjectGlobals.h
--- a/Engine/Source/Runtime/CoreUObject/Public/UObject/UObjectGlobals.h
+++ b/Engine/Source/Runtime/CoreUObject/Public/UObject/UObjectGlobals.h
@@ -864,6 +864,13 @@ private:
 	 * @return	Returns true if that property was a non-native one, otherwise false
 	 */
 	static bool InitNonNativeProperty(UProperty* Property, UObject* Data);
+	
+	/**
+	 * Finalizes a constructed UObject by initializing properties, 
+	 * instancing/initializing sub-objects, etc.
+	 */
+	void PostConstructInit();
+
 private:
 
 	/**  Littel helper struct to manage overrides from dervied classes **/
@@ -976,6 +983,11 @@ private:
 	mutable FSubobjectsToInit ComponentInits;
 	/**  Previously constructed object in the callstack */
 	UObject* LastConstructedObject;
+
+#if USE_CIRCULAR_DEPENDENCY_LOAD_DEFERRING
+	/**  */
+	bool bIsDeferredInitializer : 1;
+#endif // USE_CIRCULAR_DEPENDENCY_LOAD_DEFERRING
 };
 
 #define FPostConstructInitializeProperties \
@@ -1022,6 +1034,17 @@ public:
 	{
 		ObjectInitializer.InstanceSubobjects(Class, bNeedInstancing, bNeedSubobjectInstancing);
 	}
+
+	/**
+	 * Finalizes a constructed UObject by initializing properties, instancing &
+	 * initializing sub-objects, etc.
+	 * 
+	 * @param  ObjectInitializer    The initializer to run PostConstructInit() on.
+	 */
+	inline static void PostConstructInitObject(FObjectInitializer& ObjectInitializer)
+	{
+		ObjectInitializer.PostConstructInit();
+	}
 };
 
 /**
