Fix for a crash when setting up uclass token streams

diff --git a/Engine/Source/Runtime/CoreUObject/Private/UObject/GarbageCollection.cpp b/Engine/Source/Runtime/CoreUObject/Private/UObject/GarbageCollection.cpp
--- a/Engine/Source/Runtime/CoreUObject/Private/UObject/GarbageCollection.cpp
+++ b/Engine/Source/Runtime/CoreUObject/Private/UObject/GarbageCollection.cpp
@@ -51,6 +51,9 @@ static volatile bool GIsRunningParallelReachability = false;
 /** Whether we are currently purging an object in the GC purge pass. */
 static bool GIsPurgingObject = false;
 
+/** Helpful constant for determining how many token slots we need to store a pointer **/
+static const uint32 GNumTokensPerPointer = sizeof(void*) / sizeof(uint32);
+
 /** Locks all UObject hash tables when performing GC */
 class FGCScopeLock
 {
@@ -1872,7 +1875,7 @@ void FGCReferenceTokenStream::ReplaceOrAddAddReferencedObjectsCall(void (*AddRef
 		case GCRT_AddStructReferencedObjects:
 			{
 				// Skip pointer
-				TokenIndex++; 
+				TokenIndex += GNumTokensPerPointer;
 			}
 			break;
 		case GCRT_AddReferencedObjects:
@@ -1949,7 +1952,7 @@ void FGCReferenceTokenStream::EmitCount( uint32 Count )
 void FGCReferenceTokenStream::EmitPointer( void const* Ptr )
 {
 	const int32 StoreIndex = Tokens.Num();
-	Tokens.AddUninitialized(sizeof(void*) / sizeof(uint32));
+	Tokens.AddUninitialized(GNumTokensPerPointer);
 	StorePointer(&Tokens[StoreIndex], Ptr);
 }
 
