Add LoadSynchronous utility function to FAssetPtr and TAssetPtr

diff --git a/Engine/Source/Runtime/CoreUObject/Public/UObject/AssetPtr.h b/Engine/Source/Runtime/CoreUObject/Public/UObject/AssetPtr.h
--- a/Engine/Source/Runtime/CoreUObject/Public/UObject/AssetPtr.h
+++ b/Engine/Source/Runtime/CoreUObject/Public/UObject/AssetPtr.h
@@ -44,6 +44,20 @@ public:
 		(*this)=Object;
 	}
 
+	/**  
+	 * Synchronously load (if necessary) and return the asset object represented by this asset ptr
+	 */
+	UObject* LoadSynchronous()
+	{
+		UObject* Asset = Get();
+		if (Asset == nullptr && IsPending())
+		{
+			Asset = GetUniqueID().TryLoad();
+			*this = Asset;
+		}
+		return Asset;
+	}
+
 	using TPersistentObjectPtr<FStringAssetReference>::operator=;
 };
 
@@ -194,6 +208,19 @@ public:
 		return Get();
 	}
 
+	/**  
+	 * Synchronously load (if necessary) and return the asset object represented by this asset ptr
+	 */
+	T* LoadSynchronous()
+	{
+		UObject* Asset = AssetPtr.Get();
+		if (AssetPtr == nullptr && IsPending())
+		{
+			Asset = AssetPtr.LoadSynchronous();
+		}
+		return Cast<T>(Asset);
+	}
+
 	/**  
 	 * Test if this points to a live UObject
 	 * @return true if Get() would return a valid non-null pointer
