Fix compile error

diff --git a/Engine/Source/Runtime/CoreUObject/Public/UObject/AssetPtr.h b/Engine/Source/Runtime/CoreUObject/Public/UObject/AssetPtr.h
--- a/Engine/Source/Runtime/CoreUObject/Public/UObject/AssetPtr.h
+++ b/Engine/Source/Runtime/CoreUObject/Public/UObject/AssetPtr.h
@@ -214,7 +214,7 @@ public:
 	T* LoadSynchronous()
 	{
 		UObject* Asset = AssetPtr.Get();
-		if (AssetPtr == nullptr && IsPending())
+		if (Asset == nullptr && IsPending())
 		{
 			Asset = AssetPtr.LoadSynchronous();
 		}
