[UE-14848] Deferred sub-object initialization needs to fix up any stale archetype references.

diff --git a/Engine/Source/Runtime/CoreUObject/Private/UObject/UObjectGlobals.cpp b/Engine/Source/Runtime/CoreUObject/Private/UObject/UObjectGlobals.cpp
--- a/Engine/Source/Runtime/CoreUObject/Private/UObject/UObjectGlobals.cpp
+++ b/Engine/Source/Runtime/CoreUObject/Private/UObject/UObjectGlobals.cpp
@@ -2070,7 +2070,7 @@ FObjectInitializer::~FObjectInitializer()
 
 #if USE_CIRCULAR_DEPENDENCY_LOAD_DEFERRING
 	bool bIsPostConstructInitDeferred = false;
-	if (bIsCDO && (ObjectArchetype != nullptr))
+	if (bIsCDO && (ObjectArchetype != nullptr) && FBlueprintSupport::UseDeferredDependencyLoading())
 	{
 		UClass* ArchetypeClass = ObjectArchetype->GetClass();
 		// if this is a blueprint CDO that derives from another blueprint, and 
@@ -2124,10 +2124,41 @@ void FObjectInitializer::PostConstructInit()
 	SCOPE_CYCLE_COUNTER(STAT_PostConstructInitializeProperties);
 	const bool bIsCDO = Obj->HasAnyFlags(RF_ClassDefaultObject);
 	UClass* Class = Obj->GetClass();
+	UClass* SuperClass = Class->GetSuperClass();
+
+#if USE_CIRCULAR_DEPENDENCY_LOAD_DEFERRING
+	if (bIsDeferredInitializer && SuperClass->HasAnyClassFlags(CLASS_NewerVersionExists))
+	{
+		check(bIsCDO);
+
+		SuperClass = SuperClass->GetAuthoritativeClass();
+		Class->SetSuperStruct(SuperClass);
+		ObjectArchetype = SuperClass->GetDefaultObject(/*bCreateIfNeeded =*/false);
+
+		// iterate backwards, so we can remove elements as we go
+		for (int32 SubObjIndex = ComponentInits.SubobjectInits.Num()-1; SubObjIndex >= 0; --SubObjIndex)
+		{
+			FSubobjectsToInit::FSubobjectInit& SubObjInitInfo = ComponentInits.SubobjectInits[SubObjIndex];
+			const FName SubObjName = SubObjInitInfo.Subobject->GetFName();
+
+			UObject* OuterArchetype = SubObjInitInfo.Subobject->GetOuter()->GetArchetype();
+			UObject* NewTemplate = OuterArchetype->GetClass()->GetDefaultSubobjectByName(SubObjName);
+
+			if (ensure(NewTemplate != nullptr))
+			{
+				SubObjInitInfo.Template = NewTemplate;
+			}
+			else
+			{
+				ComponentInits.SubobjectInits.RemoveAtSwap(SubObjIndex);
+			}
+		}
+	}
+#endif // USE_CIRCULAR_DEPENDENCY_LOAD_DEFERRING
 
 	if (bShouldIntializePropsFromArchetype)
 	{
-		UClass* BaseClass = (bIsCDO && !GIsDuplicatingClassForReinstancing) ? Obj->GetClass()->GetSuperClass() : Class;
+		UClass* BaseClass = (bIsCDO && !GIsDuplicatingClassForReinstancing) ? SuperClass : Class;
 		if (BaseClass == NULL)
 		{
 			check(Class==UObject::StaticClass());
