[UE-14848] Comments to explain CL 2534515.

diff --git a/Engine/Source/Runtime/CoreUObject/Private/UObject/UObjectGlobals.cpp b/Engine/Source/Runtime/CoreUObject/Private/UObject/UObjectGlobals.cpp
--- a/Engine/Source/Runtime/CoreUObject/Private/UObject/UObjectGlobals.cpp
+++ b/Engine/Source/Runtime/CoreUObject/Private/UObject/UObjectGlobals.cpp
@@ -2127,6 +2127,10 @@ void FObjectInitializer::PostConstructInit()
 	UClass* SuperClass = Class->GetSuperClass();
 
 #if USE_CIRCULAR_DEPENDENCY_LOAD_DEFERRING
+	// if this is a deferred initializer (implying that it's for a CDO), and the 
+	// ObjectArchetype (super CDO) has since been regenerated, then we need 
+	// to update the cached archetypes (so their not pointing to TRASH/REINST 
+	// versions)
 	if (bIsDeferredInitializer && SuperClass->HasAnyClassFlags(CLASS_NewerVersionExists))
 	{
 		check(bIsCDO);
