Made InitializeCoreClasses store information staticly for reuse in following save package calls.
Converted some string comparisions to use FNames instead.

diff --git a/Engine/Source/Runtime/CoreUObject/Private/UObject/SavePackage.cpp b/Engine/Source/Runtime/CoreUObject/Private/UObject/SavePackage.cpp
--- a/Engine/Source/Runtime/CoreUObject/Private/UObject/SavePackage.cpp
+++ b/Engine/Source/Runtime/CoreUObject/Private/UObject/SavePackage.cpp
@@ -13,6 +13,37 @@ static const int32 MAX_MERGED_COMPRESSION_CHUNKSIZE = 1024 * 1024;
 static const FName WorldClassName = FName("World");
 
 
+#define VALIDATE_INITIALIZECORECLASSES 0
+#define EXPORT_SORTING_DETAILED_LOGGING 0
+
+#define UE_PROFILE_COOKSAVE 0
+#if UE_PROFILE_COOKSAVE
+
+#define UE_START_LOG_COOK_TIME(InFilename) double PreviousTime; double StartTime; PreviousTime = StartTime = FPlatformTime::Seconds(); \
+	UE_LOG(LogSavePackage, Log, TEXT("Starting save package for %s"), InFilename);
+
+#define UE_LOG_COOK_TIME(TimeType) \
+{\
+	double CurrentTime = FPlatformTime::Seconds(); \
+	UE_LOG(LogSavePackage, Log, TEXT("TIMING: %s took %fms"), TimeType, (CurrentTime - PreviousTime) * 1000.0f); \
+	PreviousTime = CurrentTime; \
+}
+
+#define UE_FINISH_LOG_COOK_TIME()\
+	UE_LOG_COOK_TIME(TEXT("Unaccounted time")); \
+{\
+	double CurrentTime = FPlatformTime::Seconds(); \
+	UE_LOG(LogSavePackage, Log, TEXT("Total save time: %fms"), (CurrentTime - StartTime) * 1000.0f); \
+}
+
+#else
+// don't do anything
+#define UE_START_LOG_COOK_TIME(InFilename)
+#define UE_LOG_COOK_TIME(TimeType) 
+#define UE_FINISH_LOG_COOK_TIME()
+#endif
+
+
 static bool HasDeprecatedOrPendingKillOuter(UObject* InObj, UPackage* InSavingPackage)
 {
 	UObject* Obj = InObj;
@@ -1194,6 +1225,34 @@ private:
 	 */
 	static FObjectExportSortHelper* Sorter;
 
+	struct FObjectFullName
+	{
+	public:
+		FObjectFullName(const UObject* Object, const UObject* Root)
+		{
+			ClassName = Object->GetClass()->GetFName();
+			const UObject* Current = Object;
+			while (Current != nullptr && Current != Root)
+			{
+				Path.Insert(Current->GetFName(), 0);
+				Current = Current->GetOuter();
+			}
+		}
+
+		FObjectFullName(FObjectFullName&& InFullName)
+		{
+			ClassName = InFullName.ClassName;
+			Swap(Path, InFullName.Path);
+		}
+		FName ClassName;
+		TArray<FName> Path;
+	};
+
+	bool bUseFObjectFullName;
+
+
+	TMap<UObject*, FObjectFullName> ObjectToObjectFullNameMap;
+
 	/**
 	 * Map of UObject => full name; optimization for sorting.
 	 */
@@ -1218,18 +1277,53 @@ private:
 		}
 		else
 		{
-			FString* FullNameA = Sorter->ObjectToFullNameMap.Find(A.Object);
-			FString* FullNameB = Sorter->ObjectToFullNameMap.Find(B.Object);
-			checkSlow(FullNameA);
-			checkSlow(FullNameB);
+			if (bUseFObjectFullName)
+			{
+				const FObjectFullName* FullNameA = Sorter->ObjectToObjectFullNameMap.Find(A.Object);
+				const FObjectFullName* FullNameB = Sorter->ObjectToObjectFullNameMap.Find(B.Object);
+				checkSlow(FullNameA);
+				checkSlow(FullNameB);
 
-			Result = FCString::Stricmp(**FullNameA, **FullNameB);
+				if (FullNameA->ClassName != FullNameB->ClassName)
+				{
+					Result = FCString::Stricmp(*FullNameA->ClassName.ToString(), *FullNameB->ClassName.ToString());
+				}
+				else
+				{
+					int Num = FMath::Min(FullNameA->Path.Num(), FullNameB->Path.Num());
+					for (int I = 0; I < Num; ++I)
+					{
+						if (FullNameA->Path[I] != FullNameB->Path[I])
+						{
+							Result = FCString::Stricmp(*FullNameA->Path[I].ToString(), *FullNameB->Path[I].ToString());
+							break;
+						}
+					}
+					if (Result == 0)
+					{
+						Result = FullNameA->Path.Num() - FullNameB->Path.Num();
+					}
+				}
+			}
+			else
+			{
+				FString* FullNameA = Sorter->ObjectToFullNameMap.Find(A.Object);
+				FString* FullNameB = Sorter->ObjectToFullNameMap.Find(B.Object);
+				checkSlow(FullNameA);
+				checkSlow(FullNameB);
+
+				Result = FCString::Stricmp(**FullNameA, **FullNameB);
+			}
 		}
 
 		return Result < 0;
 	}
 
 public:
+
+	FObjectExportSortHelper() : bUseFObjectFullName(false)
+	{}
+
 	/**
 	 * Sorts exports alphabetically.  If a package is specified to be conformed against, ensures that the order
 	 * of the exports match the order in which the corresponding exports occur in the old package.
@@ -1237,8 +1331,10 @@ public:
 	 * @param	Linker				linker containing the exports that need to be sorted
 	 * @param	LinkerToConformTo	optional linker to conform against.
 	 */
-	void SortExports( FLinkerSave* Linker, FLinkerLoad* LinkerToConformTo=NULL )
+	void SortExports( FLinkerSave* Linker, FLinkerLoad* LinkerToConformTo=NULL, bool InbUseFObjectFullName = false)
 	{
+		bUseFObjectFullName = InbUseFObjectFullName;
+
 		int32 SortStartPosition=0;
 		if ( LinkerToConformTo )
 		{
@@ -1267,6 +1363,9 @@ public:
 					// Set the index (key) in the map to the index of this object into the export map.
 					OriginalExportIndexes.Add( *ExportFullName, ExportIndex );
 					ObjectToFullNameMap.Add(Export.Object, *ExportFullName);
+
+					FObjectFullName ObjectFullName(Export.Object, Linker->LinkerRoot);
+					ObjectToObjectFullNameMap.Add(Export.Object, MoveTemp(ObjectFullName) );
 				}
 			}
 
@@ -1442,56 +1541,156 @@ class FExportReferenceSorter : public FArchiveUObject
 	 */
 	void InitializeCoreClasses()
 	{
-		// initialize the tracking maps with the core classes
-		UClass* CoreClassList[]=
-		{
-			UObject::StaticClass(),
-			UField::StaticClass(),
-			UStruct::StaticClass(),
-			UScriptStruct::StaticClass(),
-			UFunction::StaticClass(),
-			UEnum::StaticClass(),
-			UClass::StaticClass(),
-			UProperty::StaticClass(),
-			UByteProperty::StaticClass(),
-			UIntProperty::StaticClass(),
-			UBoolProperty::StaticClass(),
-			UFloatProperty::StaticClass(),
-			UDoubleProperty::StaticClass(),
-			UObjectProperty::StaticClass(),
-			UClassProperty::StaticClass(),
-			UInterfaceProperty::StaticClass(),
-			UNameProperty::StaticClass(),
-			UStrProperty::StaticClass(),
-			UArrayProperty::StaticClass(),
-			UTextProperty::StaticClass(),
-			UStructProperty::StaticClass(),
-			UDelegateProperty::StaticClass(),
-			UInterface::StaticClass(),
-			UMulticastDelegateProperty::StaticClass(),
-			UWeakObjectProperty::StaticClass(),
-			UObjectPropertyBase::StaticClass(),
-			ULazyObjectProperty::StaticClass(),
-			UAssetObjectProperty::StaticClass(),
-			UAssetClassProperty::StaticClass()
+		check(CoreClasses.Num() == 0);
+		check(ReferencedObjects.Num() == 0);
+		check(ForceLoadObjects.Num() == 0);
+		check(SerializedObjects.Num() == 0);
+		static bool bInitializedStaticCoreClasses = false;
+		static TArray<UClass*> StaticCoreClasses;
+		static TArray<UObject*> StaticCoreReferencedObjects;
+		static FOrderedObjectSet StaticProcessedObjects;
+		static TArray<UObject*> StaticForceLoadObjects;
+		static TSet<UObject*> StaticSerializedObjects;
+		static bool StaticbIgnoreFieldReferences = false;
+
+		// Helper class to register FlushInitializedStaticCoreClasses callback on first SavePackage run
+		struct FAddFlushInitalizedStaticCoreClasses
+		{
+			FAddFlushInitalizedStaticCoreClasses() 
+			{
+				FCoreUObjectDelegates::PreGarbageCollect.AddStatic(FlushInitalizedStaticCoreClasses);
+			}
+			/** Wrapper function to handle default parameter when used as function pointer */
+			static void FlushInitalizedStaticCoreClasses()
+			{
+				bInitializedStaticCoreClasses = false;
+			}
 		};
+		static FAddFlushInitalizedStaticCoreClasses MaybeAddAddFlushInitializedStaticCoreClasses;
 
-		for ( int32 CoreClassIndex = 0; CoreClassIndex < ARRAY_COUNT(CoreClassList); CoreClassIndex++ )
-		{
-			UClass* CoreClass = CoreClassList[CoreClassIndex];
-			CoreClasses.AddUnique(CoreClass);
 
-			ReferencedObjects.Add(CoreClass);
-			ReferencedObjects.Add(CoreClass->GetDefaultObject());
+		if (bIgnoreFieldReferences != StaticbIgnoreFieldReferences)
+		{
+			StaticbIgnoreFieldReferences = bIgnoreFieldReferences;
+			bInitializedStaticCoreClasses = false;
 		}
 
-		for ( int32 CoreClassIndex = 0; CoreClassIndex < CoreClasses.Num(); CoreClassIndex++ )
+#if VALIDATE_INITIALIZECORECLASSES
+		bool bWasValid = bInitializedStaticCoreClasses;
+		bInitializedStaticCoreClasses = false;
+#endif
+
+		if (!bInitializedStaticCoreClasses)
 		{
-			UClass* CoreClass = CoreClasses[CoreClassIndex];
-			ProcessStruct(CoreClass);
+			bInitializedStaticCoreClasses = true;
+
+
+			// initialize the tracking maps with the core classes
+			UClass* CoreClassList[] =
+			{
+				UObject::StaticClass(),
+				UField::StaticClass(),
+				UStruct::StaticClass(),
+				UScriptStruct::StaticClass(),
+				UFunction::StaticClass(),
+				UEnum::StaticClass(),
+				UClass::StaticClass(),
+				UProperty::StaticClass(),
+				UByteProperty::StaticClass(),
+				UIntProperty::StaticClass(),
+				UBoolProperty::StaticClass(),
+				UFloatProperty::StaticClass(),
+				UDoubleProperty::StaticClass(),
+				UObjectProperty::StaticClass(),
+				UClassProperty::StaticClass(),
+				UInterfaceProperty::StaticClass(),
+				UNameProperty::StaticClass(),
+				UStrProperty::StaticClass(),
+				UArrayProperty::StaticClass(),
+				UTextProperty::StaticClass(),
+				UStructProperty::StaticClass(),
+				UDelegateProperty::StaticClass(),
+				UInterface::StaticClass(),
+				UMulticastDelegateProperty::StaticClass(),
+				UWeakObjectProperty::StaticClass(),
+				UObjectPropertyBase::StaticClass(),
+				ULazyObjectProperty::StaticClass(),
+				UAssetObjectProperty::StaticClass(),
+				UAssetClassProperty::StaticClass()
+			};
+
+			for (int32 CoreClassIndex = 0; CoreClassIndex < ARRAY_COUNT(CoreClassList); CoreClassIndex++)
+			{
+				UClass* CoreClass = CoreClassList[CoreClassIndex];
+				CoreClasses.AddUnique(CoreClass);
+
+				ReferencedObjects.Add(CoreClass);
+				ReferencedObjects.Add(CoreClass->GetDefaultObject());
+			}
+
+			for (int32 CoreClassIndex = 0; CoreClassIndex < CoreClasses.Num(); CoreClassIndex++)
+			{
+				UClass* CoreClass = CoreClasses[CoreClassIndex];
+				ProcessStruct(CoreClass);
+			}
+
+			CoreReferencesOffset = ReferencedObjects.Num();
+
+
+#if VALIDATE_INITIALIZECORECLASSES
+			if (bWasValid)
+			{
+				// make sure everything matches up 
+				check(CoreClasses.Num() == StaticCoreClasses.Num());
+				check(ReferencedObjects.Num() == StaticCoreReferencedObjects.Num());
+				check(ProcessedObjects.Num() == StaticProcessedObjects.Num());
+				check(ForceLoadObjects.Num() == StaticForceLoadObjects.Num());
+				check(SerializedObjects.Num() == StaticSerializedObjects.Num());
+				
+				
+				for (int I = 0; I < CoreClasses.Num(); ++I)
+				{
+					check(CoreClasses[I] == StaticCoreClasses[I]);
+				}
+				for (int I = 0; I < ReferencedObjects.Num(); ++I)
+				{
+					check(ReferencedObjects[I] == StaticCoreReferencedObjects[I]);
+				}
+				for (const auto& ProcessedObject : ProcessedObjects.ObjectsSet)
+				{
+					check(ProcessedObject.Value == StaticProcessedObjects.Find(ProcessedObject.Key));
+				}
+				for (int I = 0; I < ForceLoadObjects.Num(); ++I)
+				{
+					check(ForceLoadObjects[I] == StaticForceLoadObjects[I]);
+				}
+				for (const auto& SerializedObject : SerializedObjects)
+				{
+					check(StaticSerializedObjects.Find(SerializedObject));
+				}
+			}
+#endif
+
+			StaticCoreClasses = CoreClasses;
+			StaticCoreReferencedObjects = ReferencedObjects;
+			StaticProcessedObjects = ProcessedObjects;
+			StaticForceLoadObjects = ForceLoadObjects;
+			StaticSerializedObjects = SerializedObjects;
+
+			check(CurrentClass == NULL);
+
+
+
 		}
 
-		CoreReferencesOffset = ReferencedObjects.Num();
+		CoreClasses = StaticCoreClasses;
+		ReferencedObjects = StaticCoreReferencedObjects;
+		ProcessedObjects = StaticProcessedObjects;
+		ForceLoadObjects = StaticForceLoadObjects;
+		SerializedObjects = StaticSerializedObjects;
+
+		CoreReferencesOffset = StaticCoreReferencedObjects.Num();
+		
 	}
 
 	/**
@@ -1934,7 +2133,35 @@ private:
 	/**
 	 * The list of objects that have been evaluated by this archive so far.
 	 */
-	TArray<UObject*> ProcessedObjects;
+	struct FOrderedObjectSet
+	{
+		TMap<UObject*, int32> ObjectsSet;
+		// TArray<UObject*> ObjectsList;
+
+		int32 Add(UObject* Object)
+		{
+			// int32 Index = ObjectsList.Add(Object);
+			int32 Index = ObjectsSet.Num(); // never use the list anyway so no point in even having it
+			ObjectsSet.Add(Object, Index);
+			return Index;
+		}
+
+		inline int32 Find(UObject* Object) const
+		{
+			const int32 *Index = ObjectsSet.Find(Object);
+			if (Index)
+			{
+				return *Index;
+			}
+			return INDEX_NONE;
+		}
+		inline int32 Num() const
+		{
+			return ObjectsSet.Num();
+		}
+	};
+	FOrderedObjectSet ProcessedObjects;
+	
 
 	/**
 	 * The list of objects that have been serialized; used to prevent calling Serialize on an object more than once.
@@ -1968,7 +2195,7 @@ private:
 	TArray<UObject*> ForceLoadObjects;
 };
 
-#define EXPORT_SORTING_DETAILED_LOGGING 0
+
 
 /**
  * Helper structure encapsulating functionality to sort a linker's export map to allow seek free
@@ -1985,6 +2212,9 @@ struct FObjectExportSeekFreeSorter
 	 */
 	void SortExports( FLinkerSave* Linker, FLinkerLoad* LinkerToConformTo )
 	{
+		UE_START_LOG_COOK_TIME(*Linker->Filename);
+
+
 		SortArchive.SetCookingTarget(Linker->CookingTarget());
 
 		int32					FirstSortIndex = LinkerToConformTo ? LinkerToConformTo->ExportMap.Num() : 0;
@@ -2001,6 +2231,8 @@ struct FObjectExportSeekFreeSorter
 			}
 		}
 
+		UE_LOG_COOK_TIME(TEXT("Setup Original Exports"));
+
 		bool bRetrieveInitialReferences = true;
 
 		// Now we need to sort the export list according to the order in which objects will be loaded.  For the sake of simplicity, 
@@ -2019,7 +2251,7 @@ struct FObjectExportSeekFreeSorter
 				UE_LOG(LogSavePackage, Log, TEXT("Referenced objects for (%i) %s in %s"), ExportIndex, *Export.Object->GetFullName(), *Linker->LinkerRoot->GetName());
 				for ( int32 RefIndex = 0; RefIndex < ReferencedObjects.Num(); RefIndex++ )
 				{
-					UE_LOG(LogSavePackage, Log, TEXT("\t%i) %s"), RefIndex, *ReferencedObjects(RefIndex)->GetFullName());
+					UE_LOG(LogSavePackage, Log, TEXT("\t%i) %s"), RefIndex, *ReferencedObjects[RefIndex]->GetFullName());
 				}
 				if ( ReferencedObjects.Num() > 1 )
 				{
@@ -2036,6 +2268,8 @@ struct FObjectExportSeekFreeSorter
 
 		}
 
+		UE_LOG_COOK_TIME(TEXT("Process Class Exports"));
+
 #if EXPORT_SORTING_DETAILED_LOGGING
 		UE_LOG(LogSavePackage, Log, TEXT("*************   Processed %i classes out of %i possible exports for package %s.  Beginning second pass...   *************"), SortedExports.Num(), Linker->ExportMap.Num() - FirstSortIndex, *Linker->LinkerRoot->GetName());
 #endif
@@ -2055,7 +2289,7 @@ struct FObjectExportSeekFreeSorter
 				UE_LOG(LogSavePackage, Log, TEXT("Referenced objects for (%i) %s in %s"), ExportIndex, *Export.Object->GetFullName(), *Linker->LinkerRoot->GetName());
 				for ( int32 RefIndex = 0; RefIndex < ReferencedObjects.Num(); RefIndex++ )
 				{
-					UE_LOG(LogSavePackage, Log, TEXT("\t%i) %s"), RefIndex, *ReferencedObjects(RefIndex)->GetFullName());
+					UE_LOG(LogSavePackage, Log, TEXT("\t%i) %s"), RefIndex, *ReferencedObjects[RefIndex]->GetFullName());
 				}
 				if ( ReferencedObjects.Num() > 1 )
 				{
@@ -2070,6 +2304,9 @@ struct FObjectExportSeekFreeSorter
 				bRetrieveInitialReferences = false;
 			}
 		}
+
+		UE_LOG_COOK_TIME(TEXT("Process Objects Exports"));
+
 #if EXPORT_SORTING_DETAILED_LOGGING
 		SortArchive.VerifySortingAlgorithm();
 #endif
@@ -2097,6 +2334,8 @@ struct FObjectExportSeekFreeSorter
 			}
 		}
 
+		UE_LOG_COOK_TIME(TEXT("Fill Export Map"));
+
 		// Manually add any new NULL exports last as they won't be in the SortedExportsObjects list. 
 		// A NULL Export.Object can occur if you are e.g. saving an object in the game that is 
 		// OBJECTMARK_NotForClient.
@@ -2108,6 +2347,8 @@ struct FObjectExportSeekFreeSorter
 				Linker->ExportMap.Add( Export );
 			}
 		}
+
+		UE_FINISH_LOG_COOK_TIME();
 	}
 
 private:
@@ -2353,32 +2594,6 @@ EObjectMark UPackage::GetObjectMarksForTargetPlatform( const class ITargetPlatfo
  *
  * @return	true if the package was saved successfully.
  */
-#define UE_PROFILE_COOKSAVE 0
-#if UE_PROFILE_COOKSAVE
-
-#define UE_START_LOG_COOK_TIME(InFilename) double PreviousTime; double StartTime; PreviousTime = StartTime = FPlatformTime::Seconds(); \
-	UE_LOG( LogSavePackage, Log, TEXT("Starting save package for %s"), InFilename);
-
-#define UE_LOG_COOK_TIME(TimeType) \
-	{\
-		double CurrentTime = FPlatformTime::Seconds(); \
-		UE_LOG( LogSavePackage, Log, TEXT("TIMING: %s took %fms"), TimeType, (CurrentTime - PreviousTime) * 1000.0f ); \
-		PreviousTime = CurrentTime; \
-	}
-
-#define UE_FINISH_LOG_COOK_TIME()\
-	UE_LOG_COOK_TIME( TEXT("Unaccounted time")); \
-	{\
-		double CurrentTime = FPlatformTime::Seconds(); \
-		UE_LOG( LogSavePackage, Log, TEXT("Total save time: %fms"), (CurrentTime - StartTime) * 1000.0f ); \
-	}
-
-#else
-// don't do anything
-#define UE_START_LOG_COOK_TIME(InFilename)
-#define UE_LOG_COOK_TIME(TimeType) 
-#define UE_FINISH_LOG_COOK_TIME()
-#endif
 
 #if WITH_EDITOR
 COREUOBJECT_API extern bool GOutputCookingWarnings;
@@ -3131,12 +3346,16 @@ bool UPackage::SavePackage( UPackage* InOuter, UObject* Base, EObjectFlags TopLe
 				if ( EndSavingIfCancelled( Linker, TempFilename ) ) { return false; }
 				SlowTask.EnterProgressFrame();
 
+				UE_LOG_COOK_TIME(TEXT("Sort Exports nonseekfree"));
+
 				// Sort exports for seek-free loading.
 				FObjectExportSeekFreeSorter SeekFreeSorter;
+				UE_LOG_COOK_TIME(TEXT("Sort Exports seekfree constructor"));
 				SeekFreeSorter.SortExports( Linker, Conform );
+				UE_LOG_COOK_TIME(TEXT("Sort Exports seekfree inner"));
 				Linker->Summary.ExportCount = Linker->ExportMap.Num();
 				
-				UE_LOG_COOK_TIME(TEXT("Sort Exports"));
+				UE_LOG_COOK_TIME(TEXT("Sort Exports seekfree"));
 
 				if ( EndSavingIfCancelled( Linker, TempFilename ) ) { return false; }
 				SlowTask.EnterProgressFrame();
@@ -3642,6 +3861,7 @@ bool UPackage::SavePackage( UPackage* InOuter, UObject* Base, EObjectFlags TopLe
 						{
 							Linker->Detach();
 						}
+						UE_LOG_COOK_TIME(TEXT("Compress from memory"));
 					}
 					// Compress the temporarily file to destination.
 					else if( InOuter->PackageFlags & PKG_StoreCompressed )
@@ -3649,6 +3869,8 @@ bool UPackage::SavePackage( UPackage* InOuter, UObject* Base, EObjectFlags TopLe
 						UE_LOG(LogSavePackage, Log,  TEXT("Compressing '%s' to '%s'"), *TempFilename, *NewPath );
 						FFileCompressionHelper CompressionHelper;
 						Success = CompressionHelper.CompressFile( *TempFilename, *NewPath, Linker );
+
+						UE_LOG_COOK_TIME(TEXT("Compress from temp file"));
 					}
 					// Fully compress package in one "block".
 					else if( InOuter->PackageFlags & PKG_StoreFullyCompressed )
@@ -3656,6 +3878,8 @@ bool UPackage::SavePackage( UPackage* InOuter, UObject* Base, EObjectFlags TopLe
 						UE_LOG(LogSavePackage, Log,  TEXT("Full-package compressing '%s' to '%s'"), *TempFilename, *NewPath );
 						FFileCompressionHelper CompressionHelper;
 						Success = CompressionHelper.FullyCompressFile( *TempFilename, *NewPath, Linker->ForceByteSwapping() );
+
+						UE_LOG_COOK_TIME(TEXT("Fully compress from temp file"));
 					}
 					else if (bSaveAsync)
 					{
@@ -3667,6 +3891,7 @@ bool UPackage::SavePackage( UPackage* InOuter, UObject* Base, EObjectFlags TopLe
 							AsyncWriteFile(*(FBufferArchive*)(Linker->Saver), *NewPath, FinalTimeStamp);
 							Linker->Detach();
 						}
+						UE_LOG_COOK_TIME(TEXT("Async write"));
 					}
 					// Move the temporary file.
 					else
@@ -3677,6 +3902,8 @@ bool UPackage::SavePackage( UPackage* InOuter, UObject* Base, EObjectFlags TopLe
 						{
 							IFileManager::Get().SetTimeStamp(*NewPath, FinalTimeStamp);
 						}
+
+						UE_LOG_COOK_TIME(TEXT("Move file"));
 					}
 
 					// Make sure to clean up any stale .uncompressed_size files.
@@ -3792,6 +4019,8 @@ bool UPackage::SavePackage( UPackage* InOuter, UObject* Base, EObjectFlags TopLe
 
 		UE_FINISH_LOG_COOK_TIME();
 
+		UE_LOG(LogSavePackage, Warning, TEXT("Finished SavePackage %s"), Filename);
+
 		return Success;
 	}
 	else
