[UE-13579] Relinking dependent blueprint subclasses on-load, when their parent class is recompiled (as part of FScopedClassDependencyGather).

diff --git a/Engine/Source/Runtime/CoreUObject/Private/Blueprint/BlueprintSupport.cpp b/Engine/Source/Runtime/CoreUObject/Private/Blueprint/BlueprintSupport.cpp
--- a/Engine/Source/Runtime/CoreUObject/Private/Blueprint/BlueprintSupport.cpp
+++ b/Engine/Source/Runtime/CoreUObject/Private/Blueprint/BlueprintSupport.cpp
@@ -127,18 +127,49 @@ FScopedClassDependencyGather::~FScopedClassDependencyGather()
 	// dependencies (unless compiling on load is explicitly disabled)
 	if( bMasterClass && !GForceDisableBlueprintCompileOnLoad )
 	{
-		for( auto DependencyIter = BatchClassDependencies.CreateIterator(); DependencyIter; ++DependencyIter )
+		auto DependencyIter = BatchClassDependencies.CreateIterator();
+		// implemented as a lambda, to prevent duplicated code between 
+		// BatchMasterClass and BatchClassDependencies entries
+		auto RecompileClassLambda = [&DependencyIter](UClass* Class)
+		{
+			Class->ConditionalRecompileClass(&FUObjectThreadContext::Get().ObjLoaded);
+
+			// because of the above call to ConditionalRecompileClass(), the 
+			// specified Class gets "cleaned and sanitized" (meaning its old 
+			// properties get moved to a TRASH class, and new ones are 
+			// constructed in their place)... the unfortunate side-effect of 
+			// this is that child classes that have already been linked are now
+			// referencing TRASH inherited properties; to resolve this issue, 
+			// here we go back through dependencies that were already recompiled
+			// and re-link any that are sub-classes
+			//
+			// @TODO: this isn't the most optimal solution to this problem; we 
+			//        should probably instead prevent CleanAndSanitizeClass()
+			//        from running for BytecodeOnly compiles (we would then need 
+			//        to block UField re-creation)... UE-14957 was created to 
+			//        track this issue
+			auto ReverseIt = DependencyIter;
+			for (--ReverseIt; ReverseIt.GetIndex() >= 0; --ReverseIt)
+			{
+				UClass* ProcessedDependency = *ReverseIt;
+				if (ProcessedDependency->IsChildOf(Class))
+				{
+					ProcessedDependency->StaticLink(/*bRelinkExistingProperties =*/true);
+				}
+			}
+		};
+
+		for( ; DependencyIter; ++DependencyIter )
 		{
 			UClass* Dependency = *DependencyIter;
 			if( Dependency->ClassGeneratedBy != BatchMasterClass->ClassGeneratedBy )
 			{
-				Dependency->ConditionalRecompileClass(&FUObjectThreadContext::Get().ObjLoaded);
+				RecompileClassLambda(Dependency);
 			}
 		}
 
 		// Finally, recompile the master class to make sure it gets updated too
-		BatchMasterClass->ConditionalRecompileClass(&FUObjectThreadContext::Get().ObjLoaded);
-		
+		RecompileClassLambda(BatchMasterClass);
 		BatchMasterClass = NULL;
 	}
 }
