Amending a deferred dependency check() so that it accounts for deferred functions that have been regenerated already.

diff --git a/Engine/Source/Runtime/CoreUObject/Private/Blueprint/BlueprintSupport.cpp b/Engine/Source/Runtime/CoreUObject/Private/Blueprint/BlueprintSupport.cpp
--- a/Engine/Source/Runtime/CoreUObject/Private/Blueprint/BlueprintSupport.cpp
+++ b/Engine/Source/Runtime/CoreUObject/Private/Blueprint/BlueprintSupport.cpp
@@ -935,14 +935,26 @@ int32 FLinkerLoad::ResolveDependencyPlaceholder(FLinkerPlaceholderBase* Placehol
 
 #if USE_DEFERRED_DEPENDENCY_CHECK_VERIFICATION_TESTS
 	UFunction* AsFunction = Cast<UFunction>(RealImportObj);
+	UClass* FunctionOwner = (AsFunction != nullptr) ? AsFunction->GetOwnerClass() : nullptr;
 	// it's ok if super functions come in not fully loaded (missing 
 	// RF_LoadCompleted... meaning it's in the middle of serializing in somewhere 
 	// up the stack); the function will be forcefully ran through Preload(), 
 	// when we regenerate the super class (see FRegenerationHelper::ForcedLoadMembers)
-	bool const bIsSuperFunction = (AsFunction != nullptr) && (ReferencingClass != nullptr) && ReferencingClass->IsChildOf(AsFunction->GetOwnerClass());
+	bool const bIsSuperFunction   = (AsFunction != nullptr) && (ReferencingClass != nullptr) && ReferencingClass->IsChildOf(FunctionOwner);
+	// it's also possible that the loaded version of this function has been 
+	// thrown out and replaced with a regenerated version (presumably from a
+	// blueprint compiling on load)... if that's the case, then this function 
+	// will not have a corresponding linker assigned to it
+	bool const bIsRegeneratedFunc = (AsFunction != nullptr) && (AsFunction->GetLinker() == nullptr);
+
+	bool const bExpectsLoadCompleteFlag = (RealImportObj != nullptr) && !bIsSuperFunction && !bIsRegeneratedFunc;
+	// if we can't rely on the Import object's RF_LoadCompleted flag, then its
+	// owner class should at least have it
+	DEFERRED_DEPENDENCY_CHECK( (RealImportObj == nullptr) || bExpectsLoadCompleteFlag ||
+		(FunctionOwner && FunctionOwner->HasAnyFlags(RF_LoadCompleted)) );
 
 	DEFERRED_DEPENDENCY_CHECK(RealImportObj != PlaceholderObj);
-	DEFERRED_DEPENDENCY_CHECK(RealImportObj == nullptr || bIsSuperFunction || RealImportObj->HasAnyFlags(RF_LoadCompleted));
+	DEFERRED_DEPENDENCY_CHECK(!bExpectsLoadCompleteFlag || RealImportObj->HasAnyFlags(RF_LoadCompleted));
 #endif // USE_DEFERRED_DEPENDENCY_CHECK_VERIFICATION_TESTS
 
 	int32 ReplacementCount = 0;
