[UE-13996] Retain references to other Actor instances on copy/paste from one level to another.

change summary:
- modified UObjectPropertyBase::GetExportPath() to also avoid fully-qualified names when PPF_ExportsNotFullyQualified is set and the object is a sibling of the root scope

diff --git a/Engine/Source/Runtime/CoreUObject/Private/UObject/PropertyBaseObject.cpp b/Engine/Source/Runtime/CoreUObject/Private/UObject/PropertyBaseObject.cpp
--- a/Engine/Source/Runtime/CoreUObject/Private/UObject/PropertyBaseObject.cpp
+++ b/Engine/Source/Runtime/CoreUObject/Private/UObject/PropertyBaseObject.cpp
@@ -134,6 +134,13 @@ FString UObjectPropertyBase::GetExportPath(const UObject* Object, const UObject*
 	{
 		StopOuter = (ExportRootScope || (Parent == nullptr)) ? ExportRootScope : Parent->GetOutermost();
 		bExportFullyQualified = !Object->IsIn(StopOuter);
+
+		// Also don't fully qualify the name if it's a sibling of the root scope, since it may be included in the exported set of objects
+		if (bExportFullyQualified)
+		{
+			StopOuter = StopOuter->GetOuter();
+			bExportFullyQualified = (StopOuter == nullptr) || (!Object->IsIn(StopOuter));
+		}
 	}
 
 	// if we want a full qualified object reference, use the pathname, otherwise, use just the object name
