GIsReinstancing flag added to supress warning about executing script form an obsolete class.

https://udn.unrealengine.com/questions/233181/class-newerversionexists-assert-when-compiling-paw.html

diff --git a/Engine/Source/Editor/UnrealEd/Private/Kismet2/KismetReinstanceUtilities.cpp b/Engine/Source/Editor/UnrealEd/Private/Kismet2/KismetReinstanceUtilities.cpp
--- a/Engine/Source/Editor/UnrealEd/Private/Kismet2/KismetReinstanceUtilities.cpp
+++ b/Engine/Source/Editor/UnrealEd/Private/Kismet2/KismetReinstanceUtilities.cpp
@@ -556,6 +556,8 @@ void FBlueprintCompileReinstancer::ReinstanceObjects(bool bForceAlwaysReinstance
 
 			ensure(0 == DependentBlueprintsToRecompile.Num());
 
+			TGuardValue<bool> ReinstancingGuard(GIsReinstancing, true);
+
 			TArray<TSharedPtr<FReinstanceFinalizer>> Finalizers;
 
 			// All children were recompiled. It's safe to reinstance.
diff --git a/Engine/Source/Runtime/Core/Private/Misc/Core.cpp b/Engine/Source/Runtime/Core/Private/Misc/Core.cpp
--- a/Engine/Source/Runtime/Core/Private/Misc/Core.cpp
+++ b/Engine/Source/Runtime/Core/Private/Misc/Core.cpp
@@ -90,6 +90,9 @@ bool GIsReconstructingBlueprintInstances = false;
 /** Force blueprints to not compile on load */
 bool GForceDisableBlueprintCompileOnLoad = false;
 
+/** True if actors and objects are being re-instanced. */
+bool GIsReinstancing = false;
+
 #if WITH_ENGINE
 bool					PRIVATE_GIsRunningCommandlet			= false;				/* Whether this executable is running a commandlet (custom command-line processing code) */
 bool					PRIVATE_GAllowCommandletRendering	= false;				/** If true, initialise RHI and set up scene for rendering even when running a commandlet. */
diff --git a/Engine/Source/Runtime/Core/Public/CoreGlobals.h b/Engine/Source/Runtime/Core/Public/CoreGlobals.h
--- a/Engine/Source/Runtime/Core/Public/CoreGlobals.h
+++ b/Engine/Source/Runtime/Core/Public/CoreGlobals.h
@@ -52,6 +52,8 @@ extern CORE_API bool GIsReconstructingBlueprintInstances;
 /** Force blueprints to not compile on load */
 extern CORE_API bool GForceDisableBlueprintCompileOnLoad;
 
+/** True if actors and objects are being re-instanced. */
+extern CORE_API bool GIsReinstancing;
 
 /** Helper function to flush resource streaming. */
 extern CORE_API void(*GFlushStreamingFunc)(void);
diff --git a/Engine/Source/Runtime/CoreUObject/Private/UObject/ScriptCore.cpp b/Engine/Source/Runtime/CoreUObject/Private/UObject/ScriptCore.cpp
--- a/Engine/Source/Runtime/CoreUObject/Private/UObject/ScriptCore.cpp
+++ b/Engine/Source/Runtime/CoreUObject/Private/UObject/ScriptCore.cpp
@@ -619,11 +619,14 @@ void UObject::ProcessInternal( FFrame& Stack, RESULT_DECL )
 	// remove later when stable
 	if (GetClass()->HasAnyClassFlags(CLASS_NewerVersionExists))
 	{
-		static int32 num=0;
-		num++;
-		if (num < 5)
+		if (!GIsReinstancing)
 		{
-			ensureMsgf(!GetClass()->HasAnyClassFlags(CLASS_NewerVersionExists), TEXT("Object '%s' is being used for execution, but its class is out of date and has been replaced with a recompiled class!"), *GetFullName());
+			static int32 num = 0;
+			num++;
+			if (num < 5)
+			{
+				ensureMsgf(!GetClass()->HasAnyClassFlags(CLASS_NewerVersionExists), TEXT("Object '%s' is being used for execution, but its class is out of date and has been replaced with a recompiled class!"), *GetFullName());
+			}
 		}
 		return;
 	}
