Cis fix for clang

diff --git a/Engine/Source/Runtime/CoreUObject/Public/UObject/Class.h b/Engine/Source/Runtime/CoreUObject/Public/UObject/Class.h
--- a/Engine/Source/Runtime/CoreUObject/Public/UObject/Class.h
+++ b/Engine/Source/Runtime/CoreUObject/Public/UObject/Class.h
@@ -2753,11 +2753,6 @@ inline T* GetMutableDefault(UClass *Class)
 
 template< class T > struct TBaseStructure
 {
-	static UScriptStruct* Get()
-	{
-		static_assert(false, "Unknown base structure, try to use T::StaticStruct()");
-		return nullptr;
-	}
 };
 
 template<> struct TBaseStructure<FRotator>
