Making FAsyncLoadingThread::InsertPackage (more) thread safe

diff --git a/Engine/Source/Runtime/CoreUObject/Private/Serialization/AsyncLoading.cpp b/Engine/Source/Runtime/CoreUObject/Private/Serialization/AsyncLoading.cpp
--- a/Engine/Source/Runtime/CoreUObject/Private/Serialization/AsyncLoading.cpp
+++ b/Engine/Source/Runtime/CoreUObject/Private/Serialization/AsyncLoading.cpp
@@ -383,9 +383,33 @@ int32 FAsyncLoadingThread::CreateAsyncPackagesFromQueue()
 	return NumCreated;
 }
 
-/**
-* [ASYNC THREAD] Loads all packages
-*/
+void FAsyncLoadingThread::InsertPackage(FAsyncPackage* Package)
+{
+	checkSlow(IsInAsyncLoadThread());
+
+	// Incremented on the Async Thread, decremented on the game thread
+	AsyncLoadingCounter.Increment();
+
+	// Incemented and decremented on the AsyncThread
+	AsyncPackagesCounter.Increment();
+
+	{
+#if THREADSAFE_UOBJECTS
+		FScopeLock LockAsyncPackages(&AsyncPackagesCritical);
+#endif
+		// Insert new package keeping descending priority order in AsyncPackages
+		auto InsertIndex = AsyncPackages.IndexOfByPredicate([Package](const FAsyncPackage* Element)
+		{
+			return Element->GetPriority() <= Package->GetPriority();
+		});
+
+		InsertIndex = InsertIndex == INDEX_NONE ? AsyncPackages.Num() : InsertIndex;
+
+		AsyncPackages.InsertUninitialized(InsertIndex);
+		AsyncPackages[InsertIndex] = Package;
+	}
+}
+
 EAsyncPackageState::Type FAsyncLoadingThread::ProcessAsyncLoading(int32& OutPackagesProcessed, bool bUseTimeLimit /*= false*/, bool bUseFullTimeLimit /*= false*/, float TimeLimit /*= 0.0f*/, FName ExcludeType /*= NAME_None*/)
 {
 	SCOPE_CYCLE_COUNTER(STAT_FAsyncLoadingThread_ProcessAsyncLoading);
diff --git a/Engine/Source/Runtime/CoreUObject/Private/Serialization/AsyncLoadingThread.h b/Engine/Source/Runtime/CoreUObject/Private/Serialization/AsyncLoadingThread.h
--- a/Engine/Source/Runtime/CoreUObject/Private/Serialization/AsyncLoadingThread.h
+++ b/Engine/Source/Runtime/CoreUObject/Private/Serialization/AsyncLoadingThread.h
@@ -181,31 +181,7 @@ public:
 	*
 	* @param PackageName async package name.
 	*/
-	void InsertPackage(FAsyncPackage* Package)
-	{
-		checkSlow(IsInAsyncLoadThread());
-
-		// Incremented on the Async Thread, decremented on the game thread
-		AsyncLoadingCounter.Increment();
-
-		// Incemented and decremented on the AsyncThread
-		AsyncPackagesCounter.Increment();
-
-		// Insert new package keeping descending priority order in AsyncPackages
-		auto InsertIndex = AsyncPackages.IndexOfByPredicate([Package](const FAsyncPackage* Element)
-		{
-			return Element->GetPriority() <= Package->GetPriority();
-		});
-
-		InsertIndex = InsertIndex == -1 ? AsyncPackages.Num() : InsertIndex;
-		{
-#if THREADSAFE_UOBJECTS
-			FScopeLock LockAsyncPackages(&AsyncPackagesCritical);
-#endif
-			AsyncPackages.InsertUninitialized(InsertIndex);
-			AsyncPackages[InsertIndex] = Package;
-		}
-	}
+	void InsertPackage(FAsyncPackage* Package);
 
 	/**
 	* [ASYNC THREAD] Finds an existing async package in the LoadedPackages by its name.
