Removing debug code from FLinker

diff --git a/Engine/Source/Runtime/CoreUObject/Private/UObject/Linker.cpp b/Engine/Source/Runtime/CoreUObject/Private/UObject/Linker.cpp
--- a/Engine/Source/Runtime/CoreUObject/Private/UObject/Linker.cpp
+++ b/Engine/Source/Runtime/CoreUObject/Private/UObject/Linker.cpp
@@ -94,7 +94,6 @@ FName FLinkerTables::GetExportClassName( int32 i )
 ----------------------------------------------------------------------------*/
 FLinker::FLinker(ELinkerType::Type InType, UPackage* InRoot, const TCHAR* InFilename)
 : LinkerType(InType)
-, bDestroyed(false)
 , LinkerRoot( InRoot )
 , Filename( InFilename )
 , FilterClientButNotServer(false)
@@ -308,13 +307,6 @@ void FLinker::GetScriptSHAKey(uint8* OutKey)
 
 FLinker::~FLinker()
 {
-	if (bDestroyed)
-	{
-		static volatile int32 xx = 0;
-		xx++;
-	}
-	bDestroyed = true;
-
 	// free any SHA memory
 	delete ScriptSHA;
 }
diff --git a/Engine/Source/Runtime/CoreUObject/Public/UObject/Linker.h b/Engine/Source/Runtime/CoreUObject/Public/UObject/Linker.h
--- a/Engine/Source/Runtime/CoreUObject/Public/UObject/Linker.h
+++ b/Engine/Source/Runtime/CoreUObject/Public/UObject/Linker.h
@@ -930,10 +930,6 @@ private:
 
 	ELinkerType::Type LinkerType;
 
-protected:
-
-	bool bDestroyed;
-
 public:
 
 	/** The top-level UPackage object for the package associated with this linker */
