[UE-15057] Make sure default sub-objects are preloaded after their CDO owner is initialized (accounts for deferred CDO initialization).

diff --git a/Engine/Source/Runtime/CoreUObject/Private/Blueprint/BlueprintSupport.cpp b/Engine/Source/Runtime/CoreUObject/Private/Blueprint/BlueprintSupport.cpp
--- a/Engine/Source/Runtime/CoreUObject/Private/Blueprint/BlueprintSupport.cpp
+++ b/Engine/Source/Runtime/CoreUObject/Private/Blueprint/BlueprintSupport.cpp
@@ -1258,19 +1258,14 @@ void FLinkerLoad::ResolveDeferredExports(UClass* LoadClass)
 			BlueprintCDO->SetLinker(this, DeferredCDOIndex, /*bShouldDetatchExisting =*/false);
 			BlueprintCDO->SetFlags(OldFlags);
 		}
-
-		// Here we could make sure we're up to date with our parent. Alternatively we could guaranatee that
-		// SerializeDefaultObjects was called before we construct other default objects, but that is not easy to
-		// do. Simply moving the call to SerializeDefaultObject next to the call to create the CDO means
-		// the stream order changes...
 		DEFERRED_DEPENDENCY_CHECK(BlueprintCDO->GetClass() == LoadClass);
 
-		if (FObjectInitializer* DeferredInitializer = FDeferredObjInitializerTracker::Find(LoadClass))
-		{
-			FScriptIntegrationObjectHelper::PostConstructInitObject(*DeferredInitializer);
-			FDeferredObjInitializerTracker::Remove(LoadClass);
-		}
- 		DEFERRED_DEPENDENCY_CHECK(BlueprintCDO->GetClass() == LoadClass);
+		// if this class's CDO had its initialization deferred (presumably 
+		// because its super's CDO hadn't been fully serialized), then we want 
+		// to make sure it is ran here before we serialize in the CDO (so it 
+		// gets stocked with inherited values first, that the Preload() may or
+		// may not override)
+		FDeferredObjInitializerTracker::ResolveDeferredInitialization(LoadClass);
 
 		// should load the CDO (ensuring that it has been serialized in by the 
 		// time we get to class regeneration)
@@ -1485,10 +1480,79 @@ FObjectInitializer* FDeferredObjInitializerTracker::Find(UClass* LoadClass)
 	return ThreadInst.DeferredInitializers.Find(LoadClass);
 }
 
+bool FDeferredObjInitializerTracker::IsCdoDeferred(UClass* LoadClass)
+{
+	return (Find(LoadClass) != nullptr);
+}
+
+bool FDeferredObjInitializerTracker::DeferSubObjectPreload(UObject* SubObject)
+{
+	DEFERRED_DEPENDENCY_CHECK(SubObject->HasAnyFlags(RF_DefaultSubObject));
+
+	UObject* CdoOuter = SubObject->GetOuter();
+	UClass* OuterClass = CdoOuter->GetClass();
+
+	if (IsCdoDeferred(OuterClass))
+	{
+		DEFERRED_DEPENDENCY_CHECK(CdoOuter->HasAnyFlags(RF_ClassDefaultObject));
+
+		UObject* SubObjTemplate = SubObject->GetArchetype();
+		// check that this is an inherited sub-object (that it is defined on the 
+		// parent class)... we only need to defer Preload() for inherited 
+		// components because they're what gets filled out in the CDO's InitSubobjectProperties()
+		if (SubObjTemplate && (SubObjTemplate->GetOuter() != CdoOuter))
+		{
+			FDeferredObjInitializerTracker& ThreadInst = FDeferredObjInitializerTracker::Get();
+			ThreadInst.DeferredSubObjects.AddUnique(OuterClass, SubObject);
+
+			return true;
+		}
+	}
+	return false;
+}
+
 void FDeferredObjInitializerTracker::Remove(UClass* LoadClass)
 {
 	FDeferredObjInitializerTracker& ThreadInst = FDeferredObjInitializerTracker::Get();
 	ThreadInst.DeferredInitializers.Remove(LoadClass);
+	ThreadInst.DeferredSubObjects.Remove(LoadClass);
+}
+
+bool FDeferredObjInitializerTracker::ResolveDeferredInitialization(UClass* LoadClass)
+{
+	if (FObjectInitializer* DeferredInitializer = FDeferredObjInitializerTracker::Find(LoadClass))
+	{
+		// initializes and instances CDO properties (copies inherited values 
+		// from the super's CDO)
+		FScriptIntegrationObjectHelper::PostConstructInitObject(*DeferredInitializer);
+
+		FDeferredObjInitializerTracker& ThreadInst = FDeferredObjInitializerTracker::Get();
+		TArray<UObject*> DeferredSubObjects;
+		ThreadInst.DeferredSubObjects.MultiFind(LoadClass, DeferredSubObjects);
+
+		// remove before we call Preload() for sub-objects, else we'll end up 
+		// back in DeferSubObjectPreload()
+		FDeferredObjInitializerTracker::Remove(LoadClass);
+
+		FLinkerLoad* ClassLinker = LoadClass->GetLinker();
+		DEFERRED_DEPENDENCY_CHECK(ClassLinker != nullptr);
+		if (ClassLinker != nullptr)
+		{
+			// this all needs to happen after PostConstructInitObject() (above), 
+			// since InitSubObjectProperties() is invoked there (which is where 
+			// we fill this sub-object with values from the super)... here we
+			// account for any Preload() calls that were skipped on account of
+			// the deferred CDO initialization
+			for (UObject* SubObj : DeferredSubObjects)
+			{
+				ClassLinker->Preload(SubObj);
+			}
+		}
+
+		
+		return true;
+	}
+	return false;
 }
 
 // don't want other files ending up with this internal define
diff --git a/Engine/Source/Runtime/CoreUObject/Private/UObject/LinkerLoad.cpp b/Engine/Source/Runtime/CoreUObject/Private/UObject/LinkerLoad.cpp
--- a/Engine/Source/Runtime/CoreUObject/Private/UObject/LinkerLoad.cpp
+++ b/Engine/Source/Runtime/CoreUObject/Private/UObject/LinkerLoad.cpp
@@ -2948,6 +2948,23 @@ void FLinkerLoad::Preload( UObject* Object )
 	{
 		if (Object->GetLinker() == this)
 		{
+#if USE_CIRCULAR_DEPENDENCY_LOAD_DEFERRING
+			// if this is an inherited sub-object on a CDO, and that CDO has had
+			// its initialization deferred, then we shouldn't serialize in data 
+			// for this quite yet... not until the CDO owner has had a chance to
+			// initialize itself (because, as part of CDO initialization, 
+			// inherited sub-objects get filled in with values inherited from 
+			// the parent class... it is expected that this happens prior to 
+			// sub-object serialization)
+			if (Object->HasAnyFlags(RF_DefaultSubObject) && FDeferredObjInitializerTracker::DeferSubObjectPreload(Object))
+			{
+				// don't worry, FDeferredObjInitializerTracker::DeferSubObjectPreload() 
+				// should have cached this object, and it will run Preload() on 
+				// this later (once the super CDO has been initialized)
+				return;
+			}
+#endif // USE_CIRCULAR_DEPENDENCY_LOAD_DEFERRING
+
 			SCOPE_CYCLE_COUNTER(STAT_LinkerPreload);
 			FScopeCycleCounterUObject PreloadScope(Object, GET_STATID(STAT_LinkerPreload));
 			UClass* Cls = NULL;
diff --git a/Engine/Source/Runtime/CoreUObject/Private/UObject/UObjectGlobals.cpp b/Engine/Source/Runtime/CoreUObject/Private/UObject/UObjectGlobals.cpp
--- a/Engine/Source/Runtime/CoreUObject/Private/UObject/UObjectGlobals.cpp
+++ b/Engine/Source/Runtime/CoreUObject/Private/UObject/UObjectGlobals.cpp
@@ -2075,8 +2075,8 @@ FObjectInitializer::~FObjectInitializer()
 		UClass* ArchetypeClass = ObjectArchetype->GetClass();
 		// if this is a blueprint CDO that derives from another blueprint, and 
 		// that parent (archetype) CDO isn't fully serialized
-		if ( !Class->HasAnyFlags(RF_Native) && !ArchetypeClass->HasAnyFlags(RF_Native) &&
-			(ObjectArchetype->HasAnyFlags(RF_NeedLoad) || !ObjectArchetype->HasAnyFlags(RF_LoadCompleted)) )
+		if ( !Class->HasAnyFlags(RF_Native) && !ArchetypeClass->HasAnyFlags(RF_Native) && (ObjectArchetype->HasAnyFlags(RF_NeedLoad) ||
+			 !ObjectArchetype->HasAnyFlags(RF_LoadCompleted) || FDeferredObjInitializerTracker::IsCdoDeferred(ArchetypeClass)) )
 		{
 			FLinkerLoad* ClassLinker = Class->GetLinker();
 			if ((ClassLinker != nullptr) && (ClassLinker->LoadFlags & LOAD_DeferDependencyLoads) != 0x00)
diff --git a/Engine/Source/Runtime/CoreUObject/Public/Blueprint/BlueprintSupport.h b/Engine/Source/Runtime/CoreUObject/Public/Blueprint/BlueprintSupport.h
--- a/Engine/Source/Runtime/CoreUObject/Public/Blueprint/BlueprintSupport.h
+++ b/Engine/Source/Runtime/CoreUObject/Public/Blueprint/BlueprintSupport.h
@@ -84,10 +84,21 @@ public:
 	/** Looks up a FObjectInitializer that was deferred for the specified class (an FObjectInitializer for that class's CDO). */
 	static FObjectInitializer* Find(UClass* LoadClass);
 
+	/** Checks to see if the specified class has been logged as deferred (meaning its CDO hasn't had FObjectInitializer::PostConstructInit() ran on it yet). */
+	static bool IsCdoDeferred(UClass* LoadClass);
+
+	/** Determines if the specified sub-object should have its Preload() skipped; if so, this will cache the sub-object and return true. */
+	static bool DeferSubObjectPreload(UObject* SubObject);
+
 	/** Destroys any FObjectInitializers that were cached corresponding to the specified class. */
 	static void Remove(UClass* LoadClass);
 
+	/** Runs FObjectInitializer::PostConstructInit() on the specified class's CDO (if it was deferred), and preloads any sub-objects that were skipped. */
+	static bool ResolveDeferredInitialization(UClass* LoadClass);
+
 private:
-	/** A map that tracks the relationsship between Blueprint classes and  FObjectInitializers for their CDOs */
+	/** A map that tracks the relationship between Blueprint classes and FObjectInitializers for their CDOs */
 	TMap<UClass*, FObjectInitializer> DeferredInitializers;
+	/** Track default sub-objects that had their Preload() skipped, because the owning CDO's initialization should happen first */
+	TMultiMap<UClass*, UObject*> DeferredSubObjects;
 };
\ No newline at end of file
