Fixed up initialization of FObjectExportSortHelper using uninitalized variables.
Removed uneeded variables.
Removed excessive stats.

diff --git a/Engine/Source/Runtime/CoreUObject/Private/UObject/SavePackage.cpp b/Engine/Source/Runtime/CoreUObject/Private/UObject/SavePackage.cpp
--- a/Engine/Source/Runtime/CoreUObject/Private/UObject/SavePackage.cpp
+++ b/Engine/Source/Runtime/CoreUObject/Private/UObject/SavePackage.cpp
@@ -16,7 +16,7 @@ static const FName WorldClassName = FName("World");
 #define VALIDATE_INITIALIZECORECLASSES 0
 #define EXPORT_SORTING_DETAILED_LOGGING 0
 
-#define UE_PROFILE_COOKSAVE 0
+#define UE_PROFILE_COOKSAVE 1
 #if UE_PROFILE_COOKSAVE
 
 #define UE_START_LOG_COOK_TIME(InFilename) double PreviousTime; double StartTime; PreviousTime = StartTime = FPlatformTime::Seconds(); \
@@ -1220,11 +1220,6 @@ FObjectImportSortHelper* FObjectImportSortHelper::Sorter = NULL;
 struct FObjectExportSortHelper
 {
 private:
-	/**
-	 * Allows Compare access to the object full name lookup map
-	 */
-	static FObjectExportSortHelper* Sorter;
-
 	struct FObjectFullName
 	{
 	public:
@@ -1264,7 +1259,7 @@ private:
 	/** Comparison function used by Sort */
 	bool operator()( const FObjectExport& A, const FObjectExport& B ) const
 	{
-		checkSlow(Sorter);
+		//checkSlow(Sorter);
 
 		int32 Result = 0;
 		if ( A.Object == NULL )
@@ -1279,8 +1274,8 @@ private:
 		{
 			if (bUseFObjectFullName)
 			{
-				const FObjectFullName* FullNameA = Sorter->ObjectToObjectFullNameMap.Find(A.Object);
-				const FObjectFullName* FullNameB = Sorter->ObjectToObjectFullNameMap.Find(B.Object);
+				const FObjectFullName* FullNameA = ObjectToObjectFullNameMap.Find(A.Object);
+				const FObjectFullName* FullNameB = ObjectToObjectFullNameMap.Find(B.Object);
 				checkSlow(FullNameA);
 				checkSlow(FullNameB);
 
@@ -1307,8 +1302,8 @@ private:
 			}
 			else
 			{
-				FString* FullNameA = Sorter->ObjectToFullNameMap.Find(A.Object);
-				FString* FullNameB = Sorter->ObjectToFullNameMap.Find(B.Object);
+				const FString* FullNameA = ObjectToFullNameMap.Find(A.Object);
+				const FString* FullNameB = ObjectToFullNameMap.Find(B.Object);
 				checkSlow(FullNameA);
 				checkSlow(FullNameB);
 
@@ -1362,10 +1357,15 @@ public:
 
 					// Set the index (key) in the map to the index of this object into the export map.
 					OriginalExportIndexes.Add( *ExportFullName, ExportIndex );
-					ObjectToFullNameMap.Add(Export.Object, *ExportFullName);
-
-					FObjectFullName ObjectFullName(Export.Object, Linker->LinkerRoot);
-					ObjectToObjectFullNameMap.Add(Export.Object, MoveTemp(ObjectFullName) );
+					if (bUseFObjectFullName)
+					{
+						FObjectFullName ObjectFullName(Export.Object, Linker->LinkerRoot);
+						ObjectToObjectFullNameMap.Add(Export.Object, MoveTemp(ObjectFullName)); 
+					}
+					else
+					{
+						ObjectToFullNameMap.Add(Export.Object, *ExportFullName);
+					}
 				}
 			}
 
@@ -1437,7 +1437,15 @@ public:
 				const FObjectExport& Export = Linker->ExportMap[ExportIndex];
 				if ( Export.Object )
 				{
-					ObjectToFullNameMap.Add(Export.Object, *Export.Object->GetFullName());
+					if (bUseFObjectFullName)
+					{
+						FObjectFullName ObjectFullName(Export.Object, NULL);
+						ObjectToObjectFullNameMap.Add(Export.Object, MoveTemp(ObjectFullName));
+					}
+					else
+					{
+						ObjectToFullNameMap.Add(Export.Object, *Export.Object->GetFullName());
+					}
 				}
 			}
 		}
@@ -1445,12 +1453,11 @@ public:
 		if ( SortStartPosition < Linker->ExportMap.Num() )
 		{
 			Sorter = this;
-			Sort( &Linker->ExportMap[SortStartPosition], Linker->ExportMap.Num() - SortStartPosition, FObjectExportSortHelper() );
+			Sort( &Linker->ExportMap[SortStartPosition], Linker->ExportMap.Num() - SortStartPosition, *this );
 			Sorter = NULL;
 		}
 	}
 };
-FObjectExportSortHelper* FObjectExportSortHelper::Sorter = NULL;
 
 
 class FExportReferenceSorter : public FArchiveUObject
@@ -1541,17 +1548,21 @@ class FExportReferenceSorter : public FArchiveUObject
 	 */
 	void InitializeCoreClasses()
 	{
+#if 1
 		check(CoreClasses.Num() == 0);
 		check(ReferencedObjects.Num() == 0);
 		check(ForceLoadObjects.Num() == 0);
 		check(SerializedObjects.Num() == 0);
+		check(bIgnoreFieldReferences == false);
+
 		static bool bInitializedStaticCoreClasses = false;
 		static TArray<UClass*> StaticCoreClasses;
 		static TArray<UObject*> StaticCoreReferencedObjects;
-		static FOrderedObjectSet StaticProcessedObjects;
+		static TArray<UObject*> StaticProcessedObjects;
 		static TArray<UObject*> StaticForceLoadObjects;
 		static TSet<UObject*> StaticSerializedObjects;
-		static bool StaticbIgnoreFieldReferences = false;
+		
+		
 
 		// Helper class to register FlushInitializedStaticCoreClasses callback on first SavePackage run
 		struct FAddFlushInitalizedStaticCoreClasses
@@ -1568,13 +1579,6 @@ class FExportReferenceSorter : public FArchiveUObject
 		};
 		static FAddFlushInitalizedStaticCoreClasses MaybeAddAddFlushInitializedStaticCoreClasses;
 
-
-		if (bIgnoreFieldReferences != StaticbIgnoreFieldReferences)
-		{
-			StaticbIgnoreFieldReferences = bIgnoreFieldReferences;
-			bInitializedStaticCoreClasses = false;
-		}
-
 #if VALIDATE_INITIALIZECORECLASSES
 		bool bWasValid = bInitializedStaticCoreClasses;
 		bInitializedStaticCoreClasses = false;
@@ -1678,19 +1682,71 @@ class FExportReferenceSorter : public FArchiveUObject
 			StaticSerializedObjects = SerializedObjects;
 
 			check(CurrentClass == NULL);
+			check(CurrentInsertIndex == INDEX_NONE);
+		}
+		else
+		{
+			CoreClasses = StaticCoreClasses;
+			ReferencedObjects = StaticCoreReferencedObjects;
+			ProcessedObjects = StaticProcessedObjects;
+			ForceLoadObjects = StaticForceLoadObjects;
+			SerializedObjects = StaticSerializedObjects;
 
+			CoreReferencesOffset = StaticCoreReferencedObjects.Num();
+		}
 
+#else
+		// initialize the tracking maps with the core classes
+		UClass* CoreClassList[] =
+		{
+			UObject::StaticClass(),
+			UField::StaticClass(),
+			UStruct::StaticClass(),
+			UScriptStruct::StaticClass(),
+			UFunction::StaticClass(),
+			UEnum::StaticClass(),
+			UClass::StaticClass(),
+			UProperty::StaticClass(),
+			UByteProperty::StaticClass(),
+			UIntProperty::StaticClass(),
+			UBoolProperty::StaticClass(),
+			UFloatProperty::StaticClass(),
+			UDoubleProperty::StaticClass(),
+			UObjectProperty::StaticClass(),
+			UClassProperty::StaticClass(),
+			UInterfaceProperty::StaticClass(),
+			UNameProperty::StaticClass(),
+			UStrProperty::StaticClass(),
+			UArrayProperty::StaticClass(),
+			UTextProperty::StaticClass(),
+			UStructProperty::StaticClass(),
+			UDelegateProperty::StaticClass(),
+			UInterface::StaticClass(),
+			UMulticastDelegateProperty::StaticClass(),
+			UWeakObjectProperty::StaticClass(),
+			UObjectPropertyBase::StaticClass(),
+			ULazyObjectProperty::StaticClass(),
+			UAssetObjectProperty::StaticClass(),
+			UAssetClassProperty::StaticClass()
+		};
+
+		for (int32 CoreClassIndex = 0; CoreClassIndex < ARRAY_COUNT(CoreClassList); CoreClassIndex++)
+		{
+			UClass* CoreClass = CoreClassList[CoreClassIndex];
+			CoreClasses.AddUnique(CoreClass);
 
+			ReferencedObjects.Add(CoreClass);
+			ReferencedObjects.Add(CoreClass->GetDefaultObject());
 		}
 
-		CoreClasses = StaticCoreClasses;
-		ReferencedObjects = StaticCoreReferencedObjects;
-		ProcessedObjects = StaticProcessedObjects;
-		ForceLoadObjects = StaticForceLoadObjects;
-		SerializedObjects = StaticSerializedObjects;
+		for (int32 CoreClassIndex = 0; CoreClassIndex < CoreClasses.Num(); CoreClassIndex++)
+		{
+			UClass* CoreClass = CoreClasses[CoreClassIndex];
+			ProcessStruct(CoreClass);
+		}
 
-		CoreReferencesOffset = StaticCoreReferencedObjects.Num();
-		
+		CoreReferencesOffset = ReferencedObjects.Num();
+#endif
 	}
 
 	/**
@@ -1767,7 +1823,7 @@ public:
 	 * Constructor
 	 */
 	FExportReferenceSorter()
-	: FArchiveUObject(), CurrentInsertIndex(INDEX_NONE)
+		: FArchiveUObject(), CurrentInsertIndex(INDEX_NONE), CurrentClass(NULL), bIgnoreFieldReferences(false), CoreReferencesOffset(INDEX_NONE)
 	{
 		ArIsObjectReferenceCollector = true;
 		ArIsPersistent = true;
@@ -2133,7 +2189,7 @@ private:
 	/**
 	 * The list of objects that have been evaluated by this archive so far.
 	 */
-	struct FOrderedObjectSet
+	/*struct FOrderedObjectSet
 	{
 		TMap<UObject*, int32> ObjectsSet;
 		// TArray<UObject*> ObjectsList;
@@ -2159,8 +2215,8 @@ private:
 		{
 			return ObjectsSet.Num();
 		}
-	};
-	FOrderedObjectSet ProcessedObjects;
+	};*/
+	TArray<UObject*> ProcessedObjects;
 	
 
 	/**
@@ -2212,7 +2268,7 @@ struct FObjectExportSeekFreeSorter
 	 */
 	void SortExports( FLinkerSave* Linker, FLinkerLoad* LinkerToConformTo )
 	{
-		UE_START_LOG_COOK_TIME(*Linker->Filename);
+		//UE_START_LOG_COOK_TIME(*Linker->Filename);
 
 
 		SortArchive.SetCookingTarget(Linker->CookingTarget());
@@ -2231,7 +2287,7 @@ struct FObjectExportSeekFreeSorter
 			}
 		}
 
-		UE_LOG_COOK_TIME(TEXT("Setup Original Exports"));
+		//UE_LOG_COOK_TIME(TEXT("Setup Original Exports"));
 
 		bool bRetrieveInitialReferences = true;
 
@@ -2268,7 +2324,7 @@ struct FObjectExportSeekFreeSorter
 
 		}
 
-		UE_LOG_COOK_TIME(TEXT("Process Class Exports"));
+		//UE_LOG_COOK_TIME(TEXT("Process Class Exports"));
 
 #if EXPORT_SORTING_DETAILED_LOGGING
 		UE_LOG(LogSavePackage, Log, TEXT("*************   Processed %i classes out of %i possible exports for package %s.  Beginning second pass...   *************"), SortedExports.Num(), Linker->ExportMap.Num() - FirstSortIndex, *Linker->LinkerRoot->GetName());
@@ -2305,7 +2361,7 @@ struct FObjectExportSeekFreeSorter
 			}
 		}
 
-		UE_LOG_COOK_TIME(TEXT("Process Objects Exports"));
+		//UE_LOG_COOK_TIME(TEXT("Process Objects Exports"));
 
 #if EXPORT_SORTING_DETAILED_LOGGING
 		SortArchive.VerifySortingAlgorithm();
@@ -2334,7 +2390,7 @@ struct FObjectExportSeekFreeSorter
 			}
 		}
 
-		UE_LOG_COOK_TIME(TEXT("Fill Export Map"));
+		// UE_LOG_COOK_TIME(TEXT("Fill Export Map"));
 
 		// Manually add any new NULL exports last as they won't be in the SortedExportsObjects list. 
 		// A NULL Export.Object can occur if you are e.g. saving an object in the game that is 
@@ -2348,7 +2404,7 @@ struct FObjectExportSeekFreeSorter
 			}
 		}
 
-		UE_FINISH_LOG_COOK_TIME();
+		//UE_FINISH_LOG_COOK_TIME();
 	}
 
 private:
@@ -3343,11 +3399,11 @@ bool UPackage::SavePackage( UPackage* InOuter, UObject* Base, EObjectFlags TopLe
 				FObjectExportSortHelper ExportSortHelper;
 				ExportSortHelper.SortExports( Linker, Conform );
 				
+				UE_LOG_COOK_TIME(TEXT("Sort Exports nonseekfree"));
+
 				if ( EndSavingIfCancelled( Linker, TempFilename ) ) { return false; }
 				SlowTask.EnterProgressFrame();
 
-				UE_LOG_COOK_TIME(TEXT("Sort Exports nonseekfree"));
-
 				// Sort exports for seek-free loading.
 				FObjectExportSeekFreeSorter SeekFreeSorter;
 				UE_LOG_COOK_TIME(TEXT("Sort Exports seekfree constructor"));
