Disabled profile logging for savepackage

diff --git a/Engine/Source/Runtime/CoreUObject/Private/UObject/SavePackage.cpp b/Engine/Source/Runtime/CoreUObject/Private/UObject/SavePackage.cpp
--- a/Engine/Source/Runtime/CoreUObject/Private/UObject/SavePackage.cpp
+++ b/Engine/Source/Runtime/CoreUObject/Private/UObject/SavePackage.cpp
@@ -16,7 +16,7 @@ static const FName WorldClassName = FName("World");
 #define VALIDATE_INITIALIZECORECLASSES 0
 #define EXPORT_SORTING_DETAILED_LOGGING 0
 
-#define UE_PROFILE_COOKSAVE 1
+#define UE_PROFILE_COOKSAVE 0
 #if UE_PROFILE_COOKSAVE
 
 #define UE_START_LOG_COOK_TIME(InFilename) double PreviousTime; double StartTime; PreviousTime = StartTime = FPlatformTime::Seconds(); \
