Removed sorter completely.  As it's not required.

diff --git a/Engine/Source/Runtime/CoreUObject/Private/UObject/SavePackage.cpp b/Engine/Source/Runtime/CoreUObject/Private/UObject/SavePackage.cpp
--- a/Engine/Source/Runtime/CoreUObject/Private/UObject/SavePackage.cpp
+++ b/Engine/Source/Runtime/CoreUObject/Private/UObject/SavePackage.cpp
@@ -16,7 +16,7 @@ static const FName WorldClassName = FName("World");
 #define VALIDATE_INITIALIZECORECLASSES 0
 #define EXPORT_SORTING_DETAILED_LOGGING 0
 
-#define UE_PROFILE_COOKSAVE 0
+#define UE_PROFILE_COOKSAVE 1
 #if UE_PROFILE_COOKSAVE
 
 #define UE_START_LOG_COOK_TIME(InFilename) double PreviousTime; double StartTime; PreviousTime = StartTime = FPlatformTime::Seconds(); \
@@ -1088,11 +1088,6 @@ public:
 struct FObjectImportSortHelper
 {
 private:
-	/**
-	 * Allows Compare access to the object full name lookup map
-	 */
-	static FObjectImportSortHelper* Sorter;
-
 	/**
 	 * Map of UObject => full name; optimization for sorting.
 	 */
@@ -1104,8 +1099,6 @@ private:
 	/** Comparison function used by Sort */
 	bool operator()( const FObjectImport& A, const FObjectImport& B ) const
 	{
-		checkSlow(Sorter);
-
 		int32 Result = 0;
 		if ( A.XObject == NULL )
 		{
@@ -1117,8 +1110,8 @@ private:
 		}
 		else
 		{
-			FString* FullNameA = Sorter->ObjectToFullNameMap.Find(A.XObject);
-			FString* FullNameB = Sorter->ObjectToFullNameMap.Find(B.XObject);
+			const FString* FullNameA = ObjectToFullNameMap.Find(A.XObject);
+			const FString* FullNameB = ObjectToFullNameMap.Find(B.XObject);
 			checkSlow(FullNameA);
 			checkSlow(FullNameB);
 
@@ -1206,13 +1199,10 @@ public:
 
 		if ( SortStartPosition < Linker->ImportMap.Num() )
 		{
-			Sorter = this;
-			Sort( &Linker->ImportMap[SortStartPosition], Linker->ImportMap.Num() - SortStartPosition, FObjectImportSortHelper() );
-			Sorter = NULL;
+			Sort( &Linker->ImportMap[SortStartPosition], Linker->ImportMap.Num() - SortStartPosition, *this );
 		}
 	}
 };
-FObjectImportSortHelper* FObjectImportSortHelper::Sorter = NULL;
 
 /**
  * Helper structure to encapsulate sorting a linker's export table alphabetically, taking into account conforming to other linkers.
@@ -1259,8 +1249,6 @@ private:
 	/** Comparison function used by Sort */
 	bool operator()( const FObjectExport& A, const FObjectExport& B ) const
 	{
-		//checkSlow(Sorter);
-
 		int32 Result = 0;
 		if ( A.Object == NULL )
 		{
@@ -1452,9 +1440,7 @@ public:
 
 		if ( SortStartPosition < Linker->ExportMap.Num() )
 		{
-			Sorter = this;
 			Sort( &Linker->ExportMap[SortStartPosition], Linker->ExportMap.Num() - SortStartPosition, *this );
-			Sorter = NULL;
 		}
 	}
 };
