Fixup unaligned access.

diff --git a/Engine/Source/Runtime/CoreUObject/Private/UObject/Class.cpp b/Engine/Source/Runtime/CoreUObject/Private/UObject/Class.cpp
--- a/Engine/Source/Runtime/CoreUObject/Private/UObject/Class.cpp
+++ b/Engine/Source/Runtime/CoreUObject/Private/UObject/Class.cpp
@@ -1545,7 +1545,13 @@ bool UStruct::GetStringMetaDataHierarchical(const FName& Key, FString* OutValue)
 	 */
 	static void HandlePlaceholderScriptRef(ScriptPointerType& ScriptPtr)
 	{
+#ifdef REQUIRES_ALIGNED_INT_ACCESS
+		ScriptPointerType  Temp; 
+		FMemory::Memcpy(&Temp, &ScriptPtr, sizeof(ScriptPointerType));
+		UObject*& ExprPtrRef = (UObject*&)Temp;
+#else
 		UObject*& ExprPtrRef = (UObject*&)ScriptPtr;
+#endif 
 		if (ULinkerPlaceholderClass* PlaceholderObj = Cast<ULinkerPlaceholderClass>(ExprPtrRef))
 		{
 			PlaceholderObj->AddReferencingScriptExpr((UClass**)(&ExprPtrRef));
