Fix bad bug in FStringAssetReference where TryLoad would return null instead of the loaded object

diff --git a/Engine/Source/Runtime/CoreUObject/Private/Misc/StringAssetReference.cpp b/Engine/Source/Runtime/CoreUObject/Private/Misc/StringAssetReference.cpp
--- a/Engine/Source/Runtime/CoreUObject/Private/Misc/StringAssetReference.cpp
+++ b/Engine/Source/Runtime/CoreUObject/Private/Misc/StringAssetReference.cpp
@@ -124,7 +124,7 @@ UObject* FStringAssetReference::TryLoad() const
 		}
 	}
 
-	return nullptr;
+	return LoadedObject;
 }
 
 UObject* FStringAssetReference::ResolveObject() const
