Only register and unregister the class during deserialization.

diff --git a/Engine/Source/Runtime/CoreUObject/Private/UObject/Class.cpp b/Engine/Source/Runtime/CoreUObject/Private/UObject/Class.cpp
--- a/Engine/Source/Runtime/CoreUObject/Private/UObject/Class.cpp
+++ b/Engine/Source/Runtime/CoreUObject/Private/UObject/Class.cpp
@@ -3199,11 +3199,18 @@ void UClass::SetSuperStruct(UStruct* NewSuperStruct)
 void UClass::SerializeSuperStruct(FArchive& Ar)
 {
 #if UCLASS_FAST_ISA_IMPL & 2
-	FFastIndexingClassTree::Unregister(this);
+	bool bIsLoading = Ar.IsLoading();
+	if (bIsLoading)
+	{
+		FFastIndexingClassTree::Unregister(this);
+	}
 #endif
 	Super::SerializeSuperStruct(Ar);
 #if UCLASS_FAST_ISA_IMPL & 2
-	FFastIndexingClassTree::Register(this);
+	if (bIsLoading)
+	{
+		FFastIndexingClassTree::Register(this);
+	}
 #endif
 }
 
