Core: Improved the readability of a warning message for struct type mismatches in SerializeTaggedProperties

diff --git a/Engine/Source/Runtime/CoreUObject/Private/UObject/Class.cpp b/Engine/Source/Runtime/CoreUObject/Private/UObject/Class.cpp
--- a/Engine/Source/Runtime/CoreUObject/Private/UObject/Class.cpp
+++ b/Engine/Source/Runtime/CoreUObject/Private/UObject/Class.cpp
@@ -1246,7 +1246,8 @@ void UStruct::SerializeTaggedProperties(FArchive& Ar, uint8* Data, UStruct* Defa
 			else if( Tag.Type==NAME_StructProperty && Tag.StructName!=CastChecked<UStructProperty>(Property)->Struct->GetFName() 
 				&& !CanSerializeFromStructWithDifferentName(Ar, Tag, CastChecked<UStructProperty>(Property)))
 			{
-				UE_LOG(LogClass, Warning, TEXT("Property %s of %s struct type mismatch %s/%s for package:  %s. If that property got renamed, add an ActiveStructRedirect."), *Tag.Name.ToString(), *GetName(), *Tag.StructName.ToString(), *CastChecked<UStructProperty>(Property)->Struct->GetName(), *Ar.GetArchiveName() );
+				UE_LOG(LogClass, Warning, TEXT("Property %s of %s has a struct type mismatch (tag %s != prop %s) in package:  %s. If that struct got renamed, add an entry to ActiveStructRedirects."),
+					*Tag.Name.ToString(), *GetName(), *Tag.StructName.ToString(), *CastChecked<UStructProperty>(Property)->Struct->GetName(), *Ar.GetArchiveName() );
 			}
 			else if( !Property->ShouldSerializeValue(Ar) )
 			{
