UE-15054 Having 2 empty blueprint enumerators with the same name without adding values causes crash
Unique name for max element in UserDefinedEnum.

diff --git a/Engine/Source/Editor/UnrealEd/Private/Kismet2/EnumEditorUtils.cpp b/Engine/Source/Editor/UnrealEd/Private/Kismet2/EnumEditorUtils.cpp
--- a/Engine/Source/Editor/UnrealEd/Private/Kismet2/EnumEditorUtils.cpp
+++ b/Engine/Source/Editor/UnrealEd/Private/Kismet2/EnumEditorUtils.cpp
@@ -374,7 +374,7 @@ void FEnumEditorUtils::EnsureAllDisplayNamesExist(UUserDefinedEnum* Enum)
 {
 	if (Enum)
 	{
-		const int32 EnumeratorsToEnsure = Enum->NumEnums() - 1;
+		const int32 EnumeratorsToEnsure = (Enum->NumEnums() > 0) ? (Enum->NumEnums() - 1) : 0;
 		Enum->DisplayNames.Empty(EnumeratorsToEnsure);
 		for	(int32 Index = 0; Index < EnumeratorsToEnsure; Index++)
 		{
diff --git a/Engine/Source/Runtime/CoreUObject/Private/UObject/Enum.cpp b/Engine/Source/Runtime/CoreUObject/Private/UObject/Enum.cpp
--- a/Engine/Source/Runtime/CoreUObject/Private/UObject/Enum.cpp
+++ b/Engine/Source/Runtime/CoreUObject/Private/UObject/Enum.cpp
@@ -277,23 +277,13 @@ bool UEnum::SetEnums(TArray<FName>& InNames, UEnum::ECppForm InCppForm)
 	}
 	Names   = InNames;
 	CppForm = InCppForm;
-	return GenerateMaxEnum();
-}
 
-/**
- * Adds a virtual _MAX entry to the enum's list of names, unless the
- * enum already contains one.
- *
- * @return	true unless the MAX enum already exists and isn't the last enum.
- */
-bool UEnum::GenerateMaxEnum()
-{
 	const FString EnumPrefix = GenerateEnumPrefix();
 	checkSlow(EnumPrefix.Len());
 
 	const FName MaxEnumItem = *GenerateFullEnumName(*(EnumPrefix + TEXT("_MAX")));
 	const int32 MaxEnumItemIndex = Names.Find(MaxEnumItem);
-	if ( MaxEnumItemIndex == INDEX_NONE )
+	if (MaxEnumItemIndex == INDEX_NONE)
 	{
 		if (LookupEnumName(MaxEnumItem) != INDEX_NONE)
 		{
@@ -306,7 +296,7 @@ bool UEnum::GenerateMaxEnum()
 
 	AddNamesToMasterList();
 
-	if ( MaxEnumItemIndex != INDEX_NONE && MaxEnumItemIndex != Names.Num() - 1 )
+	if (MaxEnumItemIndex != INDEX_NONE && MaxEnumItemIndex != Names.Num() - 1)
 	{
 		// The MAX enum already exists, but isn't the last enum.
 		return false;
diff --git a/Engine/Source/Runtime/CoreUObject/Public/UObject/Class.h b/Engine/Source/Runtime/CoreUObject/Public/UObject/Class.h
--- a/Engine/Source/Runtime/CoreUObject/Public/UObject/Class.h
+++ b/Engine/Source/Runtime/CoreUObject/Public/UObject/Class.h
@@ -1381,12 +1381,12 @@ protected:
 protected: 
 	
 	/** adds the Names in this enum to the master AllEnumNames list */
-	void AddNamesToMasterList();
+	COREUOBJECT_API void AddNamesToMasterList();
 
 public:
 
 	/** removes the Names in this enum from the master AllEnumNames list */
-	void RemoveNamesFromMasterList();
+	COREUOBJECT_API void RemoveNamesFromMasterList();
 
 	// UObject interface.
 	COREUOBJECT_API virtual void Serialize(FArchive& Ar) override;
@@ -1497,7 +1497,7 @@ public:
 	 * @param InCppForm The form of enum.
 	 * @return	true unless the MAX enum already exists and isn't the last enum.
 	 */
-	COREUOBJECT_API bool SetEnums(TArray<FName>& InNames, ECppForm InCppForm);
+	COREUOBJECT_API virtual bool SetEnums(TArray<FName>& InNames, ECppForm InCppForm);
 
 	/**
 	 * @return	The enum name at the specified Index.
@@ -1571,14 +1571,6 @@ public:
 	 */
 	COREUOBJECT_API FString GenerateEnumPrefix() const;
 
-	/**
-	 * Adds a virtual _MAX entry to the enum's list of names, unless the
-	 * enum already contains one.
-	 *
-	 * @return	true unless the MAX enum already exists and isn't the last enum.
-	 */
-	bool GenerateMaxEnum();
-
 #if WITH_EDITOR
 	/**
 	* Finds the localized display name or native display name as a fallback.
diff --git a/Engine/Source/Runtime/Engine/Classes/Engine/UserDefinedEnum.h b/Engine/Source/Runtime/Engine/Classes/Engine/UserDefinedEnum.h
--- a/Engine/Source/Runtime/Engine/Classes/Engine/UserDefinedEnum.h
+++ b/Engine/Source/Runtime/Engine/Classes/Engine/UserDefinedEnum.h
@@ -56,6 +56,8 @@ public:
 	 */
 	virtual FText GetEnumText(int32 InIndex) const override;
 
+	virtual bool SetEnums(TArray<FName>& InNames, ECppForm InCppForm) override;
+
 #if WITH_EDITOR
 	// Begin UObject interface
 	virtual bool Rename(const TCHAR* NewName = NULL, UObject* NewOuter = NULL, ERenameFlags Flags = REN_None) override;
diff --git a/Engine/Source/Runtime/Engine/Private/UserDefinedEnum.cpp b/Engine/Source/Runtime/Engine/Private/UserDefinedEnum.cpp
--- a/Engine/Source/Runtime/Engine/Private/UserDefinedEnum.cpp
+++ b/Engine/Source/Runtime/Engine/Private/UserDefinedEnum.cpp
@@ -95,4 +95,35 @@ FText UUserDefinedEnum::GetEnumText(int32 InIndex) const
 	}
 
 	return Super::GetEnumText(InIndex);
+}
+
+bool UUserDefinedEnum::SetEnums(TArray<FName>& InNames, ECppForm InCppForm)
+{
+	if (Names.Num() > 0)
+	{
+		RemoveNamesFromMasterList();
+	}
+	Names = InNames;
+	CppForm = InCppForm;
+
+	const FString BaseEnumPrefix = GenerateEnumPrefix();
+	checkSlow(BaseEnumPrefix.Len());
+
+	const int32 MaxTryNum = 1024;
+	for (int32 TryNum = 0; TryNum < MaxTryNum; ++TryNum)
+	{
+		const FString EnumPrefix = (TryNum == 0) ? BaseEnumPrefix : FString::Printf(TEXT("%s_%d"), *BaseEnumPrefix, TryNum - 1);
+		const FName MaxEnumItem = *GenerateFullEnumName(*(EnumPrefix + TEXT("_MAX")));
+		const int32 MaxEnumItemIndex = Names.Find(MaxEnumItem);
+		if ((MaxEnumItemIndex == INDEX_NONE) && (LookupEnumName(MaxEnumItem) == INDEX_NONE))
+		{
+			Names.Add(MaxEnumItem);
+			AddNamesToMasterList();
+			return true;
+		}
+	}
+
+	UE_LOG(LogClass, Error, TEXT("Unable to generate enum MAX entry due to name collision. Enum: %s"), *GetPathName());
+
+	return false;
 }
\ No newline at end of file
