Put in braces to avoid risk of code breakage

diff --git a/Engine/Source/Runtime/CoreUObject/Private/UObject/UObjectGlobals.cpp b/Engine/Source/Runtime/CoreUObject/Private/UObject/UObjectGlobals.cpp
--- a/Engine/Source/Runtime/CoreUObject/Private/UObject/UObjectGlobals.cpp
+++ b/Engine/Source/Runtime/CoreUObject/Private/UObject/UObjectGlobals.cpp
@@ -1764,7 +1764,7 @@ UObject* StaticAllocateObject
 		check(InClass->GetClass());
 		if( !GIsDuplicatingClassForReinstancing )
 		{
-		InName = InClass->GetDefaultObjectName();
+			InName = InClass->GetDefaultObjectName();
 		}
 		// never call PostLoad on class default objects
 		InFlags &= ~(RF_NeedPostLoad|RF_NeedPostLoadSubobjects);
@@ -1781,8 +1781,9 @@ UObject* StaticAllocateObject
 		}
 		else
 #endif
-		InName = MakeUniqueObjectName(InOuter, InClass);
-
+		{
+			InName = MakeUniqueObjectName(InOuter, InClass);
+		}
 	}
 	else
 	{
