UE-14297 - Hang on Launch On HTML5 projects - ( 4.8 blocker)

  fixing an issue for single thread mode where the IOSystem is starved and loading goes into an infinite loop,  this change may be pushing other issues under the rug, but for now checking it in as it has a limited fallout, if any.

diff --git a/Engine/Source/Runtime/CoreUObject/Private/Serialization/AsyncLoading.cpp b/Engine/Source/Runtime/CoreUObject/Private/Serialization/AsyncLoading.cpp
--- a/Engine/Source/Runtime/CoreUObject/Private/Serialization/AsyncLoading.cpp
+++ b/Engine/Source/Runtime/CoreUObject/Private/Serialization/AsyncLoading.cpp
@@ -2220,6 +2220,10 @@ void FArchiveAsync::Serialize(void* Data, int64 Count)
 			do
 			{
 				SHUTDOWN_IF_EXIT_REQUESTED;
+				if (PlatformIsSinglethreaded)
+				{
+					FIOSystem::Get().TickSingleThreaded();
+				}
 				FPlatformProcess::SleepNoStats(0.0f);
 			} while (!Precache(CurrentPos, Count));
 		}
