Retrieving localized metadata will prefer local changes to the metadata if they do not match what is expected by localization. This means that if you change the metadata in the header, it will use the updated value instead of the localized value and your change will be seen in the editor.

diff --git a/Engine/Source/Runtime/CoreUObject/Private/UObject/Class.cpp b/Engine/Source/Runtime/CoreUObject/Private/UObject/Class.cpp
--- a/Engine/Source/Runtime/CoreUObject/Private/UObject/Class.cpp
+++ b/Engine/Source/Runtime/CoreUObject/Private/UObject/Class.cpp
@@ -325,13 +325,25 @@ const FString& UField::GetMetaData(const FName& Key) const
 
 const FText UField::GetMetaDataText(const TCHAR* MetaDataKey, const FString LocalizationNamespace, const FString LocalizationKey) const
 {
+	FString DefaultMetaData;
+
+	if( HasMetaData( MetaDataKey ))
+	{
+		DefaultMetaData = GetMetaData(MetaDataKey);
+	}
+
+	// If attempting to grab the DisplayName metadata, we must correct the source string and output it as a DisplayString for lookup
+	if( DefaultMetaData.IsEmpty() && FString(MetaDataKey) == TEXT("DisplayName") )
+	{
+		DefaultMetaData = FName::NameToDisplayString( GetName(), IsA( UBoolProperty::StaticClass() ) );
+	}
+
+
 	FText LocalizedMetaData;
-	if ( !( FText::FindText( LocalizationNamespace, LocalizationKey, /*OUT*/LocalizedMetaData ) ) )
+	if ( !( FText::FindText( LocalizationNamespace, LocalizationKey, /*OUT*/LocalizedMetaData, &DefaultMetaData ) ) )
 	{
-		FString DefaultMetaData;
-		if( HasMetaData( MetaDataKey ))
+		if (!DefaultMetaData.IsEmpty())
 		{
-			DefaultMetaData = GetMetaData(MetaDataKey);
 			LocalizedMetaData = FText::AsCultureInvariant(DefaultMetaData);
 		}
 	}
@@ -341,16 +353,29 @@ const FText UField::GetMetaDataText(const TCHAR* MetaDataKey, const FString Loca
 
 const FText UField::GetMetaDataText(const FName& MetaDataKey, const FString LocalizationNamespace, const FString LocalizationKey) const
 {
+	FString DefaultMetaData;
+
+	if( HasMetaData( MetaDataKey ))
+	{
+		DefaultMetaData = GetMetaData(MetaDataKey);
+	}
+
+	// If attempting to grab the DisplayName metadata, we must correct the source string and output it as a DisplayString for lookup
+	if( DefaultMetaData.IsEmpty() && MetaDataKey == TEXT("DisplayName") )
+	{
+		DefaultMetaData = FName::NameToDisplayString( GetName(), IsA( UBoolProperty::StaticClass() ) );
+	}
+	
+
 	FText LocalizedMetaData;
-	if ( !( FText::FindText( LocalizationNamespace, LocalizationKey, /*OUT*/LocalizedMetaData ) ) )
+	if ( !( FText::FindText( LocalizationNamespace, LocalizationKey, /*OUT*/LocalizedMetaData, &DefaultMetaData ) ) )
 	{
-		FString DefaultMetaData;
-		if( HasMetaData( MetaDataKey ))
+		if (!DefaultMetaData.IsEmpty())
 		{
-			DefaultMetaData = GetMetaData(MetaDataKey);
 			LocalizedMetaData = FText::AsCultureInvariant(DefaultMetaData);
 		}
 	}
+
 	return LocalizedMetaData;
 }
 
