[UE-15439] Moving where we resolve deferred CDO initialization (making sure it happens before the super CDO has been regenerated, because on regeneration the class layout could change).

diff --git a/Engine/Source/Runtime/CoreUObject/Private/Blueprint/BlueprintSupport.cpp b/Engine/Source/Runtime/CoreUObject/Private/Blueprint/BlueprintSupport.cpp
--- a/Engine/Source/Runtime/CoreUObject/Private/Blueprint/BlueprintSupport.cpp
+++ b/Engine/Source/Runtime/CoreUObject/Private/Blueprint/BlueprintSupport.cpp
@@ -92,6 +92,16 @@ bool FBlueprintSupport::IsDeferredExportCreationDisabled()
 #endif
 }
 
+bool FBlueprintSupport::IsDeferredCDOInitializationDisabled()
+{
+#if USE_CIRCULAR_DEPENDENCY_LOAD_DEFERRING
+	static const FBoolConfigValueHelper NoDeferredCDOInit(TEXT("Kismet"), TEXT("bForceDisableDeferredCDOInitialization"), GEngineIni);
+	return !UseDeferredDependencyLoading() || NoDeferredCDOInit;
+#else
+	return false;
+#endif
+}
+
 /*******************************************************************************
  * FScopedClassDependencyGather
  ******************************************************************************/
@@ -1250,6 +1260,82 @@ void FLinkerLoad::ResolveDeferredExports(UClass* LoadClass)
 	if (!FBlueprintSupport::IsDeferredCDOSerializationDisabled())
 #endif // USE_DEFERRED_DEPENDENCY_CHECK_VERIFICATION_TESTS
 	{
+		// NOTE: this doesn't work... mainly because of how we handle the case  
+		//       where SuperClass->HasAnyClassFlags(CLASS_NewerVersionExists) is 
+		//       true; at a surface level it works, up until GC is ran (the old
+		//       CDO is still tied to the active class, and the CDO's 
+		//       DestroyNonNativeProperties() attempts to call the class's
+		//       DestroyPersistentUberGraphFrame()... this is why 
+		//       ResolveDeferredSubClassObjects() was introduced instead
+// 		UClass* SuperClass = LoadClass->GetSuperClass();
+// 		// if this class's CDO had its initialization deferred (presumably 
+// 		// because its super's CDO hadn't been fully serialized), then we want 
+// 		// to make sure it is ran here before we serialize in the CDO (so it 
+// 		// gets stocked with inherited values first, that the Preload() may or
+// 		// may not override)
+// 		if (!FPlatformProperties::RequiresCookedData() && FDeferredObjInitializerTracker::IsCdoDeferred(LoadClass))
+// 		{
+// 			// if LoadClass's super has been regenerated (presumably since we 
+// 			// deferred its CDO initialization), then we need to throw out the
+// 			// CDO and recreate a new one... we do this because the super's
+// 			// layout could have changed (and property offsets are now different)
+// 			// for the deferred initialization to work, inherited property 
+// 			// offsets have to match (FObjectInitializer uses one property to 
+// 			// copy values between the CDOs)
+// 			if (SuperClass->HasAnyClassFlags(CLASS_NewerVersionExists))
+// 			{
+// 				// reset the class's super and relink (so that the class's 
+// 				// property chain reflects the super's regenerated layout)
+// 				UClass* RegenedSuperClass = SuperClass->GetAuthoritativeClass();
+// 				LoadClass->SetSuperStruct(RegenedSuperClass);
+// 				LoadClass->StaticLink(/*bRelinkExistingProperties =*/true);
+// 				SuperClass = RegenedSuperClass;
+// 
+// 				UObject* OldCDO = BlueprintCDO;
+// 				DEFERRED_DEPENDENCY_CHECK(LoadClass->GetDefaultObject(/*bCreateIfNeeded =*/false) == BlueprintCDO);
+// 
+// 				// we have to move the OldCDO out of its package, otherwise 
+// 				// when we attempt to create a new one StaticAllocateObject() 
+// 				// will just return us this stale one
+// 				UPackage* TransientPackage = GetTransientPackage();
+// 				const FName TransientName = MakeUniqueObjectName(TransientPackage, LoadClass, *FString::Printf(TEXT("CYCLIC_%s"), *OldCDO->GetName()));
+// 				const ERenameFlags RenameFlags = (REN_DontCreateRedirectors | REN_ForceNoResetLoaders | REN_NonTransactional | REN_DoNotDirty);
+// 				OldCDO->Rename(*TransientName.ToString(), TransientPackage, RenameFlags);
+// 
+// 				// have to reset the LOAD_DeferDependencyLoads flag so that 
+// 				// CreateExport() doesn't try to regenerate the Blueprint on the spot
+// 				const uint32 DeferDependencyFlag = (LoadFlags & LOAD_DeferDependencyLoads);
+// 				LoadFlags |= LOAD_DeferDependencyLoads;
+// 				
+// 				// recreate the CDO from scratch (going through CreateExport() instead
+// 				// of GetDefaultObject() so that it gets properly marked with RF_NeedLoad, etc.)
+// 				LoadClass->ClassDefaultObject = nullptr;
+// 				CDOExport.Object = nullptr;
+// 				UObject* NewCDO = CreateExport(DeferredCDOIndex);
+// 				DEFERRED_DEPENDENCY_CHECK(NewCDO != OldCDO);
+// 				LoadClass->ClassDefaultObject = NewCDO;
+// 
+// 				// restore (clear) the LOAD_DeferDependencyLoads flag
+// 				LoadFlags = (LoadFlags & ~LOAD_DeferDependencyLoads) | DeferDependencyFlag;
+// 
+// 				BlueprintCDO = NewCDO;
+// 				// NOTE: since we're recreating the CDO, CreateExport() will take care of initializting
+// 				//       it for us (the cached FObjectInitializer is no longer needed), but we do still
+// 				//       need to load (and re-parent) deferred sub-objects
+// 				FDeferredObjInitializerTracker::ResolveDeferredSubObjects(NewCDO);
+// 				FDeferredObjInitializerTracker::Remove(LoadClass);
+// 			}
+// 			else
+// 			{
+// 				FDeferredObjInitializerTracker::ResolveDeferredInitialization(LoadClass);
+// 			}
+// 		}
+// 		else
+// 		{
+//  			DEFERRED_DEPENDENCY_CHECK(!SuperClass->HasAnyClassFlags(CLASS_NewerVersionExists));
+//  			FDeferredObjInitializerTracker::ResolveDeferredInitialization(LoadClass);
+// 		}
+
 		// have to prematurely set the CDO's linker so we can force a Preload()/
 		// Serialization of the CDO before we regenerate the Blueprint class
 		{
@@ -1260,13 +1346,6 @@ void FLinkerLoad::ResolveDeferredExports(UClass* LoadClass)
 		}
 		DEFERRED_DEPENDENCY_CHECK(BlueprintCDO->GetClass() == LoadClass);
 
-		// if this class's CDO had its initialization deferred (presumably 
-		// because its super's CDO hadn't been fully serialized), then we want 
-		// to make sure it is ran here before we serialize in the CDO (so it 
-		// gets stocked with inherited values first, that the Preload() may or
-		// may not override)
-		FDeferredObjInitializerTracker::ResolveDeferredInitialization(LoadClass);
-
 		// should load the CDO (ensuring that it has been serialized in by the 
 		// time we get to class regeneration)
 		//
@@ -1279,6 +1358,20 @@ void FLinkerLoad::ResolveDeferredExports(UClass* LoadClass)
 		//       be cleared (and this will do nothing the 2nd time around)
 		Preload(BlueprintCDO);
 
+		// sub-classes of this Blueprint could have had their CDO's 
+		// initialization deferred (this occurs when the sub-class CDO is 
+		// created before this super CDO has been fully serialized; we do this
+		// because the  sub-class's CDO would not have been initialized with 
+		// accurate values)
+		//
+		// in that case, the sub-class CDOs are waiting around until their 
+		// super CDO is fully loaded (which is now)... we want to do this here, 
+		// before this (super) Blueprint gets regenerated, because after it's
+		// regenerated the class layout (and property offsets) may no longer 
+		// match the layout that sub-class CDOs were constructed with (making 
+		// property copying dangerous)
+		FDeferredObjInitializerTracker::ResolveDeferredSubClassObjects(LoadClass);
+
 		DEFERRED_DEPENDENCY_CHECK(BlueprintCDO->HasAnyFlags(RF_LoadCompleted));
 	}
 #endif // USE_CIRCULAR_DEPENDENCY_LOAD_DEFERRING
@@ -1461,8 +1554,16 @@ FObjectInitializer* FDeferredObjInitializerTracker::Add(const FObjectInitializer
 	if (LoadClass != nullptr)
 	{
 		FDeferredObjInitializerTracker& ThreadInst = FDeferredObjInitializerTracker::Get();
-		auto& InitializersCache = ThreadInst.DeferredInitializers;
+		auto& SuperClassMap = ThreadInst.SuperClassMap;
 
+		UClass* SuperClass = LoadClass->GetSuperClass();
+		SuperClassMap.AddUnique(SuperClass, LoadClass);
+#if USE_DEFERRED_DEPENDENCY_CHECK_VERIFICATION_TESTS
+		UObject* SuperCDO = SuperClass->GetDefaultObject(/*bCreateIfNeeded =*/false);
+		DEFERRED_DEPENDENCY_CHECK( SuperCDO && (SuperCDO->HasAnyFlags(RF_NeedLoad) || SuperCDO->HasAnyFlags(RF_LoadCompleted) || IsCdoDeferred(SuperClass)) );
+#endif // USE_DEFERRED_DEPENDENCY_CHECK_VERIFICATION_TESTS
+		
+		auto& InitializersCache = ThreadInst.DeferredInitializers;
 		DEFERRED_DEPENDENCY_CHECK(InitializersCache.Find(LoadClass) == nullptr); // did we try to init the CDO twice?
 		DEFERRED_DEPENDENCY_CHECK(LoadClass->GetLinker()->LoadFlags & LOAD_DeferDependencyLoads);
 
@@ -1492,7 +1593,8 @@ bool FDeferredObjInitializerTracker::DeferSubObjectPreload(UObject* SubObject)
 	UObject* CdoOuter = SubObject->GetOuter();
 	UClass* OuterClass = CdoOuter->GetClass();
 
-	if (IsCdoDeferred(OuterClass))
+	FDeferredObjInitializerTracker& ThreadInst = FDeferredObjInitializerTracker::Get();
+	if (IsCdoDeferred(OuterClass) && (OuterClass != ThreadInst.ResolvingClass))
 	{
 		DEFERRED_DEPENDENCY_CHECK(CdoOuter->HasAnyFlags(RF_ClassDefaultObject));
 
@@ -1502,7 +1604,7 @@ bool FDeferredObjInitializerTracker::DeferSubObjectPreload(UObject* SubObject)
 		// components because they're what gets filled out in the CDO's InitSubobjectProperties()
 		if (SubObjTemplate && (SubObjTemplate->GetOuter() != CdoOuter))
 		{
-			FDeferredObjInitializerTracker& ThreadInst = FDeferredObjInitializerTracker::Get();
+			
 			ThreadInst.DeferredSubObjects.AddUnique(OuterClass, SubObject);
 
 			return true;
@@ -1516,43 +1618,87 @@ void FDeferredObjInitializerTracker::Remove(UClass* LoadClass)
 	FDeferredObjInitializerTracker& ThreadInst = FDeferredObjInitializerTracker::Get();
 	ThreadInst.DeferredInitializers.Remove(LoadClass);
 	ThreadInst.DeferredSubObjects.Remove(LoadClass);
+	ThreadInst.SuperClassMap.RemoveSingle(LoadClass->GetSuperClass(), LoadClass);
 }
 
 bool FDeferredObjInitializerTracker::ResolveDeferredInitialization(UClass* LoadClass)
 {
 	if (FObjectInitializer* DeferredInitializer = FDeferredObjInitializerTracker::Find(LoadClass))
 	{
+		FDeferredObjInitializerTracker& ThreadInst = FDeferredObjInitializerTracker::Get();
+		TGuardValue<UClass*> ResolvingGuard(ThreadInst.ResolvingClass, LoadClass);
+
+		DEFERRED_DEPENDENCY_CHECK(!LoadClass->GetSuperClass()->HasAnyClassFlags(CLASS_NewerVersionExists));
 		// initializes and instances CDO properties (copies inherited values 
 		// from the super's CDO)
 		FScriptIntegrationObjectHelper::PostConstructInitObject(*DeferredInitializer);
 
-		FDeferredObjInitializerTracker& ThreadInst = FDeferredObjInitializerTracker::Get();
-		TArray<UObject*> DeferredSubObjects;
-		ThreadInst.DeferredSubObjects.MultiFind(LoadClass, DeferredSubObjects);
-
-		// remove before we call Preload() for sub-objects, else we'll end up 
-		// back in DeferSubObjectPreload()
+		ResolveDeferredSubObjects(LoadClass->GetDefaultObject());
 		FDeferredObjInitializerTracker::Remove(LoadClass);
 
-		FLinkerLoad* ClassLinker = LoadClass->GetLinker();
-		DEFERRED_DEPENDENCY_CHECK(ClassLinker != nullptr);
-		if (ClassLinker != nullptr)
+		return true;
+	}
+	return false;
+}
+
+void FDeferredObjInitializerTracker::ResolveDeferredSubObjects(UObject* CDO)
+{
+	DEFERRED_DEPENDENCY_CHECK(CDO->HasAnyFlags(RF_ClassDefaultObject));
+	UClass* LoadClass = CDO->GetClass();
+
+	FDeferredObjInitializerTracker& ThreadInst = FDeferredObjInitializerTracker::Get();
+	TArray<UObject*> DeferredSubObjects;
+	ThreadInst.DeferredSubObjects.MultiFind(LoadClass, DeferredSubObjects);
+
+	// guard to keep sub-object preloads from ending up back in DeferSubObjectPreload()
+	TGuardValue<UClass*> ResolvingGuard(ThreadInst.ResolvingClass, LoadClass);
+
+	FLinkerLoad* ClassLinker = LoadClass->GetLinker();
+	DEFERRED_DEPENDENCY_CHECK(ClassLinker != nullptr);
+	if (ClassLinker != nullptr)
+	{
+		// this all needs to happen after PostConstructInitObject() (in 
+		// ResolveDeferredInitialization), since InitSubObjectProperties() is 
+		// invoked there (which is where we fill this sub-object with values 
+		// from the super)... here we account for any Preload() calls that were 
+		// skipped on account of the deferred CDO initialization
+		for (UObject* SubObj : DeferredSubObjects)
 		{
-			// this all needs to happen after PostConstructInitObject() (above), 
-			// since InitSubObjectProperties() is invoked there (which is where 
-			// we fill this sub-object with values from the super)... here we
-			// account for any Preload() calls that were skipped on account of
-			// the deferred CDO initialization
-			for (UObject* SubObj : DeferredSubObjects)
-			{
-				ClassLinker->Preload(SubObj);
-			}
+			// no longer need to handle this case, because we're no longer 
+			// attempting to recreate the CDO object in 
+			// FLinkerLoad::ResolveDeferredExports() (when the super class has
+			// been regenerated)
+// 			if (SubObj->GetOuter() != CDO)
+// 			{
+// 				// we may have had to recreate the CDO (because its super has had 
+// 				// its layout altered through regeneration); see 
+// 				// FLinkerLoad::ResolveDeferredExports() for more context
+// 				//
+// 				// @TODO: Are we sure that DeferredSubObjects would have all the
+// 				//        sub-objects that we need to over?
+// 				// 
+// 				// @TODO: Are we sure these are the right rename clags to use?
+// 				const ERenameFlags RenameFlags = (REN_DontCreateRedirectors | REN_ForceNoResetLoaders | REN_NonTransactional | REN_DoNotDirty);
+// 				SubObj->Rename(nullptr, CDO, RenameFlags);
+// 			}
+			DEFERRED_DEPENDENCY_CHECK(SubObj->GetOuter() == CDO);
+			ClassLinker->Preload(SubObj);
 		}
+	}
 
-		
-		return true;
+	ThreadInst.DeferredSubObjects.Remove(LoadClass);
+}
+
+void FDeferredObjInitializerTracker::ResolveDeferredSubClassObjects(UClass* SuperClass)
+{
+	FDeferredObjInitializerTracker& ThreadInst = FDeferredObjInitializerTracker::Get();
+	TArray<UClass*> DeferredSubClasses;
+	ThreadInst.SuperClassMap.MultiFind(SuperClass, DeferredSubClasses);
+
+	for (UClass* SubClass : DeferredSubClasses)
+	{
+		ResolveDeferredInitialization(SubClass);
 	}
-	return false;
 }
 
 // don't want other files ending up with this internal define
diff --git a/Engine/Source/Runtime/CoreUObject/Private/UObject/UObjectGlobals.cpp b/Engine/Source/Runtime/CoreUObject/Private/UObject/UObjectGlobals.cpp
--- a/Engine/Source/Runtime/CoreUObject/Private/UObject/UObjectGlobals.cpp
+++ b/Engine/Source/Runtime/CoreUObject/Private/UObject/UObjectGlobals.cpp
@@ -2071,7 +2071,7 @@ FObjectInitializer::~FObjectInitializer()
 
 #if USE_CIRCULAR_DEPENDENCY_LOAD_DEFERRING
 	bool bIsPostConstructInitDeferred = false;
-	if (bIsCDO && (ObjectArchetype != nullptr) && FBlueprintSupport::UseDeferredDependencyLoading())
+	if (bIsCDO && (ObjectArchetype != nullptr) && !FBlueprintSupport::IsDeferredCDOInitializationDisabled())
 	{
 		UClass* ArchetypeClass = ObjectArchetype->GetClass();
 		// if this is a blueprint CDO that derives from another blueprint, and 
@@ -2129,33 +2129,46 @@ void FObjectInitializer::PostConstructInit()
 
 #if USE_CIRCULAR_DEPENDENCY_LOAD_DEFERRING
 	// if this is a deferred initializer (implying that it's for a CDO), and the 
-	// ObjectArchetype (super CDO) has since been regenerated, then we need 
-	// to update the cached archetypes (so their not pointing to TRASH/REINST 
-	// versions)
-	if (bIsDeferredInitializer && SuperClass->HasAnyClassFlags(CLASS_NewerVersionExists))
-	{
+	// ObjectArchetype (super CDO) has since been regenerated, then we cannot
+	// reliably initialize and instance properties (if we were to use the 
+	// regenerated ObjectArchetype, the property layouts could differ and that's
+	// not viable for InitProperties())
+	// 
+	// However, this case is handled in FLinkerLoad::ResolveDeferredExports(), 
+	// where we forcefully recreate and separately init this CDO (hence it being 
+	// moved to the transient package) 
+	if (bIsDeferredInitializer)
+	{
+		const bool bSuperHasBeenRegenerated = SuperClass->HasAnyClassFlags(CLASS_NewerVersionExists);
+#if USE_DEFERRED_DEPENDENCY_CHECK_VERIFICATION_TESTS
 		check(bIsCDO);
+		check(ObjectArchetype->GetOutermost() != GetTransientPackage());
+		check(ObjectArchetype->GetClass() == SuperClass && !bSuperHasBeenRegenerated);
+#endif // USE_DEFERRED_DEPENDENCY_CHECK_VERIFICATION_TESTS
 
-		SuperClass = SuperClass->GetAuthoritativeClass();
-		Class->SetSuperStruct(SuperClass);
-		ObjectArchetype = SuperClass->GetDefaultObject(/*bCreateIfNeeded =*/false);
-
-		// iterate backwards, so we can remove elements as we go
-		for (int32 SubObjIndex = ComponentInits.SubobjectInits.Num()-1; SubObjIndex >= 0; --SubObjIndex)
+		if ( !ensureMsgf(!bSuperHasBeenRegenerated, TEXT("The super class for %s has been regenerated, we cannot properly initialize inherited properties, as the class layout may have changed."), *Obj->GetName()) )
 		{
-			FSubobjectsToInit::FSubobjectInit& SubObjInitInfo = ComponentInits.SubobjectInits[SubObjIndex];
-			const FName SubObjName = SubObjInitInfo.Subobject->GetFName();
+			// attempt to complete initialization/instancing as best we can, but
+			// it would not be surprising if our CDO was improperly initialized 
+			// as a result...
 
-			UObject* OuterArchetype = SubObjInitInfo.Subobject->GetOuter()->GetArchetype();
-			UObject* NewTemplate = OuterArchetype->GetClass()->GetDefaultSubobjectByName(SubObjName);
-
-			if (ensure(NewTemplate != nullptr))
-			{
-				SubObjInitInfo.Template = NewTemplate;
-			}
-			else
+			// iterate backwards, so we can remove elements as we go
+			for (int32 SubObjIndex = ComponentInits.SubobjectInits.Num() - 1; SubObjIndex >= 0; --SubObjIndex)
 			{
-				ComponentInits.SubobjectInits.RemoveAtSwap(SubObjIndex);
+				FSubobjectsToInit::FSubobjectInit& SubObjInitInfo = ComponentInits.SubobjectInits[SubObjIndex];
+				const FName SubObjName = SubObjInitInfo.Subobject->GetFName();
+
+				UObject* OuterArchetype = SubObjInitInfo.Subobject->GetOuter()->GetArchetype();
+				UObject* NewTemplate = OuterArchetype->GetClass()->GetDefaultSubobjectByName(SubObjName);
+
+				if (ensure(NewTemplate != nullptr))
+				{
+					SubObjInitInfo.Template = NewTemplate;
+				}
+				else
+				{
+					ComponentInits.SubobjectInits.RemoveAtSwap(SubObjIndex);
+				}
 			}
 		}
 	}
diff --git a/Engine/Source/Runtime/CoreUObject/Public/Blueprint/BlueprintSupport.h b/Engine/Source/Runtime/CoreUObject/Public/Blueprint/BlueprintSupport.h
--- a/Engine/Source/Runtime/CoreUObject/Public/Blueprint/BlueprintSupport.h
+++ b/Engine/Source/Runtime/CoreUObject/Public/Blueprint/BlueprintSupport.h
@@ -26,6 +26,7 @@ struct FBlueprintSupport
 	static bool IsResolvingDeferredDependenciesDisabled();
 	static bool IsDeferredCDOSerializationDisabled();
 	static bool IsDeferredExportCreationDisabled();
+	static bool IsDeferredCDOInitializationDisabled();
 };
 
 #if WITH_EDITOR
@@ -78,6 +79,8 @@ private:
 struct FDeferredObjInitializerTracker : TThreadSingleton<FDeferredObjInitializerTracker>
 {
 public:
+	FDeferredObjInitializerTracker() : ResolvingClass(nullptr) {}
+
 	/** Stores a copy of the specified FObjectInitializer and returns a pointer to it (could be null if a corresponding class could not be determined). */
 	static FObjectInitializer* Add(const FObjectInitializer& DeferringInitializer);
 
@@ -96,9 +99,19 @@ public:
 	/** Runs FObjectInitializer::PostConstructInit() on the specified class's CDO (if it was deferred), and preloads any sub-objects that were skipped. */
 	static bool ResolveDeferredInitialization(UClass* LoadClass);
 
+	/**  */
+	static void ResolveDeferredSubObjects(UObject* CDO);
+
+	/**  */
+	static void ResolveDeferredSubClassObjects(UClass* SuperClass);
+
 private:
 	/** A map that tracks the relationship between Blueprint classes and FObjectInitializers for their CDOs */
 	TMap<UClass*, FObjectInitializer> DeferredInitializers;
 	/** Track default sub-objects that had their Preload() skipped, because the owning CDO's initialization should happen first */
 	TMultiMap<UClass*, UObject*> DeferredSubObjects;
+	/** Used to keep ResolveDeferredSubObjects() from re-adding sub-objects via DeferSubObjectPreload() */
+	UClass* ResolvingClass;
+	/** Tracks sub-classes that have had their CDO deferred as a result of the super not being fully serialized */
+	TMultiMap<UClass*, UClass*> SuperClassMap;
 };
\ No newline at end of file
