Fixing subobject instancing when creating a blueprint instance from non-CDO archetype.

UE-15502

diff --git a/Engine/Source/Runtime/CoreUObject/Private/UObject/CoreNative.cpp b/Engine/Source/Runtime/CoreUObject/Private/UObject/CoreNative.cpp
--- a/Engine/Source/Runtime/CoreUObject/Private/UObject/CoreNative.cpp
+++ b/Engine/Source/Runtime/CoreUObject/Private/UObject/CoreNative.cpp
@@ -74,12 +74,12 @@ FObjectInstancingGraph::FObjectInstancingGraph( UObject* DestinationSubobjectRoo
 	SetDestinationRoot(DestinationSubobjectRoot);
 }
 
-void FObjectInstancingGraph::SetDestinationRoot( UObject* DestinationSubobjectRoot )
+void FObjectInstancingGraph::SetDestinationRoot(UObject* DestinationSubobjectRoot, UObject* InSourceRoot /*= nullptr*/)
 {
 	DestinationRoot = DestinationSubobjectRoot;
 	check(DestinationRoot);
 
-	SourceRoot = DestinationRoot->GetArchetype();
+	SourceRoot = InSourceRoot ? InSourceRoot : DestinationRoot->GetArchetype();
 	check(SourceRoot);
 
 	// add the subobject roots to the Source -> Destination mapping
@@ -263,26 +263,26 @@ UObject* FObjectInstancingGraph::InstancePropertyValue( class UObject* Component
 }
 
 
-void FObjectInstancingGraph::AddNewObject(class UObject* ObjectInstance)
+void FObjectInstancingGraph::AddNewObject(class UObject* ObjectInstance, UObject* InArchetype /*= nullptr*/)
 {
 	if (HasDestinationRoot())
 	{
-		AddNewInstance(ObjectInstance);
+		AddNewInstance(ObjectInstance, InArchetype);
 	}
 	else
 	{
-		SetDestinationRoot(ObjectInstance);
+		SetDestinationRoot(ObjectInstance, InArchetype);
 	}
 }
 
-void FObjectInstancingGraph::AddNewInstance( UObject* ObjectInstance )
+void FObjectInstancingGraph::AddNewInstance(UObject* ObjectInstance, UObject* InArchetype /*= nullptr*/)
 {
 	check(SourceRoot);
 	check(DestinationRoot);
 
-	if ( ObjectInstance != NULL )
+	if (ObjectInstance != nullptr)
 	{
-		UObject* SourceObject = ObjectInstance->GetArchetype();
+		UObject* SourceObject = InArchetype ? InArchetype : ObjectInstance->GetArchetype();
 		check(SourceObject);
 
 		SourceToDestinationMap.Add(SourceObject, ObjectInstance);
diff --git a/Engine/Source/Runtime/CoreUObject/Private/UObject/UObjectGlobals.cpp b/Engine/Source/Runtime/CoreUObject/Private/UObject/UObjectGlobals.cpp
--- a/Engine/Source/Runtime/CoreUObject/Private/UObject/UObjectGlobals.cpp
+++ b/Engine/Source/Runtime/CoreUObject/Private/UObject/UObjectGlobals.cpp
@@ -2296,15 +2296,17 @@ void FObjectInitializer::InstanceSubobjects(UClass* Class, bool bNeedInstancing,
 {
 	FObjectInstancingGraph TempInstancingGraph;
 	FObjectInstancingGraph* UseInstancingGraph = InstanceGraph ? InstanceGraph : &TempInstancingGraph;
-	UseInstancingGraph->AddNewObject(Obj);
+	{
+		UseInstancingGraph->AddNewObject(Obj, ObjectArchetype);
+	}
 	// Add any default subobjects
-	for (int32 Index = 0; Index < ComponentInits.SubobjectInits.Num(); Index++)
+	for (auto& SubobjectInit : ComponentInits.SubobjectInits)
 	{
-		UseInstancingGraph->AddNewObject(ComponentInits.SubobjectInits[Index].Subobject);
+		UseInstancingGraph->AddNewObject(SubobjectInit.Subobject, SubobjectInit.Template);
 	}
 	if (bNeedInstancing)
 	{
-		UObject* Archetype = Obj->GetArchetype();
+		UObject* Archetype = ObjectArchetype ? ObjectArchetype : Obj->GetArchetype();
 		Class->InstanceSubobjectTemplates(Obj, Archetype, Archetype ? Archetype->GetClass() : NULL, Obj, UseInstancingGraph);
 	}
 	if (bNeedSubobjectInstancing)
diff --git a/Engine/Source/Runtime/CoreUObject/Public/UObject/Class.h b/Engine/Source/Runtime/CoreUObject/Public/UObject/Class.h
--- a/Engine/Source/Runtime/CoreUObject/Public/UObject/Class.h
+++ b/Engine/Source/Runtime/CoreUObject/Public/UObject/Class.h
@@ -2484,8 +2484,9 @@ public:
 	 * Sets the DestinationRoot for this instancing graph.
 	 *
 	 * @param	DestinationSubobjectRoot	the top-level object that is being created
+	 * @param	InSourceRoot	Archetype of DestinationSubobjectRoot
 	 */
-	void SetDestinationRoot( class UObject* DestinationSubobjectRoot );
+	void SetDestinationRoot( class UObject* DestinationSubobjectRoot, class UObject* InSourceRoot = nullptr );
 
 	/**
 	 * Finds the destination object instance corresponding to the specified source object.
@@ -2513,17 +2514,19 @@ public:
 
 	/**
 	 * Adds a partially built object instance to the map(s) of source objects to their instances.
-	 * @param	ObjectInstance			Object that was just allocated, but has not been constructed yet
+	 * @param	ObjectInstance  Object that was just allocated, but has not been constructed yet
+	 * @param	InArchetype     Archetype of ObjectInstance
 	 */
-	void AddNewObject(class UObject* ObjectInstance);
+	void AddNewObject(class UObject* ObjectInstance, class UObject* InArchetype = nullptr);
 
 	/**
 	 * Adds an object instance to the map of source objects to their instances.  If there is already a mapping for this object, it will be replaced
 	 * and the value corresponding to ObjectInstance's archetype will now point to ObjectInstance.
 	 *
-	 * @param	ObjectInstance	the object that should be added as the corresopnding instance for ObjectSource
+	 * @param	ObjectInstance  the object that should be added as the corresopnding instance for ObjectSource
+	 * @param	InArchetype     Archetype of ObjectInstance
 	 */
-	void AddNewInstance(class UObject* ObjectInstance);
+	void AddNewInstance(class UObject* ObjectInstance, class UObject* InArchetype = nullptr);
 
 	/**
 	 * Retrieves a list of objects that have the specified Outer
