Fixing a crash when re-enetering FlushAsyncLoading in the same callstack when there was a synchronuous load request from async loading callbacks.

diff --git a/Engine/Source/Runtime/CoreUObject/Private/Serialization/AsyncLoading.cpp b/Engine/Source/Runtime/CoreUObject/Private/Serialization/AsyncLoading.cpp
--- a/Engine/Source/Runtime/CoreUObject/Private/Serialization/AsyncLoading.cpp
+++ b/Engine/Source/Runtime/CoreUObject/Private/Serialization/AsyncLoading.cpp
@@ -505,22 +505,27 @@ EAsyncPackageState::Type FAsyncLoadingThread::ProcessLoadedPackages(bool bUseTim
 			Result = Package->PostLoadDeferredObjects(TickStartTime, bUseTimeLimit, TimeLimit);
 			if (Result == EAsyncPackageState::Complete)
 			{
-				const bool bInternalCallbacks = false;
-				const EAsyncLoadingResult::Type LoadingResult = Package->HasLoadFailed() ? EAsyncLoadingResult::Failed : EAsyncLoadingResult::Succeeded;
-				Package->CallCompletionCallbacks(bInternalCallbacks, LoadingResult);
-
 				if (FPlatformProperties::RequiresCookedData())
 				{
 					// Emulates ResetLoaders on the package linker's linkerroot.
 					Package->ResetLoader();
 				}
 
-				// Incremented on the Async Thread, decremented on the Game Thread				
-				AsyncLoadingCounter.Decrement();
-				check(AsyncLoadingCounter.GetValue() >= 0);
+				// Remove the package from the list before we fore callbacks, 
+				// this is to ensure we can re-enter FlushAsyncLoading from any of the callbacks
+				LoadedPackagesToProcess.RemoveAt(PackageIndex--);
+
+				// Incremented on the Async Thread, now decrement as we're done with this package				
+				const int32 NewAsyncLoadingCounterValue = AsyncLoadingCounter.Decrement();
+				UE_CLOG(NewAsyncLoadingCounterValue < 0, LogStreaming, Fatal, TEXT("AsyncLoadingCounter is negative, this means we loaded more packages then requested so there must be a bug in async loading code."));
 
+				// Call external callbacks
+				const bool bInternalCallbacks = false;
+				const EAsyncLoadingResult::Type LoadingResult = Package->HasLoadFailed() ? EAsyncLoadingResult::Failed : EAsyncLoadingResult::Succeeded;
+				Package->CallCompletionCallbacks(bInternalCallbacks, LoadingResult);
+
+				// We don't need the package anymore
 				delete Package;
-				LoadedPackagesToProcess.RemoveAt(PackageIndex--);
 			}
 			else
 			{
