Look for duplicate exports in package and skip them while loading and saving.

diff --git a/Engine/Source/Runtime/Core/Public/UObject/NameTypes.h b/Engine/Source/Runtime/Core/Public/UObject/NameTypes.h
--- a/Engine/Source/Runtime/Core/Public/UObject/NameTypes.h
+++ b/Engine/Source/Runtime/Core/Public/UObject/NameTypes.h
@@ -577,6 +577,14 @@ public:
 		return Compare(Other) < 0;
 	}
 
+	/**
+	 * Comparison operator used for sorting alphabetically.
+	 */
+	FORCEINLINE bool operator>(const FName& Other) const
+	{
+		return Compare(Other) > 0;
+	}
+
 	FORCEINLINE bool IsNone() const
 	{
 		return GetComparisonIndexFast() == 0 && GetNumber() == 0;
diff --git a/Engine/Source/Runtime/Core/Public/UObject/ObjectVersion.h b/Engine/Source/Runtime/Core/Public/UObject/ObjectVersion.h
--- a/Engine/Source/Runtime/Core/Public/UObject/ObjectVersion.h
+++ b/Engine/Source/Runtime/Core/Public/UObject/ObjectVersion.h
@@ -510,6 +510,8 @@ enum EUnrealEngineObjectUE4Version
 	VER_UE4_SERIALIZE_PINTYPE_CONST,
 	// Change UMaterialFunction::LibraryCategories to LibraryCategoriesText (old assets were saved before auto-conversion of UArrayProperty was possible)
 	VER_UE4_LIBRARY_CATEGORIES_AS_FTEXT,
+	// Check for duplicate exports while saving packages.
+	VER_UE4_SKIP_DUPLICATE_EXPORTS_ON_SAVE_PACKAGE,
 
 	// -----<new versions can be added before this line>-------------------------------------------------
 	// - this needs to be the last line (see note below)
diff --git a/Engine/Source/Runtime/CoreUObject/Private/UObject/LinkerLoad.cpp b/Engine/Source/Runtime/CoreUObject/Private/UObject/LinkerLoad.cpp
--- a/Engine/Source/Runtime/CoreUObject/Private/UObject/LinkerLoad.cpp
+++ b/Engine/Source/Runtime/CoreUObject/Private/UObject/LinkerLoad.cpp
@@ -1401,6 +1401,7 @@ FLinkerLoad::ELinkerStatus FLinkerLoad::SerializeExportMap()
 	{
 		FObjectExport* Export = new(ExportMap)FObjectExport;
 		*this << *Export;
+		Export->ThisIndex = FPackageIndex::FromExport(ExportMapIndex);
 		ExportMapIndex++;
 	}
 
@@ -4342,6 +4343,101 @@ FName FLinkerLoad::FindNewNameForClass(FName OldClassName, bool bIsInstance)
 	return NAME_None;
 }
 
+#if WITH_EDITOR
+
+/**
+* Checks if exports' indexes and names are equal.
+*/
+bool AreObjectExportsEqualForDuplicateChecks(const FObjectExport& Lhs, const FObjectExport& Rhs)
+{
+	return Lhs.ObjectName == Rhs.ObjectName
+		&& Lhs.ClassIndex == Rhs.ClassIndex
+		&& Lhs.OuterIndex == Rhs.OuterIndex;
+}
+
+/**
+ * Helper function to sort ExportMap for duplicate checks.
+ */
+bool ExportMapSorter(const FObjectExport& Lhs, const FObjectExport& Rhs)
+{
+	// Check names first.
+	if (Lhs.ObjectName < Rhs.ObjectName)
+	{
+		return true;
+	}
+
+	if (Rhs.ObjectName > Rhs.ObjectName)
+	{
+		return false;
+	}
+
+	// Names are equal, check classes.
+	if (Lhs.ClassIndex < Rhs.ClassIndex)
+	{
+		return true;
+	}
+
+	if (Lhs.ClassIndex > Rhs.ClassIndex)
+	{
+		return false;
+	}
+
+	// Class names are equal as well, check outers.
+	return Lhs.OuterIndex < Rhs.OuterIndex;
+}
+
+void FLinkerLoad::ReplaceExportIndexes(const FPackageIndex& OldIndex, const FPackageIndex& NewIndex)
+{
+	for (auto& Export : ExportMap)
+	{
+		if (Export.ClassIndex == OldIndex)
+		{
+			Export.ClassIndex = NewIndex;
+		}
+
+		if (Export.SuperIndex == OldIndex)
+		{
+			Export.SuperIndex = NewIndex;
+		}
+
+		if (Export.OuterIndex == OldIndex)
+		{
+			Export.OuterIndex = NewIndex;
+		}
+	}
+}
+
+void FLinkerLoad::FixupDuplicateExports()
+{
+	// We need to operate on copy to avoid incorrect indexes after sorting
+	auto ExportMapSorted = ExportMap;
+	ExportMapSorted.Sort(ExportMapSorter);
+
+	// ClassIndex, SuperIndex, OuterIndex
+	int32 LastUniqueExportIndex = 0;
+	for (int32 SortedIndex = 1; SortedIndex < ExportMapSorted.Num(); ++SortedIndex)
+	{
+		const FObjectExport& Original = ExportMapSorted[LastUniqueExportIndex];
+		const FObjectExport& Duplicate = ExportMapSorted[SortedIndex];
+
+		if (AreObjectExportsEqualForDuplicateChecks(Original, Duplicate))
+		{
+			// Duplicate entry found. Look through all Exports and update their ClassIndex, SuperIndex and OuterIndex
+			// to point on original export instead of duplicate.
+			const FPackageIndex& DuplicateIndex = Duplicate.ThisIndex;
+			const FPackageIndex& OriginalIndex = Original.ThisIndex;
+			ReplaceExportIndexes(DuplicateIndex, OriginalIndex);
+
+			// Mark Duplicate as null, so we don't load it.
+			Exp(Duplicate.ThisIndex).ThisIndex = FPackageIndex();
+		}
+		else
+		{
+			LastUniqueExportIndex = SortedIndex;
+		}
+	}
+}
+#endif // WITH_EDITOR
 
 /**
 * Allows object instances to be converted to other classes upon loading a package
@@ -4350,6 +4446,13 @@ FLinkerLoad::ELinkerStatus FLinkerLoad::FixupExportMap()
 {
 	DECLARE_SCOPE_CYCLE_COUNTER( TEXT( "FLinkerLoad::FixupExportMap" ), STAT_LinkerLoad_FixupExportMap, STATGROUP_LinkerLoad );
 
+#if WITH_EDITOR
+	if (UE4Ver() < VER_UE4_SKIP_DUPLICATE_EXPORTS_ON_SAVE_PACKAGE)
+	{
+		FixupDuplicateExports();
+	}
+#endif // WITH_EDITOR
+
 	// No need to fixup exports if everything is cooked.
 	if (!FPlatformProperties::RequiresCookedData())
 	{
diff --git a/Engine/Source/Runtime/CoreUObject/Private/UObject/SavePackage.cpp b/Engine/Source/Runtime/CoreUObject/Private/UObject/SavePackage.cpp
--- a/Engine/Source/Runtime/CoreUObject/Private/UObject/SavePackage.cpp
+++ b/Engine/Source/Runtime/CoreUObject/Private/UObject/SavePackage.cpp
@@ -2616,6 +2616,95 @@ EObjectMark UPackage::GetObjectMarksForTargetPlatform( const class ITargetPlatfo
 	return ObjectMarks;
 }
 
+#if WITH_EDITOR
+/**
+ * Helper function to sort export objects by fully qualified names.
+ */
+bool ExportObjectSorter(const UObject& Lhs, const UObject& Rhs)
+{
+	// Check names first.
+	if (Lhs.GetFName() < Rhs.GetFName())
+	{
+		return true;
+	}
+
+	if (Lhs.GetFName() > Rhs.GetFName())
+	{
+		return false;
+	}
+
+	// Names equal, compare class names.
+	if (Lhs.GetClass()->GetFName() < Rhs.GetClass()->GetFName())
+	{
+		return true;
+	}
+
+	if (Lhs.GetClass()->GetFName() > Rhs.GetClass()->GetFName())
+	{
+		return false;
+	}
+
+	// Compare by outers if they exist.
+	if (Lhs.GetOuter() && Rhs.GetOuter())
+	{
+		return Lhs.GetOuter()->GetFName() < Rhs.GetOuter()->GetFName();
+	}
+
+	if (Lhs.GetOuter())
+	{
+		return true;
+	}
+
+	return false;
+}
+
+/**
+* Helper equality comparator for export objects. Compares by names, class names and outer names.
+*/
+bool ExportEqualityComparator(UObject* Lhs, UObject* Rhs)
+{
+	check(Lhs && Rhs);
+	return Lhs->GetOuter() == Rhs->GetOuter()
+		&& Lhs->GetClass() == Rhs->GetClass()
+		&& Lhs->GetFName() == Rhs->GetFName();
+}
+
+/**
+ * Remove OBJECTMARK_TagExp from duplicated objects.
+ */
+TMap<UObject*, UObject*> UnmarkExportTagFromDuplicates()
+{
+	TMap<UObject*, UObject*> RedirectDuplicatesToOriginals;
+	TArray<UObject*> Objects;
+	GetObjectsWithAnyMarks(Objects, OBJECTMARK_TagExp);
+
+	Objects.Sort(ExportObjectSorter);
+
+	int32 LastUniqueObjectIndex = 0;
+	for (int32 CurrentObjectIndex = 1; CurrentObjectIndex < Objects.Num(); ++CurrentObjectIndex)
+	{
+		UObject* LastUniqueObject = Objects[LastUniqueObjectIndex];
+		UObject* CurrentObject = Objects[CurrentObjectIndex];
+
+		// Check if duplicates with different pointers
+		if (LastUniqueObject != CurrentObject
+			// but matching names
+			&& ExportEqualityComparator(LastUniqueObject, CurrentObject))
+		{
+			// Don't export duplicates.
+			CurrentObject->UnMark(OBJECTMARK_TagExp);
+			RedirectDuplicatesToOriginals.Add(CurrentObject, LastUniqueObject);
+		}
+		else
+		{
+			LastUniqueObjectIndex = CurrentObjectIndex;
+		}
+	}
+
+	return RedirectDuplicatesToOriginals;
+}
+#endif // WITH_EDITOR
+
 /**
  * Save one specific object (along with any objects it references contained within the same Outer) into an Unreal package.
  * 
@@ -3028,6 +3117,11 @@ bool UPackage::SavePackage( UPackage* InOuter, UObject* Base, EObjectFlags TopLe
 					}
 				}
 
+#if WITH_EDITOR
+				// Remove TagExp from duplicate objects.
+				TMap<UObject*, UObject*> DuplicateRedirects = UnmarkExportTagFromDuplicates();
+#endif // WITH_EDITOR
+
 				UE_LOG_COOK_TIME(TEXT("Serialize Imports"));
 				
 				if ( EndSavingIfCancelled( Linker, TempFilename ) ) { return false; }
@@ -3462,6 +3556,14 @@ bool UPackage::SavePackage( UPackage* InOuter, UObject* Base, EObjectFlags TopLe
 								DependencyIndex = ImportToIndexMap.FindRef(DependentObject);
 							}
 					
+#if WITH_EDITOR
+							// If we still didn't find index, maybe it was a duplicate export which got removed.
+							// Check if we have a redirect to original.
+							if (DependencyIndex.IsNull())
+							{
+								DependencyIndex = ExportToIndexMap.FindRef(DuplicateRedirects[DependentObject]);
+							}
+#endif
 							// if we didn't find it (FindRef returns 0 on failure, which is good in this case), then we are in trouble, something went wrong somewhere
 							checkf(!DependencyIndex.IsNull(), TEXT("Failed to find dependency index for %s (%s)"), *DependentObject->GetFullName(), *Object->GetFullName());
 
diff --git a/Engine/Source/Runtime/CoreUObject/Public/UObject/Linker.h b/Engine/Source/Runtime/CoreUObject/Public/UObject/Linker.h
--- a/Engine/Source/Runtime/CoreUObject/Public/UObject/Linker.h
+++ b/Engine/Source/Runtime/CoreUObject/Public/UObject/Linker.h
@@ -97,6 +97,23 @@ public:
 		return Index != Other.Index;
 	}
 
+	/** Compare package indecies **/
+	FORCEINLINE bool operator<(const FPackageIndex& Other) const
+	{
+		return Index < Other.Index;
+	}
+	FORCEINLINE bool operator>(const FPackageIndex& Other) const
+	{
+		return Index > Other.Index;
+	}
+	FORCEINLINE bool operator<=(const FPackageIndex& Other) const
+	{
+		return Index <= Other.Index;
+	}
+	FORCEINLINE bool operator>=(const FPackageIndex& Other) const
+	{
+		return Index >= Other.Index;
+	}
 	/**
 	 * Serializes a package index value from or into an archive.
 	 *
@@ -160,6 +177,13 @@ struct FObjectExport : public FObjectResource
 	 */
 	FPackageIndex  	ClassIndex;
 
+	/**
+	* Location of this resource in export map. Used for export fixups while loading packages.
+	* Value of zero indicates resource is invalid and shouldn't be loaded.
+	* Not serialized.
+	*/
+	FPackageIndex ThisIndex;
+
 	/**
 	 * Location of the resource for this export's SuperField (parent).  Only valid if
 	 * this export represents a UStruct object. A value of zero indicates that the object
@@ -1166,6 +1190,10 @@ public:
 		{
 			return true;
 		}
+		if (Export.ThisIndex.IsNull()) // Export is invalid and shouldn't be processed.
+		{
+			return true;
+		}
 		return false;
 	}
 
@@ -1629,6 +1657,19 @@ public:
 	COREUOBJECT_API static FName FindSubobjectRedirectName(const FName& Name);
 
 private:
+#if WITH_EDITOR
+
+	/**
+	 * Nulls duplicated exports and fixes indexes in ExportMap to point to original objects instead of duplicates.
+	 */
+	void FixupDuplicateExports();
+
+	/**
+	 * Replaces all instances of OldIndex in ExportMap with NewIndex.
+	 */
+	void ReplaceExportIndexes(const FPackageIndex& OldIndex, const FPackageIndex& NewIndex);
+
+#endif // WITH_EDITOR
 
 	UObject* CreateExport( int32 Index );
 
