[INTEGRATE] Change 2552532  by [REDACTED_PERSON][REDACTED_HANDLE] on [REDACTED_PHONE]:14

Fixing crash on startup introduced in CL #2550790, and networking code failures due to linker being destroyed before async callbacks got triggered

diff --git a/Engine/Source/Runtime/CoreUObject/Private/Serialization/AsyncLoading.cpp b/Engine/Source/Runtime/CoreUObject/Private/Serialization/AsyncLoading.cpp
--- a/Engine/Source/Runtime/CoreUObject/Private/Serialization/AsyncLoading.cpp
+++ b/Engine/Source/Runtime/CoreUObject/Private/Serialization/AsyncLoading.cpp
@@ -505,13 +505,7 @@ EAsyncPackageState::Type FAsyncLoadingThread::ProcessLoadedPackages(bool bUseTim
 			Result = Package->PostLoadDeferredObjects(TickStartTime, bUseTimeLimit, TimeLimit);
 			if (Result == EAsyncPackageState::Complete)
 			{
-				if (FPlatformProperties::RequiresCookedData())
-				{
-					// Emulates ResetLoaders on the package linker's linkerroot.
-					Package->ResetLoader();
-				}
-
-				// Remove the package from the list before we fore callbacks, 
+				// Remove the package from the list before we trigger the callbacks, 
 				// this is to ensure we can re-enter FlushAsyncLoading from any of the callbacks
 				LoadedPackagesToProcess.RemoveAt(PackageIndex--);
 
@@ -524,6 +518,12 @@ EAsyncPackageState::Type FAsyncLoadingThread::ProcessLoadedPackages(bool bUseTim
 				const EAsyncLoadingResult::Type LoadingResult = Package->HasLoadFailed() ? EAsyncLoadingResult::Failed : EAsyncLoadingResult::Succeeded;
 				Package->CallCompletionCallbacks(bInternalCallbacks, LoadingResult);
 
+				if (FPlatformProperties::RequiresCookedData())
+				{
+					// Emulates ResetLoaders on the package linker's linkerroot.
+					Package->ResetLoader();
+				}
+
 				// We don't need the package anymore
 				delete Package;
 			}
@@ -765,7 +765,7 @@ int32 FAsyncPackage::PostLoadIndex = 0;
 */
 FAsyncPackage::FAsyncPackage(const FAsyncPackageDesc& InDesc)
 : Desc(InDesc)
-, Linker(NULL)
+, Linker(nullptr)
 , DependencyRefCount(0)
 , LoadImportIndex(0)
 , ImportIndex(0)
@@ -778,8 +778,8 @@ FAsyncPackage::FAsyncPackage(const FAsyncPackageDesc& InDesc)
 , bLoadHasFailed(false)
 , bLoadHasFinished(false)
 , TickStartTime(0)
-, LastObjectWorkWasPerformedOn(NULL)
-, LastTypeOfWorkPerformed(NULL)
+, LastObjectWorkWasPerformedOn(nullptr)
+, LastTypeOfWorkPerformed(nullptr)
 , LoadStartTime(0.0)
 , LoadPercentage(0)
 #if PERF_TRACK_DETAILED_ASYNC_STATS
@@ -1535,11 +1535,18 @@ EAsyncPackageState::Type FAsyncPackage::PostLoadDeferredObjects(double InTickSta
 	Result = (DeferredPostLoadIndex == DeferredPostLoadObjects.Num()) ? EAsyncPackageState::Complete : EAsyncPackageState::TimeOut;
 	if (Result == EAsyncPackageState::Complete)
 	{
+		// Clear async loading flags (we still want RF_Async, but RF_AsyncLoading can be cleared)
+		for (UObject* Object : DeferredPostLoadObjects)
+		{
+			Object->AtomicallyClearFlags(RF_AsyncLoading);
+		}
+
 		// Mark package as having been fully loaded and update load time.
 		if (Linker && !bLoadHasFailed)
 		{
 			if (Linker->LinkerRoot)
 			{
+				Linker->LinkerRoot->AtomicallyClearFlags(RF_AsyncLoading);
 				Linker->LinkerRoot->MarkAsFullyLoaded();
 				Linker->LinkerRoot->SetLoadTime(FPlatformTime::Seconds() - LoadStartTime);
 			}
