Don't let packages that are async loading potentially return true from IsFullyLoaded. Prevents PackageMap from finding packages whose GUIDs may not have been set yet.

diff --git a/Engine/Source/Runtime/CoreUObject/Private/UObject/Package.cpp b/Engine/Source/Runtime/CoreUObject/Private/UObject/Package.cpp
--- a/Engine/Source/Runtime/CoreUObject/Private/UObject/Package.cpp
+++ b/Engine/Source/Runtime/CoreUObject/Private/UObject/Package.cpp
@@ -173,7 +173,7 @@ bool UPackage::IsFullyLoaded()
 	// Newly created packages aren't loaded and therefore haven't been marked as being fully loaded. They are treated as fully
 	// loaded packages though in this case, which is why we are looking to see whether the package exists on disk and assume it
 	// has been fully loaded if it doesn't.
-	if( !bHasBeenFullyLoaded )
+	if( !bHasBeenFullyLoaded && !HasAnyFlags(RF_AsyncLoading) )
 	{
 		FString DummyFilename;
 		// Try to find matching package in package file cache.
