UStruct::SuperStruct is now private instead of protected and is only accessible through GetSuperStruct().
Handle SuperStruct being overwritten by archives during hot reload.
FFastIndexingClassTree state unified under a single static.
Extra tree integrity checks.

diff --git a/Engine/Source/Runtime/CoreUObject/Private/UObject/Class.cpp b/Engine/Source/Runtime/CoreUObject/Private/UObject/Class.cpp
--- a/Engine/Source/Runtime/CoreUObject/Private/UObject/Class.cpp
+++ b/Engine/Source/Runtime/CoreUObject/Private/UObject/Class.cpp
@@ -2759,9 +2759,9 @@ void UClass::DeferredRegister(UClass *UClassStaticClass,const TCHAR* PackageName
 	ClassConfigName = InClassConfigName;
 
 	// Propagate inherited flags.
-	if (SuperStruct != NULL)
+	UClass* SuperClass = GetSuperClass();
+	if (SuperClass != NULL)
 	{
-		UClass* SuperClass = GetSuperClass();
 		ClassFlags |= (SuperClass->ClassFlags & CLASS_Inherit);
 		ClassCastFlags |= SuperClass->ClassCastFlags;
 	}
@@ -2965,9 +2965,9 @@ void UClass::Link(FArchive& Ar, bool bRelinkExistingProperties)
 	if (PropertyLink != NULL)
 	{
 		NetFields.Empty();
-		if (SuperStruct)
+		if (UClass* SuperClass = GetSuperClass())
 		{
-			ClassReps = GetSuperClass()->ClassReps;
+			ClassReps = SuperClass->ClassReps;
 		}
 		else
 		{
@@ -3041,6 +3041,23 @@ void UClass::Link(FArchive& Ar, bool bRelinkExistingProperties)
 
 #if UCLASS_FAST_ISA_IMPL & 2
 
+	struct FClassParentPair
+	{
+		UClass* Class;
+		UClass* Parent;
+
+		FClassParentPair(UClass* InClass, UClass* InParent)
+			: Class (InClass)
+			, Parent(InParent)
+		{
+		}
+
+		friend bool operator==(const FClassParentPair& Lhs, const UClass* Rhs) { return Lhs.Class == Rhs; }
+		friend bool operator!=(const FClassParentPair& Lhs, const UClass* Rhs) { return Lhs.Class != Rhs; }
+		friend bool operator==(const UClass* Lhs, const FClassParentPair& Rhs) { return Lhs == Rhs.Class; }
+		friend bool operator!=(const UClass* Lhs, const FClassParentPair& Rhs) { return Lhs != Rhs.Class; }
+	};
+
 	/**
 	 * Tree for fast IsA implementation.
 	 *
@@ -3056,35 +3073,37 @@ void UClass::Link(FArchive& Ar, bool bRelinkExistingProperties)
 		static void Register(UClass* NewClass);
 		static void Unregister(UClass* NewClass);
 
-		static TArray<UClass*>& GetClasses()
-		{
-			static TArray<UClass*> Classes;
-			return Classes;
-		}
-		static TMap<UClass*, TArray<UClass*>>& GetOrphans()
+		struct StateType
 		{
-			static TMap<UClass*, TArray<UClass*>> Orphans;
-			return Orphans;
-		}
-		static FCriticalSection& GetClassesCriticalSection()
+			TArray<FClassParentPair> Classes;
+			TSet<UClass*>            Orphans;
+			FCriticalSection ClassesCriticalSection;
+		};
+
+		static StateType& GetState()
 		{
-			static FCriticalSection ClassesCriticalSection;
-			return ClassesCriticalSection;
+			static StateType State;
+			return State;
 		}
+
+	public:
+		static void Validate();
 	};
 
 	void FFastIndexingClassTree::Register(UClass* Class)
 	{
-		FScopeLock Lock(&GetClassesCriticalSection());
-		TMap<UClass*, TArray<UClass*>>& Orphans = GetOrphans();
-		TArray<UClass*>& Classes = GetClasses();
+		StateType& State = GetState();
+		FScopeLock Lock(&State.ClassesCriticalSection);
+
+		// Ensure that the class is not already registered or orphaned
+		check(!State.Classes.Contains(Class) && !State.Orphans.Contains(Class));
 
 		UClass* ParentClass = Class->GetSuperClass();
 
 		// If the parent has previously been orphaned, flag the child as orphaned
-		if (TArray<UClass*>* ParentOrphans = Orphans.Find(ParentClass))
+		if (State.Orphans.Contains(ParentClass))
 		{
-			ParentOrphans->Add(Class);
+			State.Orphans.Add(Class);
 			return;
 		}
 
@@ -3092,9 +3111,9 @@ void UClass::Link(FArchive& Ar, bool bRelinkExistingProperties)
 		if (ParentClass)
 		{
 			// Can happen if a child is registered *after* the parent
-			if (!Classes.Contains(ParentClass))
+			if (!State.Classes.Contains(ParentClass))
 			{
-				Orphans.Add(ParentClass).Add(Class);
+				State.Orphans.Add(Class);
 				return;
 			}
 
@@ -3102,13 +3121,13 @@ void UClass::Link(FArchive& Ar, bool bRelinkExistingProperties)
 		}
 		else
 		{
-			NewIndex = Classes.Num();
+			NewIndex = State.Classes.Num();
 		}
 
 		// Increment indices of following classes
-		for (auto Index = NewIndex, LastIndex = Classes.Num(); Index != LastIndex; ++Index)
+		for (auto Index = NewIndex, LastIndex = State.Classes.Num(); Index != LastIndex; ++Index)
 		{
-			++Classes[Index]->ClassTreeIndex;
+			++State.Classes[Index].Class->ClassTreeIndex;
 		}
 
 		// Update children count of all parents
@@ -3120,50 +3139,50 @@ void UClass::Link(FArchive& Ar, bool bRelinkExistingProperties)
 		// Add class
 		Class->ClassTreeIndex       = NewIndex;
 		Class->ClassTreeNumChildren = 0;
-		Classes.Insert(Class, NewIndex);
+		State.Classes.Insert(FClassParentPair(Class, ParentClass), NewIndex);
 
 		// Re-register any children orphaned by a previous Unregister call
-		if (TArray<UClass*>* FoundOrphans = Orphans.Find(Class))
+		TArray<UClass*> OrphansToReregister;
+		for (auto It = State.Orphans.CreateIterator(); It; ++It)
 		{
-			TArray<UClass*> OrphansToReregister = MoveTemp(*FoundOrphans);
-			Orphans.Remove(Class);
-			Orphans.Compact();
-
-			for (UClass* Orphan : OrphansToReregister)
+			UClass* Orphan = *It;
+			if (Orphan->GetSuperClass() == Class)
 			{
-				Register(Orphan);
+				OrphansToReregister.Add(Orphan);
+				It.RemoveCurrent();
 			}
 		}
+
+		State.Orphans.Compact();
+
+		for (UClass* Orphan : OrphansToReregister)
+		{
+			Register(Orphan);
+		}
+
+		#if DO_CHECK
+			Validate();
+		#endif
 	}
 
 	void FFastIndexingClassTree::Unregister(UClass* Class)
 	{
-		FScopeLock Lock(&GetClassesCriticalSection());
-		TMap<UClass*, TArray<UClass*>>& Orphans = GetOrphans();
-		TArray<UClass*>& Classes = GetClasses();
-
-		UClass* ParentClass = Class->GetSuperClass();
+		StateType& State = GetState();
+		FScopeLock Lock(&State.ClassesCriticalSection);
 
-		// Remove class if it was orphaned
-		if (TArray<UClass*>* Children = Orphans.Find(ParentClass))
+		// Remove class if it was already orphaned
+		if (State.Orphans.Remove(Class))
 		{
-			// Remove the child, or the entire array if it's the only child left.
-			if (Children->Num() != 1)
-			{
-				int32 ChildIndex = Children->Find(Class);
-				check(ChildIndex != INDEX_NONE);
-				Children->RemoveAt(ChildIndex);
-			}
-			else
-			{
-				check(Children->Last() == Class);
-				Orphans.Remove(ParentClass);
-				Orphans.Compact();
-			}
-
+			State.Orphans.Compact();
 			return;
 		}
 
+		UClass* ParentClass = State.Classes[Class->ClassTreeIndex].Parent;
+
+		// Ensure that the class and any parent are registered and in the expected location
+		check(                State.Classes[Class      ->ClassTreeIndex].Class == Class);
+		check(!ParentClass || State.Classes[ParentClass->ClassTreeIndex].Class == ParentClass);
+
 		// Remove it and mark its children as orphaned
 		int32 ClassIndex       = Class->ClassTreeIndex;
 		int32 ClassNumChildren = Class->ClassTreeNumChildren;
@@ -3172,13 +3191,13 @@ void UClass::Link(FArchive& Ar, bool bRelinkExistingProperties)
 		// Mark any children as orphaned
 		for (int32 Index = ClassIndex + 1, EndIndex = ClassIndex + NumRemoved; Index != EndIndex; ++Index)
 		{
-			Orphans.FindOrAdd(Classes[Index]->GetSuperClass()).Add(Classes[Index]);
+			State.Orphans.Add(State.Classes[Index].Class);
 		}
 
 		// Decrement indices of following classes
-		for (int32 Index = ClassIndex + NumRemoved, IndexEnd = Classes.Num(); Index != IndexEnd; ++Index)
+		for (int32 Index = ClassIndex + NumRemoved, IndexEnd = State.Classes.Num(); Index != IndexEnd; ++Index)
 		{
-			Classes[Index]->ClassTreeIndex -= NumRemoved;
+			State.Classes[Index].Class->ClassTreeIndex -= NumRemoved;
 		}
 
 		// Update children count of all parents
@@ -3187,7 +3206,40 @@ void UClass::Link(FArchive& Ar, bool bRelinkExistingProperties)
 			Parent->ClassTreeNumChildren -= NumRemoved;
 		}
 
-		Classes.RemoveAt(ClassIndex, NumRemoved, false);
+		State.Classes.RemoveAt(ClassIndex, NumRemoved, false);
+
+		#if DO_CHECK
+			Validate();
+		#endif
+	}
+
+	void FFastIndexingClassTree::Validate()
+	{
+		StateType& State = GetState();
+		FScopeLock Lock(&State.ClassesCriticalSection);
+
+		for (const FClassParentPair& Pair : State.Classes)
+		{
+			int32 Index = Pair.Class->ClassTreeIndex;
+
+			// Check that the class is not orphaned
+			check(!State.Orphans.Contains(Pair.Class));
+
+			// Check that the class is where it thinks it is
+			check(State.Classes[Index] == Pair.Class);
+
+			if (Pair.Parent)
+			{
+				int32 ParentIndex = Pair.Parent->ClassTreeIndex;
+
+				// Check that the parent is registered and not orphaned
+				check( State.Classes.Contains(Pair.Parent));
+				check(!State.Orphans.Contains(Pair.Parent));
+
+				// Check that class 'is' its parent
+				check(Index - Pair.Parent->ClassTreeIndex <= Pair.Parent->ClassTreeNumChildren);
+			}
+		}
 	}
 
 	FFastIndexingClassTreeRegistrar::FFastIndexingClassTreeRegistrar()
@@ -3248,7 +3300,23 @@ void UClass::Serialize( FArchive& Ar )
 		UnhashObject(this);
 	}
 
+#if UCLASS_FAST_ISA_IMPL & 2
+	UClass* SuperClassBefore = GetSuperClass();
+#endif
 	Super::Serialize( Ar );
+#if UCLASS_FAST_ISA_IMPL & 2
+	// Handle that fact that FArchive takes UObject*s by reference, and archives can just blat
+	// over our SuperStruct with impunity.
+	if (SuperClassBefore)
+	{
+		UClass* SuperClassAfter = GetSuperClass();
+		if (SuperClassBefore != SuperClassAfter)
+		{
+			FFastIndexingClassTree::Unregister(this);
+			FFastIndexingClassTree::Register(this);
+		}
+	}
+#endif
 
 	if ( Ar.IsLoading() || Ar.IsModifyingWeakAndStrongReferences() )
 	{
diff --git a/Engine/Source/Runtime/CoreUObject/Public/UObject/Class.h b/Engine/Source/Runtime/CoreUObject/Public/UObject/Class.h
--- a/Engine/Source/Runtime/CoreUObject/Public/UObject/Class.h
+++ b/Engine/Source/Runtime/CoreUObject/Public/UObject/Class.h
@@ -260,6 +260,7 @@ class COREUOBJECT_API UStruct : public UField
 	// Variables.
 protected:
 	friend COREUOBJECT_API UClass* Z_Construct_UClass_UStruct();
+private:
 	UStruct* SuperStruct;
 public:
 	UField* Children;
@@ -1272,8 +1273,9 @@ public:
 	// UFunction interface.
 	UFunction* GetSuperFunction() const
 	{
-		checkSlow(!SuperStruct||SuperStruct->IsA<UFunction>());
-		return (UFunction*)SuperStruct;
+		UStruct* Result = GetSuperStruct();
+		checkSlow(!Result || Result->IsA<UFunction>());
+		return (UFunction*)Result;
 	}
 
 	UProperty* GetReturnProperty() const;
@@ -2055,7 +2057,7 @@ public:
 
 	UClass* GetSuperClass() const
 	{
-		return (UClass*)SuperStruct;
+		return (UClass*)GetSuperStruct();
 	}
 
 	/** Feedback context for default property import **/
diff --git a/Engine/Source/Runtime/Engine/Private/BlueprintGeneratedClass.cpp b/Engine/Source/Runtime/Engine/Private/BlueprintGeneratedClass.cpp
--- a/Engine/Source/Runtime/Engine/Private/BlueprintGeneratedClass.cpp
+++ b/Engine/Source/Runtime/Engine/Private/BlueprintGeneratedClass.cpp
@@ -401,7 +401,7 @@ void UBlueprintGeneratedClass::BindDynamicDelegates(UObject* InInstance) const
 	}
 
 	// call on super class, if it's a BlueprintGeneratedClass
-	UBlueprintGeneratedClass* BGClass = Cast<UBlueprintGeneratedClass>(SuperStruct);
+	UBlueprintGeneratedClass* BGClass = Cast<UBlueprintGeneratedClass>(GetSuperStruct());
 	if(BGClass != NULL)
 	{
 		BGClass->BindDynamicDelegates(InInstance);
