Fast IsA code inlined.
Fast IsA code enabled by default.

diff --git a/Engine/Source/Runtime/CoreUObject/Private/UObject/UObjectBaseUtility.cpp b/Engine/Source/Runtime/CoreUObject/Private/UObject/UObjectBaseUtility.cpp
--- a/Engine/Source/Runtime/CoreUObject/Private/UObject/UObjectBaseUtility.cpp
+++ b/Engine/Source/Runtime/CoreUObject/Private/UObject/UObjectBaseUtility.cpp
@@ -248,6 +248,7 @@ bool UObjectBaseUtility::RootPackageHasAnyFlags( uint32 CheckFlagMask ) const
 /**
  * @return	true if this object is of the specified type.
  */
+#if UCLASS_FAST_ISA_IMPL != 2
 bool UObjectBaseUtility::IsA( const UClass* SomeBase ) const
 {
 	UE_CLOG(!SomeBase, LogObj, Fatal, TEXT("IsA(NULL) cannot yield meaningful results"));
@@ -278,6 +279,7 @@ bool UObjectBaseUtility::IsA( const UClass* SomeBase ) const
 		return bNewResult;
 	#endif
 }
+#endif
 
 
 /**
diff --git a/Engine/Source/Runtime/CoreUObject/Public/UObject/Class.h b/Engine/Source/Runtime/CoreUObject/Public/UObject/Class.h
--- a/Engine/Source/Runtime/CoreUObject/Public/UObject/Class.h
+++ b/Engine/Source/Runtime/CoreUObject/Public/UObject/Class.h
@@ -8,15 +8,6 @@
 
 #include "ObjectBase.h"
 
-// 1 = old behavior
-// 2 = new behavior
-// 3 = old behavior with checks against the new behavior
-#if !(UE_BUILD_SHIPPING || UE_BUILD_TEST)
-	#define UCLASS_FAST_ISA_IMPL 3
-#else
-	#define UCLASS_FAST_ISA_IMPL 2
-#endif
-
 /*-----------------------------------------------------------------------------
 	Mirrors of mirror structures in Object.h. These are used by generated code 
 	to facilitate correct offsets and alignments for structures containing these '
@@ -2791,4 +2782,4 @@ struct FFallbackStruct;
 template<> struct TBaseStructure<FFallbackStruct>
 {
 	COREUOBJECT_API static UScriptStruct* Get();
-};
\ No newline at end of file
+};
diff --git a/Engine/Source/Runtime/CoreUObject/Public/UObject/ObjectBase.h b/Engine/Source/Runtime/CoreUObject/Public/UObject/ObjectBase.h
--- a/Engine/Source/Runtime/CoreUObject/Public/UObject/ObjectBase.h
+++ b/Engine/Source/Runtime/CoreUObject/Public/UObject/ObjectBase.h
@@ -23,6 +23,15 @@ typedef	uint64 ScriptPointerType;
 /** Set this to 0 to disable UObject thread safety features */
 #define THREADSAFE_UOBJECTS 1
 
+// 1 = old IsA behavior
+// 2 = new IsA behavior
+// 3 = old IsA behavior with checks against the new behavior
+#if 0//!(UE_BUILD_SHIPPING || UE_BUILD_TEST)
+	#define UCLASS_FAST_ISA_IMPL 3
+#else
+	#define UCLASS_FAST_ISA_IMPL 2
+#endif
+
 /*-----------------------------------------------------------------------------
 	Core enumerations.
 -----------------------------------------------------------------------------*/
diff --git a/Engine/Source/Runtime/CoreUObject/Public/UObject/UObjectBaseUtility.h b/Engine/Source/Runtime/CoreUObject/Public/UObject/UObjectBaseUtility.h
--- a/Engine/Source/Runtime/CoreUObject/Public/UObject/UObjectBaseUtility.h
+++ b/Engine/Source/Runtime/CoreUObject/Public/UObject/UObjectBaseUtility.h
@@ -6,6 +6,8 @@
 
 #pragma once
 
+#include "Templates/PointerIsConvertibleFromTo.h"
+
 class COREUOBJECT_API UObjectBaseUtility : public UObjectBase
 {
 public:
@@ -278,7 +280,37 @@ public:
 	/**
 	 * @return	true if this object is of the specified type.
 	 */
-	bool IsA( const UClass* SomeBaseClass ) const;
+	#if UCLASS_FAST_ISA_IMPL == 2
+	private:
+		template <typename ClassType>
+		static FORCEINLINE bool IsAWorkaround(const ClassType* ObjClass, const ClassType* TestCls)
+		{
+			return ObjClass->IsAUsingFastTree(*TestCls);
+		}
+
+	public:
+		template <typename OtherClassType>
+		FORCEINLINE bool IsA( OtherClassType SomeBase ) const
+		{
+			// We have a cyclic dependency between UObjectBaseUtility and UClass,
+			// so we use a template to allow inlining of something we haven't yet seen, because it delays compilation until the function is called.
+
+			// 'static_assert' that this thing is actually a UClass pointer or convertible to it.
+			const UClass* SomeBaseClass = SomeBase;
+			(void)SomeBaseClass;
+			checkfSlow(SomeBaseClass, TEXT("IsA(NULL) cannot yield meaningful results"));
+
+			const UClass* ThisClass = GetClass();
+
+			// Stop the compiler doing some unnecessary branching for nullptr checks
+			ASSUME(SomeBaseClass);
+			ASSUME(ThisClass);
+
+			return IsAWorkaround(ThisClass, SomeBaseClass);
+		}
+	#else
+		bool IsA( const UClass* SomeBase ) const;
+	#endif
 
 	/**
 	 * @return	true if this object is of the template type.
