Fix for objects not having RF_AsyncLoading flag cleared in non-cooked builds before triggering async loading callbacks

diff --git a/Engine/Source/Runtime/CoreUObject/Private/Serialization/AsyncLoading.cpp b/Engine/Source/Runtime/CoreUObject/Private/Serialization/AsyncLoading.cpp
--- a/Engine/Source/Runtime/CoreUObject/Private/Serialization/AsyncLoading.cpp
+++ b/Engine/Source/Runtime/CoreUObject/Private/Serialization/AsyncLoading.cpp
@@ -1498,6 +1498,8 @@ EAsyncPackageState::Type FAsyncPackage::PostLoadObjects()
 		{
 			DeferredPostLoadObjects.Add(Object);
 		}
+		// All object must be finalized on the game thread
+		DeferredFinalizeObjects.Add(Object);
 	}
 
 	// New objects might have been loaded during PostLoad.
@@ -1536,7 +1538,7 @@ EAsyncPackageState::Type FAsyncPackage::PostLoadDeferredObjects(double InTickSta
 	if (Result == EAsyncPackageState::Complete)
 	{
 		// Clear async loading flags (we still want RF_Async, but RF_AsyncLoading can be cleared)
-		for (UObject* Object : DeferredPostLoadObjects)
+		for (UObject* Object : DeferredFinalizeObjects)
 		{
 			Object->AtomicallyClearFlags(RF_AsyncLoading);
 		}
diff --git a/Engine/Source/Runtime/CoreUObject/Public/Serialization/AsyncLoading.h b/Engine/Source/Runtime/CoreUObject/Public/Serialization/AsyncLoading.h
--- a/Engine/Source/Runtime/CoreUObject/Public/Serialization/AsyncLoading.h
+++ b/Engine/Source/Runtime/CoreUObject/Public/Serialization/AsyncLoading.h
@@ -161,7 +161,8 @@ private:
 	float						LoadPercentage;
 	/** Objects to be post loaded on the game thread */
 	TArray<UObject*> DeferredPostLoadObjects;
-
+	/** Objects to be finalized on the game thread */
+	TArray<UObject*> DeferredFinalizeObjects;
 public:
 #if PERF_TRACK_DETAILED_ASYNC_STATS
 	/** Number of times Tick function has been called.													*/
