Forcefully resolving macro dependencies before they get expanded into a blueprint on load (to guard against half resolved macro dependencies ending up in blueprints that call said macros).

diff --git a/Engine/Source/Editor/UnrealEd/Private/Kismet2/BlueprintEditorUtils.cpp b/Engine/Source/Editor/UnrealEd/Private/Kismet2/BlueprintEditorUtils.cpp
--- a/Engine/Source/Editor/UnrealEd/Private/Kismet2/BlueprintEditorUtils.cpp
+++ b/Engine/Source/Editor/UnrealEd/Private/Kismet2/BlueprintEditorUtils.cpp
@@ -952,6 +952,14 @@ struct FRegenerationHelper
 						Linker->Preload(BP);
 					}
 				}
+				// at the point of blueprint regeneration (on load), we are guaranteed that blueprint dependencies (like this macro) have
+				// fully formed classes (meaning the blueprint class and all its direct dependencies have been loaded)... however, we do not 
+				// get the guarantee that all of that blueprint's graph dependencies are loaded (hence, why we have to force load 
+				// everything here); in the case of cyclic dependencies, macro dependencies could already be loaded, but in the midst of 
+				// resolving thier own dependency placeholders (why a ForcedLoad() call is not enough); this ensures that 
+				// placeholder objects are properly resolved on nodes that will be injected by macro expansion
+				FLinkerLoad::PRIVATE_ForceLoadAllDependencies(BP->GetOutermost());
+				
 				ForcedLoadMembers(BP);
 			}
 		}
diff --git a/Engine/Source/Runtime/CoreUObject/Private/Blueprint/BlueprintSupport.cpp b/Engine/Source/Runtime/CoreUObject/Private/Blueprint/BlueprintSupport.cpp
--- a/Engine/Source/Runtime/CoreUObject/Private/Blueprint/BlueprintSupport.cpp
+++ b/Engine/Source/Runtime/CoreUObject/Private/Blueprint/BlueprintSupport.cpp
@@ -1020,6 +1020,50 @@
 #endif // USE_CIRCULAR_DEPENDENCY_LOAD_DEFERRING
 }
 
+void FLinkerLoad::PRIVATE_ForceLoadAllDependencies(UPackage* Package)
+{
+	if (FLinkerLoad* PkgLinker = FindExistingLinkerForPackage(Package))
+	{
+		PkgLinker->ResolveAllImports();
+	}
+}
+
+void FLinkerLoad::ResolveAllImports()
+{
+	for (int32 ImportIndex = 0; ImportIndex < ImportMap.Num() && IsBlueprintFinalizationPending(); ++ImportIndex)
+	{
+		// first, make sure every import object is available... just because 
+		// it isn't present in the map already, doesn't mean it isn't in the 
+		// middle of a resolve (the CreateImport() brings in an export 
+		// object from another package, which could be resolving itself)... 
+		// 
+		// don't fret, all these imports were bound to get created sooner or 
+		// later (like when the blueprint was regenerated)
+		//
+		// NOTE: this is a possible root point for recursion... accessing a 
+		//       separate package could continue its loading process which
+		//       in turn, could end us back in this function before we ever  
+		//       returned from this
+		FObjectImport& Import = ImportMap[ImportIndex];
+		UObject* ImportObject = CreateImport(ImportIndex);
+
+		// see if this import is currently being resolved (presumably somewhere 
+		// up the callstack)... if it is, we need to ensure that this dependency 
+		// is fully resolved before we get to regenerating the blueprint (else,
+		// we could end up with placeholder classes in our script-code)
+		if (FUnresolvedStructTracker::IsImportStructUnresolved(ImportObject))
+		{
+			// because it is tracked by FUnresolvedStructTracker, it must be a struct
+			DEFERRED_DEPENDENCY_CHECK(Cast<UStruct>(ImportObject) != nullptr);
+			auto SourceLinker = FindExistingLinkerForImport(ImportIndex);
+			if (SourceLinker)
+			{
+				SourceLinker->ResolveDeferredDependencies((UStruct*)ImportObject);
+			}
+		}
+	}
+}
+
 void FLinkerLoad::FinalizeBlueprint(UClass* LoadClass)
 {
 #if USE_CIRCULAR_DEPENDENCY_LOAD_DEFERRING
@@ -1067,38 +1111,7 @@
 	// do, otherwise other linkers could want to finish this off when they don't
 	// have to)... we do however need it here in FinalizeBlueprint(), because
 	// we need it ran for any super-classes before we regen
-	for (int32 ImportIndex = 0; ImportIndex < ImportMap.Num() && IsBlueprintFinalizationPending(); ++ImportIndex)
-	{
-		// first, make sure every import object is available... just because 
-		// it isn't present in the map already, doesn't mean it isn't in the 
-		// middle of a resolve (the CreateImport() brings in an export 
-		// object from another package, which could be resolving itself)... 
-		// 
-		// don't fret, all these imports were bound to get created sooner or 
-		// later (like when the blueprint was regenerated)
-		//
-		// NOTE: this is a possible root point for recursion... accessing a 
-		//       separate package could continue its loading process which
-		//       in turn, could end us back in this function before we ever  
-		//       returned from this
-		FObjectImport& Import = ImportMap[ImportIndex];
-		UObject* ImportObject = CreateImport(ImportIndex);
+	ResolveAllImports();
 
-		// see if this import is currently being resolved (presumably somewhere 
-		// up the callstack)... if it is, we need to ensure that this dependency 
-		// is fully resolved before we get to regenerating the blueprint (else,
-		// we could end up with placeholder classes in our script-code)
-		if (FUnresolvedStructTracker::IsImportStructUnresolved(ImportObject))
-		{
-			// because it is tracked by FUnresolvedStructTracker, it must be a struct
-			DEFERRED_DEPENDENCY_CHECK(Cast<UStruct>(ImportObject) != nullptr);
-			auto SourceLinker = FindExistingLinkerForImport(ImportIndex);
-			if (SourceLinker)
-			{
-				SourceLinker->ResolveDeferredDependencies((UStruct*)ImportObject);
-			}
-		}
-	}
-
 	// the above loop could have caused some recursion... if it ended up 
 	// finalizing a sub-class (of LoadClass), then that would have finalized
diff --git a/Engine/Source/Runtime/CoreUObject/Public/UObject/Linker.h b/Engine/Source/Runtime/CoreUObject/Public/UObject/Linker.h
--- a/Engine/Source/Runtime/CoreUObject/Public/UObject/Linker.h
+++ b/Engine/Source/Runtime/CoreUObject/Public/UObject/Linker.h
@@ -1663,6 +1663,18 @@ public:
 	 */
 	COREUOBJECT_API static void PRIVATE_PatchNewObjectIntoExport(UObject* OldObject, UObject* NewObject);
 
+	/**
+	 * Wraps a call to the package linker's ResolveAllImports().
+	 * 
+	 * @param  Package    The package whose imports you want all loaded.
+	 * 
+	 * WARNING!!!	This function shouldn't be used carelessly, and serves as a
+	 *				hacky entrypoint to FLinkerLoad's privates. It should only 
+	 *				be used at very specific times, and in very specific cases.
+	 *				If you're unsure, DON'T TRY TO USE IT!!!
+	 */
+	COREUOBJECT_API static void PRIVATE_ForceLoadAllDependencies(UPackage* Package);
+
 	/**
 	 * Invalidates the future loading of a specific object, so that subsequent loads will fail
 	 * This is used to invalidate sub objects of a replaced object that may no longer be valid
@@ -1994,6 +2006,14 @@ private:
 	 */
 	bool HasUnresolvedDependencies() const;
 
+	/**
+	 * Iterates through the ImportMap and calls CreateImport() for every entry, 
+	 * creating/loading each import as we go. This also makes sure that class 
+	 * imports have had ResolveDeferredDependencies() completely executed for 
+	 * them (even those already running through it earlier in the callstack).
+	 */
+	void ResolveAllImports();
+
 	/**
 	 * Takes the supplied serialized class and serializes in its CDO, then 
 	 * regenerates both.
