Adding archive size checks to package file summary to guard against corrupt assets

diff --git a/Engine/Source/Runtime/AssetRegistry/Private/PackageReader.h b/Engine/Source/Runtime/AssetRegistry/Private/PackageReader.h
--- a/Engine/Source/Runtime/AssetRegistry/Private/PackageReader.h
+++ b/Engine/Source/Runtime/AssetRegistry/Private/PackageReader.h
@@ -32,6 +32,10 @@ public:
 	int64 Tell();
 	int64 TotalSize();
 	FArchive& operator<<( FName& Name );
+	virtual FString GetArchiveName() const override
+	{
+		return PackageFilename;
+	}
 
 private:
 	FString PackageFilename;
diff --git a/Engine/Source/Runtime/Core/Public/HAL/FileManagerGeneric.h b/Engine/Source/Runtime/Core/Public/HAL/FileManagerGeneric.h
--- a/Engine/Source/Runtime/Core/Public/HAL/FileManagerGeneric.h
+++ b/Engine/Source/Runtime/Core/Public/HAL/FileManagerGeneric.h
@@ -187,6 +187,10 @@ public:
 	}
 	virtual bool Close() final;
 	virtual void Serialize( void* V, int64 Length ) final;
+	virtual FString GetArchiveName() const override
+	{
+		return Filename;
+	}
 
 protected:
 	bool InternalPrecache( int64 PrecacheOffset, int64 PrecacheSize );
diff --git a/Engine/Source/Runtime/CoreUObject/Private/UObject/PackageFileSummary.cpp b/Engine/Source/Runtime/CoreUObject/Private/UObject/PackageFileSummary.cpp
--- a/Engine/Source/Runtime/CoreUObject/Private/UObject/PackageFileSummary.cpp
+++ b/Engine/Source/Runtime/CoreUObject/Private/UObject/PackageFileSummary.cpp
@@ -10,7 +10,22 @@ FPackageFileSummary::FPackageFileSummary()
 
 FArchive& operator<<( FArchive& Ar, FPackageFileSummary& Sum )
 {
-	Ar << Sum.Tag;
+	bool bCanStartSerializing = true;
+	int64 ArchiveSize = 0;
+	if (Ar.IsLoading())
+	{
+		// Sanity checks before we even start serializing the archive
+		ArchiveSize = Ar.TotalSize();
+		const int64 MinimumPackageSize = 32; // That should get us safely to Sum.TotalHeaderSize
+		bCanStartSerializing = ArchiveSize >= MinimumPackageSize;
+		UE_CLOG(!bCanStartSerializing, LogLinker, Warning,
+			TEXT("Failed to read package file summary, the file \"%s\" is too small (%lld bytes, expected at least %lld bytes)"),
+			*Ar.GetArchiveName(), ArchiveSize, MinimumPackageSize);
+	}
+	if (bCanStartSerializing)
+	{
+		Ar << Sum.Tag;
+	}
 	// only keep loading if we match the magic
 	if( Sum.Tag == PACKAGE_FILE_TAG || Sum.Tag == PACKAGE_FILE_TAG_SWAPPED )
 	{
@@ -103,6 +118,15 @@ FArchive& operator<<( FArchive& Ar, FPackageFileSummary& Sum )
 			}
 		}
 		Ar << Sum.TotalHeaderSize;
+		if (Ar.IsLoading() && ArchiveSize < Sum.TotalHeaderSize || Sum.TotalHeaderSize < 0)
+		{
+			UE_LOG(LogLinker, Warning,
+				TEXT("Failed to read package file summary, the file \"%s\" is too small or corrupt (%lld bytes, expected at least %lld bytes)"),
+				*Ar.GetArchiveName(), ArchiveSize, Sum.TotalHeaderSize);
+			// Reset the tag so that the summary reports as invalid
+			Sum.Tag = 0;
+			return Ar;
+		}
 		Ar << Sum.FolderName;
 		Ar << Sum.PackageFlags;
 
