Fixing printf data type

diff --git a/Engine/Source/Runtime/CoreUObject/Private/UObject/PackageFileSummary.cpp b/Engine/Source/Runtime/CoreUObject/Private/UObject/PackageFileSummary.cpp
--- a/Engine/Source/Runtime/CoreUObject/Private/UObject/PackageFileSummary.cpp
+++ b/Engine/Source/Runtime/CoreUObject/Private/UObject/PackageFileSummary.cpp
@@ -121,7 +121,7 @@ FArchive& operator<<( FArchive& Ar, FPackageFileSummary& Sum )
 		if (Ar.IsLoading() && ArchiveSize < Sum.TotalHeaderSize || Sum.TotalHeaderSize < 0)
 		{
 			UE_LOG(LogLinker, Warning,
-				TEXT("Failed to read package file summary, the file \"%s\" is too small or corrupt (%lld bytes, expected at least %lld bytes)"),
+				TEXT("Failed to read package file summary, the file \"%s\" is too small or corrupt (%lld bytes, expected at least %d bytes)"),
 				*Ar.GetArchiveName(), ArchiveSize, Sum.TotalHeaderSize);
 			// Reset the tag so that the summary reports as invalid
 			Sum.Tag = 0;
