Fixing CIS

diff --git a/Engine/Source/Runtime/CoreUObject/Private/UObject/PackageFileSummary.cpp b/Engine/Source/Runtime/CoreUObject/Private/UObject/PackageFileSummary.cpp
--- a/Engine/Source/Runtime/CoreUObject/Private/UObject/PackageFileSummary.cpp
+++ b/Engine/Source/Runtime/CoreUObject/Private/UObject/PackageFileSummary.cpp
@@ -118,7 +118,7 @@ FArchive& operator<<( FArchive& Ar, FPackageFileSummary& Sum )
 			}
 		}
 		Ar << Sum.TotalHeaderSize;
-		if (Ar.IsLoading() && ArchiveSize < Sum.TotalHeaderSize || Sum.TotalHeaderSize < 0)
+		if (Ar.IsLoading() && (ArchiveSize < Sum.TotalHeaderSize || Sum.TotalHeaderSize < 0))
 		{
 			UE_LOG(LogLinker, Warning,
 				TEXT("Failed to read package file summary, the file \"%s\" is too small or corrupt (%lld bytes, expected at least %d bytes)"),
