UE-15552 - Fix async package loading for seamless travel not setting PKG_PlayInEditor resulting in failure to clean up the map on area transition/end PIE
Fix bug in simplified AsyncLoadPackage where priority was used as flags
Clean up some parameters and properties to be clear that they are PackageFlags, not LoadFlags

diff --git a/Engine/Source/Runtime/CoreUObject/Private/Serialization/AsyncLoading.cpp b/Engine/Source/Runtime/CoreUObject/Private/Serialization/AsyncLoading.cpp
--- a/Engine/Source/Runtime/CoreUObject/Private/Serialization/AsyncLoading.cpp
+++ b/Engine/Source/Runtime/CoreUObject/Private/Serialization/AsyncLoading.cpp
@@ -1026,7 +1026,7 @@ EAsyncPackageState::Type FAsyncPackage::CreateLinker()
 		FScopeCycleCounterUObject ConstructorScope(Package, GET_STATID(STAT_FAsyncPackage_CreateLinker));
 
 		// Set package specific data 
-		Package->PackageFlags |= Desc.Flags;
+		Package->PackageFlags |= Desc.PackageFlags;
 #if WITH_EDITOR
 		Package->PIEInstanceID = Desc.PIEInstanceID;
 #endif
@@ -1049,7 +1049,18 @@ EAsyncPackageState::Type FAsyncPackage::CreateLinker()
 			}
 
 			// Create raw async linker, requiring to be ticked till finished creating.
-			Linker = FLinkerLoad::CreateLinkerAsync( Package, *PackageFileName, (FApp::IsGame() && !GIsEditor) ? (LOAD_SeekFree | LOAD_NoVerify) : LOAD_None  );
+			uint32 LinkerFlags = LOAD_None;
+			if (FApp::IsGame() && !GIsEditor)
+			{
+				LinkerFlags |= (LOAD_SeekFree | LOAD_NoVerify);
+			}
+#if WITH_EDITOR
+			else if ((Desc.PackageFlags & PKG_PlayInEditor) != 0)
+			{
+				LinkerFlags |= LOAD_PackageForPIE;
+			}
+#endif
+			Linker = FLinkerLoad::CreateLinkerAsync( Package, *PackageFileName, LinkerFlags );
 		}
 
 		UE_LOG(LogStreaming, Verbose, TEXT("FAsyncPackage::CreateLinker for %s finished."), *Desc.NameToLoad.ToString());
@@ -1682,7 +1693,7 @@ void FAsyncPackage::UpdateLoadPercentage()
 /*-----------------------------------------------------------------------------
 	UObject async (pre)loading.
 -----------------------------------------------------------------------------*/
-void LoadPackageAsync(const FString& InName, const FGuid* InGuid /*= nullptr*/, FName InType /*= NAME_None*/, const TCHAR* InPackageToLoadFrom /*= nullptr*/, FLoadPackageAsyncDelegate InCompletionDelegate /*= FLoadPackageAsyncDelegate()*/, uint32 InFlags /*= 0*/, int32 InPIEInstanceID /*= INDEX_NONE*/, uint32 InPackagePriority /*= 0*/)
+void LoadPackageAsync(const FString& InName, const FGuid* InGuid /*= nullptr*/, FName InType /*= NAME_None*/, const TCHAR* InPackageToLoadFrom /*= nullptr*/, FLoadPackageAsyncDelegate InCompletionDelegate /*= FLoadPackageAsyncDelegate()*/, EPackageFlags InPackageFlags /*= PKG_None*/, int32 InPIEInstanceID /*= INDEX_NONE*/, uint32 InPackagePriority /*= 0*/)
 {
 	// The comments clearly state that it should be a package name but we also handle it being a filename as this function is not perf critical
 	// and LoadPackage handles having a filename being passed in as well.
@@ -1714,15 +1725,15 @@ void LoadPackageAsync(const FString& InName, const FGuid* InGuid /*= nullptr*/,
 		UE_LOG(LogStreaming, Fatal, TEXT("Async loading code requires long package names (%s)."), *PackageNameToLoad);
 	}
 
-	FAsyncPackageDesc PackageDesc(*PackageName, *PackageNameToLoad, InGuid ? *InGuid : FGuid(), InType, InCompletionDelegate, InFlags, InPIEInstanceID, InPackagePriority);
+	FAsyncPackageDesc PackageDesc(*PackageName, *PackageNameToLoad, InGuid ? *InGuid : FGuid(), InType, InCompletionDelegate, InPackageFlags, InPIEInstanceID, InPackagePriority);
 	FAsyncLoadingThread::Get().QueuePackage(PackageDesc);
 }
 
-void LoadPackageAsync(const FString& PackageName, FLoadPackageAsyncDelegate CompletionDelegate, uint32 InPackagePriority /*= 0*/)
+void LoadPackageAsync(const FString& PackageName, FLoadPackageAsyncDelegate CompletionDelegate, uint32 InPackagePriority /*= 0*/, EPackageFlags InPackageFlags /*= PKG_None*/)
 {
 	const FGuid* Guid = nullptr;
 	const TCHAR* PackageToLoadFrom = nullptr;
-	LoadPackageAsync(PackageName, Guid, NAME_None, PackageToLoadFrom, CompletionDelegate, InPackagePriority);
+	LoadPackageAsync(PackageName, Guid, NAME_None, PackageToLoadFrom, CompletionDelegate, InPackageFlags, -1, InPackagePriority );
 }
 
 void CancelAsyncLoading()
diff --git a/Engine/Source/Runtime/CoreUObject/Public/Serialization/AsyncPackage.h b/Engine/Source/Runtime/CoreUObject/Public/Serialization/AsyncPackage.h
--- a/Engine/Source/Runtime/CoreUObject/Public/Serialization/AsyncPackage.h
+++ b/Engine/Source/Runtime/CoreUObject/Public/Serialization/AsyncPackage.h
@@ -19,7 +19,7 @@ struct FAsyncPackageDesc
 	/** Delegate called on completion of loading */
 	FLoadPackageAsyncDelegate PackageLoadedDelegate;
 	/** The flags that should be applied to the package */
-	uint32 Flags;
+	EPackageFlags PackageFlags;
 	/** Package loading priority. Higher number is higher priority. */
 	uint32 Priority;
 #if WITH_EDITOR
@@ -28,13 +28,13 @@ struct FAsyncPackageDesc
 #endif
 
 
-	FAsyncPackageDesc(const FName& InName, FName InPackageToLoadFrom = NAME_None, const FGuid& InGuid = FGuid(), FName InType = NAME_None, FLoadPackageAsyncDelegate InCompletionDelegate = FLoadPackageAsyncDelegate(), uint32 InFlags = 0, int32 InPIEInstanceID = INDEX_NONE, uint32 InPriority = 0)
+	FAsyncPackageDesc(const FName& InName, FName InPackageToLoadFrom = NAME_None, const FGuid& InGuid = FGuid(), FName InType = NAME_None, FLoadPackageAsyncDelegate InCompletionDelegate = FLoadPackageAsyncDelegate(), EPackageFlags InPackageFlags = PKG_None, int32 InPIEInstanceID = INDEX_NONE, uint32 InPriority = 0)
 		: Name(InName)
 		, NameToLoad(InPackageToLoadFrom)
 		, Type(InType)	
 		, Guid(InGuid)
 		, PackageLoadedDelegate(InCompletionDelegate)
-		, Flags(InFlags)
+		, PackageFlags(InPackageFlags)
 		, Priority(InPriority)
 #if WITH_EDITOR
 		, PIEInstanceID(InPIEInstanceID)
diff --git a/Engine/Source/Runtime/CoreUObject/Public/UObject/ObjectBase.h b/Engine/Source/Runtime/CoreUObject/Public/UObject/ObjectBase.h
--- a/Engine/Source/Runtime/CoreUObject/Public/UObject/ObjectBase.h
+++ b/Engine/Source/Runtime/CoreUObject/Public/UObject/ObjectBase.h
@@ -77,6 +77,7 @@ enum ESaveFlags
 //
 enum EPackageFlags
 {
+	PKG_None						= 0x00000000,	// No flags
 	PKG_NewlyCreated				= 0x00000001,	// Newly created package, not saved yet. In editor only.
 	PKG_ClientOptional				= 0x00000002,	// Purely optional for clients.
 	PKG_ServerSideOnly				= 0x00000004,   // Only needed on the server side.
@@ -102,6 +103,7 @@ enum EPackageFlags
 
 	PKG_InMemoryOnly				= PKG_CompiledIn | PKG_NewlyCreated, // Flag mask that indicates if this package is a package that exists in memory only.
 };
+ENUM_CLASS_FLAGS(EPackageFlags);
 
 //
 // Internal enums.
diff --git a/Engine/Source/Runtime/CoreUObject/Public/UObject/UObjectGlobals.h b/Engine/Source/Runtime/CoreUObject/Public/UObject/UObjectGlobals.h
--- a/Engine/Source/Runtime/CoreUObject/Public/UObject/UObjectGlobals.h
+++ b/Engine/Source/Runtime/CoreUObject/Public/UObject/UObjectGlobals.h
@@ -292,7 +292,7 @@ DECLARE_DELEGATE_ThreeParams(FLoadPackageAsyncDelegate, const FName& /*PackageNa
  * @param	InPIEInstanceID			PIE instance ID
  * @param	InPackagePriority		Loading priority
  */
-COREUOBJECT_API void LoadPackageAsync(const FString& InName, const FGuid* InGuid = nullptr, FName InType = NAME_None, const TCHAR* InPackageToLoadFrom = nullptr, FLoadPackageAsyncDelegate InCompletionDelegate = FLoadPackageAsyncDelegate(), uint32 InFlags = 0, int32 InPIEInstanceID = INDEX_NONE, uint32 InPackagePriority = 0);
+COREUOBJECT_API void LoadPackageAsync(const FString& InName, const FGuid* InGuid = nullptr, FName InType = NAME_None, const TCHAR* InPackageToLoadFrom = nullptr, FLoadPackageAsyncDelegate InCompletionDelegate = FLoadPackageAsyncDelegate(), EPackageFlags InPackageFlags = PKG_None, int32 InPIEInstanceID = INDEX_NONE, uint32 InPackagePriority = 0);
 
 /**
 * Asynchronously load a package and all contained objects that match context flags. Non- blocking.
@@ -301,7 +301,7 @@ COREUOBJECT_API void LoadPackageAsync(const FString& InName, const FGuid* InGuid
 * @param	InCompletionDelegate	Delegate to be invoked when the packages has finished streaming
 * @param	InPackagePriority		Loading priority
 */
-COREUOBJECT_API void LoadPackageAsync(const FString& InName, FLoadPackageAsyncDelegate InCompletionDelegate, uint32 InPackagePriority = 0);
+COREUOBJECT_API void LoadPackageAsync(const FString& InName, FLoadPackageAsyncDelegate InCompletionDelegate, uint32 InPackagePriority = 0, EPackageFlags InPackageFlags = PKG_None);
 
 /**
 * Cancels all async package loading requests.
diff --git a/Engine/Source/Runtime/Engine/Classes/Engine/Engine.h b/Engine/Source/Runtime/Engine/Classes/Engine/Engine.h
--- a/Engine/Source/Runtime/Engine/Classes/Engine/Engine.h
+++ b/Engine/Source/Runtime/Engine/Classes/Engine/Engine.h
@@ -1939,7 +1939,7 @@ public:
 	/** Looks up the GUID of a package on disk. The package must NOT be in the auto-download cache.
 	 * This may require loading the header of the package in question and is therefore slow.
 	 */
-	static FGuid GetPackageGuid(FName PackageName);
+	static FGuid GetPackageGuid(FName PackageName, bool bForPIE);
 
 	/**
 	 * Returns whether we are running on a console platform or on the PC.
diff --git a/Engine/Source/Runtime/Engine/Private/GameMode.cpp b/Engine/Source/Runtime/Engine/Private/GameMode.cpp
--- a/Engine/Source/Runtime/Engine/Private/GameMode.cpp
+++ b/Engine/Source/Runtime/Engine/Private/GameMode.cpp
@@ -955,7 +955,7 @@ void AGameMode::ProcessServerTravel(const FString& URL, bool bAbsolute)
 		return;
 	}
 	
-	FGuid NextMapGuid = UEngine::GetPackageGuid(FName(*NextMap));
+	FGuid NextMapGuid = UEngine::GetPackageGuid(FName(*NextMap), GetWorld()->IsPlayInEditor());
 
 	// Notify clients we're switching level and give them time to receive.
 	FString URLMod = URL;
diff --git a/Engine/Source/Runtime/Engine/Private/LevelStreaming.cpp b/Engine/Source/Runtime/Engine/Private/LevelStreaming.cpp
--- a/Engine/Source/Runtime/Engine/Private/LevelStreaming.cpp
+++ b/Engine/Source/Runtime/Engine/Private/LevelStreaming.cpp
@@ -401,7 +401,7 @@ bool ULevelStreaming::RequestLevel(UWorld* PersistentWorld, bool bAllowLevelLoad
 		}
 	}
 
-	uint32 PackageFlags = 0;
+	EPackageFlags PackageFlags = PKG_None;
 	int32 PIEInstanceID = INDEX_NONE;
 
 	// copy streaming level on demand if we are in PIE
diff --git a/Engine/Source/Runtime/Engine/Private/UnrealEngine.cpp b/Engine/Source/Runtime/Engine/Private/UnrealEngine.cpp
--- a/Engine/Source/Runtime/Engine/Private/UnrealEngine.cpp
+++ b/Engine/Source/Runtime/Engine/Private/UnrealEngine.cpp
@@ -6275,12 +6275,17 @@ void UEngine::AddTextureStreamingSlaveLoc(FVector InLoc, float BoostFactor, bool
 /** Looks up the GUID of a package on disk. The package must NOT be in the autodownload cache.
  * This may require loading the header of the package in question and is therefore slow.
  */
-FGuid UEngine::GetPackageGuid(FName PackageName)
+FGuid UEngine::GetPackageGuid(FName PackageName, bool bForPIE)
 {
 	FGuid Result(0,0,0,0);
 
 	BeginLoad();
-	FLinkerLoad* Linker = GetPackageLinker(NULL, *PackageName.ToString(), LOAD_NoWarn | LOAD_NoVerify, NULL, NULL);
+	uint32 LoadFlags = LOAD_NoWarn | LOAD_NoVerify;
+	if (bForPIE)
+	{
+		LoadFlags |= LOAD_PackageForPIE;
+	}
+	FLinkerLoad* Linker = GetPackageLinker(NULL, *PackageName.ToString(), LoadFlags, NULL, NULL);
 	if (Linker != NULL && Linker->LinkerRoot != NULL)
 	{
 		Result = Linker->LinkerRoot->GetGuid();
diff --git a/Engine/Source/Runtime/Engine/Private/World.cpp b/Engine/Source/Runtime/Engine/Private/World.cpp
--- a/Engine/Source/Runtime/Engine/Private/World.cpp
+++ b/Engine/Source/Runtime/Engine/Private/World.cpp
@@ -4393,7 +4393,9 @@ bool FSeamlessTravelHandler::StartTravel(UWorld* InCurrentWorld, const FURL& InU
 
 				// first, load the entry level package
 				LoadPackageAsync(TransitionMap, 
-					FLoadPackageAsyncDelegate::CreateRaw(this, &FSeamlessTravelHandler::SeamlessTravelLoadCallback)
+					FLoadPackageAsyncDelegate::CreateRaw(this, &FSeamlessTravelHandler::SeamlessTravelLoadCallback),
+					0, 
+					(CurrentWorld->WorldType == EWorldType::PIE ? PKG_PlayInEditor : PKG_None)
 					);
 			}
 
@@ -4497,17 +4499,20 @@ void FSeamlessTravelHandler::StartLoadingDestination()
 		// In PIE we might want to mangle MapPackageName when traveling to a map loaded in the editor
 		FString URLMapPackageName = PendingTravelURL.Map;
 		FString URLMapPackageToLoadFrom = PendingTravelURL.Map;
-		uint32 PackageFlags = 0;
+		EPackageFlags PackageFlags = PKG_None;
 		int32 PIEInstanceID = INDEX_NONE;
 		
 #if WITH_EDITOR
 		if (GIsEditor)
 		{
+			FWorldContext &WorldContext = GEngine->GetWorldContextFromHandleChecked(WorldContextHandle);
+			if (WorldContext.WorldType == EWorldType::PIE)
+			{
+				PackageFlags |= PKG_PlayInEditor;
+			}
 			UPackage* EditorLevelPackage = (UPackage*)StaticFindObjectFast(UPackage::StaticClass(), NULL, URLMapFName, 0, 0, RF_PendingKill);
 			if (EditorLevelPackage)
 			{
-				FWorldContext &WorldContext = GEngine->GetWorldContextFromHandleChecked(WorldContextHandle);
-				PackageFlags |= PKG_PlayInEditor;
 				PIEInstanceID = WorldContext.PIEInstance;
 				URLMapPackageName = UWorld::ConvertToPIEPackageName(URLMapPackageName, PIEInstanceID);
 			}
