Avoid multiple Export duplicate fixups during loading.

diff --git a/Engine/Source/Runtime/CoreUObject/Private/UObject/LinkerLoad.cpp b/Engine/Source/Runtime/CoreUObject/Private/UObject/LinkerLoad.cpp
--- a/Engine/Source/Runtime/CoreUObject/Private/UObject/LinkerLoad.cpp
+++ b/Engine/Source/Runtime/CoreUObject/Private/UObject/LinkerLoad.cpp
@@ -692,6 +692,7 @@ FLinkerLoad::FLinkerLoad(UPackage* InParent, const TCHAR* InFilename, uint32 InL
 , TickStartTime(0.0)
 , bFixupExportMapDone(false)
 #if WITH_EDITOR
+, bExportsDuplicatesFixed(false)
 ,	LoadProgressScope( nullptr )
 #endif // WITH_EDITOR
 #if USE_CIRCULAR_DEPENDENCY_LOAD_DEFERRING
@@ -4489,9 +4490,10 @@ FLinkerLoad::ELinkerStatus FLinkerLoad::FixupExportMap()
 	DECLARE_SCOPE_CYCLE_COUNTER( TEXT( "FLinkerLoad::FixupExportMap" ), STAT_LinkerLoad_FixupExportMap, STATGROUP_LinkerLoad );
 
 #if WITH_EDITOR
-	if (UE4Ver() < VER_UE4_SKIP_DUPLICATE_EXPORTS_ON_SAVE_PACKAGE)
+	if (UE4Ver() < VER_UE4_SKIP_DUPLICATE_EXPORTS_ON_SAVE_PACKAGE && !bExportsDuplicatesFixed)
 	{
 		FixupDuplicateExports();
+		bExportsDuplicatesFixed = true;
 	}
 #endif // WITH_EDITOR
 
diff --git a/Engine/Source/Runtime/CoreUObject/Public/UObject/Linker.h b/Engine/Source/Runtime/CoreUObject/Public/UObject/Linker.h
--- a/Engine/Source/Runtime/CoreUObject/Public/UObject/Linker.h
+++ b/Engine/Source/Runtime/CoreUObject/Public/UObject/Linker.h
@@ -1394,6 +1394,10 @@ private:
 	/** Used for ActiveClassRedirects functionality */
 	bool					bFixupExportMapDone;
 
+#if WITH_EDITOR
+	/** Check to avoid multiple export duplicate fixups in case we don't save asset. */
+	bool bExportsDuplicatesFixed;
+#endif // WITH_EDITOR
 	/** Id of the thread that created this linker. This is to guard against using this linker on other threads than the one it was created on **/
 	int32					OwnerThread;
 
