[UE-15878] Reverting a portion of CL 2555878, we currently cannot use property size to compare against skeleton function params (as they're unlinked, with zero size).

diff --git a/Engine/Source/Runtime/CoreUObject/Private/UObject/Class.cpp b/Engine/Source/Runtime/CoreUObject/Private/UObject/Class.cpp
--- a/Engine/Source/Runtime/CoreUObject/Private/UObject/Class.cpp
+++ b/Engine/Source/Runtime/CoreUObject/Private/UObject/Class.cpp
@@ -4373,9 +4373,17 @@ bool UFunction::IsSignatureCompatibleWith(const UFunction* OtherFunction, uint64
 			UProperty* PropA = *IteratorA;
 			UProperty* PropB = *IteratorB;
 
+			FString TypeA(PropA->GetCPPType());
+			FString TypeB(PropB->GetCPPType());
+
 			// Check the flags as well
 			const uint64 PropertyMash = PropA->PropertyFlags ^ PropB->PropertyFlags;
-			if (!FStructUtils::ArePropertiesTheSame(PropA, PropB, false) || ((PropertyMash & ~IgnoreFlags) != 0))
+
+			// @TODO: ArePropertiesTheSame() currently has problems when 
+			//        comparing a blueprint skeleton function property to its 
+			//        native counterpart (it seems that skel properties are not 
+			//        linked, and therefore have zero size)
+			if (/*!FStructUtils::ArePropertiesTheSame(PropA, PropB, false)*/(TypeA != TypeB) || ((PropertyMash & ~IgnoreFlags) != 0))
 			{
 				// Type mismatch between an argument of A and B
 				return false;
