UE-15878 Many blueprints don't compile with similar error

In FKismetCompilerContext::PrecompileFunction the parameter list of function is linked, so the signature of function can be compared.

diff --git a/Engine/Source/Editor/KismetCompiler/Private/KismetCompiler.cpp b/Engine/Source/Editor/KismetCompiler/Private/KismetCompiler.cpp
--- a/Engine/Source/Editor/KismetCompiler/Private/KismetCompiler.cpp
+++ b/Engine/Source/Editor/KismetCompiler/Private/KismetCompiler.cpp
@@ -1371,6 +1371,9 @@ void FKismetCompilerContext::PrecompileFunction(FKismetFunctionContext& Context)
 
 		Context.Function->FunctionFlags |= Context.GetNetFlags();
 
+		// Parameter list needs to be linked before signatures are compared. 
+		Context.Function->StaticLink(true);
+
 		// Make sure the function signature is valid if this is an override
 		if (ParentFunction)
 		{
diff --git a/Engine/Source/Runtime/CoreUObject/Private/UObject/Class.cpp b/Engine/Source/Runtime/CoreUObject/Private/UObject/Class.cpp
--- a/Engine/Source/Runtime/CoreUObject/Private/UObject/Class.cpp
+++ b/Engine/Source/Runtime/CoreUObject/Private/UObject/Class.cpp
@@ -4373,17 +4373,9 @@ bool UFunction::IsSignatureCompatibleWith(const UFunction* OtherFunction, uint64
 			UProperty* PropA = *IteratorA;
 			UProperty* PropB = *IteratorB;
 
-			FString TypeA(PropA->GetCPPType());
-			FString TypeB(PropB->GetCPPType());
-
 			// Check the flags as well
 			const uint64 PropertyMash = PropA->PropertyFlags ^ PropB->PropertyFlags;
-
-			// @TODO: ArePropertiesTheSame() currently has problems when 
-			//        comparing a blueprint skeleton function property to its 
-			//        native counterpart (it seems that skel properties are not 
-			//        linked, and therefore have zero size)
-			if (/*!FStructUtils::ArePropertiesTheSame(PropA, PropB, false)*/(TypeA != TypeB) || ((PropertyMash & ~IgnoreFlags) != 0))
+			if (!FStructUtils::ArePropertiesTheSame(PropA, PropB, false) || ((PropertyMash & ~IgnoreFlags) != 0))
 			{
 				// Type mismatch between an argument of A and B
 				return false;
