Return value fixes for FObjectAndNameAsStringProxyArchive::operator<< and FNameAsStringProxyArchive::operator<<.

diff --git a/Engine/Source/Runtime/Core/Private/Serialization/Archive.cpp b/Engine/Source/Runtime/Core/Private/Serialization/Archive.cpp
--- a/Engine/Source/Runtime/Core/Private/Serialization/Archive.cpp
+++ b/Engine/Source/Runtime/Core/Private/Serialization/Archive.cpp
@@ -242,13 +242,13 @@ FArchive& FNameAsStringProxyArchive::operator<<( class FName& N )
 		FString LoadedString;
 		InnerArchive << LoadedString;
 		N = FName(*LoadedString);
-		return InnerArchive;
 	}
 	else
 	{
 		FString SavedString(N.ToString());
-		return InnerArchive << SavedString;
+		InnerArchive << SavedString;
 	}
+	return *this;
 }
 
 /** 
diff --git a/Engine/Source/Runtime/CoreUObject/Private/Serialization/ArchiveUObject.cpp b/Engine/Source/Runtime/CoreUObject/Private/Serialization/ArchiveUObject.cpp
--- a/Engine/Source/Runtime/CoreUObject/Private/Serialization/ArchiveUObject.cpp
+++ b/Engine/Source/Runtime/CoreUObject/Private/Serialization/ArchiveUObject.cpp
@@ -91,14 +91,13 @@ FArchive& FObjectAndNameAsStringProxyArchive::operator<<(class UObject*& Obj)
 		{
 			Obj = LoadObject<UObject>(NULL, *LoadedString);
 		}
-
-		return InnerArchive;
 	}
 	else
 	{
 		// save out the fully qualified object name
 		FString SavedString(Obj->GetPathName());
-		return InnerArchive << SavedString;
+		InnerArchive << SavedString;
 	}
+	return *this;
 }
 
