Added comment to enums in CoreNet, changed MTU back to 512 until discussion on this topic has finished

diff --git a/Engine/Source/Runtime/CoreUObject/Public/UObject/CoreNet.h b/Engine/Source/Runtime/CoreUObject/Public/UObject/CoreNet.h
--- a/Engine/Source/Runtime/CoreUObject/Public/UObject/CoreNet.h
+++ b/Engine/Source/Runtime/CoreUObject/Public/UObject/CoreNet.h
@@ -430,14 +430,16 @@ COREUOBJECT_API void SerializeChecksum(FArchive &Ar, uint32 x, bool ErrorOK);
 
 #endif
 
-// Up to this many reliable channel bunches may be buffered.
+/**
+* Values used for initializing UNetConnection and LanBeacon
+*/
 enum { RELIABLE_BUFFER = 256 }; // Power of 2 >= 1.
-enum { MAX_PACKETID = 16384 }; // Power of 2 >= 1, covering guaranteed loss/misorder time.
+enum { MAX_PACKETID = 16384 };  // Power of 2 >= 1, covering guaranteed loss/misorder time.
 enum { MAX_CHSEQUENCE = 1024 }; // Power of 2 >RELIABLE_BUFFER, covering loss/misorder time.
 enum { MAX_BUNCH_HEADER_BITS = 64 };
 enum { MAX_PACKET_HEADER_BITS = 16 };
 enum { MAX_PACKET_TRAILER_BITS = 1 };
-enum { MAX_PACKET_SIZE = 1000 }; // MTU for the connection
+enum { MAX_PACKET_SIZE = 512 }; // MTU for the connection
 
 /**
  * Functions to assist in detecting errors during RPC calls
