ScriptVM use FAssetPtr to store Asset and Asset class. Related to UE-15033.

diff --git a/Engine/Source/Runtime/CoreUObject/Public/UObject/UnrealType.h b/Engine/Source/Runtime/CoreUObject/Public/UObject/UnrealType.h
--- a/Engine/Source/Runtime/CoreUObject/Public/UObject/UnrealType.h
+++ b/Engine/Source/Runtime/CoreUObject/Public/UObject/UnrealType.h
@@ -1999,6 +1999,28 @@ class COREUOBJECT_API UAssetObjectProperty : public TUObjectPropertyBase<FAssetP
 	{
 		return true;
 	}
+
+	// ScriptVM should store Asset as FAssetPtr not as UObject.
+
+	virtual void CopySingleValueToScriptVM(void* Dest, void const* Src) const override
+	{
+		CopySingleValue(Dest, Src);
+	}
+
+	virtual void CopyCompleteValueToScriptVM(void* Dest, void const* Src) const override
+	{
+		CopyCompleteValue(Dest, Src);
+	}
+
+	virtual void CopySingleValueFromScriptVM(void* Dest, void const* Src) const override
+	{
+		CopySingleValue(Dest, Src);
+	}
+
+	virtual void CopyCompleteValueFromScriptVM(void* Dest, void const* Src) const override
+	{
+		CopyCompleteValue(Dest, Src);
+	}
 	// End of UObjectProperty interface
 };
 
