Removed erronious errors which were triggering in cooked builds.

diff --git a/Engine/Source/Runtime/CoreUObject/Private/UObject/PackageFileSummary.cpp b/Engine/Source/Runtime/CoreUObject/Private/UObject/PackageFileSummary.cpp
--- a/Engine/Source/Runtime/CoreUObject/Private/UObject/PackageFileSummary.cpp
+++ b/Engine/Source/Runtime/CoreUObject/Private/UObject/PackageFileSummary.cpp
@@ -12,7 +12,7 @@ FArchive& operator<<( FArchive& Ar, FPackageFileSummary& Sum )
 {
 	bool bCanStartSerializing = true;
 	int64 ArchiveSize = 0;
-	if (Ar.IsLoading())
+	/*if (Ar.IsLoading())
 	{
 		// Sanity checks before we even start serializing the archive
 		ArchiveSize = Ar.TotalSize();
@@ -21,7 +21,7 @@ FArchive& operator<<( FArchive& Ar, FPackageFileSummary& Sum )
 		UE_CLOG(!bCanStartSerializing, LogLinker, Warning,
 			TEXT("Failed to read package file summary, the file \"%s\" is too small (%lld bytes, expected at least %lld bytes)"),
 			*Ar.GetArchiveName(), ArchiveSize, MinimumPackageSize);
-	}
+	}*/
 	if (bCanStartSerializing)
 	{
 		Ar << Sum.Tag;
@@ -118,7 +118,7 @@ FArchive& operator<<( FArchive& Ar, FPackageFileSummary& Sum )
 			}
 		}
 		Ar << Sum.TotalHeaderSize;
-		if (Ar.IsLoading() && (ArchiveSize < Sum.TotalHeaderSize || Sum.TotalHeaderSize < 0))
+		/*if (Ar.IsLoading() && (ArchiveSize < Sum.TotalHeaderSize || Sum.TotalHeaderSize < 0))
 		{
 			UE_LOG(LogLinker, Warning,
 				TEXT("Failed to read package file summary, the file \"%s\" is too small or corrupt (%lld bytes, expected at least %d bytes)"),
@@ -126,7 +126,7 @@ FArchive& operator<<( FArchive& Ar, FPackageFileSummary& Sum )
 			// Reset the tag so that the summary reports as invalid
 			Sum.Tag = 0;
 			return Ar;
-		}
+		}*/
 		Ar << Sum.FolderName;
 		Ar << Sum.PackageFlags;
 
