Helper functions for TAssetPtr and FStringAssetReference to get package and asset name

diff --git a/Engine/Source/Runtime/CoreUObject/Public/Misc/StringAssetReference.h b/Engine/Source/Runtime/CoreUObject/Public/Misc/StringAssetReference.h
--- a/Engine/Source/Runtime/CoreUObject/Public/Misc/StringAssetReference.h
+++ b/Engine/Source/Runtime/CoreUObject/Public/Misc/StringAssetReference.h
@@ -47,6 +47,13 @@ struct COREUOBJECT_API FStringAssetReference
 		return AssetLongPathname;
 	}
 
+	const FString GetLongPackageName() const
+	{
+		FString PackageName;
+		AssetLongPathname.Split(TEXT("."), &PackageName, nullptr, ESearchCase::CaseSensitive, ESearchDir::FromEnd);
+		return PackageName;
+	}
+
 	/**
 	 * Attempts to load the asset.
 	 * @return Loaded UObject, or null if the asset fails to load, or if the reference is not valid.
diff --git a/Engine/Source/Runtime/CoreUObject/Public/UObject/AssetPtr.h b/Engine/Source/Runtime/CoreUObject/Public/UObject/AssetPtr.h
--- a/Engine/Source/Runtime/CoreUObject/Public/UObject/AssetPtr.h
+++ b/Engine/Source/Runtime/CoreUObject/Public/UObject/AssetPtr.h
@@ -256,6 +256,16 @@ public:
 		return AssetPtr.GetUniqueID();
 	}
 
+	FORCEINLINE const FString& ToString() const
+	{
+		return ToStringReference().ToString();
+	}
+
+	FORCEINLINE FString GetLongPackageName() const
+	{
+		return ToStringReference().GetLongPackageName();
+	}
+
 	/**  
 	 * Dereference lazy pointer to see if it points somewhere valid.
 	 */
