UE-14122 Passing FText objects accessed from Access None objects into functions crashes the game

CL# 2244148 was reverted. Hopefully the optimization can be restore once the UE-14720 is fixed.

diff --git a/Engine/Source/Runtime/CoreUObject/Private/UObject/ScriptCore.cpp b/Engine/Source/Runtime/CoreUObject/Private/UObject/ScriptCore.cpp
--- a/Engine/Source/Runtime/CoreUObject/Private/UObject/ScriptCore.cpp
+++ b/Engine/Source/Runtime/CoreUObject/Private/UObject/ScriptCore.cpp
@@ -154,6 +154,8 @@ void FFrame::Step(UObject *Context, RESULT_DECL)
 
 void FFrame::StepExplicitProperty(void*const Result, UProperty* Property)
 {
+	checkSlow(Result != NULL);
+
 	if (Property->PropertyFlags & CPF_OutParm)
 	{
 		// look through the out parameter infos and find the one that has the address of this property
@@ -170,10 +172,7 @@ void FFrame::StepExplicitProperty(void*const Result, UProperty* Property)
 	else
 	{
 		MostRecentPropertyAddress = Property->ContainerPtrToValuePtr<uint8>(Locals);
-		if (Result)
-		{
-			Property->CopyCompleteValueToScriptVM(Result, MostRecentPropertyAddress);
-		}
+		Property->CopyCompleteValueToScriptVM(Result, MostRecentPropertyAddress);
 	}
 }
 
diff --git a/Engine/Source/Runtime/CoreUObject/Public/UObject/ScriptMacros.h b/Engine/Source/Runtime/CoreUObject/Public/UObject/ScriptMacros.h
--- a/Engine/Source/Runtime/CoreUObject/Public/UObject/ScriptMacros.h
+++ b/Engine/Source/Runtime/CoreUObject/Public/UObject/ScriptMacros.h
@@ -27,17 +27,20 @@ enum {MAX_VARIABLE_SIZE = 0x0FFF };
 	Stack.StepCompiledIn<PropertyType>(&ParamName);
 
 #define PARAM_PASSED_BY_REF(ParamName, PropertyType, ParamType)									\
-	ParamType& ParamName = Stack.StepCompiledInRef<PropertyType, ParamType>(NULL);
+	ParamType ParamName##Temp;																	\
+	ParamType& ParamName = Stack.StepCompiledInRef<PropertyType, ParamType>(&ParamName##Temp);
 
 #define PARAM_PASSED_BY_REF_ZEROED(ParamName, PropertyType, ParamType)							\
-	ParamType& ParamName = Stack.StepCompiledInRef<PropertyType, ParamType>(NULL);
+	ParamType ParamName##Temp = (ParamType)0;													\
+	ParamType& ParamName = Stack.StepCompiledInRef<PropertyType, ParamType>(&ParamName##Temp);
 
 #define P_GET_PROPERTY(PropertyType, ParamName)													\
 	PropertyType::TCppType ParamName = PropertyType::GetDefaultPropertyValue();					\
 	Stack.StepCompiledIn<PropertyType>(&ParamName);
 
 #define P_GET_PROPERTY_REF(PropertyType, ParamName)												\
-	PropertyType::TCppType& ParamName = Stack.StepCompiledInRef<PropertyType, PropertyType::TCppType>(NULL);
+	PropertyType::TCppType ParamName##Temp = PropertyType::GetDefaultPropertyValue();			\
+	PropertyType::TCppType& ParamName = Stack.StepCompiledInRef<PropertyType, PropertyType::TCppType>(&ParamName##Temp);
 
 
 
diff --git a/Engine/Source/Runtime/CoreUObject/Public/UObject/Stack.h b/Engine/Source/Runtime/CoreUObject/Public/UObject/Stack.h
--- a/Engine/Source/Runtime/CoreUObject/Public/UObject/Stack.h
+++ b/Engine/Source/Runtime/CoreUObject/Public/UObject/Stack.h
@@ -416,6 +416,5 @@ FORCEINLINE_DEBUGGABLE TNativeType& FFrame::StepCompiledInRef(void*const Tempora
 		StepExplicitProperty(TemporaryBuffer, Property);
 	}
 
-	check(MostRecentPropertyAddress);
-	return *(TNativeType*)(MostRecentPropertyAddress);
+	return (MostRecentPropertyAddress != NULL) ? *(TNativeType*)(MostRecentPropertyAddress) : *(TNativeType*)TemporaryBuffer;
 }
