UE-15865: Compile errors within a struct (float 2,3's and linear colors seem to have lost their type)

diff --git a/Engine/Source/Runtime/CoreUObject/Private/UObject/UObjectGlobals.cpp b/Engine/Source/Runtime/CoreUObject/Private/UObject/UObjectGlobals.cpp
--- a/Engine/Source/Runtime/CoreUObject/Private/UObject/UObjectGlobals.cpp
+++ b/Engine/Source/Runtime/CoreUObject/Private/UObject/UObjectGlobals.cpp
@@ -107,6 +107,76 @@ UObject* StaticFindObjectFast( UClass* ObjectClass, UObject* ObjectPackage, FNam
 	return StaticFindObjectFastInternal( ObjectClass, ObjectPackage, ObjectName, ExactClass, AnyPackage, ExclusiveFlags );
 }
 
+// Anonymous namespace to not pollute global.
+namespace
+{
+	/**
+	 * Legacy static find object helper, that helps to find reflected types, that
+	 * are no longer a subobjects of UCLASS defined in the same header.
+	 *
+	 * If the class looked for is of one of the relocated types (or theirs subclass)
+	 * then it performs another search in containing package.
+	 *
+	 * If the class match wasn't exact (i.e. either nullptr or subclass of allowed
+	 * ones) and we've found an object we're revalidating it to make sure the
+	 * legacy search was valid.
+	 *
+	 * @param ObjectClass Class of the object to find.
+	 * @param ObjectPackage Package of the object to find.
+	 * @param ObjectName Name of the object to find.
+	 * @param ExactClass If the class match has to be exact. I.e. ObjectClass == FoundObjects.GetClass()
+	 * @param OrigInName Original name provided to StaticFindObject, before resolving. Used to determine if we're looking for delegate signature.
+	 *
+	 * @returns Found object.
+	 */
+	UObject* StaticFindObjectWithChangedLegacyPath(UClass* ObjectClass, UObject* ObjectPackage, FName ObjectName, bool ExactClass, const TCHAR* OrigInName)
+	{
+		UObject* MatchingObject = nullptr;
+
+		// This is another look-up for native enums, structs or delegate signatures, cause they're path changed
+		// and old packages can have invalid ones. The path now does not have a UCLASS as an outer. All mentioned
+		// types are just children of package of the file there were defined in.
+		if (!FPlatformProperties::RequiresCookedData() && // Cooked platforms will have all paths resolved.
+			ObjectPackage != nullptr &&
+			ObjectPackage->IsA<UClass>()) // Only if outer is a class.
+		{
+			bool bHasDelegateSignaturePostfix = FString(OrigInName).EndsWith(HEADER_GENERATED_DELEGATE_SIGNATURE_SUFFIX);
+
+			bool bExactPathChangedClass = ObjectClass == UEnum::StaticClass() // Enums
+				|| ObjectClass == UScriptStruct::StaticClass() || ObjectClass == UStruct::StaticClass() // Structs
+				|| (ObjectClass == UFunction::StaticClass() && bHasDelegateSignaturePostfix); // Delegates
+
+			bool bSubclassOfPathChangedClass = !bExactPathChangedClass && !ExactClass
+				&& (ObjectClass == nullptr // Any class
+				|| UEnum::StaticClass()->IsChildOf(ObjectClass) // Enums
+				|| UScriptStruct::StaticClass()->IsChildOf(ObjectClass) || UStruct::StaticClass()->IsChildOf(ObjectClass) // Structs
+				|| (UFunction::StaticClass()->IsChildOf(ObjectClass) && bHasDelegateSignaturePostfix)); // Delegates
+
+			if (!bExactPathChangedClass && !bSubclassOfPathChangedClass)
+			{
+				return nullptr;
+			}
+
+			MatchingObject = StaticFindObject(ObjectClass, ObjectPackage->GetOutermost(), *ObjectName.ToString(), ExactClass);
+
+			if (MatchingObject && bSubclassOfPathChangedClass)
+			{
+				// If the class wasn't given exactly, check if found object is of class that outers were changed.
+				UClass* MatchingObjectClass = MatchingObject->GetClass();
+				if (!(MatchingObjectClass == UEnum::StaticClass()	// Enums
+					|| MatchingObjectClass == UScriptStruct::StaticClass() || MatchingObjectClass == UStruct::StaticClass() // Structs
+					|| (MatchingObjectClass == UFunction::StaticClass() && bHasDelegateSignaturePostfix)) // Delegates
+					)
+				{
+					return nullptr;
+				}
+			}
+		}
+
+		return MatchingObject;
+	}
+}
+
 //
 // Find an optional object.
 //
@@ -165,20 +235,9 @@ UObject* StaticFindObject( UClass* ObjectClass, UObject* InObjectPackage, const
 	FName ObjectName(*InName, FNAME_Add, true);
 	MatchingObject = StaticFindObjectFast( ObjectClass, ObjectPackage, ObjectName, ExactClass, bAnyPackage );
 
-	// This is another look-up for native enums, structs or delegate signatures, cause they're path changed
-	// and old packages can have invalid ones. The path now does not have a UCLASS as an outer. All mentioned
-	// types are just children of package of the file there were defined in.
-	if (!MatchingObject && ObjectPackage != nullptr &&
-		ObjectPackage->IsA<UClass>() && // Only if outer is a class.
-		!FPlatformProperties::RequiresCookedData() && // Cooked platforms will have all paths resolved.
-			(
-				ObjectClass == UEnum::StaticClass()	// Enums
-				|| ObjectClass == UScriptStruct::StaticClass() || ObjectClass == UStruct::StaticClass() // Structs
-				|| (ObjectClass == UFunction::StaticClass() && FString(OrigInName).EndsWith(HEADER_GENERATED_DELEGATE_SIGNATURE_SUFFIX)) // Delegates
-			)
-		)
+	if (!MatchingObject)
 	{
-		MatchingObject = StaticFindObject(ObjectClass, ObjectPackage->GetOutermost(), *ObjectName.ToString(), ExactClass);
+		return StaticFindObjectWithChangedLegacyPath(ObjectClass, ObjectPackage, ObjectName, ExactClass, OrigInName);
 	}
 
 	return MatchingObject;
