Added support for async save of compressed packages.
Used in launch on helps with in editor hitches.

diff --git a/Engine/Source/Runtime/CoreUObject/Private/UObject/SavePackage.cpp b/Engine/Source/Runtime/CoreUObject/Private/UObject/SavePackage.cpp
--- a/Engine/Source/Runtime/CoreUObject/Private/UObject/SavePackage.cpp
+++ b/Engine/Source/Runtime/CoreUObject/Private/UObject/SavePackage.cpp
@@ -735,7 +735,7 @@ public:
 	 * @param	DstFilename		Output name of compressed file, cannot be the same as SrcFilename
 	 * @param	SrcLinker		FLinkerSave object used to save src file
 	 *
-	 * @return true if sucessful, false otherwise
+	 * @return true if successful, false otherwise
 	 */
 	bool CompressFile( const TCHAR* DstFilename, FLinkerSave* SrcLinker )
 	{
@@ -764,6 +764,119 @@ public:
 		bool bMoveSucceded = IFileManager::Get().Move( DstFilename, *TmpFilename );
 		return bMoveSucceded;
 	}
+	/**
+	* Compresses the passed in src archive and writes it to destination archive.
+	*
+	* @param	FileReader		archive to read from
+	* @param	FileWriter		archive to write to
+	* @param	SrcLinker		FLinkerSave object used to save src file
+	*
+	* @return true if successful, false otherwise
+	*/
+	void CompressArchive(FArchive* FileReader, FArchive* FileWriter, const bool bForceByteSwapping, const int32 TotalHeaderSize, const TArray<int32> &ExportSizes)
+	{
+
+		// Read package file summary from source file.
+		FPackageFileSummary FileSummary;
+		(*FileReader) << FileSummary;
+
+		// Propagate byte swapping.
+		FileWriter->SetByteSwapping(bForceByteSwapping);
+
+		// We don't compress the package file summary but treat everything afterwards
+		// till the first export as a single chunk. This basically lumps name and import 
+		// tables into one compressed block.
+		int32 StartOffset = FileReader->Tell();
+		int32 RemainingHeaderSize = TotalHeaderSize - StartOffset;
+		CurrentChunk.UncompressedSize = RemainingHeaderSize;
+		CurrentChunk.UncompressedOffset = StartOffset;
+
+		// finish header chunk
+		FinishCurrentAndCreateNewChunk(0);
+
+		// Iterate over all exports and add them separately. The underlying code will take
+		// care of merging small blocks.
+		for (int32 ExportIndex = 0; ExportIndex<ExportSizes.Num(); ExportIndex++)
+		{
+			AddToChunk(ExportSizes[ExportIndex]);
+		}
+
+		// Finish chunk in flight and reset current chunk with size 0.
+		FinishCurrentAndCreateNewChunk(0);
+
+		ECompressionFlags BaseCompressionMethod = COMPRESS_Default;
+		if (FileWriter->IsCooking())
+		{
+			BaseCompressionMethod = FileWriter->CookingTarget()->GetBaseCompressionMethod();
+		}
+
+		// Write base version of package file summary after updating compressed chunks array and compression flags.
+		FileSummary.CompressionFlags = BaseCompressionMethod;
+		FileSummary.CompressedChunks = CompressedChunks;
+		(*FileWriter) << FileSummary;
+
+		// Reset internal state so subsequent calls will work.
+		CompressedChunks.Empty();
+		CurrentChunk = FCompressedChunk();
+
+		// Allocate temporary buffers for reading and compression.
+		int32	SrcBufferSize = RemainingHeaderSize;
+		void*	SrcBuffer = FMemory::Malloc(SrcBufferSize);
+
+		// Iterate over all chunks, read the data, compress and write it out to destination file.
+		for (int32 ChunkIndex = 0; ChunkIndex<FileSummary.CompressedChunks.Num(); ChunkIndex++)
+		{
+			FCompressedChunk& Chunk = FileSummary.CompressedChunks[ChunkIndex];
+
+			// Increase temporary buffer sizes if they are too small.
+			if (SrcBufferSize < Chunk.UncompressedSize)
+			{
+				SrcBufferSize = Chunk.UncompressedSize;
+				SrcBuffer = FMemory::Realloc(SrcBuffer, SrcBufferSize);
+			}
+
+			// Verify that we're not skipping any data.
+			check(Chunk.UncompressedOffset == FileReader->Tell());
+
+			// Read src/ uncompressed data.
+			FileReader->Serialize(SrcBuffer, Chunk.UncompressedSize);
+
+			// Keep track of offset.
+			Chunk.CompressedOffset = FileWriter->Tell();
+			// Serialize compressed. This is compatible with async LoadCompressedData.
+			FileWriter->SerializeCompressed(SrcBuffer, Chunk.UncompressedSize, (ECompressionFlags)FileSummary.CompressionFlags);
+			// Keep track of compressed size.
+			Chunk.CompressedSize = FileWriter->Tell() - Chunk.CompressedOffset;
+		}
+
+		// get the start of the bulkdata and update it in the summary
+		FileSummary.BulkDataStartOffset = FileWriter->Tell();
+
+		// serialize the bulkdata directly, as bulkdata is handling the compression already
+		int64 SizeToCopy = FileReader->TotalSize() - FileReader->Tell();
+		if (SrcBufferSize < SizeToCopy)
+		{
+			SrcBufferSize = SizeToCopy;
+			SrcBuffer = FMemory::Realloc(SrcBuffer, SizeToCopy);
+		}
+
+		// Read src/ uncompressed data.
+		FileReader->Serialize(SrcBuffer, SizeToCopy);
+		FileWriter->Serialize(SrcBuffer, SizeToCopy);
+
+		// Verify that we've compressed everything.
+		check(FileReader->AtEnd());
+
+		// Serialize file summary again - this time CompressedChunks array is going to contain compressed size and offsets.
+		FileWriter->Seek(0);
+		(*FileWriter) << FileSummary;
+
+		// Free intermediate buffers.
+		FMemory::Free(SrcBuffer);
+
+	}
+
+
 	/**
 	 * Compresses the passed in src archive and writes it to destination archive.
 	 *
@@ -771,7 +884,7 @@ public:
 	 * @param	FileWriter		archive to write to
 	 * @param	SrcLinker		FLinkerSave object used to save src file
 	 *
-	 * @return true if sucessful, false otherwise
+	 * @return true if successful, false otherwise
 	 */
 	void CompressArchive( FArchive* FileReader, FArchive* FileWriter, FLinkerSave* SrcLinker )
 	{
@@ -977,6 +1090,89 @@ private:
 };
 
 
+
+void AsyncWriteCompressedFile(const TArray<uint8>& InData, const TCHAR* Filename, const FDateTime& TimeStamp, const bool bForceByteSwapping, const int32 TotalHeaderSize, const TArray<int32>& ExportSizes)
+{
+	class FAsyncWriteWorker : public FNonAbandonableTask
+	{
+	public:
+		/** Filename To write to**/
+		FString Filename;
+		/** Data for the file **/
+		TArray<uint8> Data;
+		/** Timestamp to give the file. MinValue if shouldn't be modified */
+		FDateTime FinalTimeStamp;
+
+		bool bForceByteSwapping;
+		int32 TotalHeaderSize;
+		TArray<int32> ExportSizes;
+
+		/** Constructor
+		*/
+		FAsyncWriteWorker(const TArray<uint8>& InData, const TCHAR* InFilename, const FDateTime& InTimeStamp, bool InBForceByteSwapping, const int32 InTotalHeaderSize, const TArray<int32>& InExportSizes)
+			: Filename(InFilename)
+			, Data(InData)
+			, FinalTimeStamp(InTimeStamp)
+			, bForceByteSwapping(InBForceByteSwapping)
+			, TotalHeaderSize(InTotalHeaderSize)
+			, ExportSizes(InExportSizes)
+		{
+		}
+
+		/** Write the file  */
+		void DoWork()
+		{
+			FString TmpFilename = FPaths::GetPath(Filename) / (FPaths::GetBaseFilename(Filename) + TEXT("_SaveCompressed.tmp"));
+			// Create file reader and writer...
+			FMemoryReader Reader(Data, true);
+			FArchive* FileWriter = IFileManager::Get().CreateFileWriter(*TmpFilename);
+			// ... and abort if either operation wasn't successful.
+			if (!FileWriter)
+			{
+				// Delete potentially created reader or writer.
+				delete FileWriter;
+				// Delete temporary file.
+				IFileManager::Get().Delete(*TmpFilename);
+				// Failure.
+				UE_LOG(LogSavePackage, Fatal, TEXT("Could not write to %s!"), *TmpFilename);
+			}
+
+			FFileCompressionHelper CompressionHelper;
+			CompressionHelper.CompressArchive(&Reader, FileWriter, bForceByteSwapping, TotalHeaderSize, ExportSizes);
+
+			delete FileWriter;
+
+			// Clean-up the memory as soon as we save the file to reduce the memory footprint.
+			const int64 DataSize = Data.Num();
+			Data.Empty();
+			if (!IFileManager::Get().Move(*Filename, *TmpFilename, true, true, false, false))
+			{
+				UE_LOG(LogSavePackage, Fatal, TEXT("Could not move from %s to %s."), *TmpFilename, *Filename);
+			}
+			else
+			{
+				if (FinalTimeStamp != FDateTime::MinValue())
+				{
+					IFileManager::Get().SetTimeStamp(*Filename, FinalTimeStamp);
+				}
+			}
+
+			OutstandingAsyncWrites.Decrement();
+		}
+
+		FORCEINLINE TStatId GetStatId() const
+		{
+			RETURN_QUICK_DECLARE_CYCLE_STAT(FAsyncWriteWorker, STATGROUP_ThreadPoolAsyncTasks);
+		}
+	};
+
+	OutstandingAsyncWrites.Increment();
+	(new FAutoDeleteAsyncTask<FAsyncWriteWorker>(InData, Filename, TimeStamp, bForceByteSwapping, TotalHeaderSize, ExportSizes))->StartBackgroundTask();
+}
+
+
+
+
 /**
  * Find most likely culprit that caused the objects in the passed in array to be considered for saving.
  *
@@ -4039,7 +4235,24 @@ bool UPackage::SavePackage( UPackage* InOuter, UObject* Base, EObjectFlags TopLe
 				if( Success == true )
 				{
 					// Compress the temporarily file to destination.
-					if( bCompressFromMemory )
+					if (bCompressFromMemory && bSaveAsync)
+					{
+						UE_LOG(LogSavePackage, Log, TEXT("Async compressing from memory to '%s'"), *NewPath);
+
+						// Detach archive used for memory saving.
+						if (Linker)
+						{
+							TArray<int32> ExportSizes;
+							for (int I = 0; I < Linker->ExportMap.Num(); ++I)
+							{
+								ExportSizes.Add(Linker->ExportMap[I].SerialSize);
+							}
+							AsyncWriteCompressedFile(*(FBufferArchive*)(Linker->Saver), *NewPath, FinalTimeStamp, Linker->ForceByteSwapping(), Linker->Summary.TotalHeaderSize, ExportSizes);
+							Linker->Detach();
+						}
+						UE_LOG_COOK_TIME(TEXT("Async write"));
+					}
+					else if( bCompressFromMemory )
 					{
 						UE_LOG(LogSavePackage, Log,  TEXT("Compressing from memory to '%s'"), *NewPath );
 						FFileCompressionHelper CompressionHelper;
