Fixed warning as error on mono.

diff --git a/Engine/Source/Runtime/CoreUObject/Private/UObject/SavePackage.cpp b/Engine/Source/Runtime/CoreUObject/Private/UObject/SavePackage.cpp
--- a/Engine/Source/Runtime/CoreUObject/Private/UObject/SavePackage.cpp
+++ b/Engine/Source/Runtime/CoreUObject/Private/UObject/SavePackage.cpp
@@ -1091,7 +1091,7 @@ private:
 
 
 
-void AsyncWriteCompressedFile(const TArray<uint8>& InData, const TCHAR* Filename, const FDateTime& TimeStamp, const bool bForceByteSwapping, const int32 TotalHeaderSize, const TArray<int32>& ExportSizes)
+void AsyncWriteCompressedFile(const TArray<uint8>& Data, const TCHAR* Filename, const FDateTime& TimeStamp, const bool bForceByteSwapping, const int32 TotalHeaderSize, const TArray<int32>& ExportSizes)
 {
 	class FAsyncWriteWorker : public FNonAbandonableTask
 	{
@@ -1167,7 +1167,7 @@ void AsyncWriteCompressedFile(const TArray<uint8>& InData, const TCHAR* Filename
 	};
 
 	OutstandingAsyncWrites.Increment();
-	(new FAutoDeleteAsyncTask<FAsyncWriteWorker>(InData, Filename, TimeStamp, bForceByteSwapping, TotalHeaderSize, ExportSizes))->StartBackgroundTask();
+	(new FAutoDeleteAsyncTask<FAsyncWriteWorker>(Data, Filename, TimeStamp, bForceByteSwapping, TotalHeaderSize, ExportSizes))->StartBackgroundTask();
 }
 
 
