Examining the value of an enum array pin will show the enum value's display name.

Changed UByteProperty::ExportTextItem to use GetEnumText when PPF_BlueprintDebugView is passed in.

diff --git a/Engine/Source/Runtime/CoreUObject/Private/UObject/PropertyByte.cpp b/Engine/Source/Runtime/CoreUObject/Private/UObject/PropertyByte.cpp
--- a/Engine/Source/Runtime/CoreUObject/Private/UObject/PropertyByte.cpp
+++ b/Engine/Source/Runtime/CoreUObject/Private/UObject/PropertyByte.cpp
@@ -113,9 +113,21 @@ void UByteProperty::ExportTextItem( FString& ValueStr, const void* PropertyValue
 	{
 		// if the value is the max value (the autogenerated *_MAX value), export as "INVALID", unless we're exporting text for copy/paste (for copy/paste,
 		// the property text value must actually match an entry in the enum's names array)
-		ValueStr += ((*(const uint8*)PropertyValue < Enum->NumEnums() - 1 || ((PortFlags&PPF_Copy) != 0 && *(const uint8*)PropertyValue < Enum->NumEnums())) 
-			? Enum->GetEnumName(*(const uint8*)PropertyValue) 
-			: TEXT("(INVALID)"));
+		if (*(const uint8*)PropertyValue < Enum->NumEnums() - 1 || ((PortFlags&PPF_Copy) != 0 && *(const uint8*)PropertyValue < Enum->NumEnums()))
+		{
+			if (PortFlags & PPF_BlueprintDebugView)
+			{
+				ValueStr += Enum->GetEnumText(*(const uint8*)PropertyValue).ToString();
+			}
+			else
+			{
+				ValueStr += Enum->GetEnumName(*(const uint8*)PropertyValue);
+			}
+		}
+		else
+		{
+			ValueStr += TEXT("(INVALID)");
+		}
 	}
 	else
 	{
