Fixing loading objects synchronously from PostLoad while streaming

diff --git a/Engine/Source/Runtime/CoreUObject/Private/Serialization/AsyncLoading.cpp b/Engine/Source/Runtime/CoreUObject/Private/Serialization/AsyncLoading.cpp
--- a/Engine/Source/Runtime/CoreUObject/Private/Serialization/AsyncLoading.cpp
+++ b/Engine/Source/Runtime/CoreUObject/Private/Serialization/AsyncLoading.cpp
@@ -1530,6 +1530,8 @@ EAsyncPackageState::Type FAsyncPackage::PostLoadDeferredObjects(double InTickSta
 	LastObjectWorkWasPerformedOn = nullptr;
 	LastTypeOfWorkPerformed = TEXT("postloading_gamethread");
 
+	TArray<UObject*>& ObjLoadedInPostLoad = FUObjectThreadContext::Get().ObjLoaded;
+
 	while (DeferredPostLoadIndex < DeferredPostLoadObjects.Num() && !::IsTimeLimitExceeded(InTickStartTime, bInUseTimeLimit, InOutTimeLimit, LastTypeOfWorkPerformed, LastObjectWorkWasPerformedOn))
 	{
 		UObject* Object = DeferredPostLoadObjects[DeferredPostLoadIndex++];
@@ -1539,6 +1541,26 @@ EAsyncPackageState::Type FAsyncPackage::PostLoadDeferredObjects(double InTickSta
 
 		Object->ConditionalPostLoad();
 
+		if (ObjLoadedInPostLoad.Num())
+		{
+			// If there were any LoadObject calls inside of PostLoad, we need to pre-load those objects here. 
+			// There's no going back to the async tick loop from here.
+			UE_LOG(LogStreaming, Warning, TEXT("Detected %d objects loaded in PostLoad while streaming, this may cause hitches as we're blocking async loading to pre-load them."), ObjLoadedInPostLoad.Num());
+			
+			// Make sure all objects loaded in PostLoad get post-loaded too
+			DeferredPostLoadObjects.Append(ObjLoadedInPostLoad);
+
+			// Preload (aka serialize) the objects loaded in PostLoad.
+			for (UObject* PreLoadObject : ObjLoadedInPostLoad)
+			{
+				if (PreLoadObject && PreLoadObject->GetLinker())
+				{
+					PreLoadObject->GetLinker()->Preload(PreLoadObject);
+				}
+			}
+			ObjLoadedInPostLoad.Empty();
+		}
+
 		LastObjectWorkWasPerformedOn = Object;		
 
 		UpdateLoadPercentage();
diff --git a/Engine/Source/Runtime/CoreUObject/Private/UObject/UObjectGlobals.cpp b/Engine/Source/Runtime/CoreUObject/Private/UObject/UObjectGlobals.cpp
--- a/Engine/Source/Runtime/CoreUObject/Private/UObject/UObjectGlobals.cpp
+++ b/Engine/Source/Runtime/CoreUObject/Private/UObject/UObjectGlobals.cpp
@@ -1041,10 +1041,12 @@ UPackage* LoadPackageInternal(UPackage* InOuter, const TCHAR* InLongPackageName,
 			FIOSystem::Get().HintDoneWithFile(*Linker->Filename);
 		}
 
-		// with UE4 and single asset per package, we load so many packages that some platforms will run out
+		// With UE4 and single asset per package, we load so many packages that some platforms will run out
 		// of file handles. So, this will close the package, but just things like bulk data loading will
-		// fail, so we only currently do this when loading on consoles
-		if (FPlatformProperties::RequiresCookedData())
+		// fail, so we only currently do this when loading on consoles.
+		// The only exception here is when we're in the middle of async loading where we can't reset loaders yet. This should only happen when
+		// doing synchronous load in the middle of streaming.
+		if (FPlatformProperties::RequiresCookedData() && !IsInAsyncLoadingThread())
 		{
 			if (FUObjectThreadContext::Get().ObjBeginLoadCount == 0)
 			{
@@ -1053,10 +1055,10 @@ UPackage* LoadPackageInternal(UPackage* InOuter, const TCHAR* InLongPackageName,
 					ResetLoaders(Result);
 				}
 				delete Linker->Loader;
-				Linker->Loader = NULL;
+				Linker->Loader = nullptr;
 			}
 			// Async loading removes delayed linkers on the game thread after streaming has finished
-			else if (!IsInAsyncLoadingThread())
+			else
 			{
 				FUObjectThreadContext::Get().DelayedLinkerClosePackages.Add(Linker);
 			}
