[INTEGRATE] Change 2564765  by [REDACTED_PERSON][REDACTED_HANDLE] on [REDACTED_PHONE]:38

Making DebugPostLoad thread safe

diff --git a/Engine/Source/Runtime/CoreUObject/Private/UObject/Obj.cpp b/Engine/Source/Runtime/CoreUObject/Private/UObject/Obj.cpp
--- a/Engine/Source/Runtime/CoreUObject/Private/UObject/Obj.cpp
+++ b/Engine/Source/Runtime/CoreUObject/Private/UObject/Obj.cpp
@@ -30,8 +30,6 @@ DEFINE_LOG_CATEGORY(LogObj);
 static UPackage*			GObjTransientPkg								= NULL;		
 
 #if !(UE_BUILD_SHIPPING || UE_BUILD_TEST)
-	/** Used to verify that the Super::PostLoad chain is intact.			*/
-	static TArray<UObject*,TInlineAllocator<16> >		DebugPostLoad;
 	/** Used to verify that the Super::BeginDestroyed chain is intact.			*/
 	static TArray<UObject*,TInlineAllocator<16> >		DebugBeginDestroyed;
 	/** Used to verify that the Super::FinishDestroyed chain is intact.			*/
@@ -258,7 +256,7 @@ void UObject::PostLoad()
 {
 	// Note that it has propagated.
 #if !(UE_BUILD_SHIPPING || UE_BUILD_TEST)
-	DebugPostLoad.RemoveSingle(this);
+	FUObjectThreadContext::Get().DebugPostLoad.RemoveSingle(this);
 #endif
 
 	/*
@@ -700,8 +698,9 @@ void UObject::ConditionalPostLoad()
 		check(IsInGameThread() || HasAnyFlags(RF_ClassDefaultObject|RF_ArchetypeObject) || IsPostLoadThreadSafe() || IsA(UClass::StaticClass()))
 
 #if !(UE_BUILD_SHIPPING || UE_BUILD_TEST)
-		checkSlow(!DebugPostLoad.Contains(this));
-		DebugPostLoad.Add(this);
+		FUObjectThreadContext& ThreadContext = FUObjectThreadContext::Get();
+		checkSlow(!ThreadContext.DebugPostLoad.Contains(this));
+		ThreadContext.DebugPostLoad.Add(this);
 #endif
 		ClearFlags( RF_NeedPostLoad );
 
@@ -716,9 +715,9 @@ void UObject::ConditionalPostLoad()
 		PostLoad();
 
 #if !(UE_BUILD_SHIPPING || UE_BUILD_TEST)
-		if( DebugPostLoad.Contains(this) )
+		if (ThreadContext.DebugPostLoad.Contains(this))
 		{
-			UE_LOG(LogObj, Fatal, TEXT("%s failed to route PostLoad.  Please call Super::PostLoad() in your <className>::PostLoad() function. "), *GetFullName() );
+			UE_LOG(LogObj, Fatal, TEXT("%s failed to route PostLoad.  Please call Super::PostLoad() in your <className>::PostLoad() function."), *GetFullName());
 		}
 #endif
 	}
diff --git a/Engine/Source/Runtime/CoreUObject/Public/UObject/UObjectThreadContext.h b/Engine/Source/Runtime/CoreUObject/Public/UObject/UObjectThreadContext.h
--- a/Engine/Source/Runtime/CoreUObject/Public/UObject/UObjectThreadContext.h
+++ b/Engine/Source/Runtime/CoreUObject/Public/UObject/UObjectThreadContext.h
@@ -46,5 +46,7 @@ public:
 #if !(UE_BUILD_SHIPPING || UE_BUILD_TEST)
 	/** Stack to ensure that PostInitProperties is routed through Super:: calls. **/
 	TArray<UObject*> PostInitPropertiesCheck;
+	/** Used to verify that the Super::PostLoad chain is intact.			*/
+	TArray<UObject*, TInlineAllocator<16> > DebugPostLoad;
 #endif
 };
\ No newline at end of file
