[INTEGRATE] Change 2564781  by [REDACTED_PERSON][REDACTED_HANDLE] on [REDACTED_PHONE]:04

Fix for potential crash when resetting loaders in cooked build

diff --git a/Engine/Source/Runtime/CoreUObject/Private/UObject/UObjectGlobals.cpp b/Engine/Source/Runtime/CoreUObject/Private/UObject/UObjectGlobals.cpp
--- a/Engine/Source/Runtime/CoreUObject/Private/UObject/UObjectGlobals.cpp
+++ b/Engine/Source/Runtime/CoreUObject/Private/UObject/UObjectGlobals.cpp
@@ -1050,12 +1050,20 @@ UPackage* LoadPackageInternal(UPackage* InOuter, const TCHAR* InLongPackageName,
 		{
 			if (FUObjectThreadContext::Get().ObjBeginLoadCount == 0)
 			{
+				// Sanity check to make sure that Linker is the linker that loaded our Result package
+				check(!Result || Result->LinkerLoad == Linker);
 				if (Result && Linker->Loader)
 				{
 					ResetLoaders(Result);
 				}
-				delete Linker->Loader;
-				Linker->Loader = nullptr;
+				// Reset loaders could have already deleted Linker so guard against deleting stale pointers
+				if (Result && Result->LinkerLoad)
+				{
+					delete Linker->Loader;
+					Linker->Loader = nullptr;
+				}
+				// And make sure no one can use it after it's been deleted
+				Linker = nullptr;
 			}
 			// Async loading removes delayed linkers on the game thread after streaming has finished
 			else
