These validation functions take a lot of time when loading blueprints, disabled after consulting with Steve.Robb

diff --git a/Engine/Source/Runtime/CoreUObject/Private/UObject/Class.cpp b/Engine/Source/Runtime/CoreUObject/Private/UObject/Class.cpp
--- a/Engine/Source/Runtime/CoreUObject/Private/UObject/Class.cpp
+++ b/Engine/Source/Runtime/CoreUObject/Private/UObject/Class.cpp
@@ -11,6 +11,11 @@
 #include "LinkerPlaceholderFunction.h"
 #include "StructScriptLoader.h"
 
+// This flag enables some expensive class tree validation that is meant to catch mutations of 
+// the class tree outside of SetSuperStruct. It has been disabled because loading blueprints 
+// does a lot of mutation of the class tree, and the validation checks impact iteration time.
+#define DO_CLASS_TREE_VALIDATION 0
+
 DECLARE_LOG_CATEGORY_EXTERN(LogScriptSerialization, Log, All);
 DEFINE_LOG_CATEGORY(LogScriptSerialization);
 DEFINE_LOG_CATEGORY(LogClass);
@@ -3160,7 +3165,7 @@ void UClass::Link(FArchive& Ar, bool bRelinkExistingProperties)
 			Register(Orphan);
 		}
 
-		#if DO_CHECK
+		#if DO_CLASS_TREE_VALIDATION
 			Validate();
 		#endif
 	}
@@ -3208,7 +3213,7 @@ void UClass::Link(FArchive& Ar, bool bRelinkExistingProperties)
 
 		State.Classes.RemoveAt(ClassIndex, NumRemoved, false);
 
-		#if DO_CHECK
+		#if DO_CLASS_TREE_VALIDATION
 			Validate();
 		#endif
 	}
