Protect against crash if Class is null (likely due to referencing a class in an unloaded module)

diff --git a/Engine/Source/Runtime/CoreUObject/Private/UObject/ScriptCore.cpp b/Engine/Source/Runtime/CoreUObject/Private/UObject/ScriptCore.cpp
--- a/Engine/Source/Runtime/CoreUObject/Private/UObject/ScriptCore.cpp
+++ b/Engine/Source/Runtime/CoreUObject/Private/UObject/ScriptCore.cpp
@@ -2022,37 +2022,39 @@ void UObject::execDynamicCast( FFrame& Stack, RESULT_DECL )
 	//*(UObject**)RESULT_PARAM = (Castee && Castee->IsA(Class)) ? Castee : NULL;
 	*(UObject**)RESULT_PARAM = NULL; // default value
 
-
-	// if we were passed in a null value
-	if( Castee == NULL )
+	if (Class)
 	{
-		if( Class->HasAnyClassFlags(CLASS_Interface) )
+		// if we were passed in a null value
+		if( Castee == NULL )
 		{
-			((FScriptInterface*)RESULT_PARAM)->SetObject(NULL);
+			if( Class->HasAnyClassFlags(CLASS_Interface) )
+			{
+				((FScriptInterface*)RESULT_PARAM)->SetObject(NULL);
+			}
+			else
+			{
+				*(UObject**)RESULT_PARAM = NULL;
+			}
+			return;
 		}
-		else
+
+		// check to see if the Castee is an implemented interface by looking up the
+		// class hierarchy and seeing if any class in said hierarchy implements the interface
+		if( Class->HasAnyClassFlags(CLASS_Interface) )
 		{
-			*(UObject**)RESULT_PARAM = NULL;
+			if ( Castee->GetClass()->ImplementsInterface(Class) )
+			{
+				// interface property type - convert to FScriptInterface
+				((FScriptInterface*)RESULT_PARAM)->SetObject(Castee);
+				((FScriptInterface*)RESULT_PARAM)->SetInterface(Castee->GetInterfaceAddress(Class));
+			}
 		}
-		return;
-	}
-
-	// check to see if the Castee is an implemented interface by looking up the
-	// class hierarchy and seeing if any class in said hierarchy implements the interface
-	if( Class->HasAnyClassFlags(CLASS_Interface) )
-	{
-		if ( Castee->GetClass()->ImplementsInterface(Class) )
+		// check to see if the Castee is a castable class
+		else if( Castee->IsA(Class) )
 		{
-			// interface property type - convert to FScriptInterface
-			((FScriptInterface*)RESULT_PARAM)->SetObject(Castee);
-			((FScriptInterface*)RESULT_PARAM)->SetInterface(Castee->GetInterfaceAddress(Class));
+			*(UObject**)RESULT_PARAM = Castee;
 		}
 	}
-	// check to see if the Castee is a castable class
-	else if( Castee->IsA(Class) )
-	{
-		*(UObject**)RESULT_PARAM = Castee;
-	}
 }
 IMPLEMENT_VM_FUNCTION( EX_DynamicCast, execDynamicCast );
 
