Support for suspending async loading.

- Async loading will now be suspended when render thread is suspended. This fixes a rare crash on startup.
- Fixed an infinite loading screen in situation when a package hasn't finished loading but was in a queue to be finalized on the game thread and something requested it to load again.

diff --git a/Engine/Source/Runtime/Core/Private/Misc/Core.cpp b/Engine/Source/Runtime/Core/Private/Misc/Core.cpp
--- a/Engine/Source/Runtime/Core/Private/Misc/Core.cpp
+++ b/Engine/Source/Runtime/Core/Private/Misc/Core.cpp
@@ -170,6 +170,11 @@ TCHAR					GInternalGameName[64];
 IMPLEMENT_FOREIGN_ENGINE_DIR()
 #endif
 
+/** A function that does nothing. Allows for a default behavior for callback function pointers. */
+static void appNoop()
+{
+}
+
 /** Exec handler for game debugging tool, allowing commands like "editactor", ...							*/
 FExec*					GDebugToolExec					= NULL;
 /** Whether we're currently in the async loading codepath or not											*/
@@ -178,7 +183,9 @@ static bool IsAsyncLoadingCoreInternal()
 	// No Async loading in Core
 	return false;
 }
-bool(*IsAsyncLoading)() = &IsAsyncLoadingCoreInternal;
+bool (*IsAsyncLoading)() = &IsAsyncLoadingCoreInternal;
+void (*SuspendAsyncLoading)() = &appNoop;
+void (*ResumeAsyncLoading)() = &appNoop;
 /** Whether the editor is currently loading a package or not												*/
 bool					GIsEditorLoadingPackage				= false;
 /** Whether GWorld points to the play in editor world														*/
@@ -219,11 +226,6 @@ uint32					GSlateLoadingThreadId			= 0;
 /** Has GGameThreadId been set yet?																			*/
 bool					GIsGameThreadIdInitialized		= false;
 
-/** A function that does nothing. Allows for a default behavior for callback function pointers. */
-static void appNoop()
-{
-}
-
 /** Helper function to flush resource streaming.															*/
 void					(*GFlushStreamingFunc)(void)	  = &appNoop;
 /** Whether to emit begin/ end draw events.																	*/
diff --git a/Engine/Source/Runtime/Core/Public/CoreGlobals.h b/Engine/Source/Runtime/Core/Public/CoreGlobals.h
--- a/Engine/Source/Runtime/Core/Public/CoreGlobals.h
+++ b/Engine/Source/Runtime/Core/Public/CoreGlobals.h
@@ -192,6 +192,12 @@ extern CORE_API FExec* GDebugToolExec;
 /** Whether we're currently in the async loading code path or not */
 extern CORE_API bool(*IsAsyncLoading)();
 
+/** Suspends async package loading. */
+extern CORE_API void (*SuspendAsyncLoading)();
+
+/** Resumes async package loading. */
+extern CORE_API void (*ResumeAsyncLoading)();
+
 /** Whether the editor is currently loading a package or not */
 extern CORE_API bool GIsEditorLoadingPackage;
 
diff --git a/Engine/Source/Runtime/CoreUObject/Private/Serialization/AsyncLoading.cpp b/Engine/Source/Runtime/CoreUObject/Private/Serialization/AsyncLoading.cpp
--- a/Engine/Source/Runtime/CoreUObject/Private/Serialization/AsyncLoading.cpp
+++ b/Engine/Source/Runtime/CoreUObject/Private/Serialization/AsyncLoading.cpp
@@ -359,6 +359,8 @@ int32 FAsyncLoadingThread::CreateAsyncPackagesFromQueue()
 			{
 				const bool bInternalCallback = false;
 				LoadedPackages[ExistingPackageIndex]->AddCompletionCallback(PackageRequest->PackageLoadedDelegate, bInternalCallback);
+				const int32 QueuedPackagesCount = QueuedPackagesCounter.Decrement();
+				check(QueuedPackagesCount >= 0);
 				continue;
 			}
 		}
@@ -375,8 +377,10 @@ int32 FAsyncLoadingThread::CreateAsyncPackagesFromQueue()
 			const bool bInternalCallback = false;
 			Package->AddCompletionCallback(PackageRequest->PackageLoadedDelegate, bInternalCallback);
 		}
-		QueuedPackagesCounter.Decrement();
-		check(QueuedPackagesCounter.GetValue() >= 0);
+		{
+			const int32 QueuedPackagesCount = QueuedPackagesCounter.Decrement();
+			check(QueuedPackagesCount >= 0);
+		}
 	}
 	NumCreated = QueueCopy.Num();
 
@@ -492,7 +496,7 @@ EAsyncPackageState::Type FAsyncLoadingThread::ProcessLoadedPackages(bool bUseTim
 		LoadedPackages.Empty();
 	}
 		
-	for (int32 PackageIndex = 0; PackageIndex < LoadedPackagesToProcess.Num(); ++PackageIndex)
+	for (int32 PackageIndex = 0; PackageIndex < LoadedPackagesToProcess.Num() && !IsAsyncLoadingSuspended(); ++PackageIndex)
 	{
 		if (PackageIndex % 20 == 0 && IsTimeLimitExceeded(TickStartTime, bUseTimeLimit, TimeLimit, TEXT("ProcessLoadedPackages")))
 		{
@@ -547,32 +551,38 @@ EAsyncPackageState::Type FAsyncLoadingThread::ProcessLoadedPackages(bool bUseTim
 
 EAsyncPackageState::Type FAsyncLoadingThread::TickAsyncLoading(bool bUseTimeLimit, bool bUseFullTimeLimit, float TimeLimit, FName ExcludeType)
 {
-	EAsyncPackageState::Type Result = EAsyncPackageState::Complete;	
-	double TickStartTime = FPlatformTime::Seconds();
-	double TimeLimitUsedForPreloading = 0;
+	const bool bLoadingSuspended = IsAsyncLoadingSuspended();
+	const bool bIsMultithreaded = FAsyncLoadingThread::IsMultithreaded();
+	EAsyncPackageState::Type Result = bLoadingSuspended ? EAsyncPackageState::PendingImports : EAsyncPackageState::Complete;
 
-	if (!FAsyncLoadingThread::IsMultithreaded())
-	{		
-		TickAsyncThread(bUseTimeLimit, bUseFullTimeLimit, TimeLimit);
-		TimeLimitUsedForPreloading = FPlatformTime::Seconds() - TickStartTime;
-	}
-
-	if (!IsTimeLimitExceeded(TickStartTime, bUseTimeLimit, TimeLimit, TEXT("Pre-ProcessLoadedPackages")))
+	if (!bLoadingSuspended)
 	{
-		double RemainingTimeLimit = FMath::Max(0.0, TimeLimit - TimeLimitUsedForPreloading);
-		Result = ProcessLoadedPackages(bUseTimeLimit, bUseFullTimeLimit, RemainingTimeLimit, ExcludeType);
-	}
+		double TickStartTime = FPlatformTime::Seconds();
+		double TimeLimitUsedForPreloading = 0;
 
-	if (Result != EAsyncPackageState::TimeOut && !IsTimeLimitExceeded(TickStartTime, bUseTimeLimit, TimeLimit, TEXT("Pre-EmptyReferencedObjects")))
-	{
+		if (!bIsMultithreaded)
+		{
+			Result = TickAsyncThread(bUseTimeLimit, bUseFullTimeLimit, TimeLimit);
+			TimeLimitUsedForPreloading = FPlatformTime::Seconds() - TickStartTime;
+		}
+
+		if (Result != EAsyncPackageState::TimeOut && !IsTimeLimitExceeded(TickStartTime, bUseTimeLimit, TimeLimit, TEXT("Pre-ProcessLoadedPackages")))
+		{
+			double RemainingTimeLimit = FMath::Max(0.0, TimeLimit - TimeLimitUsedForPreloading);
+			Result = ProcessLoadedPackages(bUseTimeLimit, bUseFullTimeLimit, RemainingTimeLimit, ExcludeType);
+		}
+
+		if (Result != EAsyncPackageState::TimeOut && !IsTimeLimitExceeded(TickStartTime, bUseTimeLimit, TimeLimit, TEXT("Pre-EmptyReferencedObjects")))
+		{
 #if THREADSAFE_UOBJECTS
-		FScopeLock QueueLock(&QueueCritical);
-		FScopeLock LoadedLock(&LoadedPackagesCritical);
+			FScopeLock QueueLock(&QueueCritical);
+			FScopeLock LoadedLock(&LoadedPackagesCritical);
 #endif
-		if (AsyncPackagesCounter.GetValue() == 0 && LoadedPackagesToProcess.Num() == 0)
-		{
-			FDeferredMessageLog::Flush();
-			FAsyncObjectsReferencer::Get().EmptyReferencedObjects();
+			if (AsyncPackagesCounter.GetValue() == 0 && LoadedPackagesToProcess.Num() == 0)
+			{
+				FDeferredMessageLog::Flush();
+				FAsyncObjectsReferencer::Get().EmptyReferencedObjects();
+			}
 		}
 	}
 
@@ -586,12 +596,15 @@ FAsyncLoadingThread::FAsyncLoadingThread()
 #endif
 	QueuedRequestsEvent = FPlatformProcess::GetSynchEventFromPool();
 	CancelLoadingEvent = FPlatformProcess::GetSynchEventFromPool();
+	ThreadSuspendedEvent = FPlatformProcess::GetSynchEventFromPool();
 	if (FAsyncLoadingThread::IsMultithreaded())
 	{
+		UE_LOG(LogStreaming, Display, TEXT("Async loading is multithreaded."));
 		Thread = FRunnableThread::Create(this, TEXT("FAsyncLoadingThread"), 0, TPri_Normal);
 	}
 	else
 	{
+		UE_LOG(LogStreaming, Display, TEXT("Async loading is time-sliced."));
 		Thread = nullptr;
 		Init();
 	}
@@ -606,6 +619,8 @@ FAsyncLoadingThread::~FAsyncLoadingThread()
 	QueuedRequestsEvent = nullptr;
 	FPlatformProcess::ReturnSynchEventToPool(CancelLoadingEvent);
 	CancelLoadingEvent = nullptr;
+	FPlatformProcess::ReturnSynchEventToPool(ThreadSuspendedEvent);
+	ThreadSuspendedEvent = nullptr;
 }
 
 bool FAsyncLoadingThread::Init()
@@ -617,22 +632,37 @@ uint32 FAsyncLoadingThread::Run()
 {
 	AsyncLoadingThreadID = FPlatformTLS::GetCurrentThreadId();
 
+	bool bWasSuspendedLastFrame = false;
 	while (StopTaskCounter.GetValue() == 0)
 	{
-		TickAsyncThread(false, true, 0.0f);
+		if (IsLoadingSuspended.GetValue() == 0)
+		{
+			bWasSuspendedLastFrame = false;
+			TickAsyncThread(false, true, 0.0f);
+		}
+		else if (!bWasSuspendedLastFrame)
+		{
+			bWasSuspendedLastFrame = true;
+			ThreadSuspendedEvent->Trigger();			
+		}
+		else
+		{
+			FPlatformProcess::SleepNoStats(0.001f);
+		}		
 	}
 	return 0;
 }
 
-void FAsyncLoadingThread::TickAsyncThread(bool bUseTimeLimit, bool bUseFullTimeLimit, float TimeLimit)
+EAsyncPackageState::Type  FAsyncLoadingThread::TickAsyncThread(bool bUseTimeLimit, bool bUseFullTimeLimit, float TimeLimit)
 {
+	EAsyncPackageState::Type Result = EAsyncPackageState::Complete;
 	if (!bShouldCancelLoading)
 	{
 		int32 ProcessedRequests = 0;
 		if (AsyncThreadReady.GetValue())
 		{
 			CreateAsyncPackagesFromQueue();
-			ProcessAsyncLoading(ProcessedRequests, bUseTimeLimit, bUseFullTimeLimit, TimeLimit, NAME_None);
+			Result = ProcessAsyncLoading(ProcessedRequests, bUseTimeLimit, bUseFullTimeLimit, TimeLimit, NAME_None);
 		}
 		if (ProcessedRequests == 0 && IsMultithreaded())
 		{
@@ -652,6 +682,8 @@ void FAsyncLoadingThread::TickAsyncThread(bool bUseTimeLimit, bool bUseFullTimeL
 	SET_FLOAT_STAT( STAT_AsyncIO_AsyncLoadingBlockingTime, FPlatformTime::ToSeconds( BlockingCycles.GetValue() ) );
 	BlockingCycles.Set( 0 );
 #endif
+
+	return Result;
 }
 
 void FAsyncLoadingThread::Stop()
@@ -661,7 +693,7 @@ void FAsyncLoadingThread::Stop()
 
 void FAsyncLoadingThread::CancelAsyncLoading()
 {
-	checkSlow(IsInGameThread());
+	checkSlow(IsInGameThread());	
 
 	bShouldCancelLoading = true;
 	if (IsMultithreaded())
@@ -677,6 +709,23 @@ void FAsyncLoadingThread::CancelAsyncLoading()
 	}
 }
 
+void FAsyncLoadingThread::SuspendLoading()
+{
+	check(IsInGameThread());
+	const int32 SuspendCount = IsLoadingSuspended.Increment();
+	if (IsMultithreaded() && SuspendCount == 1)
+	{
+		ThreadSuspendedEvent->Wait();
+	}
+}
+
+void FAsyncLoadingThread::ResumeLoading()
+{
+	check(IsInGameThread());
+	const int32 SuspendCount = IsLoadingSuspended.Decrement();
+	UE_CLOG(SuspendCount < 0, LogStreaming, Fatal, TEXT("ResumeAsyncLoadingThread: Async loading was resumed more times than it was suspended."));
+}
+
 float FAsyncLoadingThread::GetAsyncLoadPercentage(const FName& PackageName)
 {
 	float LoadPercentage = -1.0f;
@@ -782,6 +831,7 @@ FAsyncPackage::FAsyncPackage(const FAsyncPackageDesc& InDesc)
 , LastTypeOfWorkPerformed(nullptr)
 , LoadStartTime(0.0)
 , LoadPercentage(0)
+, AsyncLoadingThread(FAsyncLoadingThread::Get())
 #if PERF_TRACK_DETAILED_ASYNC_STATS
 , TickCount(0)
 , TickLoopCount(0)
@@ -832,7 +882,7 @@ void FAsyncPackage::ResetLoader()
  */
 bool FAsyncPackage::IsTimeLimitExceeded()
 {
-	return ::IsTimeLimitExceeded(TickStartTime, bUseTimeLimit, TimeLimit, LastTypeOfWorkPerformed, LastObjectWorkWasPerformedOn);
+	return AsyncLoadingThread.IsAsyncLoadingSuspended() || ::IsTimeLimitExceeded(TickStartTime, bUseTimeLimit, TimeLimit, LastTypeOfWorkPerformed, LastObjectWorkWasPerformedOn);
 }
 
 /**
@@ -992,7 +1042,7 @@ EAsyncPackageState::Type FAsyncPackage::Tick(bool InbUseTimeLimit, bool InbUseFu
 		}
 	} while (!IsTimeLimitExceeded() && LoadingState == EAsyncPackageState::TimeOut);
 
-	check(bUseTimeLimit || LoadingState != EAsyncPackageState::TimeOut);
+	check(bUseTimeLimit || LoadingState != EAsyncPackageState::TimeOut || AsyncLoadingThread.IsAsyncLoadingSuspended());
 
 	// We can't have a reference to a UObject.
 	LastObjectWorkWasPerformedOn = nullptr;
@@ -1532,7 +1582,9 @@ EAsyncPackageState::Type FAsyncPackage::PostLoadDeferredObjects(double InTickSta
 
 	TArray<UObject*>& ObjLoadedInPostLoad = FUObjectThreadContext::Get().ObjLoaded;
 
-	while (DeferredPostLoadIndex < DeferredPostLoadObjects.Num() && !::IsTimeLimitExceeded(InTickStartTime, bInUseTimeLimit, InOutTimeLimit, LastTypeOfWorkPerformed, LastObjectWorkWasPerformedOn))
+	while (DeferredPostLoadIndex < DeferredPostLoadObjects.Num() && 
+		!AsyncLoadingThread.IsAsyncLoadingSuspended() &&
+		!::IsTimeLimitExceeded(InTickStartTime, bInUseTimeLimit, InOutTimeLimit, LastTypeOfWorkPerformed, LastObjectWorkWasPerformedOn))
 	{
 		UObject* Object = DeferredPostLoadObjects[DeferredPostLoadIndex++];
 		check(Object);
@@ -1731,7 +1783,7 @@ void LoadPackageAsync(const FString& InName, const FGuid* InGuid /*= nullptr*/,
 		FString ClassName;
 		if (!FPackageName::ParseExportTextPath(PackageName, &ClassName, &PackageName))
 		{
-			UE_LOG(LogStreaming, Fatal, TEXT("LoadPackageAsync failed to begin to load a level because the supplied package name was neither a valid long package name nor a filename of a map within a content folder: '%s'"), *PackageName);
+			UE_LOG(LogStreaming, Fatal, TEXT("LoadPackageAsync failed to begin to load a package because the supplied package name was neither a valid long package name nor a filename of a map within a content folder: '%s'"), *PackageName);
 		}
 	}
 
@@ -1760,6 +1812,9 @@ void LoadPackageAsync(const FString& PackageName, FLoadPackageAsyncDelegate Comp
 
 void CancelAsyncLoading()
 {
+	// Cancelling async loading while loading is suspend will result in infinite stall
+	UE_CLOG(FAsyncLoadingThread::Get().IsAsyncLoadingSuspended(), LogStreaming, Fatal, TEXT("Cannot Cancel Async Loading while async loading is suspended."));
+
 	FAsyncLoadingThread::Get().CancelAsyncLoading();
 	CollectGarbage(GARBAGE_COLLECTION_KEEPFLAGS, false);
 }
@@ -1779,15 +1834,13 @@ bool IsInAsyncLoadingThreadCoreUObjectInternal()
 	return FAsyncLoadingThread::Get().IsInAsyncLoadThread();
 }
 
-/**
- * Blocks till all pending package/ linker requests are fulfilled.
- *
- * @param	ExcludeType					Do not flush packages associated with this specific type name
- */
 void FlushAsyncLoading(FName ExcludeType/*=NAME_None*/)
 {
 	if (IsAsyncLoading())
 	{
+		// Flushing async loading while loading is suspend will result in infinite stall
+		UE_CLOG(FAsyncLoadingThread::Get().IsAsyncLoadingSuspended(), LogStreaming, Fatal, TEXT("Cannot Flush Async Loading while async loading is suspended."));
+
 		SCOPE_CYCLE_COUNTER(STAT_FAsyncPackage_FlushAsyncLoadingGameThread);
 
 		// Disallow low priority requests like texture streaming while we are flushing streaming
@@ -1819,23 +1872,11 @@ void FlushAsyncLoading(FName ExcludeType/*=NAME_None*/)
 	}
 }
 
-/**
- * @return number of active async load package requests
- */
 int32 GetNumAsyncPackages()
 {
 	return FAsyncLoadingThread::Get().GetAsyncPackagesCount();
 }
 
-/**
- * Serializes a bit of data each frame with a soft time limit. The function is designed to be able
- * to fully load a package in a single pass given sufficient time.
- *
- * @param	bUseTimeLimit	Whether to use a time limit
- * @param	bUseFullTimeLimit	If true, use the entire time limit even if blocked on I/O
- * @param	TimeLimit		Soft limit of time this function is allowed to consume
- * @param	ExcludeType		Do not process packages associated with this specific type name
- */
 EAsyncPackageState::Type ProcessAsyncLoading(bool bUseTimeLimit, bool bUseFullTimeLimit, float TimeLimit, FName ExcludeType)
 {
 	SCOPE_CYCLE_COUNTER(STAT_AsyncLoadingTime);
@@ -1854,6 +1895,16 @@ bool IsAsyncLoadingCoreUObjectInternal()
 	return !GIsInitialLoad && FAsyncLoadingThread::Get().IsAsyncLoadingPackages();
 }
 
+void SuspendAsyncLoadingInternal()
+{
+	FAsyncLoadingThread::Get().SuspendLoading();
+}
+
+void ResumeAsyncLoadingInternal()
+{
+	FAsyncLoadingThread::Get().ResumeLoading();
+}
+
 /*----------------------------------------------------------------------------
 	FArchiveAsync.
 ----------------------------------------------------------------------------*/
diff --git a/Engine/Source/Runtime/CoreUObject/Private/Serialization/AsyncLoadingThread.h b/Engine/Source/Runtime/CoreUObject/Private/Serialization/AsyncLoadingThread.h
--- a/Engine/Source/Runtime/CoreUObject/Private/Serialization/AsyncLoadingThread.h
+++ b/Engine/Source/Runtime/CoreUObject/Private/Serialization/AsyncLoadingThread.h
@@ -21,6 +21,8 @@ class FAsyncLoadingThread : public FRunnable
 	FEvent* QueuedRequestsEvent;
 	/** [ASYNC/GAME THREAD] Event used to signal loading should be cancelled */
 	FEvent* CancelLoadingEvent;
+	/** [ASYNC/GAME THREAD] Event used to signal that the async loading thread should be suspended */
+	FEvent* ThreadSuspendedEvent;
 	/** [ASYNC/GAME THREAD] List of queued packages to stream */
 	TArray<FAsyncPackageDesc*> QueuedPackages;
 #if THREADSAFE_UOBJECTS
@@ -29,6 +31,8 @@ class FAsyncLoadingThread : public FRunnable
 #endif
 	/** [ASYNC/GAME THREAD] True if the async loading thread received a request to cancel async loading **/
 	FThreadSafeBool bShouldCancelLoading;
+	/** [ASYNC/GAME THREAD] True if the async loading thread received a request to suspend **/
+	FThreadSafeCounter IsLoadingSuspended;
 	/** [ASYNC/GAME THREAD] Event used to signal there's queued packages to stream */
 	TArray<FAsyncPackage*> LoadedPackages;	
 #if THREADSAFE_UOBJECTS
@@ -157,6 +161,13 @@ public:
 		return bResult;
 	}
 
+	/** Returns true if async loading is suspended */
+	FORCEINLINE bool IsAsyncLoadingSuspended()
+	{
+		FPlatformMisc::MemoryBarrier();
+		return IsLoadingSuspended.GetValue() != 0;
+	}
+
 	/** Returns the number of async packages that are currently being processed */
 	FORCEINLINE int32 GetAsyncPackagesCount()
 	{
@@ -207,6 +218,16 @@ public:
 	*/
 	void CancelAsyncLoading();
 
+	/**
+	* [GAME THREAD] Suspends async loading thread
+	*/
+	void SuspendLoading();
+
+	/**
+	* [GAME THREAD] Resumes async loading thread
+	*/
+	void ResumeLoading();
+
 	/**
 	* [ASYNC/GAME THREAD] Queues a package for streaming.
 	*
@@ -248,7 +269,7 @@ public:
 	* @param bUseFullTimeLimit True if full time limit should be used [time-slicing].
 	* @param TimeLimit Maximum amount of time that can be spent in this call [time-slicing].
 	*/
-	void TickAsyncThread(bool bUseTimeLimit, bool bUseFullTimeLimit, float TimeLimit);
+	EAsyncPackageState::Type TickAsyncThread(bool bUseTimeLimit, bool bUseFullTimeLimit, float TimeLimit);
 
 	/** Initializes async loading thread */
 	void InitializeAsyncThread();
diff --git a/Engine/Source/Runtime/CoreUObject/Private/UObject/CoreNative.cpp b/Engine/Source/Runtime/CoreUObject/Private/UObject/CoreNative.cpp
--- a/Engine/Source/Runtime/CoreUObject/Private/UObject/CoreNative.cpp
+++ b/Engine/Source/Runtime/CoreUObject/Private/UObject/CoreNative.cpp
@@ -9,6 +9,8 @@
 void UClassRegisterAllCompiledInClasses();
 bool IsInAsyncLoadingThreadCoreUObjectInternal();
 bool IsAsyncLoadingCoreUObjectInternal();
+void SuspendAsyncLoadingInternal();
+void ResumeAsyncLoadingInternal();
 
 // CoreUObject module. Handles UObject system pre-init (registers init function with Core callbacks).
 class FCoreUObjectModule : public FDefaultModuleImpl
@@ -25,6 +27,8 @@ public:
 		// Substitute Core version of async loading functions with CoreUObject ones.
 		IsInAsyncLoadingThread = &IsInAsyncLoadingThreadCoreUObjectInternal;
 		IsAsyncLoading = &IsAsyncLoadingCoreUObjectInternal;
+		SuspendAsyncLoading = &SuspendAsyncLoadingInternal;
+		ResumeAsyncLoading = &ResumeAsyncLoadingInternal;
 
 		// Make sure that additional content mount points can be registered after CoreUObject loads
 		FPackageName::EnsureContentPathsAreRegistered();		
diff --git a/Engine/Source/Runtime/CoreUObject/Public/Serialization/AsyncLoading.h b/Engine/Source/Runtime/CoreUObject/Public/Serialization/AsyncLoading.h
--- a/Engine/Source/Runtime/CoreUObject/Public/Serialization/AsyncLoading.h
+++ b/Engine/Source/Runtime/CoreUObject/Public/Serialization/AsyncLoading.h
@@ -163,6 +163,8 @@ private:
 	TArray<UObject*> DeferredPostLoadObjects;
 	/** Objects to be finalized on the game thread */
 	TArray<UObject*> DeferredFinalizeObjects;
+	/** Cached async loading thread object this package was created by */
+	class FAsyncLoadingThread& AsyncLoadingThread;
 public:
 #if PERF_TRACK_DETAILED_ASYNC_STATS
 	/** Number of times Tick function has been called.													*/
diff --git a/Engine/Source/Runtime/RenderCore/Private/RenderingThread.cpp b/Engine/Source/Runtime/RenderCore/Private/RenderingThread.cpp
--- a/Engine/Source/Runtime/RenderCore/Private/RenderingThread.cpp
+++ b/Engine/Source/Runtime/RenderCore/Private/RenderingThread.cpp
@@ -59,8 +59,8 @@ static void WaitAndResumeRendering()
 		FPlatformProcess::Sleep( 0.001f ); //@todo this should be a more principled wait
 	}
     
-    // set the thread back to real time mode
-    FPlatformProcess::SetRealTimeMode();
+	// set the thread back to real time mode
+	FPlatformProcess::SetRealTimeMode();
 }
 
 /**
@@ -69,6 +69,10 @@ static void WaitAndResumeRendering()
  */
 FSuspendRenderingThread::FSuspendRenderingThread( bool bInRecreateThread )
 {
+	// Suspend async loading thread so that it doesn't start queueing render commands 
+	// while the render thread is suspended.
+	SuspendAsyncLoading();
+
 	bRecreateThread = bInRecreateThread;
 	bUseRenderingThread = GUseThreadedRendering;
 	bWasRenderingThreadRunning = GIsThreadedRendering;
@@ -169,6 +173,7 @@ FSuspendRenderingThread::~FSuspendRenderingThread()
 		// Resume the render thread again. 
 		FPlatformAtomics::InterlockedDecrement( &GIsRenderingThreadSuspended );
 	}
+	ResumeAsyncLoading();
 }
 
 
