Move checking if CreateSubobject* functions are called from constructor from UObject::CreateDefaultSubobject to FObjectInitializer::CreateDefaultSubobject.

diff --git a/Engine/Source/Runtime/CoreUObject/Private/UObject/Obj.cpp b/Engine/Source/Runtime/CoreUObject/Private/UObject/Obj.cpp
--- a/Engine/Source/Runtime/CoreUObject/Private/UObject/Obj.cpp
+++ b/Engine/Source/Runtime/CoreUObject/Private/UObject/Obj.cpp
@@ -79,8 +79,6 @@ void UObject::EnsureNotRetrievingVTablePtr() const
 
 UObject* UObject::CreateDefaultSubobject(FName SubobjectFName, UClass* ReturnType, UClass* ClassToCreateByDefault, bool bIsRequired, bool bAbstract, bool bIsTransient)
 {
-	auto& ThreadContext = FUObjectThreadContext::Get();
-	UE_CLOG(!ThreadContext.IsInConstructor, LogObj, Fatal, TEXT("CreateDefultSubobject can only be used inside of UObject constructors. UObject constructing subobjects cannot be created using new or placement new operator."));
 	auto CurrentInitializer = FTlsObjectInitializers::Top();
 	UE_CLOG(!CurrentInitializer, LogObj, Fatal, TEXT("No object initializer found during construction."));
 	UE_CLOG(CurrentInitializer->Obj != this, LogObj, Fatal, TEXT("Using incorrect object initializer."));
diff --git a/Engine/Source/Runtime/CoreUObject/Private/UObject/UObjectGlobals.cpp b/Engine/Source/Runtime/CoreUObject/Private/UObject/UObjectGlobals.cpp
--- a/Engine/Source/Runtime/CoreUObject/Private/UObject/UObjectGlobals.cpp
+++ b/Engine/Source/Runtime/CoreUObject/Private/UObject/UObjectGlobals.cpp
@@ -2932,6 +2932,7 @@ UScriptStruct* GetFallbackStruct()
 
 UObject* FObjectInitializer::CreateDefaultSubobject(UObject* Outer, FName SubobjectFName, UClass* ReturnType, UClass* ClassToCreateByDefault, bool bIsRequired, bool bAbstract, bool bIsTransient) const
 {
+	UE_CLOG(!FUObjectThreadContext::Get().IsInConstructor, LogClass, Fatal, TEXT("Subobjects cannot be created outside of UObject constructors. UObject constructing subobjects cannot be created using new or placement new operator."));
 	if (SubobjectFName == NAME_None)
 	{
 		UE_LOG(LogClass, Fatal, TEXT("Illegal default subobject name: %s"), *SubobjectFName.ToString());
