Fixed TStructOpsTypeTraits<FTwoVectors> to prevent crash when setting defaults in Distribution Vector Uniform Curve objects.

diff --git a/Engine/Source/Runtime/CoreUObject/Private/UObject/Property.cpp b/Engine/Source/Runtime/CoreUObject/Private/UObject/Property.cpp
--- a/Engine/Source/Runtime/CoreUObject/Private/UObject/Property.cpp
+++ b/Engine/Source/Runtime/CoreUObject/Private/UObject/Property.cpp
@@ -183,6 +183,7 @@ struct TStructOpsTypeTraits<FTwoVectors> : public TStructOpsTypeTraitsBase
 	enum 
 	{
 		WithZeroConstructor = true,
+		WithNoDestructor = true,
 	};
 };
 IMPLEMENT_STRUCT(TwoVectors);
