Changing verbosity of some async loading logs.

diff --git a/Engine/Source/Runtime/CoreUObject/Private/Serialization/AsyncLoading.cpp b/Engine/Source/Runtime/CoreUObject/Private/Serialization/AsyncLoading.cpp
--- a/Engine/Source/Runtime/CoreUObject/Private/Serialization/AsyncLoading.cpp
+++ b/Engine/Source/Runtime/CoreUObject/Private/Serialization/AsyncLoading.cpp
@@ -599,12 +599,12 @@ FAsyncLoadingThread::FAsyncLoadingThread()
 	ThreadSuspendedEvent = FPlatformProcess::GetSynchEventFromPool();
 	if (FAsyncLoadingThread::IsMultithreaded())
 	{
-		UE_LOG(LogStreaming, Display, TEXT("Async loading is multithreaded."));
+		UE_LOG(LogStreaming, Log, TEXT("Async loading is multithreaded."));
 		Thread = FRunnableThread::Create(this, TEXT("FAsyncLoadingThread"), 0, TPri_Normal);
 	}
 	else
 	{
-		UE_LOG(LogStreaming, Display, TEXT("Async loading is time-sliced."));
+		UE_LOG(LogStreaming, Log, TEXT("Async loading is time-sliced."));
 		Thread = nullptr;
 		Init();
 	}
