Enum properties (PropertyByte) will export their enum as text for the PropertyWindow instead of the BlueprintWatch window. This allows more use cases to display the property using localized/display text.

Switched a case in BehaviorTrees to export the property using PPF_PropertyWindow to get the localized display string.

diff --git a/Engine/Source/Runtime/AIModule/Private/BehaviorTree/BlueprintNodeHelpers.cpp b/Engine/Source/Runtime/AIModule/Private/BehaviorTree/BlueprintNodeHelpers.cpp
--- a/Engine/Source/Runtime/AIModule/Private/BehaviorTree/BlueprintNodeHelpers.cpp
+++ b/Engine/Source/Runtime/AIModule/Private/BehaviorTree/BlueprintNodeHelpers.cpp
@@ -96,7 +96,7 @@ namespace BlueprintNodeHelpers
 		}
 		else
 		{
-			Prop->ExportTextItem(ExportedStringValue, PropertyAddr, NULL, NULL, 0, NULL);
+			Prop->ExportTextItem(ExportedStringValue, PropertyAddr, NULL, NULL, PPF_PropertyWindow, NULL);
 		}
 
 		const bool bIsBool = Prop->IsA(UBoolProperty::StaticClass());
diff --git a/Engine/Source/Runtime/CoreUObject/Private/UObject/PropertyByte.cpp b/Engine/Source/Runtime/CoreUObject/Private/UObject/PropertyByte.cpp
--- a/Engine/Source/Runtime/CoreUObject/Private/UObject/PropertyByte.cpp
+++ b/Engine/Source/Runtime/CoreUObject/Private/UObject/PropertyByte.cpp
@@ -115,7 +115,8 @@ void UByteProperty::ExportTextItem( FString& ValueStr, const void* PropertyValue
 		// the property text value must actually match an entry in the enum's names array)
 		if (*(const uint8*)PropertyValue < Enum->NumEnums() - 1 || ((PortFlags&PPF_Copy) != 0 && *(const uint8*)PropertyValue < Enum->NumEnums()))
 		{
-			if (PortFlags & PPF_BlueprintDebugView)
+			// We do not want to export the enum text for non-display uses, localization text is very dynamic and would cause issues on import
+			if (PortFlags & PPF_PropertyWindow)
 			{
 				ValueStr += Enum->GetEnumText(*(const uint8*)PropertyValue).ToString();
 			}
