AsyncLoading: removing redundant code

diff --git a/Engine/Source/Runtime/CoreUObject/Private/Serialization/AsyncLoading.cpp b/Engine/Source/Runtime/CoreUObject/Private/Serialization/AsyncLoading.cpp
--- a/Engine/Source/Runtime/CoreUObject/Private/Serialization/AsyncLoading.cpp
+++ b/Engine/Source/Runtime/CoreUObject/Private/Serialization/AsyncLoading.cpp
@@ -781,26 +781,6 @@ void NotifyConstructedDuringAsyncLoading(UObject* Object, bool bSubObject)
 	FAsyncObjectsReferencer::Get().AddObject(Object);
 }
 
-
-class FAsyncLoadingThreadManager : private FSelfRegisteringExec
-{
-public:
-	// FSelfRegisteringExec interface
-	virtual bool Exec(class UWorld* InWorld, const TCHAR* Cmd, FOutputDevice& Ar) override
-	{
-		if (FParse::Command(&Cmd, TEXT("ASYNCCREATE")))
-		{
-			return true;
-		}
-		return false;
-	}
-
-private:
-
-};
-static FAsyncLoadingThreadManager GAsyncLoaderManager;
-
-
 /*-----------------------------------------------------------------------------
 	FAsyncPackage implementation.
 -----------------------------------------------------------------------------*/
