fix for UE-16290, strategy game hang on launch on for #html5

     -  IO was getting starved in single threaded mode.

diff --git a/Engine/Source/Runtime/CoreUObject/Private/Serialization/AsyncLoading.cpp b/Engine/Source/Runtime/CoreUObject/Private/Serialization/AsyncLoading.cpp
--- a/Engine/Source/Runtime/CoreUObject/Private/Serialization/AsyncLoading.cpp
+++ b/Engine/Source/Runtime/CoreUObject/Private/Serialization/AsyncLoading.cpp
@@ -872,6 +872,11 @@ bool FAsyncPackage::IsTimeLimitExceeded()
  */
 bool FAsyncPackage::GiveUpTimeSlice()
 {
+	if (!FPlatformProcess::SupportsMultithreading())
+	{
+		FIOSystem::Get().TickSingleThreaded();
+	}
+
 	if (bUseTimeLimit && !bUseFullTimeLimit)
 	{
 		bTimeLimitExceeded = true;
