Don't add Linker imports to GC (potential fix for GC perf regression)

diff --git a/Engine/Source/Runtime/CoreUObject/Private/UObject/Linker.cpp b/Engine/Source/Runtime/CoreUObject/Private/UObject/Linker.cpp
--- a/Engine/Source/Runtime/CoreUObject/Private/UObject/Linker.cpp
+++ b/Engine/Source/Runtime/CoreUObject/Private/UObject/Linker.cpp
@@ -160,10 +160,6 @@ void FLinker::AddReferencedObjects(FReferenceCollector& Collector)
 		Collector.AddReferencedObject(*(UObject**)&LinkerRoot);
 	}
 #endif
-	for (auto& Import : ImportMap)
-	{
-		Collector.AddReferencedObject(Import.XObject);
-	}
 }
 
 // FLinker interface.
