UE-14720 After context fail (Access None) some data dependent on the context aren't clear

diff --git a/Engine/Source/Editor/KismetCompiler/Private/KismetCompilerVMBackend.cpp b/Engine/Source/Editor/KismetCompiler/Private/KismetCompilerVMBackend.cpp
--- a/Engine/Source/Editor/KismetCompiler/Private/KismetCompilerVMBackend.cpp
+++ b/Engine/Source/Editor/KismetCompiler/Private/KismetCompilerVMBackend.cpp
@@ -241,15 +241,15 @@ protected:
 		}
 
 		/** Starts a context if the Term isn't NULL */
-		void TryStartContext(FBPTerminal* Term, bool bUnsafeToSkip = false, bool bIsInterfaceContext = false, FBPTerminal* RValueTerm = NULL)
+		void TryStartContext(FBPTerminal* Term, bool bUnsafeToSkip = false, bool bIsInterfaceContext = false, UProperty* RValueProperty = nullptr)
 		{
 			if (Term != NULL)
 			{
-				StartContext(Term, bUnsafeToSkip, bIsInterfaceContext, RValueTerm);
+				StartContext(Term, bUnsafeToSkip, bIsInterfaceContext, RValueProperty);
 			}
 		}
 
-		void StartContext(FBPTerminal* Term, bool bUnsafeToSkip = false, bool bIsInterfaceContext = false, FBPTerminal* RValueTerm = NULL)
+		void StartContext(FBPTerminal* Term, bool bUnsafeToSkip = false, bool bIsInterfaceContext = false, UProperty* RValueProperty = nullptr)
 		{
 			bInContext = true;
 
@@ -277,15 +277,8 @@ protected:
 			Skipper.Emit();
 
 			// R-Value property, see ReadVariableSize in UObject::ProcessContextOpcode() for usage
-			//@TODO: Not sure what to use for yet
-			UProperty* RValueProperty = RValueTerm ? RValueTerm->AssociatedVarProperty : NULL;
 			Writer << RValueProperty;
 
-			// Property type if needed, it seems it's not used in ue4
-			//@TODO: Not sure what to use for yet
-			uint8 PropetyType = 0;
-			Writer << PropetyType;
-
 			// Context expression (this is the part that gets skipped if the object turns out NULL)
 			Skipper.BeginCounting();
 
@@ -748,7 +741,11 @@ public:
 
 		// Handle the function calling context if needed
 		FContextEmitter CallContextWriter(*this);
-		CallContextWriter.TryStartContext(Statement.FunctionContext, /*bUnsafeToSkip=*/ bHasOutputValue, Statement.bIsInterfaceContext);
+
+		// RValue property is used to clear value after Access Violation. See UObject::ProcessContextOpcod
+		// If the property from LHS is used, then the retured property (with CPF_ReturnParm) is cleared. But properties returned by ref are not cleared. 
+		UProperty* RValueProperty = nullptr; //Statement.LHS ? Statement.LHS->AssociatedVarProperty : nullptr;
+		CallContextWriter.TryStartContext(Statement.FunctionContext, /*bUnsafeToSkip=*/ bHasOutputValue, Statement.bIsInterfaceContext, RValueProperty);
 
 		// Emit the call type
 		if (FunctionToCall->HasAnyFunctionFlags(FUNC_Delegate))
@@ -871,7 +868,8 @@ public:
 				}
 
  				FContextEmitter CallContextWriter(*this);
-				CallContextWriter.TryStartContext(Term->Context, /*@TODO: bUnsafeToSkip*/ true, /*bIsInterfaceContext*/ false, RValueTerm);
+				UProperty* RValueProperty = RValueTerm ? RValueTerm->AssociatedVarProperty : nullptr;
+				CallContextWriter.TryStartContext(Term->Context, /*@TODO: bUnsafeToSkip*/ true, /*bIsInterfaceContext*/ false, RValueProperty);
 
 				EmitTermExpr(Term, CoerceProperty);
 			}
diff --git a/Engine/Source/Editor/UnrealEd/Private/ScriptDisassembler.cpp b/Engine/Source/Editor/UnrealEd/Private/ScriptDisassembler.cpp
--- a/Engine/Source/Editor/UnrealEd/Private/ScriptDisassembler.cpp
+++ b/Engine/Source/Editor/UnrealEd/Private/ScriptDisassembler.cpp
@@ -534,10 +534,6 @@ void FKismetBytecodeDisassembler::ProcessCommon(int32& ScriptIndex, EExprToken O
 			UField* Field = ReadPointer<UField>(ScriptIndex);
 			Ar.Logf(TEXT("%s R-Value Property: %s"), *Indents, Field ? *Field->GetName() : TEXT("(null)"));
 
-			// Property type, in case the r-value is a non-property such as dynamic array length
-			uint8 PropType = ReadBYTE(ScriptIndex);
-			Ar.Logf(TEXT("%s PropertyTypeIfNeeded: %d"), *Indents, PropType);
-
 			// Context expression.
 			Ar.Logf(TEXT("%s ContextExpression:"), *Indents);
 			SerializeExpr( ScriptIndex );
diff --git a/Engine/Source/Runtime/CoreUObject/Private/UObject/ScriptCore.cpp b/Engine/Source/Runtime/CoreUObject/Private/UObject/ScriptCore.cpp
--- a/Engine/Source/Runtime/CoreUObject/Private/UObject/ScriptCore.cpp
+++ b/Engine/Source/Runtime/CoreUObject/Private/UObject/ScriptCore.cpp
@@ -638,6 +638,7 @@ void UObject::ProcessInternal( FFrame& Stack, RESULT_DECL )
 
 	if (FunctionCallspace & FunctionCallspace::Local)
 	{
+		// No POD struct can ever be stored in this buffer. 
 		MS_ALIGN(16) uint8 Buffer[MAX_SIMPLE_RETURN_VALUE_SIZE] GCC_ALIGN(16);
 
 #if DO_GUARD
@@ -1015,15 +1016,17 @@ void UObject::ProcessEvent( UFunction* Function, void* Parms )
 		}
 
 		// Call native function or UObject::ProcessInternal.
+		const bool bHasReturnParam = Function->ReturnValueOffset != MAX_uint16;
+		uint8* ReturnValueAdress = bHasReturnParam ? ((uint8*)Parms + Function->ReturnValueOffset) : nullptr;
 		if (Function->FunctionFlags & FUNC_Native)
 		{
 			FScopeCycleCounterUObject ContextScope(this);
 			FScopeCycleCounterUObject FunctionScope(Function);
-			Function->Invoke(this, NewStack, (uint8*)Parms + Function->ReturnValueOffset);
+			Function->Invoke(this, NewStack, ReturnValueAdress);
 		}
 		else
 		{
-			Function->Invoke(this, NewStack, (uint8*)Parms + Function->ReturnValueOffset);
+			Function->Invoke(this, NewStack, ReturnValueAdress);
 		}
 
 		if (!bUsePersistentFrame)
@@ -1353,11 +1356,11 @@ void UObject::execLet( FFrame& Stack, RESULT_DECL )
 		static FBlueprintExceptionInfo ExceptionInfo(EBlueprintExceptionType::AccessViolation, TEXT("Attempt to assign variable through None"));
 		FBlueprintCoreDelegates::ThrowScriptException(this, Stack, ExceptionInfo);
 
-		//@TODO: ScriptParallel: Contended static usage
-		static uint8 Crud[1024];//@temp
-		Stack.MostRecentPropertyAddress = Crud;
-		FMemory::Memzero( Stack.MostRecentPropertyAddress, sizeof(FString) );
-	}
+			//@TODO: ScriptParallel: Contended static usage
+			static uint8 Crud[1024];//@temp
+			Stack.MostRecentPropertyAddress = Crud;
+			FMemory::Memzero(Stack.MostRecentPropertyAddress, sizeof(FString));
+		}
 
 	// Evaluate expression into variable.
 	Stack.Step( Stack.Object, Stack.MostRecentPropertyAddress );
@@ -1548,8 +1551,7 @@ void UObject::ProcessContextOpcode( FFrame& Stack, RESULT_DECL, bool bCanFailSil
 	if (IsValid(NewContext))
 	{
 		Stack.Code += sizeof(CodeSkipSizeType)	// Code offset for NULL expressions.
-			+ sizeof(ScriptPointerType)			// Property corresponding to the r-value data, in case the l-value needs to be cleared
-			+ sizeof(uint8);					// Property type, in case the r-value is a non-property - in ue4 it seems to be unused
+			+ sizeof(ScriptPointerType);		// Property corresponding to the r-value data, in case the l-value needs to be cleared
 		Stack.Step( NewContext, RESULT_PARAM );
 	}
 	else
@@ -1580,20 +1582,15 @@ void UObject::ProcessContextOpcode( FFrame& Stack, RESULT_DECL, bool bCanFailSil
 		}
 
 		const CodeSkipSizeType wSkip = Stack.ReadCodeSkipCount(); // Code offset for NULL expressions. Code += sizeof(CodeSkipSizeType)
-		UField* RValueField = nullptr;
-		const VariableSizeType bSize = Stack.ReadVariableSize(&RValueField); // Code += sizeof(ScriptPointerType) + sizeof(uint8)
+		UProperty* RValueProperty = nullptr;
+		const VariableSizeType bSize = Stack.ReadVariableSize(&RValueProperty); // Code += sizeof(ScriptPointerType) + sizeof(uint8)
 		Stack.Code += wSkip;
 		Stack.MostRecentPropertyAddress = NULL;
 		Stack.MostRecentProperty = NULL;
 
-		if (RESULT_PARAM)
+		if (RESULT_PARAM && RValueProperty)
 		{
-			auto RValueProperty = Cast<const UProperty>(RValueField);
-			ensure(RValueProperty || !RValueField);
-			if (RValueProperty)
-			{
-				RValueProperty->ClearValue(RESULT_PARAM);
-			}
+			RValueProperty->ClearValue(RESULT_PARAM);
 		}
 	}
 }
diff --git a/Engine/Source/Runtime/CoreUObject/Public/UObject/ScriptSerialization.h b/Engine/Source/Runtime/CoreUObject/Public/UObject/ScriptSerialization.h
--- a/Engine/Source/Runtime/CoreUObject/Public/UObject/ScriptSerialization.h
+++ b/Engine/Source/Runtime/CoreUObject/Public/UObject/ScriptSerialization.h
@@ -241,7 +241,6 @@
 			SerializeExpr( iCode, Ar ); // Object expression.
 			XFER(CodeSkipSizeType);		// Code offset for NULL expressions.
 			XFERPTR(UField*);			// Property corresponding to the r-value data, in case the l-value needs to be mem-zero'd
-			XFER(uint8);					// Property type, in case the r-value is a non-property such as dynamic array length
 			SerializeExpr( iCode, Ar ); // Context expression.
 			break;
 		}
diff --git a/Engine/Source/Runtime/CoreUObject/Public/UObject/Stack.h b/Engine/Source/Runtime/CoreUObject/Public/UObject/Stack.h
--- a/Engine/Source/Runtime/CoreUObject/Public/UObject/Stack.h
+++ b/Engine/Source/Runtime/CoreUObject/Public/UObject/Stack.h
@@ -156,7 +156,7 @@ public:
 	 * @param	ExpressionField		receives a pointer to the field representing the expression; used by various execs
 	 *								to drive VM logic
 	 */
-	VariableSizeType ReadVariableSize( UField** ExpressionField=NULL );
+	VariableSizeType ReadVariableSize(UProperty** ExpressionField);
 
 	/**
  	 * This will return the StackTrace of the current callstack from script land
@@ -271,95 +271,20 @@ inline CodeSkipSizeType FFrame::ReadCodeSkipCount()
 	return Result;
 }
 
-inline VariableSizeType FFrame::ReadVariableSize( UField** ExpressionField/*=NULL*/ )
+inline VariableSizeType FFrame::ReadVariableSize( UProperty** ExpressionField )
 {
 	VariableSizeType Result=0;
 
 	UObject* Field = ReadObject();
-	uint8 NullPropertyType = *Code++;
-
-	if ( Field != NULL )
-	{
-		// currently only Property (or null) seems to be a valid option.
-		if (UProperty* Property = dynamic_cast<UProperty*>(Field))
-		{
-			Result = Property->GetSize();
-		}
-		else if (UEnum* ExplicitEnumValue = dynamic_cast<UEnum*>(Field))
-		{
-			Result = 1;
-		}
-		else if (UFunction* FunctionRef = dynamic_cast<UFunction*>(Field))
-		{
-			Result = sizeof(ScriptPointerType);
-		}
-	}
-	else
+	UProperty* Property = dynamic_cast<UProperty*>(Field);
+	if (Property)
 	{
-		switch ( NullPropertyType )
-		{
-		case CPT_None:
-			// nothing...
-			break;
-		case CPT_Byte: Result = sizeof(uint8);
-			break;
-		case CPT_UInt16: Result = sizeof(uint16);
-			break;
-		case CPT_UInt32: Result = sizeof(uint32);
-			break;
-		case CPT_UInt64: Result = sizeof(uint64);
-			break;
-		case CPT_Int8: Result = sizeof(int8);
-			break;
-		case CPT_Int16: Result = sizeof(int16);
-			break;
-		case CPT_Int: Result = sizeof(int32);
-			break;
-		case CPT_Int64: Result = sizeof(int64);
-			break;
-		case CPT_Bool: Result = sizeof(bool);
-			break;
-		case CPT_Bool8: Result = sizeof(uint8);
-			break;
-		case CPT_Bool16: Result = sizeof(uint16);
-			break;
-		case CPT_Bool32: Result = sizeof(uint32);
-			break;
-		case CPT_Bool64: Result = sizeof(uint64);
-			break;
-		case CPT_Float: Result = sizeof(float);
-			break;
-		case CPT_Double: Result = sizeof(double);
-			break;
-		case CPT_Name: Result = sizeof(FScriptName);
-			break;
-		case CPT_Vector: Result = sizeof(FVector);
-			break;
-		case CPT_Rotation: Result = sizeof(FRotator);
-			break;
-		case CPT_Delegate: Result = sizeof(FScriptDelegate);
-			break;
-		case CPT_MulticastDelegate: Result = sizeof(FMulticastScriptDelegate);
-			break;
-		case CPT_WeakObjectReference: Result = sizeof(FWeakObjectPtr);
-			break;
-		case CPT_LazyObjectReference: Result = sizeof(FLazyObjectPtr);
-			break;
-		case CPT_AssetObjectReference: Result = sizeof(FAssetPtr);
-			break;
-		case CPT_Text: Result = sizeof(FText);
-			break;
-		case CPT_Map: Result = sizeof(FScriptMap);
-			break;
-		default:
-			UE_LOG(LogScriptFrame, Fatal, TEXT("Unhandled property type in FFrame::ReadVariableSize(): %u"), NullPropertyType);
-			break;
-		}
+		Result = Property->GetSize();
 	}
 
-	if ( ExpressionField != NULL )
+	if ( ExpressionField != nullptr )
 	{
-		*ExpressionField = dynamic_cast<UField*>(Field);
+		*ExpressionField = Property;
 	}
 
 	return Result;
