Temporarily disabling parallel GC in the Mac editor builds

diff --git a/Engine/Source/Runtime/CoreUObject/Private/UObject/GarbageCollection.cpp b/Engine/Source/Runtime/CoreUObject/Private/UObject/GarbageCollection.cpp
--- a/Engine/Source/Runtime/CoreUObject/Private/UObject/GarbageCollection.cpp
+++ b/Engine/Source/Runtime/CoreUObject/Private/UObject/GarbageCollection.cpp
@@ -1171,7 +1171,7 @@ COREUOBJECT_API EditorPostReachabilityAnalysisCallbackType EditorPostReachabilit
 
 // Allow parralel GC to be overriden to single threaded via console command.
 static const auto CVarAllowParallelGC = 
-	IConsoleManager::Get().RegisterConsoleVariable( TEXT("AllowParallelGC"), 1, TEXT("Used to control parallel GC.") )->AsVariableInt();
+IConsoleManager::Get().RegisterConsoleVariable(TEXT("AllowParallelGC"), (!PLATFORM_MAC || !WITH_EDITORONLY_DATA) ? 1 : 0, TEXT("Used to control parallel GC."))->AsVariableInt();
 
 /** 
  * Deletes all unreferenced objects, keeping objects that have any of the passed in KeepFlags set
