Fix crash in UMapProperty

diff --git a/Engine/Source/Runtime/CoreUObject/Private/UObject/PropertyMap.cpp b/Engine/Source/Runtime/CoreUObject/Private/UObject/PropertyMap.cpp
--- a/Engine/Source/Runtime/CoreUObject/Private/UObject/PropertyMap.cpp
+++ b/Engine/Source/Runtime/CoreUObject/Private/UObject/PropertyMap.cpp
@@ -741,7 +741,7 @@ void UMapProperty::InstanceSubobjects(void* Data, void const* DefaultData, UObje
 
 				if (bInstancedValue)
 				{
-					ValueProp->InstanceSubobjects(PairPtr + MapLayout.ValueOffset, DefaultPairPtr + MapLayout.ValueOffset, Owner, InstanceGraph);
+					ValueProp->InstanceSubobjects(PairPtr + MapLayout.ValueOffset, DefaultPairPtr ? DefaultPairPtr + MapLayout.ValueOffset : nullptr, Owner, InstanceGraph);
 				}
 
 				--Num;
