Reverting bad merge

diff --git a/Engine/Source/Runtime/CoreUObject/Private/Serialization/AsyncLoading.cpp b/Engine/Source/Runtime/CoreUObject/Private/Serialization/AsyncLoading.cpp
--- a/Engine/Source/Runtime/CoreUObject/Private/Serialization/AsyncLoading.cpp
+++ b/Engine/Source/Runtime/CoreUObject/Private/Serialization/AsyncLoading.cpp
@@ -510,6 +510,14 @@ EAsyncPackageState::Type FAsyncLoadingThread::ProcessLoadedPackages(bool bUseTim
 			if (Result == EAsyncPackageState::Complete)
 			{
 				// Remove the package from the list before we trigger the callbacks, 
+				// this is to ensure we can re-enter FlushAsyncLoading from any of the callbacks
+				LoadedPackagesToProcess.RemoveAt(PackageIndex--);
+
+				// Incremented on the Async Thread, now decrement as we're done with this package				
+				const int32 NewAsyncLoadingCounterValue = AsyncLoadingCounter.Decrement();
+				UE_CLOG(NewAsyncLoadingCounterValue < 0, LogStreaming, Fatal, TEXT("AsyncLoadingCounter is negative, this means we loaded more packages then requested so there must be a bug in async loading code."));
+
+				// Call external callbacks
 				const bool bInternalCallbacks = false;
 				const EAsyncLoadingResult::Type LoadingResult = Package->HasLoadFailed() ? EAsyncLoadingResult::Failed : EAsyncLoadingResult::Succeeded;
 				Package->CallCompletionCallbacks(bInternalCallbacks, LoadingResult);
@@ -520,18 +528,8 @@ EAsyncPackageState::Type FAsyncLoadingThread::ProcessLoadedPackages(bool bUseTim
 					Package->ResetLoader();
 				}
 
-				// Incremented on the Async Thread, decremented on the Game Thread				
-				AsyncLoadingCounter.Decrement();
-				check(AsyncLoadingCounter.GetValue() >= 0);
-
-				if (FPlatformProperties::RequiresCookedData())
-				{
-					// Emulates ResetLoaders on the package linker's linkerroot.
-					Package->ResetLoader();
-				}
-
+				// We don't need the package anymore
 				delete Package;
-				LoadedPackagesToProcess.RemoveAt(PackageIndex--);
 			}
 			else
 			{
