FORCEINLINE'd FStringAssetReference's GetTypeHash()

diff --git a/Engine/Source/Runtime/CoreUObject/Public/Misc/StringAssetReference.h b/Engine/Source/Runtime/CoreUObject/Public/Misc/StringAssetReference.h
--- a/Engine/Source/Runtime/CoreUObject/Public/Misc/StringAssetReference.h
+++ b/Engine/Source/Runtime/CoreUObject/Public/Misc/StringAssetReference.h
@@ -93,7 +93,7 @@ struct COREUOBJECT_API FStringAssetReference
 	bool ImportTextItem( const TCHAR*& Buffer, int32 PortFlags, UObject* Parent, FOutputDevice* ErrorText );
 	bool SerializeFromMismatchedTag(struct FPropertyTag const& Tag, FArchive& Ar);
 
-	friend uint32 GetTypeHash(FStringAssetReference const& This)
+	FORCEINLINE friend uint32 GetTypeHash(FStringAssetReference const& This)
 	{
 		return GetTypeHash(This.AssetLongPathname);
 	}
