Added support for UObjectProperty, ULazyObjectProperty and UAssetObjectProperty as UMapProperty keys
Fixed the non-FObjectInitializer constructors of properties to correctly add CPF_HasGetValueTypeHash where applicable

diff --git a/Engine/Source/Runtime/CoreUObject/Public/UObject/UnrealType.h b/Engine/Source/Runtime/CoreUObject/Public/UObject/UnrealType.h
--- a/Engine/Source/Runtime/CoreUObject/Public/UObject/UnrealType.h
+++ b/Engine/Source/Runtime/CoreUObject/Public/UObject/UnrealType.h
@@ -1173,7 +1173,7 @@ public:
 	}
 
 	TProperty_Numeric(ECppProperty, int32 InOffset, uint64 InFlags)
-		: Super(FObjectInitializer::Get(), EC_CppProperty, InOffset, InFlags)
+		: Super(FObjectInitializer::Get(), EC_CppProperty, InOffset, InFlags | CPF_HasGetValueTypeHash)
 	{
 	}
 
@@ -1846,12 +1846,12 @@ class COREUOBJECT_API UObjectProperty : public TUObjectPropertyBase<UObject*>
 	DECLARE_CASTED_CLASS_INTRINSIC(UObjectProperty,TUObjectPropertyBase<UObject*>,0,CoreUObject,CASTCLASS_UObjectProperty)
 
 	UObjectProperty(ECppProperty, int32 InOffset, uint64 InFlags, UClass* InClass)
-		: TUObjectPropertyBase(FObjectInitializer::Get(), EC_CppProperty, InOffset, InFlags, InClass)
+		: TUObjectPropertyBase(FObjectInitializer::Get(), EC_CppProperty, InOffset, InFlags | CPF_HasGetValueTypeHash, InClass)
 	{
 	}
 
 	UObjectProperty( const FObjectInitializer& ObjectInitializer, ECppProperty, int32 InOffset, uint64 InFlags, UClass* InClass )
-	:	TUObjectPropertyBase( ObjectInitializer, EC_CppProperty, InOffset, InFlags, InClass )
+		: TUObjectPropertyBase(ObjectInitializer, EC_CppProperty, InOffset, InFlags | CPF_HasGetValueTypeHash, InClass)
 	{
 	}
 
@@ -1864,7 +1864,15 @@ class COREUOBJECT_API UObjectProperty : public TUObjectPropertyBase<UObject*>
 	// UProperty interface
 	virtual void SerializeItem( FArchive& Ar, void* Value, void const* Defaults ) const override;
 	virtual void EmitReferenceInfo(UClass& OwnerClass, int32 BaseOffset) override;
+
+private:
+	virtual uint32 GetValueTypeHashInternal(const void* Src) const override
+	{
+		return GetTypeHash(GetPropertyValue(Src));
+	}
+public:
 	// End of UProperty interface
+
 	// UObjectPropertyBase interface
 	virtual UObject* GetObjectPropertyValue(const void* PropertyValueAddress) const override
 	{
@@ -1921,12 +1929,12 @@ class COREUOBJECT_API ULazyObjectProperty : public TUObjectPropertyBase<FLazyObj
 	DECLARE_CASTED_CLASS_INTRINSIC(ULazyObjectProperty,TUObjectPropertyBase<FLazyObjectPtr>,0,CoreUObject,CASTCLASS_ULazyObjectProperty)
 
 	ULazyObjectProperty(ECppProperty, int32 InOffset, uint64 InFlags, UClass* InClass)
-		: TUObjectPropertyBase(FObjectInitializer::Get(), EC_CppProperty, InOffset, InFlags, InClass)
+		: TUObjectPropertyBase(FObjectInitializer::Get(), EC_CppProperty, InOffset, InFlags | CPF_HasGetValueTypeHash, InClass)
 	{
 	}
 
 	ULazyObjectProperty( const FObjectInitializer& ObjectInitializer, ECppProperty, int32 InOffset, uint64 InFlags, UClass* InClass )
-		:	TUObjectPropertyBase( ObjectInitializer, EC_CppProperty, InOffset, InFlags, InClass )
+		: TUObjectPropertyBase(ObjectInitializer, EC_CppProperty, InOffset, InFlags | CPF_HasGetValueTypeHash, InClass)
 	{
 	}
 
@@ -1954,6 +1962,12 @@ class COREUOBJECT_API ULazyObjectProperty : public TUObjectPropertyBase<FLazyObj
 	{
 		return true;
 	}
+private:
+	virtual uint32 GetValueTypeHashInternal(const void* Src) const override
+	{
+		return GetTypeHash(GetPropertyValue(Src));
+	}
+public:
 	// End of UObjectProperty interface
 };
 
@@ -1965,11 +1979,11 @@ class COREUOBJECT_API UAssetObjectProperty : public TUObjectPropertyBase<FAssetP
 	DECLARE_CASTED_CLASS_INTRINSIC(UAssetObjectProperty,TUObjectPropertyBase<FAssetPtr>,0,CoreUObject,CASTCLASS_UAssetObjectProperty)
 
 	UAssetObjectProperty(ECppProperty, int32 InOffset, uint64 InFlags, UClass* InClass)
-		: TUObjectPropertyBase(FObjectInitializer::Get(), EC_CppProperty, InOffset, InFlags, InClass)
+		: TUObjectPropertyBase(FObjectInitializer::Get(), EC_CppProperty, InOffset, InFlags | CPF_HasGetValueTypeHash, InClass)
 	{}
 
 	UAssetObjectProperty( const FObjectInitializer& ObjectInitializer, ECppProperty, int32 InOffset, uint64 InFlags, UClass* InClass )
-		:	TUObjectPropertyBase( ObjectInitializer, EC_CppProperty, InOffset, InFlags, InClass )
+		: TUObjectPropertyBase(ObjectInitializer, EC_CppProperty, InOffset, InFlags | CPF_HasGetValueTypeHash, InClass)
 	{}
 
 	// UHT interface
@@ -1999,6 +2013,12 @@ class COREUOBJECT_API UAssetObjectProperty : public TUObjectPropertyBase<FAssetP
 	{
 		return true;
 	}
+private:
+	virtual uint32 GetValueTypeHashInternal(const void* Src) const override
+	{
+		return GetTypeHash(GetPropertyValue(Src));
+	}
+public:
 
 	// ScriptVM should store Asset as FAssetPtr not as UObject.
 
@@ -2218,7 +2238,7 @@ public:
 	typedef TTypeFundamentals::TCppType TCppType;
 
 	UNameProperty(ECppProperty, int32 InOffset, uint64 InFlags)
-		: UNameProperty_Super(FObjectInitializer::Get(), EC_CppProperty, InOffset, InFlags)
+		: UNameProperty_Super(FObjectInitializer::Get(), EC_CppProperty, InOffset, InFlags | CPF_HasGetValueTypeHash)
 	{
 	}
 
@@ -2261,7 +2281,7 @@ public:
 	typedef TTypeFundamentals::TCppType TCppType;
 
 	UStrProperty(ECppProperty, int32 InOffset, uint64 InFlags)
-		: UStrProperty_Super(FObjectInitializer::Get(), EC_CppProperty, InOffset, InFlags)
+		: UStrProperty_Super(FObjectInitializer::Get(), EC_CppProperty, InOffset, InFlags | CPF_HasGetValueTypeHash)
 	{
 	}
 
