[INTEGRATE] Change 2571987 by [REDACTED_PERSON][REDACTED_HANDLE] on 2015/05/31 11:57:45

	PR #1197: fixed: crash if list of delayed linkers includes multiple entries of the same linker (Contributed by [REDACTED_PERSON])

	#git author MarcKamradt [REDACTED_EMAIL]

diff --git a/Engine/Source/Runtime/CoreUObject/Private/UObject/UObjectGlobals.cpp b/Engine/Source/Runtime/CoreUObject/Private/UObject/UObjectGlobals.cpp
--- a/Engine/Source/Runtime/CoreUObject/Private/UObject/UObjectGlobals.cpp
+++ b/Engine/Source/Runtime/CoreUObject/Private/UObject/UObjectGlobals.cpp
@@ -1068,7 +1068,7 @@ UPackage* LoadPackageInternal(UPackage* InOuter, const TCHAR* InLongPackageName,
 			// Async loading removes delayed linkers on the game thread after streaming has finished
 			else
 			{
-				FUObjectThreadContext::Get().DelayedLinkerClosePackages.Add(Linker);
+				FUObjectThreadContext::Get().DelayedLinkerClosePackages.AddUnique(Linker);
 			}
 		}
 	}
