Fix for crash when using a multicast dynamic delegate type in a Blueprint function declaration.

diff --git a/Engine/Source/Runtime/CoreUObject/Private/UObject/PropertyMulticastDelegate.cpp b/Engine/Source/Runtime/CoreUObject/Private/UObject/PropertyMulticastDelegate.cpp
--- a/Engine/Source/Runtime/CoreUObject/Private/UObject/PropertyMulticastDelegate.cpp
+++ b/Engine/Source/Runtime/CoreUObject/Private/UObject/PropertyMulticastDelegate.cpp
@@ -113,6 +113,15 @@ bool UMulticastDelegateProperty::NetSerializeItem( FArchive& Ar, UPackageMap* Ma
 
 FString UMulticastDelegateProperty::GetCPPType( FString* ExtendedTypeText/*=NULL*/, uint32 CPPExportFlags/*=0*/ ) const
 {
+#if HACK_HEADER_GENERATOR
+	// We have this test because sometimes the delegate hasn't been set up by FixupDelegateProperties at the time
+	// we need the type for an error message.  We deliberately format it so that it's unambiguously not CPP code, but is still human-readable.
+	if (!SignatureFunction)
+	{
+		return FString(TEXT("{multicast delegate type}"));
+	}
+#endif
+
 	FString UnmangledFunctionName = SignatureFunction->GetName().LeftChop( FString( HEADER_GENERATED_DELEGATE_SIGNATURE_SUFFIX ).Len() );
 	if (0 != (CPPExportFlags & EPropertyExportCPPFlags::CPPF_CustomTypeName))
 	{
