Engine - Fixing a memory leak in the FLinkerLoad in debug builds.

diff --git a/Engine/Source/Runtime/CoreUObject/Private/UObject/LinkerLoad.cpp b/Engine/Source/Runtime/CoreUObject/Private/UObject/LinkerLoad.cpp
--- a/Engine/Source/Runtime/CoreUObject/Private/UObject/LinkerLoad.cpp
+++ b/Engine/Source/Runtime/CoreUObject/Private/UObject/LinkerLoad.cpp
@@ -713,7 +713,7 @@ FLinkerLoad::FLinkerLoad(UPackage* InParent, const TCHAR* InFilename, uint32 InL
 FLinkerLoad::~FLinkerLoad()
 {
 #if UE_BUILD_DEBUG
-	FLinkerManager::Get().GetLiveLinkers().Add(this);
+	FLinkerManager::Get().GetLiveLinkers().Remove(this);
 #endif
 
 	// Detaches linker.
@@ -749,9 +749,9 @@ bool FLinkerLoad::IsTimeLimitExceeded( const TCHAR* CurrentTask, int32 Granulari
 			// Log single operations that take longer than timelimit.
 			if( (CurrentTime - TickStartTime) > (2.5 * TimeLimit) )
 			{
- 				UE_LOG(LogStreaming, Log, TEXT("FLinkerLoad: %s took (less than) %5.2f ms"), 
- 					CurrentTask, 
- 					(CurrentTime - TickStartTime) * 1000);
+				UE_LOG(LogStreaming, Log, TEXT("FLinkerLoad: %s took (less than) %5.2f ms"), 
+					CurrentTask, 
+					(CurrentTime - TickStartTime) * 1000);
 			}
 		}
 	}
@@ -3740,12 +3740,12 @@ UObject* FLinkerLoad::CreateExport( int32 Index )
 
 					ClassObject->Bind();
 
- 					// Preload classes on first access.  Note that this may update the Export.Object, so ClassObject is not guaranteed to be valid after this point
+					// Preload classes on first access.  Note that this may update the Export.Object, so ClassObject is not guaranteed to be valid after this point
 					// If we're async loading on a cooked build we can skip this as there's no chance we will need to recompile the class. 
 					// Preload will be called during async package tick when the data has been precached
 					if( !FPlatformProperties::RequiresCookedData() )
 					{
- 						Preload( Export.Object );
+						Preload( Export.Object );
 					}
 				}
 			}
