Removing mutex lock from FAsyncObjectsReferencer::AddReferencedObjects

diff --git a/Engine/Source/Runtime/CoreUObject/Private/Serialization/AsyncLoading.cpp b/Engine/Source/Runtime/CoreUObject/Private/Serialization/AsyncLoading.cpp
--- a/Engine/Source/Runtime/CoreUObject/Private/Serialization/AsyncLoading.cpp
+++ b/Engine/Source/Runtime/CoreUObject/Private/Serialization/AsyncLoading.cpp
@@ -45,7 +45,14 @@ DECLARE_CYCLE_STAT(TEXT("Flush Async Loading GT"), STAT_FAsyncPackage_FlushAsync
 DECLARE_FLOAT_ACCUMULATOR_STAT( TEXT( "Async loading block time" ), STAT_AsyncIO_AsyncLoadingBlockingTime, STATGROUP_AsyncIO );
 DECLARE_FLOAT_ACCUMULATOR_STAT( TEXT( "Async package precache wait time" ), STAT_AsyncIO_AsyncPackagePrecacheWaitTime, STATGROUP_AsyncIO );
 
-/** Keeps a reference to all objects created during async load until streaming has finished */
+/** Returns true if we're inside a FGCScopeLock */
+bool IsGarbageCollectionLocked();
+
+/** 
+ * Keeps a reference to all objects created during async load until streaming has finished 
+ *
+ * ASSUMPTION: AddObject can't be called while GC is running and we don't want to lock when calling AddReferencedObjects
+ */
 class FAsyncObjectsReferencer : FGCObject
 {
 	/** Private constructor */
@@ -58,7 +65,8 @@ class FAsyncObjectsReferencer : FGCObject
 	FCriticalSection ReferencedObjectsCritical;
 #endif
 
-	FORCEINLINE int32 IndexOf(UObject* InObj)
+#if !UE_BUILD_SHIPPING
+	int32 IndexOf(UObject* InObj)
 	{
 #if THREADSAFE_UOBJECTS
 		FScopeLock ReferencedObjectsLock(&ReferencedObjectsCritical);
@@ -72,6 +80,7 @@ class FAsyncObjectsReferencer : FGCObject
 		}
 		return INDEX_NONE;
 	}
+#endif
 
 public:
 	/** Returns the one and only instance of this object */
@@ -79,36 +88,45 @@ public:
 	/** FGCObject interface */
 	virtual void AddReferencedObjects(FReferenceCollector& Collector) override
 	{
+		// Note we don't lock here as we're guaranteed that AddObject can only be called from
+		// within FGCScopeGuard scope where GC does not run
 		Collector.AllowEliminatingReferences(false);
 		{
-#if THREADSAFE_UOBJECTS
-			FScopeLock ReferencedObjectsLock(&ReferencedObjectsCritical);
-#endif
 			Collector.AddReferencedObjects(ReferencedObjects);
 		}
 		Collector.AllowEliminatingReferences(true);
 	}
-	/** Adds an object to be referenced */
+	/** 
+	 * Adds an object to be referenced 
+	 * The assumption here is that this can only happen from inside of FGCScopeGuard (@see IsGarbageCollectionLocked()) where we're sure GC is not currently running,
+	 * unless we're on the game thread where atm GC can run simultaneously with async loading.
+	 */
 	FORCEINLINE void AddObject(UObject* InObject)
 	{
 		if (InObject)
 		{
+			UE_CLOG(!IsInGameThread() && !IsGarbageCollectionLocked(), LogStreaming, Fatal, TEXT("Trying to add an object %s to FAsyncObjectsReferencer outside of a FGCScopeLock."), *InObject->GetFullName());
+			{
 #if THREADSAFE_UOBJECTS
-			FScopeLock ReferencedObjectsLock(&ReferencedObjectsCritical);
+				// Still want to lock as AddObject may be called on the game thread and async loading thread,
+				// but in any case it may not happen when GC runs.
+				FScopeLock ReferencedObjectsLock(&ReferencedObjectsCritical);
 #else
-			check(IsInGameThread());
+				check(IsInGameThread());
 #endif
-			ReferencedObjects.Add(InObject);
+				ReferencedObjects.Add(InObject);
+			}
 			InObject->ThisThreadAtomicallyClearedRFUnreachable();
 		}
 	}
 	/** Removes all objects from the list and clears async loading flags */
-	FORCENOINLINE void EmptyReferencedObjects()
+	void EmptyReferencedObjects()
 	{
-		const EObjectFlags AsyncFlags = RF_Async | RF_AsyncLoading;
+		check(IsInGameThread());
 #if THREADSAFE_UOBJECTS
 		FScopeLock ReferencedObjectsLock(&ReferencedObjectsCritical);
 #endif
+		const EObjectFlags AsyncFlags = RF_Async | RF_AsyncLoading;
 		for (UObject* Obj : ReferencedObjects)
 		{
 			check(Obj);
@@ -118,8 +136,9 @@ public:
 		ReferencedObjects.Empty(ReferencedObjects.Num());
 	}
 	/** Removes all referenced objects and markes them for GC */
-	FORCENOINLINE void EmptyReferencedObjectsAndCancelLoading()
+	void EmptyReferencedObjectsAndCancelLoading()
 	{
+		check(IsInGameThread());
 		const EObjectFlags LoadFlags = RF_NeedLoad | RF_NeedPostLoad | RF_NeedPostLoadSubobjects;
 		const EObjectFlags AsyncFlags = RF_Async | RF_AsyncLoading;
 
@@ -178,7 +197,14 @@ public:
 	{
 		if (FParse::Command(&Cmd, TEXT("VerifyAsyncLoadAssumptions")))
 		{
-			FAsyncObjectsReferencer::Get().VerifyAssumptions();
+			if (!IsAsyncLoading())
+			{
+				FAsyncObjectsReferencer::Get().VerifyAssumptions();
+			}
+			else
+			{
+				Ar.Logf(TEXT("Unable to verify async loading assumptions while streaming."));
+			}
 			return true;
 		}
 		return false;
@@ -1057,7 +1083,11 @@ EAsyncPackageState::Type FAsyncPackage::CreateLinker()
 		LastTypeOfWorkPerformed			= TEXT("creating Linker");
 
 		// Try to find existing package or create it if not already present.
-		UPackage* Package = CreatePackage(nullptr, *Desc.Name.ToString());
+		UPackage* Package = nullptr;
+		{
+			FGCScopeGuard GCGuard;
+			Package = CreatePackage(nullptr, *Desc.Name.ToString());
+		}
 		FScopeCycleCounterUObject ConstructorScope(Package, GET_STATID(STAT_FAsyncPackage_CreateLinker));
 
 		// Set package specific data 
diff --git a/Engine/Source/Runtime/CoreUObject/Private/UObject/GarbageCollection.cpp b/Engine/Source/Runtime/CoreUObject/Private/UObject/GarbageCollection.cpp
--- a/Engine/Source/Runtime/CoreUObject/Private/UObject/GarbageCollection.cpp
+++ b/Engine/Source/Runtime/CoreUObject/Private/UObject/GarbageCollection.cpp
@@ -135,6 +135,11 @@ public:
 
 		GCCounter.Increment();
 	}
+	/** Checks if any async thread has a lock */
+	bool IsAsyncLocked()
+	{
+		return AsyncCounter.GetValue() != 0;
+	}
 	/** Lock for GC. Will not block and return false if any other thread has already locked. */
 	bool TryGCLock()
 	{		
@@ -170,6 +175,11 @@ bool IsGarbageCollecting()
 	return FGCScopeLock::GetGarbageCollectingFlag();
 }
 
+bool IsGarbageCollectionLocked()
+{
+	return GGarbageCollectionGuardCritical.IsAsyncLocked();
+}
+
 /**
  * If set and VERIFY_DISREGARD_GC_ASSUMPTIONS is true, we verify GC assumptions about "Disregard For GC" objects. We also
  * verify that no unreachable actors/ components are referenced if VERIFY_NO_UNREACHABLE_OBJECTS_ARE_REFERENCED
