Engine - EPackageFlags & EObjectFlags can now be properly viewed in the debugger as a collection of enum entries, e.g. RF_Public | RF_Transactional, rather than just the raw numerical value.  Don't overlap any flags in an enum declaration, otherwise the VS debugger won't treat it like it's a collection of flags any more.

diff --git a/Engine/Source/Runtime/CoreUObject/Public/UObject/ObjectBase.h b/Engine/Source/Runtime/CoreUObject/Public/UObject/ObjectBase.h
--- a/Engine/Source/Runtime/CoreUObject/Public/UObject/ObjectBase.h
+++ b/Engine/Source/Runtime/CoreUObject/Public/UObject/ObjectBase.h
@@ -44,9 +44,17 @@ enum ELoadFlags
 	LOAD_None						= 0x00000000,	// No flags.
 	LOAD_SeekFree					= 0x00000001,	// Loads the package via the seek free loading path/ reader
 	LOAD_NoWarn						= 0x00000002,	// Don't display warning if load fails.
+//	LOAD_Unused						= 0x00000004,
+//	LOAD_Unused						= 0x00000008,
 	LOAD_Verify						= 0x00000010,	// Only verify existance; don't actually load.
 	LOAD_AllowDll					= 0x00000020,	// Allow plain DLLs.
+//	LOAD_Unused						= 0x00000040
 	LOAD_NoVerify					= 0x00000080,   // Don't verify imports yet.
+//	LOAD_Unused						= 0x00000100,
+//	LOAD_Unused						= 0x00000200,
+//	LOAD_Unused						= 0x00000400,
+//	LOAD_Unused						= 0x00000800,
+//	LOAD_Unused						= 0x00001000,
 	LOAD_Quiet						= 0x00002000,   // No log warnings.
 	LOAD_FindIfFail					= 0x00004000,	// Tries FindObject if a linker cannot be obtained (e.g. package is currently being compiled)
 	LOAD_MemoryReader				= 0x00008000,	// Loads the file into memory and serializes from there.
@@ -83,10 +91,19 @@ enum EPackageFlags
 	PKG_ServerSideOnly				= 0x00000004,   // Only needed on the server side.
 	PKG_CompiledIn					= 0x00000010,   // This package is from "compiled in" classes.
 	PKG_ForDiffing					= 0x00000020,	// This package was loaded just for the purposes of diff'ing
+//	PKG_Unused						= 0x00000040
+//  PKG_Unused						= 0x00000080,
+//	PKG_Unused						= 0x00000100,
+//	PKG_Unused						= 0x00000200,
+//	PKG_Unused						= 0x00000400,
+//	PKG_Unused						= 0x00000800,
+//	PKG_Unused						= 0x00001000,
+//	PKG_Unused						= 0x00002000,
+//	PKG_Unused						= 0x00004000,
 	PKG_Need						= 0x00008000,	// Client needs to download this package.
 	PKG_Compiling					= 0x00010000,	// package is currently being compiled
 	PKG_ContainsMap					= 0x00020000,	// Set if the package contains a ULevel/ UWorld object
-	PKG_RequiresLocalizationGather		= 0x00040000,	// Set if the package contains any data to be gathered by localization
+	PKG_RequiresLocalizationGather	= 0x00040000,	// Set if the package contains any data to be gathered by localization
 	PKG_DisallowLazyLoading			= 0x00080000,	// Set if the archive serializing this package cannot use lazy loading
 	PKG_PlayInEditor				= 0x00100000,	// Set if the package was created for the purpose of PIE
 	PKG_ContainsScript				= 0x00200000,	// Package is allowed to contain UClass objects
@@ -100,9 +117,10 @@ enum EPackageFlags
 //	PKG_Unused						= 0x20000000,
 	PKG_ReloadingForCooker			= 0x40000000,   // this package is reloading in the cooker, try to avoid getting data we will never need. We won't save this package.
 	PKG_FilterEditorOnly			= 0x80000000,	// Package has editor-only data filtered
-
-	PKG_InMemoryOnly				= PKG_CompiledIn | PKG_NewlyCreated, // Flag mask that indicates if this package is a package that exists in memory only.
 };
+
+#define PKG_InMemoryOnly	(EPackageFlags)(PKG_CompiledIn | PKG_NewlyCreated) // Flag mask that indicates if this package is a package that exists in memory only.
+
 ENUM_CLASS_FLAGS(EPackageFlags);
 
 //
@@ -212,7 +230,7 @@ enum EClassFlags
 
 	//@}
 
-
+	
 	/** @name Flags to inherit from base class */
 	//@{
 	CLASS_Inherit           = CLASS_Transient | CLASS_DefaultConfig | CLASS_Config | CLASS_PerObjectConfig | CLASS_ConfigDoNotCheckDefaults | CLASS_NotPlaceable
@@ -259,6 +277,10 @@ enum EClassFlags
 	CLASS_AllFlags			= 0xFFFFFFFF,
 };
 
+
+
+
+
 /**
  * Flags used for quickly casting classes of certain types; all class cast flags are inherited
  */
@@ -406,51 +428,49 @@ enum EObjectFlags
 {
 	// Do not add new flags unless they truly belong here. There are alternatives.
 	// if you change any the bit of any of the RF_Load flags, then you will need legacy serialization
+	RF_NoFlags						= 0x00000000,	///< No flags, used to avoid a cast
 
 	// This first group of flags mostly has to do with what kind of object it is. Other than transient, these are the persistent object flags.
 	// The garbage collector also tends to look at these.
-	RF_Public					=0x00000001,	///< Object is visible outside its package.
-	RF_Standalone				=0x00000002,	///< Keep object around for editing even if unreferenced.
-	RF_Native					=0x00000004,	///< Native (UClass only).
-	RF_Transactional			=0x00000008,	///< Object is transactional.
-	RF_ClassDefaultObject		=0x00000010,	///< This object is its class's default object
-	RF_ArchetypeObject			=0x00000020,	///< This object is a template for another object - treat like a class default object
-	RF_Transient				=0x00000040,	///< Don't save object.
+	RF_Public						= 0x00000001,	///< Object is visible outside its package.
+	RF_Standalone					= 0x00000002,	///< Keep object around for editing even if unreferenced.
+	RF_Native						= 0x00000004,	///< Native (UClass only).
+	RF_Transactional				= 0x00000008,	///< Object is transactional.
+	RF_ClassDefaultObject			= 0x00000010,	///< This object is its class's default object
+	RF_ArchetypeObject				= 0x00000020,	///< This object is a template for another object - treat like a class default object
+	RF_Transient					= 0x00000040,	///< Don't save object.
 
 	// This group of flags is primarily concerned with garbage collection.
-	RF_RootSet					=0x00000080,	///< Object will not be garbage collected, even if unreferenced.
-	RF_Unreachable				=0x00000100,	///< Object is not reachable on the object graph.
-	RF_TagGarbageTemp			=0x00000200,	///< This is a temp user flag for various utilities that need to use the garbage collector. The garbage collector itself does not interpret it.
+	RF_RootSet						= 0x00000080,	///< Object will not be garbage collected, even if unreferenced.
+	RF_Unreachable					= 0x00000100,	///< Object is not reachable on the object graph.
+	RF_TagGarbageTemp				= 0x00000200,	///< This is a temp user flag for various utilities that need to use the garbage collector. The garbage collector itself does not interpret it.
 
 	// The group of flags tracks the stages of the lifetime of a uobject
-	RF_NeedLoad					=0x00000400,	///< During load, indicates object needs loading.
-	RF_AsyncLoading				=0x00000800,	///< Object is being asynchronously loaded.
-	RF_NeedPostLoad				=0x00001000,	///< Object needs to be postloaded.
-	RF_NeedPostLoadSubobjects	=0x00002000,	///< During load, indicates that the object still needs to instance subobjects and fixup serialized component references
-	RF_PendingKill				=0x00004000,	///< Objects that are pending destruction (invalid for gameplay but valid objects)
-	RF_BeginDestroyed			=0x00008000,	///< BeginDestroy has been called on the object.
-	RF_FinishDestroyed			=0x00010000,	///< FinishDestroy has been called on the object.
+	RF_NeedLoad						= 0x00000400,	///< During load, indicates object needs loading.
+	RF_AsyncLoading					= 0x00000800,	///< Object is being asynchronously loaded.
+	RF_NeedPostLoad					= 0x00001000,	///< Object needs to be postloaded.
+	RF_NeedPostLoadSubobjects		= 0x00002000,	///< During load, indicates that the object still needs to instance subobjects and fixup serialized component references
+	RF_PendingKill					= 0x00004000,	///< Objects that are pending destruction (invalid for gameplay but valid objects)
+	RF_BeginDestroyed				= 0x00008000,	///< BeginDestroy has been called on the object.
+	RF_FinishDestroyed				= 0x00010000,	///< FinishDestroy has been called on the object.
 
 	// Misc. Flags
-	RF_BeingRegenerated			=0x00020000,	///< Flagged on UObjects that are used to create UClasses (e.g. Blueprints) while they are regenerating their UClass on load (See FLinkerLoad::CreateExport())
-	RF_DefaultSubObject			=0x00040000,	///< Flagged on subobjects that are defaults
-	RF_WasLoaded				=0x00080000,	///< Flagged on UObjects that were loaded
-	RF_TextExportTransient		=0x00100000,	///< Do not export object to text form (e.g. copy/paste). Generally used for sub-objects that can be regenerated from data in their parent object.
-	RF_LoadCompleted			=0x00200000,	///< Object has been completely serialized by linkerload at least once. DO NOT USE THIS FLAG, It should be replaced with RF_WasLoaded.
-	RF_InheritableComponentTemplate = 0x00400000, ///< Archetype of the object can be in its super class
-	RF_Async = 0x00800000, ///< Object exists only on a different thread than the game thread.
-
-	// Special all and none masks
-	RF_AllFlags					=0x00ffffff,	///< All flags, used mainly for error checking
-	RF_NoFlags					=0x00000000,	///< No flags, used to avoid a cast
-
-	// Predefined groups of the above
-	RF_Load						= RF_Public | RF_Standalone | RF_Native | RF_Transactional | RF_ClassDefaultObject | RF_ArchetypeObject | RF_DefaultSubObject | RF_TextExportTransient | RF_InheritableComponentTemplate, // Flags to load from Unrealfiles.
-	RF_PropagateToSubObjects	= RF_Public | RF_ArchetypeObject | RF_Transactional | RF_Transient,		// Sub-objects will inherit these flags from their SuperObject.
-
-	
+	RF_BeingRegenerated				= 0x00020000,	///< Flagged on UObjects that are used to create UClasses (e.g. Blueprints) while they are regenerating their UClass on load (See FLinkerLoad::CreateExport())
+	RF_DefaultSubObject				= 0x00040000,	///< Flagged on subobjects that are defaults
+	RF_WasLoaded					= 0x00080000,	///< Flagged on UObjects that were loaded
+	RF_TextExportTransient			= 0x00100000,	///< Do not export object to text form (e.g. copy/paste). Generally used for sub-objects that can be regenerated from data in their parent object.
+	RF_LoadCompleted				= 0x00200000,	///< Object has been completely serialized by linkerload at least once. DO NOT USE THIS FLAG, It should be replaced with RF_WasLoaded.
+	RF_InheritableComponentTemplate = 0x00400000,	///< Archetype of the object can be in its super class
+	RF_Async						= 0x00800000,	///< Object exists only on a different thread than the game thread.
 };
 
+// Special all and none masks
+#define RF_AllFlags				(EObjectFlags)0x00ffffff	///< All flags, used mainly for error checking
+
+// Predefined groups of the above
+#define RF_Load						((EObjectFlags)(RF_Public | RF_Standalone | RF_Native | RF_Transactional | RF_ClassDefaultObject | RF_ArchetypeObject | RF_DefaultSubObject | RF_TextExportTransient | RF_InheritableComponentTemplate)) // Flags to load from Unrealfiles.
+#define RF_PropagateToSubObjects	((EObjectFlags)(RF_Public | RF_ArchetypeObject | RF_Transactional | RF_Transient))		// Sub-objects will inherit these flags from their SuperObject.
+
 FORCEINLINE EObjectFlags operator|(EObjectFlags Arg1,EObjectFlags Arg2)
 {
 	return EObjectFlags(uint32(Arg1) | uint32(Arg2));
@@ -1210,7 +1230,7 @@ Class declaration macros.
 
 #define DECLARE_CLASS( TClass, TSuperClass, TStaticFlags, TStaticCastFlags, TPackage, TRequiredAPI  ) \
 private: \
-    TClass & operator=(TClass const &);   \
+	TClass & operator=(TClass const &);   \
 	TRequiredAPI static UClass* GetPrivateStaticClass(const TCHAR* Package); \
 public: \
 	/** Bitwise union of #EClassFlags pertaining to this class.*/ \
@@ -1598,14 +1618,14 @@ public:
 		return *this;
 	}
 	/**
- 	 * Dereference back into a UClass
+	 * Dereference back into a UClass
 	 * @return	the embedded UClass
 	 * The body for this method is in Class.h
 	 */
 	FORCEINLINE UClass* operator*() const;
 
 	/**
- 	 * Dereference back into a UClass
+	 * Dereference back into a UClass
 	 * @return	the embedded UClass
 	 */
 	FORCEINLINE UClass* operator->() const
