UE-13798 Make the shared state of the event graph use something akin to weak ptrs

diff --git a/Engine/Source/Runtime/CoreUObject/Private/UObject/GarbageCollection.cpp b/Engine/Source/Runtime/CoreUObject/Private/UObject/GarbageCollection.cpp
--- a/Engine/Source/Runtime/CoreUObject/Private/UObject/GarbageCollection.cpp
+++ b/Engine/Source/Runtime/CoreUObject/Private/UObject/GarbageCollection.cpp
@@ -293,7 +293,7 @@ void SerializeRootSet( FArchive& Ar, EObjectFlags KeepFlags )
  * @param ReferencingObject UObject which owns the reference (can be NULL)
  * @param bAllowReferenceElimination	Whether to allow NULL'ing the reference if RF_PendingKill is set
  */
-static FORCEINLINE void HandleObjectReference(TArray<UObject*>& ObjectsToSerialize, UObject* ReferencingObject, UObject*& Object, bool bAllowReferenceElimination)
+static FORCEINLINE void HandleObjectReference(TArray<UObject*>& ObjectsToSerialize, UObject* ReferencingObject, UObject*& Object, bool bAllowReferenceElimination, bool bDoHandle = true)
 {
 	// Disregard NULL objects and perform very fast check to see whether object is part of permanent
 	// object pool and should therefore be disregarded. The check doesn't touch the object and is
@@ -310,7 +310,7 @@ static FORCEINLINE void HandleObjectReference(TArray<UObject*>& ObjectsToSeriali
 				Object = NULL;
 			}
 			// Add encountered object reference to list of to be serialized objects if it hasn't already been added.
-			else if( Object->HasAnyFlags( RF_Unreachable ) )
+			else if (Object->HasAnyFlags(RF_Unreachable) && bDoHandle)
 			{				
 				if( GIsRunningParallelReachability )
 				{
@@ -383,12 +383,14 @@ class FGCCollector : public FReferenceCollector
 {
 	TArray<UObject*>& ObjectArray;
 	bool bAllowEliminatingReferences;
+	bool bShouldHandleAsWeakRef;
 
 public:
 
 	FGCCollector(TArray<UObject*>& InObjectArray)
 		: ObjectArray(InObjectArray)
 		, bAllowEliminatingReferences(true)
+		, bShouldHandleAsWeakRef(false)
 	{
 	}
 
@@ -403,7 +405,7 @@ public:
 				ReferencingProperty ? *ReferencingProperty->GetFullName() : TEXT("NULL"));
 		}
 #endif
-		::HandleObjectReference(ObjectArray, const_cast<UObject*>(ReferencingObject), Object, bAllowEliminatingReferences);
+		::HandleObjectReference(ObjectArray, const_cast<UObject*>(ReferencingObject), Object, bAllowEliminatingReferences, !bShouldHandleAsWeakRef);
 	}
 	virtual bool IsIgnoringArchetypeRef() const override
 	{
@@ -417,6 +419,11 @@ public:
 	{
 		bAllowEliminatingReferences = bAllow;
 	}
+
+	virtual void SetShouldHandleAsWeakRef(bool bWeakRef) override
+	{
+		bShouldHandleAsWeakRef = bWeakRef;
+	}
 };
 
 
diff --git a/Engine/Source/Runtime/CoreUObject/Public/UObject/ObjectBase.h b/Engine/Source/Runtime/CoreUObject/Public/UObject/ObjectBase.h
--- a/Engine/Source/Runtime/CoreUObject/Public/UObject/ObjectBase.h
+++ b/Engine/Source/Runtime/CoreUObject/Public/UObject/ObjectBase.h
@@ -103,7 +103,7 @@ enum EPackageFlags
 	PKG_Need						= 0x00008000,	// Client needs to download this package.
 	PKG_Compiling					= 0x00010000,	// package is currently being compiled
 	PKG_ContainsMap					= 0x00020000,	// Set if the package contains a ULevel/ UWorld object
-	PKG_RequiresLocalizationGather	= 0x00040000,	// Set if the package contains any data to be gathered by localization
+	PKG_RequiresLocalizationGather		= 0x00040000,	// Set if the package contains any data to be gathered by localization
 	PKG_DisallowLazyLoading			= 0x00080000,	// Set if the archive serializing this package cannot use lazy loading
 	PKG_PlayInEditor				= 0x00100000,	// Set if the package was created for the purpose of PIE
 	PKG_ContainsScript				= 0x00200000,	// Package is allowed to contain UClass objects
@@ -230,7 +230,7 @@ enum EClassFlags
 
 	//@}
 
-	
+
 	/** @name Flags to inherit from base class */
 	//@{
 	CLASS_Inherit           = CLASS_Transient | CLASS_DefaultConfig | CLASS_Config | CLASS_PerObjectConfig | CLASS_ConfigDoNotCheckDefaults | CLASS_NotPlaceable
@@ -432,42 +432,43 @@ enum EObjectFlags
 
 	// This first group of flags mostly has to do with what kind of object it is. Other than transient, these are the persistent object flags.
 	// The garbage collector also tends to look at these.
-	RF_Public						= 0x00000001,	///< Object is visible outside its package.
-	RF_Standalone					= 0x00000002,	///< Keep object around for editing even if unreferenced.
-	RF_Native						= 0x00000004,	///< Native (UClass only).
-	RF_Transactional				= 0x00000008,	///< Object is transactional.
-	RF_ClassDefaultObject			= 0x00000010,	///< This object is its class's default object
-	RF_ArchetypeObject				= 0x00000020,	///< This object is a template for another object - treat like a class default object
-	RF_Transient					= 0x00000040,	///< Don't save object.
+	RF_Public					=0x00000001,	///< Object is visible outside its package.
+	RF_Standalone				=0x00000002,	///< Keep object around for editing even if unreferenced.
+	RF_Native					=0x00000004,	///< Native (UClass only).
+	RF_Transactional			=0x00000008,	///< Object is transactional.
+	RF_ClassDefaultObject		=0x00000010,	///< This object is its class's default object
+	RF_ArchetypeObject			=0x00000020,	///< This object is a template for another object - treat like a class default object
+	RF_Transient				=0x00000040,	///< Don't save object.
 
 	// This group of flags is primarily concerned with garbage collection.
-	RF_RootSet						= 0x00000080,	///< Object will not be garbage collected, even if unreferenced.
-	RF_Unreachable					= 0x00000100,	///< Object is not reachable on the object graph.
-	RF_TagGarbageTemp				= 0x00000200,	///< This is a temp user flag for various utilities that need to use the garbage collector. The garbage collector itself does not interpret it.
+	RF_RootSet					=0x00000080,	///< Object will not be garbage collected, even if unreferenced.
+	RF_Unreachable				=0x00000100,	///< Object is not reachable on the object graph.
+	RF_TagGarbageTemp			=0x00000200,	///< This is a temp user flag for various utilities that need to use the garbage collector. The garbage collector itself does not interpret it.
 
 	// The group of flags tracks the stages of the lifetime of a uobject
-	RF_NeedLoad						= 0x00000400,	///< During load, indicates object needs loading.
-	RF_AsyncLoading					= 0x00000800,	///< Object is being asynchronously loaded.
-	RF_NeedPostLoad					= 0x00001000,	///< Object needs to be postloaded.
-	RF_NeedPostLoadSubobjects		= 0x00002000,	///< During load, indicates that the object still needs to instance subobjects and fixup serialized component references
-	RF_PendingKill					= 0x00004000,	///< Objects that are pending destruction (invalid for gameplay but valid objects)
-	RF_BeginDestroyed				= 0x00008000,	///< BeginDestroy has been called on the object.
-	RF_FinishDestroyed				= 0x00010000,	///< FinishDestroy has been called on the object.
+	RF_NeedLoad					=0x00000400,	///< During load, indicates object needs loading.
+	RF_AsyncLoading				=0x00000800,	///< Object is being asynchronously loaded.
+	RF_NeedPostLoad				=0x00001000,	///< Object needs to be postloaded.
+	RF_NeedPostLoadSubobjects	=0x00002000,	///< During load, indicates that the object still needs to instance subobjects and fixup serialized component references
+	RF_PendingKill				=0x00004000,	///< Objects that are pending destruction (invalid for gameplay but valid objects)
+	RF_BeginDestroyed			=0x00008000,	///< BeginDestroy has been called on the object.
+	RF_FinishDestroyed			=0x00010000,	///< FinishDestroy has been called on the object.
 
 	// Misc. Flags
-	RF_BeingRegenerated				= 0x00020000,	///< Flagged on UObjects that are used to create UClasses (e.g. Blueprints) while they are regenerating their UClass on load (See FLinkerLoad::CreateExport())
-	RF_DefaultSubObject				= 0x00040000,	///< Flagged on subobjects that are defaults
-	RF_WasLoaded					= 0x00080000,	///< Flagged on UObjects that were loaded
-	RF_TextExportTransient			= 0x00100000,	///< Do not export object to text form (e.g. copy/paste). Generally used for sub-objects that can be regenerated from data in their parent object.
-	RF_LoadCompleted				= 0x00200000,	///< Object has been completely serialized by linkerload at least once. DO NOT USE THIS FLAG, It should be replaced with RF_WasLoaded.
-	RF_InheritableComponentTemplate = 0x00400000,	///< Archetype of the object can be in its super class
-	RF_Async						= 0x00800000,	///< Object exists only on a different thread than the game thread.
+	RF_BeingRegenerated			=0x00020000,	///< Flagged on UObjects that are used to create UClasses (e.g. Blueprints) while they are regenerating their UClass on load (See FLinkerLoad::CreateExport())
+	RF_DefaultSubObject			=0x00040000,	///< Flagged on subobjects that are defaults
+	RF_WasLoaded				=0x00080000,	///< Flagged on UObjects that were loaded
+	RF_TextExportTransient		=0x00100000,	///< Do not export object to text form (e.g. copy/paste). Generally used for sub-objects that can be regenerated from data in their parent object.
+	RF_LoadCompleted			=0x00200000,	///< Object has been completely serialized by linkerload at least once. DO NOT USE THIS FLAG, It should be replaced with RF_WasLoaded.
+	RF_InheritableComponentTemplate = 0x00400000, ///< Archetype of the object can be in its super class
+	RF_Async = 0x00800000, ///< Object exists only on a different thread than the game thread.
+	RF_StrongRefOnFrame			= 0x01000000,	///< References to this object from persistent function frame are handled as strong ones.
 };
 
-// Special all and none masks
+	// Special all and none masks
 #define RF_AllFlags				(EObjectFlags)0x00ffffff	///< All flags, used mainly for error checking
 
-// Predefined groups of the above
+	// Predefined groups of the above
 #define RF_Load						((EObjectFlags)(RF_Public | RF_Standalone | RF_Native | RF_Transactional | RF_ClassDefaultObject | RF_ArchetypeObject | RF_DefaultSubObject | RF_TextExportTransient | RF_InheritableComponentTemplate)) // Flags to load from Unrealfiles.
 #define RF_PropagateToSubObjects	((EObjectFlags)(RF_Public | RF_ArchetypeObject | RF_Transactional | RF_Transient))		// Sub-objects will inherit these flags from their SuperObject.
 
@@ -1230,7 +1231,7 @@ Class declaration macros.
 
 #define DECLARE_CLASS( TClass, TSuperClass, TStaticFlags, TStaticCastFlags, TPackage, TRequiredAPI  ) \
 private: \
-	TClass & operator=(TClass const &);   \
+    TClass & operator=(TClass const &);   \
 	TRequiredAPI static UClass* GetPrivateStaticClass(const TCHAR* Package); \
 public: \
 	/** Bitwise union of #EClassFlags pertaining to this class.*/ \
@@ -1618,14 +1619,14 @@ public:
 		return *this;
 	}
 	/**
-	 * Dereference back into a UClass
+ 	 * Dereference back into a UClass
 	 * @return	the embedded UClass
 	 * The body for this method is in Class.h
 	 */
 	FORCEINLINE UClass* operator*() const;
 
 	/**
-	 * Dereference back into a UClass
+ 	 * Dereference back into a UClass
 	 * @return	the embedded UClass
 	 */
 	FORCEINLINE UClass* operator->() const
diff --git a/Engine/Source/Runtime/CoreUObject/Public/UObject/UObjectGlobals.h b/Engine/Source/Runtime/CoreUObject/Public/UObject/UObjectGlobals.h
--- a/Engine/Source/Runtime/CoreUObject/Public/UObject/UObjectGlobals.h
+++ b/Engine/Source/Runtime/CoreUObject/Public/UObject/UObjectGlobals.h
@@ -1563,6 +1563,8 @@ public:
 	* Gets the property that is currently being serialized
 	*/
 	virtual class UProperty* GetSerializedProperty() const { return nullptr; }
+
+	virtual void SetShouldHandleAsWeakRef(bool bWeakRef) {}
 protected:
 	/**
 	 * Handle object reference. Called by AddReferencedObject.
diff --git a/Engine/Source/Runtime/Engine/Private/BlueprintGeneratedClass.cpp b/Engine/Source/Runtime/Engine/Private/BlueprintGeneratedClass.cpp
--- a/Engine/Source/Runtime/Engine/Private/BlueprintGeneratedClass.cpp
+++ b/Engine/Source/Runtime/Engine/Private/BlueprintGeneratedClass.cpp
@@ -666,6 +666,22 @@ void UBlueprintGeneratedClass::Bind()
 	}
 }
 
+class FPersistentFrameCollectorArchive : public FSimpleObjectReferenceCollectorArchive
+{
+public:
+	FPersistentFrameCollectorArchive(const UObject* InSerializingObject, FReferenceCollector& InCollector)
+		: FSimpleObjectReferenceCollectorArchive(InSerializingObject, InCollector)
+	{}
+
+protected:
+	virtual FArchive& operator<<(UObject*& Object) override
+	{
+		const bool bWeakRef = false; //@TODO - enable: Object ? !Object->HasAnyFlags(RF_StrongRefOnFrame) : false;
+		Collector.SetShouldHandleAsWeakRef(bWeakRef); 
+		return FSimpleObjectReferenceCollectorArchive::operator<<(Object);
+	}
+};
+
 void UBlueprintGeneratedClass::AddReferencedObjectsInUbergraphFrame(UObject* InThis, FReferenceCollector& Collector)
 {
 	checkSlow(InThis);
@@ -680,8 +696,9 @@ void UBlueprintGeneratedClass::AddReferencedObjectsInUbergraphFrame(UObject* InT
 				checkSlow(PointerToUberGraphFrame)
 				if (PointerToUberGraphFrame->RawPointer)
 				{
-					FSimpleObjectReferenceCollectorArchive ObjectReferenceCollector(InThis, Collector);
+					FPersistentFrameCollectorArchive ObjectReferenceCollector(InThis, Collector);
 					BPGC->UberGraphFunction->SerializeBin(ObjectReferenceCollector, PointerToUberGraphFrame->RawPointer);
+					Collector.SetShouldHandleAsWeakRef(false);
 				}
 			}
 		}
