Changed the default return value for GetResourceSize() back to 0.  ResourceSize asset registry tag is now always added for non BP CDOs to prevent the ResourceSize tag from being filtered out when building asset tooltips.

diff --git a/Engine/Source/Editor/SizeMap/SSizeMap.cpp b/Engine/Source/Editor/SizeMap/SSizeMap.cpp
--- a/Engine/Source/Editor/SizeMap/SSizeMap.cpp
+++ b/Engine/Source/Editor/SizeMap/SSizeMap.cpp
@@ -320,12 +320,9 @@ void SSizeMap::GatherDependenciesRecursively( FAssetRegistryModule& AssetRegistr
 							NodeSizeMapData.AssetSize = Asset->GetResourceSize( EResourceSizeMode::Exclusive );
 						}
 
-						NodeSizeMapData.bHasKnownSize = NodeSizeMapData.AssetSize != UObject::RESOURCE_SIZE_NONE && NodeSizeMapData.AssetSize != 0;
+						NodeSizeMapData.bHasKnownSize = (NodeSizeMapData.AssetSize != 0);
 						if( !NodeSizeMapData.bHasKnownSize )
 						{
-							// Asset has no meaningful size
-							NodeSizeMapData.AssetSize = 0;
-
 							// @todo sizemap urgent: Try to serialize to figure out how big it is (not into sub-assets though!)
 							// FObjectMemoryAnalyzer ObjectMemoryAnalyzer( Asset );
 						}
diff --git a/Engine/Source/Runtime/CoreUObject/Private/UObject/Obj.cpp b/Engine/Source/Runtime/CoreUObject/Private/UObject/Obj.cpp
--- a/Engine/Source/Runtime/CoreUObject/Private/UObject/Obj.cpp
+++ b/Engine/Source/Runtime/CoreUObject/Private/UObject/Obj.cpp
@@ -1254,7 +1254,7 @@ void UObject::GetAssetRegistryTags(TArray<FAssetRegistryTag>& OutTags) const
 {
 	// Add ResourceSize if non-zero. GetResourceSize is not const because many override implementations end up calling Serialize on this pointers.
 	SIZE_T ResourceSize = const_cast<UObject*>(this)->GetResourceSize(EResourceSizeMode::Exclusive);
-	if ( ResourceSize != RESOURCE_SIZE_NONE )
+	if ( ResourceSize > 0 || ( !GetClass()->HasAnyClassFlags(CLASS_CompiledFromBlueprint) && HasAnyFlags(RF_ClassDefaultObject) ) )
 	{
 		OutTags.Add( FAssetRegistryTag("ResourceSize", FString::Printf(TEXT("%d"), (ResourceSize + 512) / 1024), FAssetRegistryTag::TT_Numerical) );
 	}
diff --git a/Engine/Source/Runtime/CoreUObject/Public/UObject/UObject.h b/Engine/Source/Runtime/CoreUObject/Public/UObject/UObject.h
--- a/Engine/Source/Runtime/CoreUObject/Public/UObject/UObject.h
+++ b/Engine/Source/Runtime/CoreUObject/Public/UObject/UObject.h
@@ -444,21 +444,17 @@ public:
 		return false;
 	}
 
-	/** Special value meaning that the resource size is not defined */
-	static const SIZE_T RESOURCE_SIZE_NONE = static_cast<SIZE_T>(-1);
-
 	/**
 	 * Returns the size of the object/ resource for display to artists/ LDs in the Editor. The
-	 * default behavior is to return RESOURCE_SIZE_NONE which indicates that the resource shouldn't
-	 * display its size which is used to not confuse people by displaying small sizes
-	 * e.g. for objects like materials.
+	 * default behavior is to return 0 which indicates that the resource shouldn't
+	 * display its size.
 	 *
 	 * @param	Type	Indicates which resource size should be returned
 	 * @return	Size of resource as to be displayed to artists/ LDs in the Editor.
 	 */
 	virtual SIZE_T GetResourceSize(EResourceSizeMode::Type Mode)
 	{
-		return RESOURCE_SIZE_NONE;
+		return 0;
 	}
 
 	/** 
