Fixing StaticFindObject returning objects not yet fully streamed in.

diff --git a/Engine/Source/Runtime/CoreUObject/Private/UObject/UObjectHash.cpp b/Engine/Source/Runtime/CoreUObject/Private/UObject/UObjectHash.cpp
--- a/Engine/Source/Runtime/CoreUObject/Private/UObject/UObjectHash.cpp
+++ b/Engine/Source/Runtime/CoreUObject/Private/UObject/UObjectHash.cpp
@@ -601,7 +601,7 @@ void GetObjectsWithOuter(const class UObjectBase* Outer, TArray<UObject *>& Resu
 
 	// We don't want to return any objects that are currently being background loaded unless we're using the object iterator during async loading.
 	ExclusionFlags |= RF_Unreachable;
-	if (!IsAsyncLoading())
+	if (!IsInAsyncLoadingThread())
 	{
 		ExclusionFlags = EObjectFlags(ExclusionFlags | RF_AsyncLoading);
 	}
@@ -653,7 +653,7 @@ UObjectBase* FindObjectWithOuter(class UObjectBase* Outer, class UClass* ClassTo
 	check( Outer && ClassToLookFor );
 	// We don't want to return any objects that are currently being background loaded unless we're using the object iterator during async loading.
 	EObjectFlags ExclusionFlags = RF_Unreachable;
-	if( !IsAsyncLoading() )
+	if (!IsInAsyncLoadingThread())
 	{
 		ExclusionFlags = EObjectFlags( ExclusionFlags | RF_AsyncLoading );
 	}
@@ -731,7 +731,7 @@ void GetObjectsOfClass(UClass* ClassToLookFor, TArray<UObject *>& Results, bool
 
 	// We don't want to return any objects that are currently being background loaded unless we're using the object iterator during async loading.
 	EObjectFlags ExclusionFlags = RF_Unreachable;
-	if( !IsAsyncLoading() )
+	if (!IsInAsyncLoadingThread())
 	{
 		ExclusionFlags |= RF_AsyncLoading;
 	}
