Missing file from StaticFindObjectFast fix while streaming

diff --git a/Engine/Source/Runtime/CoreUObject/Private/UObject/UObjectGlobals.cpp b/Engine/Source/Runtime/CoreUObject/Private/UObject/UObjectGlobals.cpp
--- a/Engine/Source/Runtime/CoreUObject/Private/UObject/UObjectGlobals.cpp
+++ b/Engine/Source/Runtime/CoreUObject/Private/UObject/UObjectGlobals.cpp
@@ -103,7 +103,7 @@ UObject* StaticFindObjectFast( UClass* ObjectClass, UObject* ObjectPackage, FNam
 	}
 
 	// We don't want to return any objects that are currently being background loaded unless we're using FindObject during async loading.
-	ExclusiveFlags |= IsAsyncLoading() ? RF_NoFlags : RF_AsyncLoading;
+	ExclusiveFlags |= IsInAsyncLoadingThread() ? RF_NoFlags : RF_AsyncLoading;
 	return StaticFindObjectFastInternal( ObjectClass, ObjectPackage, ObjectName, ExactClass, AnyPackage, ExclusiveFlags );
 }
 
