Preventing packages from becoming dirty during hot-reload.

diff --git a/Engine/Source/Runtime/CoreUObject/Private/UObject/UObjectBaseUtility.cpp b/Engine/Source/Runtime/CoreUObject/Private/UObject/UObjectBaseUtility.cpp
--- a/Engine/Source/Runtime/CoreUObject/Private/UObject/UObjectBaseUtility.cpp
+++ b/Engine/Source/Runtime/CoreUObject/Private/UObject/UObjectBaseUtility.cpp
@@ -140,7 +140,7 @@ void UObjectBaseUtility::MarkPackageDirty() const
 			// we explicitly disable the ability to dirty a package or map during load.  Commandlets can still
 			// set the dirty state on load.
 			if( IsRunningCommandlet() || 
-				(GIsEditor && !GIsEditorLoadingPackage && !GIsPlayInEditorWorld))
+				(GIsEditor && !GIsEditorLoadingPackage && !GIsPlayInEditorWorld && !GIsHotReload))
 			{
 				const bool bIsDirty = Package->IsDirty();
 
