CIS fix.

diff --git a/Engine/Source/Runtime/CoreUObject/Private/UObject/UObjectBaseUtility.cpp b/Engine/Source/Runtime/CoreUObject/Private/UObject/UObjectBaseUtility.cpp
--- a/Engine/Source/Runtime/CoreUObject/Private/UObject/UObjectBaseUtility.cpp
+++ b/Engine/Source/Runtime/CoreUObject/Private/UObject/UObjectBaseUtility.cpp
@@ -6,6 +6,7 @@
 
 #include "CoreUObjectPrivate.h"
 #include "Interface.h"
+#include "ModuleManager.h"
 
 
 /***********************/
@@ -140,7 +141,11 @@ void UObjectBaseUtility::MarkPackageDirty() const
 			// we explicitly disable the ability to dirty a package or map during load.  Commandlets can still
 			// set the dirty state on load.
 			if( IsRunningCommandlet() || 
-				(GIsEditor && !GIsEditorLoadingPackage && !GIsPlayInEditorWorld && !GIsHotReload))
+				(GIsEditor && !GIsEditorLoadingPackage && !GIsPlayInEditorWorld
+#if WITH_HOT_RELOAD
+				&& !GIsHotReload
+#endif // WITH_HOT_RELOAD
+				))
 			{
 				const bool bIsDirty = Package->IsDirty();
 
