Adding the ability to force GC even when triggered by TryCollectGarbage

diff --git a/Engine/Config/BaseEngine.ini b/Engine/Config/BaseEngine.ini
--- a/Engine/Config/BaseEngine.ini
+++ b/Engine/Config/BaseEngine.ini
@@ -1331,6 +1331,7 @@ FlushStreamingOnGC=False
 MinBulkDataSizeForAsyncLoading=131072
 AsyncLoadingThreadEnabled=False
 WarnIfTimeLimitExceeded=True
+NumRetriesBeforeForcingGC=0
 UseVTableConstructors=True
 
 [Internationalization]
diff --git a/Engine/Source/Runtime/CoreUObject/Private/UObject/GarbageCollection.cpp b/Engine/Source/Runtime/CoreUObject/Private/UObject/GarbageCollection.cpp
--- a/Engine/Source/Runtime/CoreUObject/Private/UObject/GarbageCollection.cpp
+++ b/Engine/Source/Runtime/CoreUObject/Private/UObject/GarbageCollection.cpp
@@ -1190,6 +1190,9 @@ COREUOBJECT_API EditorPostReachabilityAnalysisCallbackType EditorPostReachabilit
 static const auto CVarAllowParallelGC = 
 IConsoleManager::Get().RegisterConsoleVariable(TEXT("AllowParallelGC"), (!PLATFORM_MAC || !WITH_EDITORONLY_DATA) ? 1 : 0, TEXT("Used to control parallel GC."))->AsVariableInt();
 
+// This counts how many times GC was skipped
+static int32 GNumAttemptsSinceLastGC = 0;
+
 /** 
  * Deletes all unreferenced objects, keeping objects that have any of the passed in KeepFlags set
  *
@@ -1201,6 +1204,9 @@ void CollectGarbageInternal(EObjectFlags KeepFlags, bool bPerformFullPurge)
 	// We can't collect garbage while there's a load in progress. E.g. one potential issue is Import.XObject
 	check(!IsLoading());
 
+	// Reset GC skip counter
+	GNumAttemptsSinceLastGC = 0;
+
 	// Helper class to register FlushAsyncLoadingCallback on first GC run.
 	struct FAddFlushAsyncLoadingCallback
 	{
@@ -1352,6 +1358,26 @@ bool TryCollectGarbage(EObjectFlags KeepFlags, bool bPerformFullPurge)
 {
 	// No other thread may be performing UOBject operations while we're running
 	bool bCanRunGC = GGarbageCollectionGuardCritical.TryGCLock();
+	if (!bCanRunGC)
+	{
+		// One time init for ini override if we should log a warning when time limit is exceeded
+		static struct FForceGCSettings
+		{
+			int32 NumRetriesBeforeForcingGC;
+			FForceGCSettings()
+				: NumRetriesBeforeForcingGC(0) // 0 means 'never force GC'
+			{
+				check(GConfig);
+				GConfig->GetInt(TEXT("Core.System"), TEXT("NumRetriesBeforeForcingGC"), NumRetriesBeforeForcingGC, GEngineIni);
+			}
+		} ForceGCSettings;
+		if (ForceGCSettings.NumRetriesBeforeForcingGC > 0 && GNumAttemptsSinceLastGC > ForceGCSettings.NumRetriesBeforeForcingGC)
+		{
+			// Force GC and block main thread
+			bCanRunGC = true;
+			UE_LOG(LogGarbage, Warning, TEXT("TryCollectGarbage: forcing GC after %d skipped attempts."), GNumAttemptsSinceLastGC)
+		}
+	}
 	if (bCanRunGC)
 	{
 		// Perform actual garbage collection
@@ -1360,6 +1386,11 @@ bool TryCollectGarbage(EObjectFlags KeepFlags, bool bPerformFullPurge)
 		// Other threads are free to use UObjects
 		GGarbageCollectionGuardCritical.GCUnlock();
 	}
+	else
+	{
+		GNumAttemptsSinceLastGC++;
+	}
+
 	return bCanRunGC;
 }
 
