Adding more configurable settings to time limit exceeded warning when async loading

diff --git a/Engine/Config/BaseEngine.ini b/Engine/Config/BaseEngine.ini
--- a/Engine/Config/BaseEngine.ini
+++ b/Engine/Config/BaseEngine.ini
@@ -1331,6 +1331,8 @@ FlushStreamingOnGC=False
 MinBulkDataSizeForAsyncLoading=131072
 AsyncLoadingThreadEnabled=False
 WarnIfTimeLimitExceeded=True
+TimeLimitExceededMultiplier=1.5
+TimeLimitExceededMinTime=0.005
 NumRetriesBeforeForcingGC=0
 UseVTableConstructors=True
 
diff --git a/Engine/Source/Runtime/CoreUObject/Private/Serialization/AsyncLoading.cpp b/Engine/Source/Runtime/CoreUObject/Private/Serialization/AsyncLoading.cpp
--- a/Engine/Source/Runtime/CoreUObject/Private/Serialization/AsyncLoading.cpp
+++ b/Engine/Source/Runtime/CoreUObject/Private/Serialization/AsyncLoading.cpp
@@ -224,18 +224,26 @@ static FORCEINLINE bool IsTimeLimitExceeded(double InTickStartTime, bool bUseTim
 		// One time init for ini override if we should log a warning when time limit is exceeded
 		static struct FWarnIfTimeLimitExceeded
 		{
-			bool Value;
+			bool WarnignEnabled;
+			double TimeMultiplier;
+			double MinTime;
 			FWarnIfTimeLimitExceeded()
+				: WarnignEnabled(true)
+				, TimeMultiplier(1.0)
+				, MinTime(0.005)
 			{
 				check(GConfig);
-				bool bConfigValue = true; // Default to true
-				GConfig->GetBool(TEXT("Core.System"), TEXT("WarnIfTimeLimitExceeded"), bConfigValue, GEngineIni);
-				Value = bConfigValue;
+				GConfig->GetBool(TEXT("Core.System"), TEXT("WarnIfTimeLimitExceeded"), WarnignEnabled, GEngineIni);
+				GConfig->GetDouble(TEXT("Core.System"), TEXT("TimeLimitExceededMultiplier"), TimeMultiplier, GEngineIni);
+				GConfig->GetDouble(TEXT("Core.System"), TEXT("TimeLimitExceededMinTime"), MinTime, GEngineIni);
 			}
 		} WarnIfTimeLimitExceeded;
 
 		// Log single operations that take longer than time limit (but only in cooked builds)
-		if (FPlatformProperties::RequiresCookedData() && WarnIfTimeLimitExceeded.Value && (CurrentTime - InTickStartTime) > (1.5 * InTimeLimit))
+		if (FPlatformProperties::RequiresCookedData() && 
+			WarnIfTimeLimitExceeded.WarnignEnabled && 
+			(CurrentTime - InTickStartTime) > WarnIfTimeLimitExceeded.MinTime &&
+			(CurrentTime - InTickStartTime) > (WarnIfTimeLimitExceeded.TimeMultiplier * InTimeLimit))
 		{			
 			UE_LOG(LogStreaming, Warning, TEXT("IsTimeLimitExceeded: %s %s took (less than) %5.2f ms"),
 				InLastTypeOfWorkPerformed ? InLastTypeOfWorkPerformed : TEXT("unknown"),
