Remove unused UObject::CreateArchetype and unimplemented UObject::UpdateArchetype

diff --git a/Engine/Source/Runtime/CoreUObject/Private/UObject/Obj.cpp b/Engine/Source/Runtime/CoreUObject/Private/UObject/Obj.cpp
--- a/Engine/Source/Runtime/CoreUObject/Private/UObject/Obj.cpp
+++ b/Engine/Source/Runtime/CoreUObject/Private/UObject/Obj.cpp
@@ -2197,34 +2197,6 @@ void UObject::RetrieveReferencers( TArray<FReferencerInformation>* OutInternalRe
 	}
 }
 
-UObject* UObject::CreateArchetype( const TCHAR* ArchetypeName, UObject* ArchetypeOuter, UObject* AlternateArchetype/*=NULL*/, FObjectInstancingGraph* InstanceGraph/*=NULL*/ )
-{
-	check(ArchetypeName);
-	check(ArchetypeOuter);
-
-	EObjectFlags ArchetypeObjectFlags = RF_Public | RF_ArchetypeObject;
-	
-	// Archetypes residing directly in packages need to be marked RF_Standalone
-	if( dynamic_cast<UPackage*>(ArchetypeOuter) )
-	{
-		ArchetypeObjectFlags |= RF_Standalone;
-	}
-
-	UObject* ArchetypeObject = StaticConstructObject_Internal(GetClass(), ArchetypeOuter, ArchetypeName, ArchetypeObjectFlags, this, true, InstanceGraph);
-	check(ArchetypeObject);
-
-	UObject* NewArchetype = AlternateArchetype == NULL
-		? GetArchetype()
-		: AlternateArchetype;
-
-	check(NewArchetype);
-	// make sure the alternate archetype has the same class
-	check(NewArchetype->GetClass()==GetClass());
-
-	return ArchetypeObject;
-}
-
-
 void UObject::ParseParms( const TCHAR* Parms )
 {
 	if( !Parms )
diff --git a/Engine/Source/Runtime/CoreUObject/Public/UObject/UObject.h b/Engine/Source/Runtime/CoreUObject/Public/UObject/UObject.h
--- a/Engine/Source/Runtime/CoreUObject/Public/UObject/UObject.h
+++ b/Engine/Source/Runtime/CoreUObject/Public/UObject/UObject.h
@@ -861,26 +861,6 @@ public:
 	 */
 	void SetLinker( FLinkerLoad* LinkerLoad, int32 LinkerIndex, bool bShouldDetachExisting=true );
 
-	/**
-	 * Creates a new archetype based on this UObject.  The archetype's property values will match
-	 * the current values of this UObject.
-	 *
-	 * @param	ArchetypeName			the name for the new class
-	 * @param	ArchetypeOuter			the outer to create the new class in (package?)
-	 * @param	AlternateArchetype		if specified, is set as the ObjectArchetype for the newly created archetype, after the new archetype
-	 *									is initialized against "this".  Should only be specified in cases where you need the new archetype to
-	 *									inherit the property values of this object, but don't want this object to be the new archetype's ObjectArchetype.
-	 * @param	InstanceGraph			contains the mappings of instanced objects and components to their templates
-	 *
-	 * @return	a pointer to a UObject which has values identical to this object
-	 */
-	UObject* CreateArchetype( const TCHAR* ArchetypeName, UObject* ArchetypeOuter, UObject* AlternateArchetype=NULL, struct FObjectInstancingGraph* InstanceGraph=NULL );
-
-	/**
-	 *	Update the ObjectArchetype of this UObject based on this UObject's properties.
-	 */
-	void UpdateArchetype();
-
 	/**
 	 * Return the template that an object with this class, outer and name would be
 	 * 
