Core - Added native serialization for immutable structures
Core - Added POD for various math classes

diff --git a/Engine/Source/Runtime/Core/Public/Math/Box.h b/Engine/Source/Runtime/Core/Public/Math/Box.h
--- a/Engine/Source/Runtime/Core/Public/Math/Box.h
+++ b/Engine/Source/Runtime/Core/Public/Math/Box.h
@@ -445,6 +445,12 @@ public:
 	{
 		return Ar << Box.Min << Box.Max << Box.IsValid;
 	}
+
+	bool Serialize( FArchive& Ar )
+	{
+		Ar << *this;
+		return true;
+	}
 };
 
 
diff --git a/Engine/Source/Runtime/Core/Public/Math/BoxSphereBounds.h b/Engine/Source/Runtime/Core/Public/Math/BoxSphereBounds.h
--- a/Engine/Source/Runtime/Core/Public/Math/BoxSphereBounds.h
+++ b/Engine/Source/Runtime/Core/Public/Math/BoxSphereBounds.h
@@ -323,4 +323,6 @@ FORCEINLINE FBoxSphereBounds FBoxSphereBounds::operator+( const FBoxSphereBounds
 FORCEINLINE FString FBoxSphereBounds::ToString() const
 {
 	return FString::Printf(TEXT("Origin=%s, BoxExtent=(%s), SphereRadius=(%f)"), *Origin.ToString(), *BoxExtent.ToString(), SphereRadius);
-}
\ No newline at end of file
+}
+
+template <> struct TIsPODType<FBoxSphereBounds> { enum { Value = true }; };
\ No newline at end of file
diff --git a/Engine/Source/Runtime/Core/Public/Math/Color.h b/Engine/Source/Runtime/Core/Public/Math/Color.h
--- a/Engine/Source/Runtime/Core/Public/Math/Color.h
+++ b/Engine/Source/Runtime/Core/Public/Math/Color.h
@@ -64,6 +64,12 @@ struct FLinearColor
 		return Ar << Color.R << Color.G << Color.B << Color.A;
 	}
 
+	bool Serialize( FArchive& Ar )
+	{
+		Ar << *this;
+		return true;
+	}
+
 	// Conversions.
 	CORE_API FColor ToRGBE() const;
 
@@ -445,6 +451,12 @@ public:
 		return Ar << Color.DWColor();
 	}
 
+	bool Serialize( FArchive& Ar )
+	{
+		Ar << *this;
+		return true;
+	}
+
 	// Operators.
 	FORCEINLINE bool operator==( const FColor &C ) const
 	{
diff --git a/Engine/Source/Runtime/Core/Public/Math/IntPoint.h b/Engine/Source/Runtime/Core/Public/Math/IntPoint.h
--- a/Engine/Source/Runtime/Core/Public/Math/IntPoint.h
+++ b/Engine/Source/Runtime/Core/Public/Math/IntPoint.h
@@ -274,6 +274,12 @@ public:
 	{
 		return Ar << Point.X << Point.Y;
 	}
+
+	bool Serialize( FArchive& Ar )
+	{
+		Ar << *this;
+		return true;
+	}
 };
 
 
@@ -481,3 +487,5 @@ FORCEINLINE FString FIntPoint::ToString() const
 {
 	return FString::Printf(TEXT("X=%d Y=%d"), X, Y);
 }
+
+template <> struct TIsPODType<FIntPoint> { enum { Value = true }; };
diff --git a/Engine/Source/Runtime/Core/Public/Math/IntVector.h b/Engine/Source/Runtime/Core/Public/Math/IntVector.h
--- a/Engine/Source/Runtime/Core/Public/Math/IntVector.h
+++ b/Engine/Source/Runtime/Core/Public/Math/IntVector.h
@@ -238,6 +238,12 @@ public:
 	{
 		return Ar << Vector.X << Vector.Y << Vector.Z;
 	}
+
+	bool Serialize( FArchive& Ar )
+	{
+		Ar << *this;
+		return true;
+	}
 };
 
 
@@ -416,3 +422,5 @@ FORCEINLINE uint32 GetTypeHash(const FIntVector& Vector)
 {
 	return FCrc::MemCrc_DEPRECATED(&Vector,sizeof(FIntVector));
 }
+
+template <> struct TIsPODType<FIntVector> { enum { Value = true }; };
\ No newline at end of file
diff --git a/Engine/Source/Runtime/Core/Public/Math/Matrix.h b/Engine/Source/Runtime/Core/Public/Math/Matrix.h
--- a/Engine/Source/Runtime/Core/Public/Math/Matrix.h
+++ b/Engine/Source/Runtime/Core/Public/Math/Matrix.h
@@ -321,6 +321,16 @@ public:
 	 */
 	friend CORE_API FArchive& operator<<(FArchive& Ar,FMatrix& M);
 
+	bool Serialize( FArchive& Ar )
+	{
+		if (Ar.UE4Ver() >= VER_UE4_ADDED_NATIVE_SERIALIZATION_FOR_IMMUTABLE_STRUCTURES)
+		{
+			Ar << *this;
+			return true;
+		}
+		return false;
+	}
+
 	/**
 	 * Convert this Atom to the 3x4 transpose of the transformation matrix.
 	 */
diff --git a/Engine/Source/Runtime/Core/Public/Math/OrientedBox.h b/Engine/Source/Runtime/Core/Public/Math/OrientedBox.h
--- a/Engine/Source/Runtime/Core/Public/Math/OrientedBox.h
+++ b/Engine/Source/Runtime/Core/Public/Math/OrientedBox.h
@@ -113,3 +113,5 @@ FORCEINLINE FFloatInterval FOrientedBox::Project( const FVector& Axis ) const
 
 	return ProjectionInterval;
 }
+
+template <> struct TIsPODType<FOrientedBox> { enum { Value = true }; };
\ No newline at end of file
diff --git a/Engine/Source/Runtime/Core/Public/Math/Plane.h b/Engine/Source/Runtime/Core/Public/Math/Plane.h
--- a/Engine/Source/Runtime/Core/Public/Math/Plane.h
+++ b/Engine/Source/Runtime/Core/Public/Math/Plane.h
@@ -239,6 +239,16 @@ public:
 		return Ar << (FVector&)P << P.W;
 	}
 
+	bool Serialize( FArchive& Ar )
+	{
+		if (Ar.UE4Ver() >= VER_UE4_ADDED_NATIVE_SERIALIZATION_FOR_IMMUTABLE_STRUCTURES)
+		{
+			Ar << *this;
+			return true;
+		}
+		return false;
+	}
+
 	/**
 	 * Serializes the vector compressed for e.g. network transmission.
 	 * @param Ar Archive to serialize to/ from.
diff --git a/Engine/Source/Runtime/Core/Public/Math/Quat.h b/Engine/Source/Runtime/Core/Public/Math/Quat.h
--- a/Engine/Source/Runtime/Core/Public/Math/Quat.h
+++ b/Engine/Source/Runtime/Core/Public/Math/Quat.h
@@ -455,6 +455,12 @@ public:
 		return Ar << F.X << F.Y << F.Z << F.W;
 	}
 
+	bool Serialize( FArchive& Ar )
+	{
+		Ar << *this;
+		return true;
+	}
+
 } GCC_ALIGN(16);
 
 
diff --git a/Engine/Source/Runtime/Core/Public/Math/Rotator.h b/Engine/Source/Runtime/Core/Public/Math/Rotator.h
--- a/Engine/Source/Runtime/Core/Public/Math/Rotator.h
+++ b/Engine/Source/Runtime/Core/Public/Math/Rotator.h
@@ -378,6 +378,12 @@ public:
 		Ar << R.Pitch << R.Yaw << R.Roll;
 		return Ar;
 	}
+
+	bool Serialize( FArchive& Ar )
+	{
+		Ar << *this;
+		return true;
+	}
 };
 
 
diff --git a/Engine/Source/Runtime/Core/Public/Math/TwoVectors.h b/Engine/Source/Runtime/Core/Public/Math/TwoVectors.h
--- a/Engine/Source/Runtime/Core/Public/Math/TwoVectors.h
+++ b/Engine/Source/Runtime/Core/Public/Math/TwoVectors.h
@@ -217,6 +217,12 @@ public:
 	{
 		return Ar << TwoVectors.v1 << TwoVectors.v2;
 	}
+
+	bool Serialize( FArchive& Ar )
+	{
+		Ar << *this;
+		return true;
+	}
 };
 
 
@@ -424,3 +430,5 @@ FORCEINLINE FString FTwoVectors::ToString() const
 {
 	return FString::Printf(TEXT("V1=(%s) V2=(%s)"), *v1.ToString(), *v2.ToString());
 }
+
+template <> struct TIsPODType<FTwoVectors> { enum { Value = true }; };
\ No newline at end of file
diff --git a/Engine/Source/Runtime/Core/Public/Math/Vector.h b/Engine/Source/Runtime/Core/Public/Math/Vector.h
--- a/Engine/Source/Runtime/Core/Public/Math/Vector.h
+++ b/Engine/Source/Runtime/Core/Public/Math/Vector.h
@@ -909,6 +909,12 @@ public:
 		return Ar << V.X << V.Y << V.Z;
 	}
 	
+	bool Serialize( FArchive& Ar )
+	{
+		Ar << *this;
+		return true;
+	}
+
 	/** 
 	 * Network serialization function.
 	 * FVectors NetSerialize without quantization (ie exact values are serialized).
diff --git a/Engine/Source/Runtime/Core/Public/Math/Vector2D.h b/Engine/Source/Runtime/Core/Public/Math/Vector2D.h
--- a/Engine/Source/Runtime/Core/Public/Math/Vector2D.h
+++ b/Engine/Source/Runtime/Core/Public/Math/Vector2D.h
@@ -462,6 +462,12 @@ public:
 		return Ar << V.X << V.Y;
 	}
 
+	bool Serialize( FArchive& Ar )
+	{
+		Ar << *this;
+		return true;
+	}
+
 #if ENABLE_NAN_DIAGNOSTIC
 	FORCEINLINE void DiagnosticCheckNaN() const
 	{
diff --git a/Engine/Source/Runtime/Core/Public/Math/Vector4.h b/Engine/Source/Runtime/Core/Public/Math/Vector4.h
--- a/Engine/Source/Runtime/Core/Public/Math/Vector4.h
+++ b/Engine/Source/Runtime/Core/Public/Math/Vector4.h
@@ -368,6 +368,12 @@ public:
 		return Ar << V.X << V.Y << V.Z << V.W;
 	}
 
+	bool Serialize( FArchive& Ar )
+	{
+		Ar << *this;
+		return true;
+	}
+
 } GCC_ALIGN(16);
 
 
@@ -678,3 +684,5 @@ FORCEINLINE FVector4 FVector4::operator/( const FVector4& V ) const
 {
 	return FVector4( X / V.X, Y / V.Y, Z / V.Z, W / V.W );
 }
+
+template <> struct TIsPODType<FVector4> { enum { Value = true }; };
\ No newline at end of file
diff --git a/Engine/Source/Runtime/Core/Public/Misc/Guid.h b/Engine/Source/Runtime/Core/Public/Misc/Guid.h
--- a/Engine/Source/Runtime/Core/Public/Misc/Guid.h
+++ b/Engine/Source/Runtime/Core/Public/Misc/Guid.h
@@ -179,6 +179,12 @@ public:
 		return Ar << G.A << G.B << G.C << G.D;
 	}
 
+	bool Serialize( FArchive& Ar )
+	{
+		Ar << *this;
+		return true;
+	}
+
 public:
 
 	/**
@@ -305,3 +311,5 @@ public:
 	/** Holds the fourth component. */
 	uint32 D;
 };
+
+template <> struct TIsPODType<FGuid> { enum { Value = true }; };
\ No newline at end of file
diff --git a/Engine/Source/Runtime/Core/Public/UObject/ObjectVersion.h b/Engine/Source/Runtime/Core/Public/UObject/ObjectVersion.h
--- a/Engine/Source/Runtime/Core/Public/UObject/ObjectVersion.h
+++ b/Engine/Source/Runtime/Core/Public/UObject/ObjectVersion.h
@@ -1,4 +1,4 @@
-﻿// Copyright 1998-2015 Epic Games, Inc. All Rights Reserved.
+// Copyright 1998-2015 Epic Games, Inc. All Rights Reserved.
 
 #pragma once
 
@@ -526,7 +526,8 @@ enum EUnrealEngineObjectUE4Version
 	VER_UE4_ASSET_IMPORT_DATA_AS_JSON,
 	// Legacy gamma support for textures.
 	VER_UE4_TEXTURE_LEGACY_GAMMA,
-
+	// Added WithSerializer for basic native structures like FVector, FColor etc to improve serialization performance
+	VER_UE4_ADDED_NATIVE_SERIALIZATION_FOR_IMMUTABLE_STRUCTURES,
 	// -----<new versions can be added before this line>-------------------------------------------------
 	// - this needs to be the last line (see note below)
 	VER_UE4_AUTOMATIC_VERSION_PLUS_ONE,
diff --git a/Engine/Source/Runtime/CoreUObject/Private/UObject/Property.cpp b/Engine/Source/Runtime/CoreUObject/Private/UObject/Property.cpp
--- a/Engine/Source/Runtime/CoreUObject/Private/UObject/Property.cpp
+++ b/Engine/Source/Runtime/CoreUObject/Private/UObject/Property.cpp
@@ -20,6 +20,7 @@ struct TStructOpsTypeTraits<FVector> : public TStructOpsTypeTraitsBase
 		WithNoInitConstructor = true,
 		WithZeroConstructor = true,
 		WithNetSerializer = true,
+		WithSerializer = true,
 	};
 };
 IMPLEMENT_STRUCT(Vector);
@@ -31,20 +32,11 @@ struct TStructOpsTypeTraits<FIntPoint> : public TStructOpsTypeTraitsBase
 	{
 		WithNoInitConstructor = true,
 		WithZeroConstructor = true,
+		WithSerializer = true,
 	};
 };
 IMPLEMENT_STRUCT(IntPoint);
 
-template<>
-struct TStructOpsTypeTraits<FIntRect> : public TStructOpsTypeTraitsBase
-{
-	enum 
-	{
-		WithZeroConstructor = true,
-	};
-};
-IMPLEMENT_STRUCT(IntRect);
-
 template<>
 struct TStructOpsTypeTraits<FVector2D> : public TStructOpsTypeTraitsBase
 {
@@ -53,6 +45,7 @@ struct TStructOpsTypeTraits<FVector2D> : public TStructOpsTypeTraitsBase
 		WithNoInitConstructor = true,
 		WithZeroConstructor = true,
 		WithNetSerializer = true,
+		WithSerializer = true,
 	};
 };
 IMPLEMENT_STRUCT(Vector2D);
@@ -64,6 +57,7 @@ struct TStructOpsTypeTraits<FVector4> : public TStructOpsTypeTraitsBase
 	{
 		WithNoInitConstructor = true,
 		WithZeroConstructor = true,
+		WithSerializer = true,
 	};
 };
 IMPLEMENT_STRUCT(Vector4);
@@ -76,6 +70,7 @@ struct TStructOpsTypeTraits<FPlane> : public TStructOpsTypeTraitsBase
 		WithNoInitConstructor = true,
 		WithZeroConstructor = true,
 		WithNetSerializer = true,
+		WithSerializer = true,
 	};
 };
 IMPLEMENT_STRUCT(Plane);
@@ -88,6 +83,7 @@ struct TStructOpsTypeTraits<FRotator> : public TStructOpsTypeTraitsBase
 		WithNoInitConstructor = true,
 		WithZeroConstructor = true,
 		WithNetSerializer = true,
+		WithSerializer = true,
 	};
 };
 IMPLEMENT_STRUCT(Rotator);
@@ -99,6 +95,7 @@ struct TStructOpsTypeTraits<FBox> : public TStructOpsTypeTraitsBase
 	{
 		WithNoInitConstructor = true,
 		WithZeroConstructor = true,
+		WithSerializer = true,
 	};
 };
 IMPLEMENT_STRUCT(Box);
@@ -110,21 +107,11 @@ struct TStructOpsTypeTraits<FMatrix> : public TStructOpsTypeTraitsBase
 	{
 		WithNoInitConstructor = true,
 		WithZeroConstructor = true,
+		WithSerializer = true,
 	};
 };
 IMPLEMENT_STRUCT(Matrix);
 
-template<>
-struct TStructOpsTypeTraits<FSphere> : public TStructOpsTypeTraitsBase
-{
-	enum 
-	{
-		WithNoInitConstructor = true,
-		WithZeroConstructor = true,
-	};
-};
-IMPLEMENT_STRUCT(Sphere);
-
 template<>
 struct TStructOpsTypeTraits<FBoxSphereBounds> : public TStructOpsTypeTraitsBase
 {
@@ -149,6 +136,7 @@ struct TStructOpsTypeTraits<FLinearColor> : public TStructOpsTypeTraitsBase
 	{
 		WithNoInitConstructor = true,
 		WithZeroConstructor = true,
+		WithSerializer = true,
 	};
 };
 IMPLEMENT_STRUCT(LinearColor);
@@ -160,6 +148,7 @@ struct TStructOpsTypeTraits<FColor> : public TStructOpsTypeTraitsBase
 	{
 		WithNoInitConstructor = true,
 		WithZeroConstructor = true,
+		WithSerializer = true,
 	};
 };
 IMPLEMENT_STRUCT(Color);
@@ -183,6 +172,7 @@ struct TStructOpsTypeTraits<FTwoVectors> : public TStructOpsTypeTraitsBase
 	enum 
 	{
 		WithZeroConstructor = true,
+		WithSerializer = true,
 		WithNoDestructor = true,
 	};
 };
@@ -196,6 +186,7 @@ struct TStructOpsTypeTraits<FGuid> : public TStructOpsTypeTraitsBase
 		WithExportTextItem = true,
 		WithImportTextItem = true,
 		WithZeroConstructor = true,
+		WithSerializer = true,
 	};
 };
 IMPLEMENT_STRUCT(Guid);
