[UE-16907] Properly handling deferred Blueprint exports that get loaded as null (like when a dependent plugin has been disabled).

diff --git a/Engine/Source/Runtime/CoreUObject/Private/Blueprint/BlueprintSupport.cpp b/Engine/Source/Runtime/CoreUObject/Private/Blueprint/BlueprintSupport.cpp
--- a/Engine/Source/Runtime/CoreUObject/Private/Blueprint/BlueprintSupport.cpp
+++ b/Engine/Source/Runtime/CoreUObject/Private/Blueprint/BlueprintSupport.cpp
@@ -1214,17 +1214,27 @@ void FLinkerLoad::ResolveDeferredExports(UClass* LoadClass)
 			if (ULinkerPlaceholderExportObject* PlaceholderExport = Cast<ULinkerPlaceholderExportObject>(Export.Object))
 			{
 				DEFERRED_DEPENDENCY_CHECK(Export.ClassIndex.IsImport());
-				UClass* ExportClass = GetExportLoadClass(ExportIndex);
-				DEFERRED_DEPENDENCY_CHECK((ExportClass != nullptr) && !ExportClass->HasAnyClassFlags(CLASS_Intrinsic) && ExportClass->HasAnyClassFlags(CLASS_CompiledFromBlueprint));
-				FLinkerLoad* ClassLinker = ExportClass->GetLinker();
-				DEFERRED_DEPENDENCY_CHECK((ClassLinker != nullptr) && (ClassLinker != this));
 
+				// make sure we're not already in ForceRegenerateClass() for
+				// this export (that could cause some bad infinite recursion)
+				DEFERRED_DEPENDENCY_CHECK(!IsExportBeingResolved(ExportIndex));
+
+				UClass* ExportClass = GetExportLoadClass(ExportIndex);
+				// export class could be null... we create these placeholder 
+				// exports for objects that are instances of an external 
+				// (Blueprint) type, not knowing if that type (class) will 
+				// successfully load... it may resolve to null in scenarios 
+				// where its super class has been deleted, or its super belonged 
+				// to a plugin that has been removed/disabled 
+				if (ExportClass != nullptr)
 				{
-					// make sure we're not already in ForceRegenerateClass() for
-					// this export (that could cause some bad infinite recursion)
-					DEFERRED_DEPENDENCY_CHECK(!IsExportBeingResolved(ExportIndex));
-					FScopedResolvingExportTracker ForceRegenGuard(this, ExportIndex);
+#if USE_DEFERRED_DEPENDENCY_CHECK_VERIFICATION_TESTS
+					DEFERRED_DEPENDENCY_CHECK(!ExportClass->HasAnyClassFlags(CLASS_Intrinsic) && ExportClass->HasAnyClassFlags(CLASS_CompiledFromBlueprint));
+					FLinkerLoad* ClassLinker = ExportClass->GetLinker();
+					DEFERRED_DEPENDENCY_CHECK((ClassLinker != nullptr) && (ClassLinker != this));
+#endif // USE_DEFERRED_DEPENDENCY_CHECK_VERIFICATION_TESTS
 
+					FScopedResolvingExportTracker ForceRegenGuard(this, ExportIndex);
 					// make sure this export's class is fully regenerated before  
 					// we instantiate it (so we don't have to re-inst on load)
 					ForceRegenerateClass(ExportClass);
@@ -1243,12 +1253,15 @@ void FLinkerLoad::ResolveDeferredExports(UClass* LoadClass)
 
 				// if we hadn't used a ULinkerPlaceholderExportObject in place of 
 				// the expected export, then someone may have wanted it preloaded
-				Preload(ExportObj);
+				if (ExportObj != nullptr)
+				{
+					Preload(ExportObj);
+				}
 
 #if USE_DEFERRED_DEPENDENCY_CHECK_VERIFICATION_TESTS
 				FReferencerInformationList UnresolvedReferences;
 				UObject* PlaceholderObj = PlaceholderExport;
-				bool const bIsReferenced = false;// IsReferenced(PlaceholderObj, RF_NoFlags, /*bCheckSubObjects =*/false, &UnresolvedReferences);
+				bool const bIsReferenced = IsReferenced(PlaceholderObj, RF_NoFlags, /*bCheckSubObjects =*/false, &UnresolvedReferences);
 				DEFERRED_DEPENDENCY_CHECK(!bIsReferenced);
 #endif // USE_DEFERRED_DEPENDENCY_CHECK_VERIFICATION_TESTS
 			}
