Backing out bad merge (CL #2584863)

diff --git a/Engine/Source/Runtime/CoreUObject/Private/Serialization/AsyncLoading.cpp b/Engine/Source/Runtime/CoreUObject/Private/Serialization/AsyncLoading.cpp
--- a/Engine/Source/Runtime/CoreUObject/Private/Serialization/AsyncLoading.cpp
+++ b/Engine/Source/Runtime/CoreUObject/Private/Serialization/AsyncLoading.cpp
@@ -245,29 +245,10 @@ static FORCEINLINE bool IsTimeLimitExceeded(double InTickStartTime, bool bUseTim
 		double CurrentTime = FPlatformTime::Seconds();
 		bTimeLimitExceeded = CurrentTime - InTickStartTime > InTimeLimit;
 
-		// One time init for ini override if we should log a warning when time limit is exceeded
-		static struct FWarnIfTimeLimitExceeded
-		{
-			bool WarningEnabled;
-			double TimeMultiplier;
-			double MinTime;
-			FWarnIfTimeLimitExceeded()
-				: WarningEnabled(true)
-				, TimeMultiplier(1.0)
-				, MinTime(0.005)
-			{
-				check(GConfig);
-				GConfig->GetBool(TEXT("Core.System"), TEXT("WarnIfTimeLimitExceeded"), WarningEnabled, GEngineIni);
-				GConfig->GetDouble(TEXT("Core.System"), TEXT("TimeLimitExceededMultiplier"), TimeMultiplier, GEngineIni);
-				GConfig->GetDouble(TEXT("Core.System"), TEXT("TimeLimitExceededMinTime"), MinTime, GEngineIni);
-			}
-		} WarnIfTimeLimitExceeded;
-
 		// Log single operations that take longer than time limit (but only in cooked builds)
-		if (FPlatformProperties::RequiresCookedData() && 
-			WarnIfTimeLimitExceeded.WarningEnabled && 
-			(CurrentTime - InTickStartTime) > WarnIfTimeLimitExceeded.MinTime &&
-			(CurrentTime - InTickStartTime) > (WarnIfTimeLimitExceeded.TimeMultiplier * InTimeLimit))
+		if (GWarnIfTimeLimitExceeded && 
+			(CurrentTime - InTickStartTime) > GTimeLimitExceededMinTime &&
+			(CurrentTime - InTickStartTime) > (GTimeLimitExceededMultiplier * InTimeLimit))
 		{			
 			UE_LOG(LogStreaming, Warning, TEXT("IsTimeLimitExceeded: %s %s took (less than) %5.2f ms"),
 				InLastTypeOfWorkPerformed ? InLastTypeOfWorkPerformed : TEXT("unknown"),
