CoreUObject - Better stats for garbage collecting

diff --git a/Engine/Source/Runtime/CoreUObject/Private/UObject/GarbageCollection.cpp b/Engine/Source/Runtime/CoreUObject/Private/UObject/GarbageCollection.cpp
--- a/Engine/Source/Runtime/CoreUObject/Private/UObject/GarbageCollection.cpp
+++ b/Engine/Source/Runtime/CoreUObject/Private/UObject/GarbageCollection.cpp
@@ -13,6 +13,8 @@
    Garbage collection.
 -----------------------------------------------------------------------------*/
 
+DECLARE_STATS_GROUP( TEXT( "Garbage Collection" ), STATGROUP_GC, STATCAT_Advanced );
+
 DEFINE_LOG_CATEGORY_STATIC(LogGarbage, Warning, All);
 
 #define PERF_DETAILED_PER_CLASS_GC_STATS				(LOOKING_FOR_PERF_ISSUES)
@@ -563,6 +565,8 @@ public:
 	 */
 	void PerformReachabilityAnalysis( EObjectFlags KeepFlags, bool bForceSingleThreaded = false )
 	{
+		DECLARE_SCOPE_CYCLE_COUNTER( TEXT( "FArchiveRealtimeGC::PerformReachabilityAnalysis" ), STAT_FArchiveRealtimeGC_PerformReachabilityAnalysis, STATGROUP_GC );
+
 		UObject* CurrentObject = NULL;
 
 		/** Growing array of objects that require serialization */
@@ -665,6 +669,8 @@ public:
 
 	void ProcessObjectArray(TArray<UObject*>& InObjectsToSerializeArray, FGraphEventRef& MyCompletionGraphEvent)
 	{		
+		DECLARE_SCOPE_CYCLE_COUNTER( TEXT( "FArchiveRealtimeGC::ProcessObjectArray" ), STAT_FArchiveRealtimeGC_ProcessObjectArray, STATGROUP_GC );
+
 		UObject* CurrentObject = NULL;
 
 		const int32 MinDesiredObjectsPerSubTask = 128; // sometimes there will be less, a lot less
@@ -931,6 +937,8 @@ public:
  */
 void IncrementalPurgeGarbage( bool bUseTimeLimit, float TimeLimit )
 {
+	DECLARE_SCOPE_CYCLE_COUNTER( TEXT( "IncrementalPurgeGarbage" ), STAT_IncrementalPurgeGarbage, STATGROUP_GC );
+
 	if (GExitPurge)
 	{
 		GObjPurgeIsRequired = true;
@@ -1224,6 +1232,8 @@ static FAutoConsoleVariableRef CVarFlushStreamingOnGC(
  */
 void CollectGarbageInternal(EObjectFlags KeepFlags, bool bPerformFullPurge)
 {
+	DECLARE_SCOPE_CYCLE_COUNTER( TEXT( "CollectGarbageInternal" ), STAT_CollectGarbageInternal, STATGROUP_GC );
+
 	// We can't collect garbage while there's a load in progress. E.g. one potential issue is Import.XObject
 	check(!IsLoading());
 
@@ -1258,6 +1268,8 @@ void CollectGarbageInternal(EObjectFlags KeepFlags, bool bPerformFullPurge)
 	// Only verify assumptions if option is enabled. This avoids false positives in the Editor or commandlets.
 	if( GetUObjectArray().DisregardForGCEnabled() && GShouldVerifyGCAssumptions )
 	{
+		DECLARE_SCOPE_CYCLE_COUNTER( TEXT( "CollectGarbageInternal.VerifyGCAssumptions" ), STAT_CollectGarbageInternal_VerifyGCAssumptions, STATGROUP_GC );
+
 		// Verify that objects marked to be disregarded for GC are not referencing objects that are not part of the root set.
 		for( FObjectIterator It; It; ++It )
 		{
@@ -1320,20 +1332,24 @@ void CollectGarbageInternal(EObjectFlags KeepFlags, bool bPerformFullPurge)
 	}
 #endif // WITH_EDITOR
 
-	// Unhash all unreachable objects.
-	const double StartTime = FPlatformTime::Seconds();
-	for ( FRawObjectIterator It(true); It; ++It )
 	{
-		//@todo UE4 - A prefetch was removed here. Re-add it. It wasn't right anyway, since it was ten items ahead and the consoles on have 8 prefetch slots
+		DECLARE_SCOPE_CYCLE_COUNTER( TEXT( "CollectGarbageInternal.UnhashUnreachable" ), STAT_CollectGarbageInternal_UnhashUnreachable, STATGROUP_GC );
 
-		UObject* Object = *It;
-		if( Object->HasAnyFlags( RF_Unreachable ) )
+		// Unhash all unreachable objects.
+		const double StartTime = FPlatformTime::Seconds();
+		for (FRawObjectIterator It( true ); It; ++It)
 		{
-			// Begin the object's asynchronous destruction.
-			Object->ConditionalBeginDestroy();
+			//@todo UE4 - A prefetch was removed here. Re-add it. It wasn't right anyway, since it was ten items ahead and the consoles on have 8 prefetch slots
+
+			UObject* Object = *It;
+			if (Object->HasAnyFlags( RF_Unreachable ))
+			{
+				// Begin the object's asynchronous destruction.
+				Object->ConditionalBeginDestroy();
+			}
 		}
+		UE_LOG( LogGarbage, Log, TEXT( "%f ms for unhashing unreachable objects" ), (FPlatformTime::Seconds() - StartTime) * 1000 );
 	}
-	UE_LOG(LogGarbage, Log, TEXT("%f ms for unhashing unreachable objects"), (FPlatformTime::Seconds() - StartTime) * 1000 );
 
 	// Set flag to indicate that we are relying on a purge to be performed.
 	GObjPurgeIsRequired = true;
