UE-16806: [CrashReport] UE4Editor_HotReload!GetCachedCDODuplicate() [hotreload.cpp:636]

diff --git a/Engine/Source/Runtime/CoreUObject/Private/UObject/UObjectBase.cpp b/Engine/Source/Runtime/CoreUObject/Private/UObject/UObjectBase.cpp
--- a/Engine/Source/Runtime/CoreUObject/Private/UObject/UObjectBase.cpp
+++ b/Engine/Source/Runtime/CoreUObject/Private/UObject/UObjectBase.cpp
@@ -5,6 +5,7 @@
 =============================================================================*/
 
 #include "CoreUObjectPrivate.h"
+#include "GCObject.h"
 
 DEFINE_LOG_CATEGORY_STATIC(LogUObjectBase, Log, All);
 DEFINE_STAT(STAT_UObjectsStatGroupTester);
@@ -525,8 +526,30 @@ static TArray<FFieldCompiledInInfo*>& GetDeferredClassRegistration()
 #if WITH_HOT_RELOAD
 TMap<UClass*, UObject*>& GetDuplicatedCDOMap()
 {
-	static TMap<UClass*, UObject*> Map;
-	return Map;
+	/**
+	 * GC referencer object, so duplicated CDOs aren't destroyed if still
+	 * referenced by the map.
+	 */
+	class FDuplicatedCDOMap : public FGCObject
+	{
+	public:
+		/**
+		 * Add referenced objects to reference collector.
+		 */
+		virtual void AddReferencedObjects(FReferenceCollector& Collector) override
+		{
+			for (auto& KVP : Map)
+			{
+				Collector.AddReferencedObject(KVP.Value);
+			}
+		}
+
+		/** Duplicated CDOs map. */
+		TMap<UClass*, UObject*> Map;
+	};
+
+	static FDuplicatedCDOMap Cache;
+	return Cache.Map;
 }
 
 /** Map of deferred class registration info (including size and reflection info) */
