AsyncLoading: Making sure DeferredFinalizeObjects are always referenced by GC until streaming has finished

diff --git a/Engine/Source/Runtime/CoreUObject/Private/Serialization/AsyncLoading.cpp b/Engine/Source/Runtime/CoreUObject/Private/Serialization/AsyncLoading.cpp
--- a/Engine/Source/Runtime/CoreUObject/Private/Serialization/AsyncLoading.cpp
+++ b/Engine/Source/Runtime/CoreUObject/Private/Serialization/AsyncLoading.cpp
@@ -59,26 +59,19 @@ class FAsyncObjectsReferencer : FGCObject
 	FAsyncObjectsReferencer() {}
 
 	/** List of referenced objects */
-	TArray<UObject*> ReferencedObjects;
+	TSet<UObject*> ReferencedObjects;
 #if THREADSAFE_UOBJECTS
 	/** Critical section for referenced objects list */
 	FCriticalSection ReferencedObjectsCritical;
 #endif
 
 #if !UE_BUILD_SHIPPING
-	int32 IndexOf(UObject* InObj)
+	bool Contains(UObject* InObj)
 	{
 #if THREADSAFE_UOBJECTS
 		FScopeLock ReferencedObjectsLock(&ReferencedObjectsCritical);
 #endif
-		for (int32 Index = 0; Index < ReferencedObjects.Num(); ++Index)
-		{
-			if (ReferencedObjects[Index] == InObj)
-			{
-				return Index;
-			}
-		}
-		return INDEX_NONE;
+		return ReferencedObjects.Contains(InObj);
 	}
 #endif
 
@@ -114,7 +107,10 @@ public:
 #else
 				check(IsInGameThread());
 #endif
-				ReferencedObjects.Add(InObject);
+				if (!ReferencedObjects.Contains(InObject))
+				{
+					ReferencedObjects.Add(InObject);
+				}
 			}
 			InObject->ThisThreadAtomicallyClearedRFUnreachable();
 		}
@@ -169,7 +165,7 @@ public:
 			UObject* Obj = *It;
 			if (Obj->HasAnyFlags(RF_AsyncLoading|RF_Async))
 			{
-				if (IndexOf(Obj) == INDEX_NONE)
+				if (!Contains(Obj))
 				{
 					UE_LOG(LogStreaming, Error, TEXT("%s has RF_AsyncLoading|RF_Async set but is not referenced by FAsyncObjectsReferencer"), *Obj->GetPathName());
 				}
@@ -1593,6 +1589,8 @@ EAsyncPackageState::Type FAsyncPackage::PostLoadObjects()
 		}
 		// All object must be finalized on the game thread
 		DeferredFinalizeObjects.Add(Object);
+		// Make sure all objects in DeferredFinalizeObjects are referenced too
+		FAsyncObjectsReferencer::Get().AddObject(Object);
 	}
 
 	// New objects might have been loaded during PostLoad.
