[UE-16281] Fix stale components persisting in a Blueprint class asset after changing a native parent class's inheritance hierarchy.

change summary:
- modified FLinkerLoad::ResolveDeferredExports() to ensure that all default subobjects in the export map for a Blueprint class asset package have been constructed and registered prior to class/CDO regeneration during compile-on-load.

diff --git a/Engine/Source/Runtime/CoreUObject/Private/Blueprint/BlueprintSupport.cpp b/Engine/Source/Runtime/CoreUObject/Private/Blueprint/BlueprintSupport.cpp
--- a/Engine/Source/Runtime/CoreUObject/Private/Blueprint/BlueprintSupport.cpp
+++ b/Engine/Source/Runtime/CoreUObject/Private/Blueprint/BlueprintSupport.cpp
@@ -1384,6 +1384,20 @@ void FLinkerLoad::ResolveDeferredExports(UClass* LoadClass)
 		//       be cleared (and this will do nothing the 2nd time around)
 		Preload(BlueprintCDO);
 
+		// Ensure that all default subobject exports belonging to the CDO have been created. DSOs may no longer be
+		// referenced by a tagged property and thus may not get created and registered until after class regeneration.
+		// This can cause invalid subobjects to register themselves with a regenerated CDO if the native parent class
+		// has been changed to inherit from an entirely different type since the last time the class asset was saved.
+		// By constructing them here, we make sure that LoadAllObjects() won't construct them after class regeneration.
+		for (int32 ExportIndex = 0; ExportIndex < ExportMap.Num(); ++ExportIndex)
+		{
+			FObjectExport& Export = ExportMap[ExportIndex];
+			if (Export.Object == nullptr && (Export.ObjectFlags & RF_DefaultSubObject) != 0 && Export.OuterIndex.IsExport() && Export.OuterIndex.ToExport() == DeferredCDOIndex)
+			{
+				CreateExport(ExportIndex);
+			}
+		}
+
 		// sub-classes of this Blueprint could have had their CDO's 
 		// initialization deferred (this occurs when the sub-class CDO is 
 		// created before this super CDO has been fully serialized; we do this
