Making FLinkerLoad aware of FAsyncPackage association

diff --git a/Engine/Source/Runtime/CoreUObject/Private/Serialization/AsyncLoading.cpp b/Engine/Source/Runtime/CoreUObject/Private/Serialization/AsyncLoading.cpp
--- a/Engine/Source/Runtime/CoreUObject/Private/Serialization/AsyncLoading.cpp
+++ b/Engine/Source/Runtime/CoreUObject/Private/Serialization/AsyncLoading.cpp
@@ -65,6 +65,7 @@ class FAsyncObjectsReferencer : FGCObject
 	FCriticalSection ReferencedObjectsCritical;
 #endif
 
+public:
 #if !UE_BUILD_SHIPPING
 	bool Contains(UObject* InObj)
 	{
@@ -75,7 +76,6 @@ class FAsyncObjectsReferencer : FGCObject
 	}
 #endif
 
-public:
 	/** Returns the one and only instance of this object */
 	static FAsyncObjectsReferencer& Get();
 	/** FGCObject interface */
@@ -869,6 +869,11 @@ FAsyncPackage::FAsyncPackage(const FAsyncPackageDesc& InDesc)
 {
 }
 
+FAsyncPackage::~FAsyncPackage()
+{
+	DetachLinker();
+}
+
 /**
  * @return Time load begun. This is NOT the time the load was requested in the case of other pending requests.
  */
@@ -883,13 +888,22 @@ double FAsyncPackage::GetLoadStartTime() const
 void FAsyncPackage::ResetLoader()
 {
 	// Reset loader.
-	if (bLoadHasFailed)
+	if (Linker)
 	{
+		check(Linker->AsyncRoot == this || Linker->AsyncRoot == nullptr);
+		Linker->AsyncRoot = nullptr;
+		delete Linker;
 		Linker = nullptr;
 	}
-	else if (Linker)
+}
+
+void FAsyncPackage::DetachLinker()
+{	
+	if (Linker)
 	{
-		delete Linker;
+		check(bLoadHasFinished || bLoadHasFailed);
+		check(Linker->AsyncRoot == this || Linker->AsyncRoot == nullptr);
+		Linker->AsyncRoot = nullptr;
 		Linker = nullptr;
 	}
 }
@@ -1141,6 +1155,10 @@ EAsyncPackageState::Type FAsyncPackage::CreateLinker()
 			Linker = FLinkerLoad::CreateLinkerAsync( Package, *PackageFileName, LinkerFlags );
 		}
 
+		// Associate this async package with the linker
+		check(Linker->AsyncRoot == nullptr);
+		Linker->AsyncRoot = this;
+
 		UE_LOG(LogStreaming, Verbose, TEXT("FAsyncPackage::CreateLinker for %s finished."), *Desc.NameToLoad.ToString());
 	}
 	return EAsyncPackageState::Complete;
diff --git a/Engine/Source/Runtime/CoreUObject/Private/UObject/LinkerLoad.cpp b/Engine/Source/Runtime/CoreUObject/Private/UObject/LinkerLoad.cpp
--- a/Engine/Source/Runtime/CoreUObject/Private/UObject/LinkerLoad.cpp
+++ b/Engine/Source/Runtime/CoreUObject/Private/UObject/LinkerLoad.cpp
@@ -12,6 +12,7 @@
 #include "Serialization/DeferredMessageLog.h"
 #include "UObject/UObjectThreadContext.h"
 #include "GatherableTextData.h"
+#include "Serialization/AsyncLoading.h"
 
 #define LOCTEXT_NAMESPACE "LinkerLoad"
 
@@ -673,6 +674,7 @@ FLinkerLoad::FLinkerLoad(UPackage* InParent, const TCHAR* InFilename, uint32 InL
 , LoadFlags(InLoadFlags)
 , bHaveImportsBeenVerified(false)
 , Loader(nullptr)
+, AsyncRoot(nullptr)
 , NameMapIndex(0)
 , GatherableTextDataMapIndex(0)
 , ImportMapIndex(0)
@@ -4013,6 +4015,11 @@ void FLinkerLoad::Detach()
 		LinkerRoot->LinkerLoad = nullptr;
 		LinkerRoot = nullptr;
 	}
+	if (AsyncRoot)
+	{
+		AsyncRoot->DetachLinker();
+		AsyncRoot = nullptr;
+	}
 }
 
 #if WITH_EDITOR
diff --git a/Engine/Source/Runtime/CoreUObject/Public/Serialization/AsyncLoading.h b/Engine/Source/Runtime/CoreUObject/Public/Serialization/AsyncLoading.h
--- a/Engine/Source/Runtime/CoreUObject/Public/Serialization/AsyncLoading.h
+++ b/Engine/Source/Runtime/CoreUObject/Public/Serialization/AsyncLoading.h
@@ -19,6 +19,7 @@ struct FAsyncPackage
 	 * Constructor
 	 */
 	FAsyncPackage(const FAsyncPackageDesc& InDesc);
+	~FAsyncPackage();
 
 	/**
 	 * Ticks the async loading code.
@@ -49,6 +50,11 @@ struct FAsyncPackage
 	 */
 	void ResetLoader();
 
+	/**
+	* Disassociates linker from this package
+	*/
+	void DetachLinker();
+
 	/**
 	 * Returns the name of the package to load.
 	 */
diff --git a/Engine/Source/Runtime/CoreUObject/Public/UObject/Linker.h b/Engine/Source/Runtime/CoreUObject/Public/UObject/Linker.h
--- a/Engine/Source/Runtime/CoreUObject/Public/UObject/Linker.h
+++ b/Engine/Source/Runtime/CoreUObject/Public/UObject/Linker.h
@@ -1303,6 +1303,8 @@ public:
 #endif // WITH_EDITOR
 	/** The archive that actually reads the raw data from disk.																*/
 	FArchive*				Loader;
+	/** The async package associated with this linker */
+	struct FAsyncPackage* AsyncRoot;
 
 	/** OldClassName to NewClassName for ImportMap */
 	static TMap<FName, FName> ObjectNameRedirects;
