Added missing implementation of FLinkerLoad::FlushCache(), removed unused InitLinkers

diff --git a/Engine/Source/Runtime/CoreUObject/Private/UObject/LinkerLoad.cpp b/Engine/Source/Runtime/CoreUObject/Private/UObject/LinkerLoad.cpp
--- a/Engine/Source/Runtime/CoreUObject/Private/UObject/LinkerLoad.cpp
+++ b/Engine/Source/Runtime/CoreUObject/Private/UObject/LinkerLoad.cpp
@@ -4688,8 +4688,12 @@ FLinkerLoad::ELinkerStatus FLinkerLoad::FixupExportMap()
 	}
 }
 
-void FLinkerLoad::InitLinkers()
+void FLinkerLoad::FlushCache()
 {
+	if (Loader)
+	{
+		Loader->FlushCache();
+	}
 }
 
 #undef LOCTEXT_NAMESPACE
diff --git a/Engine/Source/Runtime/CoreUObject/Public/UObject/Linker.h b/Engine/Source/Runtime/CoreUObject/Public/UObject/Linker.h
--- a/Engine/Source/Runtime/CoreUObject/Public/UObject/Linker.h
+++ b/Engine/Source/Runtime/CoreUObject/Public/UObject/Linker.h
@@ -1281,9 +1281,6 @@ class FLinkerLoad : public FLinker, public FArchiveUObject
 	// Variables.
 public:
 
-	/** Initialize everything related to linkers */
-	static void InitLinkers();
-
 	FORCEINLINE static ELinkerType::Type StaticType()
 	{
 		return ELinkerType::Load;
@@ -1510,6 +1507,9 @@ public:
 	 */
 	virtual FLinker* GetLinker() override { return this; }
 
+	/** Flush Loader Cache */
+	virtual void FlushCache() override;
+
 	/**
 	 * Creates and returns a FLinkerLoad object.
 	 *
