Deprecated FAsyncPackage Type as it was not used by anything.

diff --git a/Engine/Source/Runtime/CoreUObject/Private/Serialization/AsyncLoading.cpp b/Engine/Source/Runtime/CoreUObject/Private/Serialization/AsyncLoading.cpp
--- a/Engine/Source/Runtime/CoreUObject/Private/Serialization/AsyncLoading.cpp
+++ b/Engine/Source/Runtime/CoreUObject/Private/Serialization/AsyncLoading.cpp
@@ -449,7 +449,7 @@ void FAsyncLoadingThread::InsertPackage(FAsyncPackage* Package)
 	}
 }
 
-EAsyncPackageState::Type FAsyncLoadingThread::ProcessAsyncLoading(int32& OutPackagesProcessed, bool bUseTimeLimit /*= false*/, bool bUseFullTimeLimit /*= false*/, float TimeLimit /*= 0.0f*/, FName ExcludeType /*= NAME_None*/)
+EAsyncPackageState::Type FAsyncLoadingThread::ProcessAsyncLoading(int32& OutPackagesProcessed, bool bUseTimeLimit /*= false*/, bool bUseFullTimeLimit /*= false*/, float TimeLimit /*= 0.0f*/)
 {
 	SCOPE_CYCLE_COUNTER(STAT_FAsyncLoadingThread_ProcessAsyncLoading);
 	
@@ -465,12 +465,6 @@ EAsyncPackageState::Type FAsyncLoadingThread::ProcessAsyncLoading(int32& OutPack
 		// Package to be loaded.
 		FAsyncPackage* Package = AsyncPackages[PackageIndex];
 
-		if (ExcludeType != NAME_None && ExcludeType == Package->GetPackageType())
-		{
-			// We should skip packages of this type
-			continue;
-		}
-
 		if (Package->HasFinishedLoading() == false)
 		{
 			// Package tick returns EAsyncPackageState::Complete on completion.
@@ -518,7 +512,7 @@ EAsyncPackageState::Type FAsyncLoadingThread::ProcessAsyncLoading(int32& OutPack
 }
 
 
-EAsyncPackageState::Type FAsyncLoadingThread::ProcessLoadedPackages(bool bUseTimeLimit, bool bUseFullTimeLimit, float TimeLimit, FName ExcludeType)
+EAsyncPackageState::Type FAsyncLoadingThread::ProcessLoadedPackages(bool bUseTimeLimit, bool bUseFullTimeLimit, float TimeLimit)
 {
 	EAsyncPackageState::Type Result = EAsyncPackageState::Complete;
 
@@ -584,7 +578,7 @@ EAsyncPackageState::Type FAsyncLoadingThread::ProcessLoadedPackages(bool bUseTim
 
 
 
-EAsyncPackageState::Type FAsyncLoadingThread::TickAsyncLoading(bool bUseTimeLimit, bool bUseFullTimeLimit, float TimeLimit, FName ExcludeType)
+EAsyncPackageState::Type FAsyncLoadingThread::TickAsyncLoading(bool bUseTimeLimit, bool bUseFullTimeLimit, float TimeLimit)
 {
 	const bool bLoadingSuspended = IsAsyncLoadingSuspended();
 	const bool bIsMultithreaded = FAsyncLoadingThread::IsMultithreaded();
@@ -604,7 +598,7 @@ EAsyncPackageState::Type FAsyncLoadingThread::TickAsyncLoading(bool bUseTimeLimi
 		if (Result != EAsyncPackageState::TimeOut && !IsTimeLimitExceeded(TickStartTime, bUseTimeLimit, TimeLimit, TEXT("Pre-ProcessLoadedPackages")))
 		{
 			double RemainingTimeLimit = FMath::Max(0.0, TimeLimit - TimeLimitUsedForPreloading);
-			Result = ProcessLoadedPackages(bUseTimeLimit, bUseFullTimeLimit, RemainingTimeLimit, ExcludeType);
+			Result = ProcessLoadedPackages(bUseTimeLimit, bUseFullTimeLimit, RemainingTimeLimit);
 		}
 
 		if (Result != EAsyncPackageState::TimeOut && !IsTimeLimitExceeded(TickStartTime, bUseTimeLimit, TimeLimit, TEXT("Pre-EmptyReferencedObjects")))
@@ -697,7 +691,7 @@ EAsyncPackageState::Type  FAsyncLoadingThread::TickAsyncThread(bool bUseTimeLimi
 		if (AsyncThreadReady.GetValue())
 		{
 			CreateAsyncPackagesFromQueue();
-			Result = ProcessAsyncLoading(ProcessedRequests, bUseTimeLimit, bUseFullTimeLimit, TimeLimit, NAME_None);
+			Result = ProcessAsyncLoading(ProcessedRequests, bUseTimeLimit, bUseFullTimeLimit, TimeLimit);
 		}
 		if (ProcessedRequests == 0 && IsMultithreaded())
 		{
@@ -1815,7 +1809,7 @@ void FAsyncPackage::UpdateLoadPercentage()
 /*-----------------------------------------------------------------------------
 	UObject async (pre)loading.
 -----------------------------------------------------------------------------*/
-void LoadPackageAsync(const FString& InName, const FGuid* InGuid /*= nullptr*/, FName InType /*= NAME_None*/, const TCHAR* InPackageToLoadFrom /*= nullptr*/, FLoadPackageAsyncDelegate InCompletionDelegate /*= FLoadPackageAsyncDelegate()*/, EPackageFlags InPackageFlags /*= PKG_None*/, int32 InPIEInstanceID /*= INDEX_NONE*/, uint32 InPackagePriority /*= 0*/)
+void LoadPackageAsync(const FString& InName, const FGuid* InGuid /*= nullptr*/, const TCHAR* InPackageToLoadFrom /*= nullptr*/, FLoadPackageAsyncDelegate InCompletionDelegate /*= FLoadPackageAsyncDelegate()*/, EPackageFlags InPackageFlags /*= PKG_None*/, int32 InPIEInstanceID /*= INDEX_NONE*/, uint32 InPackagePriority /*= 0*/)
 {
 	// The comments clearly state that it should be a package name but we also handle it being a filename as this function is not perf critical
 	// and LoadPackage handles having a filename being passed in as well.
@@ -1847,7 +1841,7 @@ void LoadPackageAsync(const FString& InName, const FGuid* InGuid /*= nullptr*/,
 		UE_LOG(LogStreaming, Fatal, TEXT("Async loading code requires long package names (%s)."), *PackageNameToLoad);
 	}
 
-	FAsyncPackageDesc PackageDesc(*PackageName, *PackageNameToLoad, InGuid ? *InGuid : FGuid(), InType, InCompletionDelegate, InPackageFlags, InPIEInstanceID, InPackagePriority);
+	FAsyncPackageDesc PackageDesc(*PackageName, *PackageNameToLoad, InGuid ? *InGuid : FGuid(), InCompletionDelegate, InPackageFlags, InPIEInstanceID, InPackagePriority);
 	FAsyncLoadingThread::Get().QueuePackage(PackageDesc);
 }
 
@@ -1855,7 +1849,12 @@ void LoadPackageAsync(const FString& PackageName, FLoadPackageAsyncDelegate Comp
 {
 	const FGuid* Guid = nullptr;
 	const TCHAR* PackageToLoadFrom = nullptr;
-	LoadPackageAsync(PackageName, Guid, NAME_None, PackageToLoadFrom, CompletionDelegate, InPackageFlags, -1, InPackagePriority );
+	LoadPackageAsync(PackageName, Guid, PackageToLoadFrom, CompletionDelegate, InPackageFlags, -1, InPackagePriority );
+}
+
+void LoadPackageAsync(const FString& InName, const FGuid* InGuid, FName InType /* Unused, deprecated */, const TCHAR* InPackageToLoadFrom /*= nullptr*/, FLoadPackageAsyncDelegate InCompletionDelegate /*= FLoadPackageAsyncDelegate()*/, EPackageFlags InPackageFlags /*= PKG_None*/, int32 InPIEInstanceID /*= INDEX_NONE*/, uint32 InPackagePriority /*= 0*/)
+{
+	LoadPackageAsync(InName, InGuid, InPackageToLoadFrom, InCompletionDelegate, InPackageFlags, InPIEInstanceID, InPackagePriority);
 }
 
 void CancelAsyncLoading()
@@ -1882,7 +1881,7 @@ bool IsInAsyncLoadingThreadCoreUObjectInternal()
 	return FAsyncLoadingThread::Get().IsInAsyncLoadThread();
 }
 
-void FlushAsyncLoading(FName ExcludeType/*=NAME_None*/)
+void FlushAsyncLoading()
 {
 	if (IsAsyncLoading())
 	{
@@ -1901,7 +1900,7 @@ void FlushAsyncLoading(FName ExcludeType/*=NAME_None*/)
 			SCOPE_CYCLE_COUNTER(STAT_FAsyncPackage_TickAsyncLoadingGameThread);
 			while (IsAsyncLoading())
 			{
-				FAsyncLoadingThread::Get().TickAsyncLoading(false, false, 0, ExcludeType);
+				FAsyncLoadingThread::Get().TickAsyncLoading(false, false, 0);
 				if (FAsyncLoadingThread::Get().IsMultithreaded())
 				{
 					FPlatformProcess::SleepNoStats(0.0001f);
@@ -1909,29 +1908,30 @@ void FlushAsyncLoading(FName ExcludeType/*=NAME_None*/)
 			}
 		}
 
-		if (ExcludeType == NAME_None)
-		{
-			// It's fine to have pending loads if we excluded some from the check
-			check( !IsAsyncLoading() );
-		}
+		check( !IsAsyncLoading() );
 
 		// Reset min priority again.
 		FIOSystem::Get().SetMinPriority( AIOP_MIN );
 	}
 }
 
+void FlushAsyncLoading(FName ExcludeType)
+{
+	FlushAsyncLoading();
+}
+
 int32 GetNumAsyncPackages()
 {
 	return FAsyncLoadingThread::Get().GetAsyncPackagesCount();
 }
 
-EAsyncPackageState::Type ProcessAsyncLoading(bool bUseTimeLimit, bool bUseFullTimeLimit, float TimeLimit, FName ExcludeType)
+EAsyncPackageState::Type ProcessAsyncLoading(bool bUseTimeLimit, bool bUseFullTimeLimit, float TimeLimit)
 {
 	SCOPE_CYCLE_COUNTER(STAT_AsyncLoadingTime);
 
 	{
 		SCOPE_CYCLE_COUNTER(STAT_FAsyncPackage_TickAsyncLoadingGameThread);
-		FAsyncLoadingThread::Get().TickAsyncLoading(bUseTimeLimit, bUseFullTimeLimit, TimeLimit, ExcludeType);
+		FAsyncLoadingThread::Get().TickAsyncLoading(bUseTimeLimit, bUseFullTimeLimit, TimeLimit);
 	}
 
 	return IsAsyncLoading() ? EAsyncPackageState::TimeOut : EAsyncPackageState::Complete;
diff --git a/Engine/Source/Runtime/CoreUObject/Private/Serialization/AsyncLoadingThread.h b/Engine/Source/Runtime/CoreUObject/Private/Serialization/AsyncLoadingThread.h
--- a/Engine/Source/Runtime/CoreUObject/Private/Serialization/AsyncLoadingThread.h
+++ b/Engine/Source/Runtime/CoreUObject/Private/Serialization/AsyncLoadingThread.h
@@ -246,10 +246,9 @@ public:
 	* @param bUseTimeLimit True if time limit should be used [time-slicing].
 	* @param bUseFullTimeLimit True if full time limit should be used [time-slicing].
 	* @param TimeLimit Maximum amount of time that can be spent in this call [time-slicing].
-	* @param ExcludeType Packages to exclude.
 	* @return The current state of async loading
 	*/
-	EAsyncPackageState::Type ProcessAsyncLoading(int32& OutPackagesProcessed, bool bUseTimeLimit = false, bool bUseFullTimeLimit = false, float TimeLimit = 0.0f, FName ExcludeType = NAME_None);
+	EAsyncPackageState::Type ProcessAsyncLoading(int32& OutPackagesProcessed, bool bUseTimeLimit = false, bool bUseFullTimeLimit = false, float TimeLimit = 0.0f);
 
 	/**
 	* [GAME THREAD] Ticks game thread side of async loading.
@@ -257,10 +256,9 @@ public:
 	* @param bUseTimeLimit True if time limit should be used [time-slicing].
 	* @param bUseFullTimeLimit True if full time limit should be used [time-slicing].
 	* @param TimeLimit Maximum amount of time that can be spent in this call [time-slicing].
-	* @param ExcludeType Packages to exclude.
 	* @return The current state of async loading
 	*/
-	EAsyncPackageState::Type TickAsyncLoading(bool bUseTimeLimit, bool bUseFullTimeLimit, float TimeLimit, FName ExcludeType);
+	EAsyncPackageState::Type TickAsyncLoading(bool bUseTimeLimit, bool bUseFullTimeLimit, float TimeLimit);
 
 	/**
 	* [ASYNC THREAD] Main thread loop
@@ -289,10 +287,9 @@ private:
 	* @param bUseTimeLimit True if time limit should be used [time-slicing].
 	* @param bUseFullTimeLimit True if full time limit should be used [time-slicing].
 	* @param TimeLimit Maximum amount of time that can be spent in this call [time-slicing].
-	* @param ExcludeType Packages to exclude.
 	* @return The current state of async loading
 	*/
-	EAsyncPackageState::Type ProcessLoadedPackages(bool bUseTimeLimit, bool bUseFullTimeLimit, float TimeLimit, FName ExcludeType);
+	EAsyncPackageState::Type ProcessLoadedPackages(bool bUseTimeLimit, bool bUseFullTimeLimit, float TimeLimit);
 
 	/**
 	* [ASYNC THREAD] Creates async packages from the queued requests
diff --git a/Engine/Source/Runtime/CoreUObject/Private/UObject/UObjectGlobals.cpp b/Engine/Source/Runtime/CoreUObject/Private/UObject/UObjectGlobals.cpp
--- a/Engine/Source/Runtime/CoreUObject/Private/UObject/UObjectGlobals.cpp
+++ b/Engine/Source/Runtime/CoreUObject/Private/UObject/UObjectGlobals.cpp
@@ -1121,7 +1121,7 @@ void BeginLoad()
 	if (ThreadContext.ObjBeginLoadCount == 0 && !IsInAsyncLoadingThread())
 	{
 		// Make sure we're finishing up all pending async loads, and trigger texture streaming next tick if necessary.
-		FlushAsyncLoading( NAME_None );
+		FlushAsyncLoading();
 
 		// Validate clean load state.
 		check(ThreadContext.ObjLoaded.Num() == 0);
diff --git a/Engine/Source/Runtime/CoreUObject/Public/Serialization/AsyncLoading.h b/Engine/Source/Runtime/CoreUObject/Public/Serialization/AsyncLoading.h
--- a/Engine/Source/Runtime/CoreUObject/Public/Serialization/AsyncLoading.h
+++ b/Engine/Source/Runtime/CoreUObject/Public/Serialization/AsyncLoading.h
@@ -63,14 +63,6 @@ struct FAsyncPackage
 		return Desc.Name;
 	}
 
-	/**
-	 * Returns the type name associated with this package
-	 */
-	const FName &GetPackageType() const
-	{
-		return Desc.Type;
-	}
-
 	void AddCompletionCallback(const FLoadPackageAsyncDelegate& Callback, bool bInternal);
 
 	/** Gets the number of references to this package from other packages in the dependency tree. */
diff --git a/Engine/Source/Runtime/CoreUObject/Public/Serialization/AsyncPackage.h b/Engine/Source/Runtime/CoreUObject/Public/Serialization/AsyncPackage.h
--- a/Engine/Source/Runtime/CoreUObject/Public/Serialization/AsyncPackage.h
+++ b/Engine/Source/Runtime/CoreUObject/Public/Serialization/AsyncPackage.h
@@ -12,8 +12,6 @@ struct FAsyncPackageDesc
 	FName Name;
 	/** Name of the package to load. */
 	FName NameToLoad;
-	/** An abstract type name associated with this package, for tagging use	*/
-	FName Type;
 	/** GUID of the package to load, or the zeroed invalid GUID for "don't care" */
 	FGuid Guid;
 	/** Delegate called on completion of loading */
@@ -28,10 +26,9 @@ struct FAsyncPackageDesc
 #endif
 
 
-	FAsyncPackageDesc(const FName& InName, FName InPackageToLoadFrom = NAME_None, const FGuid& InGuid = FGuid(), FName InType = NAME_None, FLoadPackageAsyncDelegate InCompletionDelegate = FLoadPackageAsyncDelegate(), EPackageFlags InPackageFlags = PKG_None, int32 InPIEInstanceID = INDEX_NONE, uint32 InPriority = 0)
+	FAsyncPackageDesc(const FName& InName, FName InPackageToLoadFrom = NAME_None, const FGuid& InGuid = FGuid(), FLoadPackageAsyncDelegate InCompletionDelegate = FLoadPackageAsyncDelegate(), EPackageFlags InPackageFlags = PKG_None, int32 InPIEInstanceID = INDEX_NONE, uint32 InPriority = 0)
 		: Name(InName)
 		, NameToLoad(InPackageToLoadFrom)
-		, Type(InType)	
 		, Guid(InGuid)
 		, PackageLoadedDelegate(InCompletionDelegate)
 		, PackageFlags(InPackageFlags)
diff --git a/Engine/Source/Runtime/CoreUObject/Public/UObject/UObjectGlobals.h b/Engine/Source/Runtime/CoreUObject/Public/UObject/UObjectGlobals.h
--- a/Engine/Source/Runtime/CoreUObject/Public/UObject/UObjectGlobals.h
+++ b/Engine/Source/Runtime/CoreUObject/Public/UObject/UObjectGlobals.h
@@ -281,7 +281,7 @@ namespace EAsyncLoadingResult
 DECLARE_DELEGATE_ThreeParams(FLoadPackageAsyncDelegate, const FName& /*PackageName*/, UPackage* /*LoadedPackage*/, EAsyncLoadingResult::Type /*Result*/)
 
 /**
- * Asynchronously load a package and all contained objects that match context flags. Non- blocking.
+ * [Deprecated] Asynchronously load a package and all contained objects that match context flags. Non- blocking.
  *
  * @param	InName					Name of package to load
  * @param	InGuid					GUID of the package to load, or NULL for "don't care"
@@ -292,7 +292,21 @@ DECLARE_DELEGATE_ThreeParams(FLoadPackageAsyncDelegate, const FName& /*PackageNa
  * @param	InPIEInstanceID			PIE instance ID
  * @param	InPackagePriority		Loading priority
  */
-COREUOBJECT_API void LoadPackageAsync(const FString& InName, const FGuid* InGuid = nullptr, FName InType = NAME_None, const TCHAR* InPackageToLoadFrom = nullptr, FLoadPackageAsyncDelegate InCompletionDelegate = FLoadPackageAsyncDelegate(), EPackageFlags InPackageFlags = PKG_None, int32 InPIEInstanceID = INDEX_NONE, uint32 InPackagePriority = 0);
+DEPRECATED(4.9, "LoadPackageAsync override that takes package type parameter FName InType is deprecated.")
+COREUOBJECT_API void LoadPackageAsync(const FString& InName, const FGuid* InGuid, FName InType, const TCHAR* InPackageToLoadFrom = nullptr, FLoadPackageAsyncDelegate InCompletionDelegate = FLoadPackageAsyncDelegate(), EPackageFlags InPackageFlags = PKG_None, int32 InPIEInstanceID = INDEX_NONE, uint32 InPackagePriority = 0);
+
+/**
+* Asynchronously load a package and all contained objects that match context flags. Non- blocking.
+*
+* @param	InName					Name of package to load
+* @param	InGuid					GUID of the package to load, or NULL for "don't care"
+* @param	InPackageToLoadFrom		If non-null, this is another package name. We load from this package name, into a (probably new) package named PackageName
+* @param	InCompletionDelegate	Delegate to be invoked when the packages has finished streaming
+* @param	InFlags					Package flags
+* @param	InPIEInstanceID			PIE instance ID
+* @param	InPackagePriority		Loading priority
+*/
+COREUOBJECT_API void LoadPackageAsync(const FString& InName, const FGuid* InGuid = nullptr, const TCHAR* InPackageToLoadFrom = nullptr, FLoadPackageAsyncDelegate InCompletionDelegate = FLoadPackageAsyncDelegate(), EPackageFlags InPackageFlags = PKG_None, int32 InPIEInstanceID = INDEX_NONE, uint32 InPackagePriority = 0);
 
 /**
 * Asynchronously load a package and all contained objects that match context flags. Non- blocking.
@@ -392,12 +406,19 @@ COREUOBJECT_API FName MakeObjectNameFromActorLabel( const FString& InActorLabel,
  * @return true if object is referenced, false otherwise
  */
 COREUOBJECT_API bool IsReferenced( UObject*& Res, EObjectFlags KeepFlags, bool bCheckSubObjects = false, FReferencerInformationList* FoundReferences = NULL );
+
 /**
  * Blocks till all pending package/ linker requests are fulfilled.
- *
- * @param	ExcludeType		Do not flush packages associated with this specific type name
  */
-COREUOBJECT_API void FlushAsyncLoading( FName ExcludeType = NAME_None );
+COREUOBJECT_API void FlushAsyncLoading();
+
+/**
+* [Deprecated] Blocks till all pending package/ linker requests are fulfilled.
+*
+* @param	ExcludeType		Do not flush packages associated with this specific type name
+*/
+DEPRECATED(4.9, "FlushAsyncLoading override that takes package type parameter FName ExcludeType is deprecated.")
+COREUOBJECT_API void FlushAsyncLoading(FName ExcludeType);
 
 /**
  * @return number of active async load package requests
@@ -433,10 +454,9 @@ namespace EAsyncPackageState
  * @param	bUseTimeLimit	Whether to use a time limit
  * @param	bUseFullTimeLimit	If true, use the entire time limit even if blocked on I/O
  * @param	TimeLimit		Soft limit of time this function is allowed to consume
- * @param	ExcludeType		Do not process packages associated with this specific type name
  * @return The minimum state of any of the queued packages.
  */
-COREUOBJECT_API EAsyncPackageState::Type ProcessAsyncLoading( bool bUseTimeLimit, bool bUseFullTimeLimit, float TimeLimit, FName ExcludeType = NAME_None );
+COREUOBJECT_API EAsyncPackageState::Type ProcessAsyncLoading( bool bUseTimeLimit, bool bUseFullTimeLimit, float TimeLimit);
 COREUOBJECT_API void BeginLoad();
 COREUOBJECT_API void EndLoad();
 
diff --git a/Engine/Source/Runtime/Engine/Classes/Engine/World.h b/Engine/Source/Runtime/Engine/Classes/Engine/World.h
--- a/Engine/Source/Runtime/Engine/Classes/Engine/World.h
+++ b/Engine/Source/Runtime/Engine/Classes/Engine/World.h
@@ -2420,9 +2420,19 @@ public:
 	 * on all async operation like updating components.
 	 *
 	 * @param FlushType					Whether to only flush level visibility operations (optional)
+	 */
+	void FlushLevelStreaming(EFlushLevelStreamingType FlushType = EFlushLevelStreamingType::Full);
+
+	/**
+	 * [Deprecated] Flushes level streaming in blocking fashion and returns when all levels are loaded/ visible/ hidden
+	 * so further calls to UpdateLevelStreaming won't do any work unless state changes. Basically blocks
+	 * on all async operation like updating components.
+	 *
+	 * @param FlushType					Whether to only flush level visibility operations
 	 * @param ExcludeType				Exclude packages of this type from flushing
 	 */
-	void FlushLevelStreaming(EFlushLevelStreamingType FlushType = EFlushLevelStreamingType::Full, FName ExcludeType = NAME_None);
+	DEPRECATED(4.9, "FlushLevelStreaming override that takes ExcludeType parameter is deprecated.")
+	void FlushLevelStreaming(EFlushLevelStreamingType FlushType, FName ExcludeType);
 
 	/**
 	 * Triggers a call to ULevel::BuildStreamingData(this,NULL,NULL) within a few seconds.
diff --git a/Engine/Source/Runtime/Engine/Private/LevelStreaming.cpp b/Engine/Source/Runtime/Engine/Private/LevelStreaming.cpp
--- a/Engine/Source/Runtime/Engine/Private/LevelStreaming.cpp
+++ b/Engine/Source/Runtime/Engine/Private/LevelStreaming.cpp
@@ -470,7 +470,7 @@ bool ULevelStreaming::RequestLevel(UWorld* PersistentWorld, bool bAllowLevelLoad
 			{
 				// Load localized part of level first in case it exists. We don't need to worry about GC or completion 
 				// callback as we always kick off another async IO for the level below.
-				LoadPackageAsync(*(GetWorldAssetPackageName() + LOCALIZED_SEEKFREE_SUFFIX), nullptr, NAME_None, *LocalizedPackageName, FLoadPackageAsyncDelegate(), PackageFlags, PIEInstanceID);
+				LoadPackageAsync(*(GetWorldAssetPackageName() + LOCALIZED_SEEKFREE_SUFFIX), nullptr, *LocalizedPackageName, FLoadPackageAsyncDelegate(), PackageFlags, PIEInstanceID);
 			}
 		}
 
@@ -482,14 +482,14 @@ bool ULevelStreaming::RequestLevel(UWorld* PersistentWorld, bool bAllowLevelLoad
 			UWorld::WorldTypePreLoadMap.FindOrAdd(DesiredPackageName) = PersistentWorld->WorldType;
 
 			// Kick off async load request.
-			LoadPackageAsync(DesiredPackageName.ToString(), nullptr, NAME_None, *PackageNameToLoadFrom, FLoadPackageAsyncDelegate::CreateUObject(this, &ULevelStreaming::AsyncLevelLoadComplete), PackageFlags, PIEInstanceID);
+			LoadPackageAsync(DesiredPackageName.ToString(), nullptr, *PackageNameToLoadFrom, FLoadPackageAsyncDelegate::CreateUObject(this, &ULevelStreaming::AsyncLevelLoadComplete), PackageFlags, PIEInstanceID);
 
 			// streamingServer: server loads everything?
 			// Editor immediately blocks on load and we also block if background level streaming is disabled.
 			if (bBlockOnLoad || ShouldBeAlwaysLoaded())
 			{
 				// Finish all async loading.
-				FlushAsyncLoading( NAME_None );
+				FlushAsyncLoading();
 			}
 		}
 		else
diff --git a/Engine/Source/Runtime/Engine/Private/UnrealEngine.cpp b/Engine/Source/Runtime/Engine/Private/UnrealEngine.cpp
--- a/Engine/Source/Runtime/Engine/Private/UnrealEngine.cpp
+++ b/Engine/Source/Runtime/Engine/Private/UnrealEngine.cpp
@@ -10040,7 +10040,7 @@ void UEngine::ConditionalCommitMapChange(FWorldContext &Context)
 		// Block on remaining async data.
 		if( !IsReadyForMapChange(Context) )
 		{
-			FlushAsyncLoading( NAME_None );
+			FlushAsyncLoading();
 			check( IsReadyForMapChange(Context) );
 		}
 		
diff --git a/Engine/Source/Runtime/Engine/Private/World.cpp b/Engine/Source/Runtime/Engine/Private/World.cpp
--- a/Engine/Source/Runtime/Engine/Private/World.cpp
+++ b/Engine/Source/Runtime/Engine/Private/World.cpp
@@ -2496,7 +2496,7 @@ void UWorld::UpdateLevelStreaming()
 	}
 }
 
-void UWorld::FlushLevelStreaming(EFlushLevelStreamingType FlushType, FName ExcludeType)
+void UWorld::FlushLevelStreaming(EFlushLevelStreamingType FlushType)
 {
 	AWorldSettings* WorldSettings = GetWorldSettings();
 
@@ -2506,7 +2506,7 @@ void UWorld::FlushLevelStreaming(EFlushLevelStreamingType FlushType, FName Exclu
 	UpdateLevelStreaming();
 
 	// Make sure all outstanding loads are taken care of, other than ones associated with the excluded type
-	FlushAsyncLoading(ExcludeType);
+	FlushAsyncLoading();
 
 	// Kick off making levels visible if loading finished by flushing.
 	UpdateLevelStreaming();
@@ -2524,7 +2524,7 @@ void UWorld::FlushLevelStreaming(EFlushLevelStreamingType FlushType, FName Exclu
 			if (FlushLevelStreamingType == EFlushLevelStreamingType::Full)
 			{
 				// Make sure all outstanding loads are taken care of...
-				FlushAsyncLoading(NAME_None);
+				FlushAsyncLoading();
 			}
 	
 			// Update level streaming.
@@ -2545,6 +2545,11 @@ void UWorld::FlushLevelStreaming(EFlushLevelStreamingType FlushType, FName Exclu
 	}
 }
 
+void UWorld::FlushLevelStreaming(EFlushLevelStreamingType FlushType, FName ExcludeType)
+{
+	FlushLevelStreaming(FlushType);
+}
+
 /**
  * Forces streaming data to be rebuilt for the current world.
  */
@@ -4529,7 +4534,6 @@ void FSeamlessTravelHandler::StartLoadingDestination()
 		LoadPackageAsync(
 			URLMapPackageName, 
 			PendingTravelGuid.IsValid() ? &PendingTravelGuid : NULL,
-			NAME_None,
 			*URLMapPackageToLoadFrom,
 			FLoadPackageAsyncDelegate::CreateRaw(this, &FSeamlessTravelHandler::SeamlessTravelLoadCallback), 			
 			PackageFlags,
