Fix crash in UMapProperty serialization

diff --git a/Engine/Source/Runtime/CoreUObject/Private/UObject/PropertyMap.cpp b/Engine/Source/Runtime/CoreUObject/Private/UObject/PropertyMap.cpp
--- a/Engine/Source/Runtime/CoreUObject/Private/UObject/PropertyMap.cpp
+++ b/Engine/Source/Runtime/CoreUObject/Private/UObject/PropertyMap.cpp
@@ -252,8 +252,11 @@ void UMapProperty::SerializeItem(FArchive& Ar, void* Value, const void* Defaults
 		uint8* TempKeyStorage = nullptr;
 		ON_SCOPE_EXIT
 		{
-			KeyProp->DestroyValue(TempKeyStorage);
-			FMemory::Free(TempKeyStorage);
+			if (TempKeyStorage)
+			{
+				KeyProp->DestroyValue(TempKeyStorage);
+				FMemory::Free(TempKeyStorage);
+			}
 		};
 
 		// Delete any explicitly-removed keys
