Add ForEachObjectWithOuter and ForEachObjectOfClass utility functions that will call a TFunctionRef for each object instead of using GetObjectsWithOuter/OfClass which requires copying each object in to a (potentially very large) array

diff --git a/Engine/Source/Runtime/CoreUObject/Private/UObject/UObjectHash.cpp b/Engine/Source/Runtime/CoreUObject/Private/UObject/UObjectHash.cpp
--- a/Engine/Source/Runtime/CoreUObject/Private/UObject/UObjectHash.cpp
+++ b/Engine/Source/Runtime/CoreUObject/Private/UObject/UObjectHash.cpp
@@ -645,6 +645,43 @@ void GetObjectsWithOuter(const class UObjectBase* Outer, TArray<UObject *>& Resu
 	}
 }
 
+void ForEachObjectWithOuter(const class UObjectBase* Outer, TFunctionRef<void (UObject*)> Operation, bool bIncludeNestedObjects, EObjectFlags ExclusionFlags)
+{
+	// We don't want to return any objects that are currently being background loaded unless we're using the object iterator during async loading.
+	ExclusionFlags |= RF_Unreachable;
+	if (!IsInAsyncLoadingThread())
+	{
+		ExclusionFlags = EObjectFlags(ExclusionFlags | RF_AsyncLoading);
+	}
+	FUObjectHashTables& ThreadHash = FUObjectHashTables::Get();
+	FHashTableLock HashLock(ThreadHash);
+	TArray<TSet<UObjectBase*> const*> AllInners;
+
+	if (TSet<UObjectBase*> const* Inners = ThreadHash.ObjectOuterMap.Find(Outer))
+	{
+		AllInners.Add(Inners);
+	}
+	while (AllInners.Num())
+	{
+		TSet<UObjectBase*> const* Inners = AllInners.Pop();
+		for (TSet<UObjectBase*>::TConstIterator It(*Inners); It; ++It)
+		{
+			UObject *Object = static_cast<UObject*>(*It);
+			if (!Object->HasAnyFlags(ExclusionFlags))
+			{
+				Operation(Object);
+				if (bIncludeNestedObjects)
+				{
+					if (TSet<UObjectBase*> const* ObjectInners = ThreadHash.ObjectOuterMap.Find(Object))
+					{
+						AllInners.Add(ObjectInners);
+					}
+				}
+			}
+		}
+	}
+}
+
 UObjectBase* FindObjectWithOuter(class UObjectBase* Outer, class UClass* ClassToLookFor, FName NameToLookFor)
 {
 	SCOPE_CYCLE_COUNTER( STAT_Hash_FindObjectWithOuter );
@@ -751,6 +788,43 @@ void GetObjectsOfClass(UClass* ClassToLookFor, TArray<UObject *>& Results, bool
 	check( Results.Num() <= GetUObjectArray().GetObjectArrayNum() ); // otherwise we have a cycle in the outer chain, which should not be possible
 }
 
+void ForEachObjectOfClass(UClass* ClassToLookFor, TFunctionRef<void (UObject*)> Operation, bool bIncludeDerivedClasses, EObjectFlags AdditionalExcludeFlags)
+{
+	// We don't want to return any objects that are currently being background loaded unless we're using the object iterator during async loading.
+	EObjectFlags ExclusionFlags = RF_Unreachable;
+	if (!IsInAsyncLoadingThread())
+	{
+		ExclusionFlags |= RF_AsyncLoading;
+	}
+	ExclusionFlags |= AdditionalExcludeFlags;
+
+	FUObjectHashTables& ThreadHash = FUObjectHashTables::Get();
+	FHashTableLock HashLock(ThreadHash);
+
+	TSet<UClass*> ClassesToSearch;
+	ClassesToSearch.Add( ClassToLookFor );
+	if( bIncludeDerivedClasses )
+	{
+		RecursivelyPopulateDerivedClasses( ThreadHash, ClassToLookFor, ClassesToSearch );
+	}
+
+	for (auto ClassIt = ClassesToSearch.CreateConstIterator(); ClassIt; ++ClassIt)
+	{
+		TSet<UObjectBase*> const* List = ThreadHash.ClassToObjectListMap.Find(*ClassIt);
+		if (List)
+		{
+			for (auto ObjectIt = List->CreateConstIterator(); ObjectIt; ++ObjectIt)
+			{
+				UObject *Object = static_cast<UObject*>(*ObjectIt);
+				if (!Object->HasAnyFlags(ExclusionFlags))
+				{
+					Operation(Object);
+				}
+			}
+		}
+	}
+}
+
 void GetDerivedClasses(UClass* ClassToLookFor, TArray<UClass *>& Results, bool bRecursive)
 {
 	if (bRecursive)
diff --git a/Engine/Source/Runtime/CoreUObject/Public/UObject/UObjectHash.h b/Engine/Source/Runtime/CoreUObject/Public/UObject/UObjectHash.h
--- a/Engine/Source/Runtime/CoreUObject/Public/UObject/UObjectHash.h
+++ b/Engine/Source/Runtime/CoreUObject/Public/UObject/UObjectHash.h
@@ -37,11 +37,22 @@ UObject* StaticFindObjectFastExplicit( UClass* ObjectClass, FName ObjectName, co
  * Return all objects with a given outer
  *
  * @param	Outer						Outer to search for
- * @param	Results						returned results
+ * @param	Results						Returned results
  * @param	bIncludeNestedObjects		If true, then things whose outers directly or indirectly have Outer as an outer are included, these are the nested objects.
+ * @param	ExclusionFlags				Specifies flags to use as a filter for which objects to return
  */
 COREUOBJECT_API void GetObjectsWithOuter(const class UObjectBase* Outer, TArray<UObject *>& Results, bool bIncludeNestedObjects = true, EObjectFlags ExclusionFlags = RF_NoFlags);
 
+/**
+ * Performs an operation on all objects with a given outer
+ *
+ * @param	Outer						Outer to search for
+ * @param	Operation					Function to be called for each object
+ * @param	bIncludeNestedObjects		If true, then things whose outers directly or indirectly have Outer as an outer are included, these are the nested objects.
+ * @param	ExclusionFlags				Specifies flags to use as a filter for which objects to return
+ */
+COREUOBJECT_API void ForEachObjectWithOuter(const class UObjectBase* Outer, TFunctionRef<void (UObject*)> Operation, bool bIncludeNestedObjects = true, EObjectFlags ExclusionFlags = RF_NoFlags);
+
 /**
  * Find an objects with a given name and or class within an outer
  *
@@ -61,6 +72,16 @@ COREUOBJECT_API class UObjectBase* FindObjectWithOuter(class UObjectBase* Outer,
  */
 COREUOBJECT_API void GetObjectsOfClass(UClass* ClassToLookFor, TArray<UObject *>& Results, bool bIncludeDerivedClasses = true, EObjectFlags AdditionalExcludeFlags=RF_ClassDefaultObject);
 
+/**
+ * Performs an operation on all objects with a given outer
+ *
+ * @param	Outer						Outer to search for
+ * @param	Operation					Function to be called for each object
+ * @param	bIncludeDerivedClasses		If true, the results will include objects of child classes as well.
+ * @param	AdditionalExcludeFlags		Objects with any of these flags will be excluded from the results.
+ */
+COREUOBJECT_API void ForEachObjectOfClass(UClass* ClassToLookFor, TFunctionRef<void (UObject*)> Operation, bool bIncludeDerivedClasses = true, EObjectFlags AdditionalExcludeFlags=RF_ClassDefaultObject);
+
 /**
  * Returns an array of classes that were derived from the specified class.
  *
