Changing sanity check condition when checking for the linker after sync loading.

diff --git a/Engine/Source/Runtime/CoreUObject/Private/UObject/UObjectGlobals.cpp b/Engine/Source/Runtime/CoreUObject/Private/UObject/UObjectGlobals.cpp
--- a/Engine/Source/Runtime/CoreUObject/Private/UObject/UObjectGlobals.cpp
+++ b/Engine/Source/Runtime/CoreUObject/Private/UObject/UObjectGlobals.cpp
@@ -1050,8 +1050,8 @@ UPackage* LoadPackageInternal(UPackage* InOuter, const TCHAR* InLongPackageName,
 		{
 			if (FUObjectThreadContext::Get().ObjBeginLoadCount == 0)
 			{
-				// Sanity check to make sure that Linker is the linker that loaded our Result package
-				check(!Result || Result->LinkerLoad == Linker);
+				// Sanity check to make sure that Linker is the linker that loaded our Result package or the linker has already been detached
+				check(!Result || Result->LinkerLoad == Linker || Result->LinkerLoad == nullptr);
 				if (Result && Linker->Loader)
 				{
 					ResetLoaders(Result);
