Stats - Added markers for garbage collection

diff --git a/Engine/Source/Runtime/CoreUObject/Private/UObject/GarbageCollection.cpp b/Engine/Source/Runtime/CoreUObject/Private/UObject/GarbageCollection.cpp
--- a/Engine/Source/Runtime/CoreUObject/Private/UObject/GarbageCollection.cpp
+++ b/Engine/Source/Runtime/CoreUObject/Private/UObject/GarbageCollection.cpp
@@ -1233,6 +1233,7 @@ static FAutoConsoleVariableRef CVarFlushStreamingOnGC(
 void CollectGarbageInternal(EObjectFlags KeepFlags, bool bPerformFullPurge)
 {
 	DECLARE_SCOPE_CYCLE_COUNTER( TEXT( "CollectGarbageInternal" ), STAT_CollectGarbageInternal, STATGROUP_GC );
+	STAT_ADD_CUSTOMMESSAGE_NAME( STAT_NamedMarker, TEXT( "GarbageCollection - Begin" ) );
 
 	// We can't collect garbage while there's a load in progress. E.g. one potential issue is Import.XObject
 	check(!IsLoading());
@@ -1364,6 +1365,8 @@ void CollectGarbageInternal(EObjectFlags KeepFlags, bool bPerformFullPurge)
 
 	// Route callbacks to verify GC assumptions
 	FCoreUObjectDelegates::PostGarbageCollect.Broadcast();
+
+	STAT_ADD_CUSTOMMESSAGE_NAME( STAT_NamedMarker, TEXT( "GarbageCollection - End" ) );
 }
 
 void CollectGarbage(EObjectFlags KeepFlags, bool bPerformFullPurge)
