Deferred linker load delete to fix potential crashes when code attempts to access linkers that have already been deleted earlier in the callstack

diff --git a/Engine/Source/Editor/UnrealEd/Private/FeedbackContextEditor.cpp b/Engine/Source/Editor/UnrealEd/Private/FeedbackContextEditor.cpp
--- a/Engine/Source/Editor/UnrealEd/Private/FeedbackContextEditor.cpp
+++ b/Engine/Source/Editor/UnrealEd/Private/FeedbackContextEditor.cpp
@@ -451,6 +451,8 @@ void FFeedbackContextEditor::FinalizeSlowTask()
 	FFeedbackContext::FinalizeSlowTask( );
 }
 
+COREUOBJECT_API void DeleteLoaders();
+
 void FFeedbackContextEditor::ProgressReported( const float TotalProgressInterp, FText DisplayMessage )
 {
 	// Clean up deferred cleanup objects from rendering thread every once in a while.
@@ -463,8 +465,10 @@ void FFeedbackContextEditor::ProgressReported( const float TotalProgressInterp,
 		FlushRenderingCommands();
 		// It is now safe to delete the pending clean objects.
 		delete PendingCleanupObjects;
+		// This is also a good time to destroy any linkers pending delete		
+		DeleteLoaders();
 		// Keep track of time this operation was performed so we don't do it too often.
-		LastTimePendingCleanupObjectsWhereDeleted = FPlatformTime::Seconds();
+		LastTimePendingCleanupObjectsWhereDeleted = FPlatformTime::Seconds();		
 	}
 
 	if (FSlateApplication::Get().CanDisplayWindows())
diff --git a/Engine/Source/Runtime/CoreUObject/Private/Serialization/AsyncLoading.cpp b/Engine/Source/Runtime/CoreUObject/Private/Serialization/AsyncLoading.cpp
--- a/Engine/Source/Runtime/CoreUObject/Private/Serialization/AsyncLoading.cpp
+++ b/Engine/Source/Runtime/CoreUObject/Private/Serialization/AsyncLoading.cpp
@@ -9,6 +9,7 @@
 #include "Serialization/DeferredMessageLog.h"
 #include "Serialization/AsyncPackage.h"
 #include "UObject/UObjectThreadContext.h"
+#include "UObject/LinkerManager.h"
 #include "Serialization/AsyncLoadingThread.h"
 
 /*-----------------------------------------------------------------------------
@@ -886,7 +887,8 @@ void FAsyncPackage::ResetLoader()
 	{
 		check(Linker->AsyncRoot == this || Linker->AsyncRoot == nullptr);
 		Linker->AsyncRoot = nullptr;
-		delete Linker;
+		Linker->Detach();
+		FLinkerManager::Get().RemoveLinker(Linker);
 		Linker = nullptr;
 	}
 }
diff --git a/Engine/Source/Runtime/CoreUObject/Private/UObject/GarbageCollection.cpp b/Engine/Source/Runtime/CoreUObject/Private/UObject/GarbageCollection.cpp
--- a/Engine/Source/Runtime/CoreUObject/Private/UObject/GarbageCollection.cpp
+++ b/Engine/Source/Runtime/CoreUObject/Private/UObject/GarbageCollection.cpp
@@ -1224,6 +1224,8 @@ static FAutoConsoleVariableRef CVarFlushStreamingOnGC(
 	ECVF_Default
 	);
 
+COREUOBJECT_API void DeleteLoaders();
+
 /** 
  * Deletes all unreferenced objects, keeping objects that have any of the passed in KeepFlags set
  *
@@ -1363,6 +1365,9 @@ void CollectGarbageInternal(EObjectFlags KeepFlags, bool bPerformFullPurge)
 		IncrementalPurgeGarbage( false );	
 	}
 
+	// Destroy all pending delete linkers
+	DeleteLoaders();
+
 	// Route callbacks to verify GC assumptions
 	FCoreUObjectDelegates::PostGarbageCollect.Broadcast();
 
diff --git a/Engine/Source/Runtime/CoreUObject/Private/UObject/Linker.cpp b/Engine/Source/Runtime/CoreUObject/Private/UObject/Linker.cpp
--- a/Engine/Source/Runtime/CoreUObject/Private/UObject/Linker.cpp
+++ b/Engine/Source/Runtime/CoreUObject/Private/UObject/Linker.cpp
@@ -319,17 +319,18 @@ FLinker::~FLinker()
 	Global functions
 -----------------------------------------------------------------------------*/
 
-//
-// Empty the loaders.
-//
-void ResetLoaders( UObject* InPkg )
+void ResetLoaders(UObject* InPkg)
 {
 	// Make sure we're not in the middle of loading something in the background.
 	FlushAsyncLoading();
 	FLinkerManager::Get().ResetLoaders(InPkg);
 }
 
-
+/** Deletes all linkers that have finished loading */
+COREUOBJECT_API void DeleteLoaders()
+{
+	FLinkerManager::Get().DeleteLinkers();
+}
 
 /**
  * Dissociates all linker import and forced export object references. This currently needs to 
diff --git a/Engine/Source/Runtime/CoreUObject/Private/UObject/LinkerManager.cpp b/Engine/Source/Runtime/CoreUObject/Private/UObject/LinkerManager.cpp
--- a/Engine/Source/Runtime/CoreUObject/Private/UObject/LinkerManager.cpp
+++ b/Engine/Source/Runtime/CoreUObject/Private/UObject/LinkerManager.cpp
@@ -108,7 +108,8 @@ void FLinkerManager::ResetLoaders(UObject* InPkg)
 			}
 			// Detach linker, also removes from array and sets LinkerRoot to NULL.
 			LinkerToReset->LoadAndDetachAllBulkData();
-			delete LinkerToReset;
+			LinkerToReset->Detach();
+			RemoveLinker(LinkerToReset);
 			LinkerToReset = nullptr;
 		}
 	}
@@ -121,7 +122,8 @@ void FLinkerManager::ResetLoaders(UObject* InPkg)
 		{
 			// Detach linker, also removes from array and sets LinkerRoot to NULL.
 			Linker->LoadAndDetachAllBulkData();
-			delete Linker;
+			Linker->Detach();
+			RemoveLinker(Linker);
 		}
 	}
 }
@@ -174,4 +176,41 @@ void FLinkerManager::DissociateImportsAndForcedExports()
 		}	
 		ForcedExportCount = 0;
 	}	
+}
+
+void FLinkerManager::DeleteLinkers()
+{
+	check(IsInGameThread());
+
+	QUICK_SCOPE_CYCLE_COUNTER(STAT_FLinkerManager_DeleteLinkers);
+
+	TArray<FLinkerLoad*> CleanupArray;
+	{
+#if THREADSAFE_UOBJECTS
+		FScopeLock PendingCleanupListLock(&PendingCleanupListCritical);
+#endif
+		CleanupArray = PendingCleanupList.Array();
+		PendingCleanupList.Empty();
+	}
+
+	// Note that even though DeleteLinkers can only be called on the main thread,
+	// we store the IsDeleteingLinkers in TLS so that we're sure nothing on
+	// another thread can delete linkers except FLinkerManager at the time
+	// we enter this loop.
+	FUObjectThreadContext& ThreadContext = FUObjectThreadContext::Get();
+	for (FLinkerLoad* Linker : CleanupArray)
+	{
+		delete Linker;
+	}
+}
+
+void FLinkerManager::RemoveLinker(FLinkerLoad* Linker)
+{
+#if THREADSAFE_UOBJECTS
+	FScopeLock PendingCleanupListLock(&PendingCleanupListCritical);
+#endif
+	if (Linker && !PendingCleanupList.Contains(Linker))
+	{
+		PendingCleanupList.Add(Linker);
+	}
 }
\ No newline at end of file
diff --git a/Engine/Source/Runtime/CoreUObject/Private/UObject/LinkerManager.h b/Engine/Source/Runtime/CoreUObject/Private/UObject/LinkerManager.h
--- a/Engine/Source/Runtime/CoreUObject/Private/UObject/LinkerManager.h
+++ b/Engine/Source/Runtime/CoreUObject/Private/UObject/LinkerManager.h
@@ -115,6 +115,12 @@ public:
 	*/
 	void DissociateImportsAndForcedExports();
 
+	/** Deletes all linkers that finished loading */
+	void DeleteLinkers();
+
+	/** Adds a linker to deferred cleanup list */
+	void RemoveLinker(FLinkerLoad* Linker);
+
 private:
 
 	/** Map of packages to their open linkers **/
@@ -133,4 +139,9 @@ private:
 	TArray<FLinkerLoad*> LiveLinkers;
 #endif
 	
+	/** List of linkers to delete **/
+	TSet<FLinkerLoad*>	PendingCleanupList;
+#if THREADSAFE_UOBJECTS
+	FCriticalSection PendingCleanupListCritical;
+#endif
 };
diff --git a/Engine/Source/Runtime/CoreUObject/Private/UObject/Obj.cpp b/Engine/Source/Runtime/CoreUObject/Private/UObject/Obj.cpp
--- a/Engine/Source/Runtime/CoreUObject/Private/UObject/Obj.cpp
+++ b/Engine/Source/Runtime/CoreUObject/Private/UObject/Obj.cpp
@@ -3632,6 +3632,8 @@ void StaticUObjectInit()
 	UE_LOG(LogInit, Log, TEXT("Object subsystem initialized") );
 }
 
+COREUOBJECT_API void DeleteLoaders();
+
 //
 // Shut down the object manager.
 //
@@ -3719,6 +3721,8 @@ void StaticExit()
 	UObjectBaseShutdown();
 	// Empty arrays to prevent falsely-reported memory leaks.
 	FUObjectThreadContext::Get().ObjLoaded.Empty();
+	// Delete all linkers are pending destroy
+	DeleteLoaders();
 	UE_LOG(LogExit, Log, TEXT("Object subsystem successfully closed.") );
 }
 
diff --git a/Engine/Source/Runtime/CoreUObject/Private/UObject/Package.cpp b/Engine/Source/Runtime/CoreUObject/Private/UObject/Package.cpp
--- a/Engine/Source/Runtime/CoreUObject/Private/UObject/Package.cpp
+++ b/Engine/Source/Runtime/CoreUObject/Private/UObject/Package.cpp
@@ -1,6 +1,7 @@
 // Copyright 1998-2015 Epic Games, Inc. All Rights Reserved.
 
 #include "CoreUObjectPrivate.h"
+#include "UObject/LinkerManager.h"
 
 /*-----------------------------------------------------------------------------
 	UPackage.
@@ -193,7 +194,7 @@ void UPackage::BeginDestroy()
 	// Detach linker if still attached
 	if (LinkerLoad)
 	{
-		delete LinkerLoad;
+		FLinkerManager::Get().RemoveLinker(LinkerLoad);
 		LinkerLoad = nullptr;
 	}
 
diff --git a/Engine/Source/Runtime/CoreUObject/Public/UObject/Linker.h b/Engine/Source/Runtime/CoreUObject/Public/UObject/Linker.h
--- a/Engine/Source/Runtime/CoreUObject/Public/UObject/Linker.h
+++ b/Engine/Source/Runtime/CoreUObject/Public/UObject/Linker.h
@@ -1805,16 +1805,19 @@ private:
 #endif
 	void DetachAllBulkData(bool bEnsureBulkDataIsLoaded);
 public:
+
 	/**
-	 * Detaches linker from bulk data/ exports and removes itself from array of loaders.
-	 *
-	 * @param	bEnsureAllBulkDataIsLoaded	Whether to load all bulk data first before detaching.
+	 * Detaches linker from bulk data.
 	 */
 	void LoadAndDetachAllBulkData();
 
+	/**
+	* Detaches linker from bulk data/ exports and removes itself from array of loaders.
+	*/
+	void Detach();
+
 private:
 
-	void Detach();
 	void Seek( int64 InPos ) override;
 	int64 Tell() override;
 	int64 TotalSize() override;
@@ -2256,6 +2259,7 @@ typedef uint32 ELazyLoaderFlags;
 	Global functions
 -----------------------------------------------------------------------------*/
 
+/** Resets linkers on packages after they have finished loading */
 COREUOBJECT_API void ResetLoaders( UObject* InOuter );
 
 /**
diff --git a/Engine/Source/Runtime/Launch/Private/LaunchEngineLoop.cpp b/Engine/Source/Runtime/Launch/Private/LaunchEngineLoop.cpp
--- a/Engine/Source/Runtime/Launch/Private/LaunchEngineLoop.cpp
+++ b/Engine/Source/Runtime/Launch/Private/LaunchEngineLoop.cpp
@@ -2252,6 +2252,10 @@ bool FEngineLoop::ShouldUseIdleMode() const
 	return bIdleMode;
 }
 
+#if WITH_COREUOBJECT
+COREUOBJECT_API void DeleteLoaders();
+#endif
+
 void FEngineLoop::Tick()
 {
 	// Ensure we aren't starting a frame while loading or playing a loading movie
@@ -2486,6 +2490,11 @@ void FEngineLoop::Tick()
 			// Delete the objects which were enqueued for deferred cleanup before the previous frame.
 			delete PreviousPendingCleanupObjects;
 
+			// Destroy all linkers pending delete
+#if WITH_COREUOBJECT
+			DeleteLoaders();
+#endif
+
 			FTicker::GetCoreTicker().Tick(FApp::GetDeltaTime());
 
 			FSingleThreadManager::Get().Tick();
