Adding guards to linker deletion so that nothing else can delete them except FLinkerManager.

Moved DeleteLoaders() declaration to Linker.h

diff --git a/Engine/Source/Editor/UnrealEd/Private/FeedbackContextEditor.cpp b/Engine/Source/Editor/UnrealEd/Private/FeedbackContextEditor.cpp
--- a/Engine/Source/Editor/UnrealEd/Private/FeedbackContextEditor.cpp
+++ b/Engine/Source/Editor/UnrealEd/Private/FeedbackContextEditor.cpp
@@ -451,8 +451,6 @@ void FFeedbackContextEditor::FinalizeSlowTask()
 	FFeedbackContext::FinalizeSlowTask( );
 }
 
-COREUOBJECT_API void DeleteLoaders();
-
 void FFeedbackContextEditor::ProgressReported( const float TotalProgressInterp, FText DisplayMessage )
 {
 	// Clean up deferred cleanup objects from rendering thread every once in a while.
diff --git a/Engine/Source/Runtime/CoreUObject/Private/UObject/GarbageCollection.cpp b/Engine/Source/Runtime/CoreUObject/Private/UObject/GarbageCollection.cpp
--- a/Engine/Source/Runtime/CoreUObject/Private/UObject/GarbageCollection.cpp
+++ b/Engine/Source/Runtime/CoreUObject/Private/UObject/GarbageCollection.cpp
@@ -1224,8 +1224,6 @@ static FAutoConsoleVariableRef CVarFlushStreamingOnGC(
 	ECVF_Default
 	);
 
-COREUOBJECT_API void DeleteLoaders();
-
 /** 
  * Deletes all unreferenced objects, keeping objects that have any of the passed in KeepFlags set
  *
diff --git a/Engine/Source/Runtime/CoreUObject/Private/UObject/Linker.cpp b/Engine/Source/Runtime/CoreUObject/Private/UObject/Linker.cpp
--- a/Engine/Source/Runtime/CoreUObject/Private/UObject/Linker.cpp
+++ b/Engine/Source/Runtime/CoreUObject/Private/UObject/Linker.cpp
@@ -327,7 +327,7 @@ void ResetLoaders(UObject* InPkg)
 }
 
 /** Deletes all linkers that have finished loading */
-COREUOBJECT_API void DeleteLoaders()
+void DeleteLoaders()
 {
 	FLinkerManager::Get().DeleteLinkers();
 }
diff --git a/Engine/Source/Runtime/CoreUObject/Private/UObject/LinkerLoad.cpp b/Engine/Source/Runtime/CoreUObject/Private/UObject/LinkerLoad.cpp
--- a/Engine/Source/Runtime/CoreUObject/Private/UObject/LinkerLoad.cpp
+++ b/Engine/Source/Runtime/CoreUObject/Private/UObject/LinkerLoad.cpp
@@ -718,6 +718,8 @@ FLinkerLoad::~FLinkerLoad()
 	FLinkerManager::Get().GetLiveLinkers().Remove(this);
 #endif
 
+	UE_CLOG(!FUObjectThreadContext::Get().IsDeletingLinkers, LogLinker, Fatal, TEXT("Linkers can only be deleted by FLinkerManager."));
+
 	// Detaches linker.
 	Detach();
 
diff --git a/Engine/Source/Runtime/CoreUObject/Private/UObject/LinkerManager.cpp b/Engine/Source/Runtime/CoreUObject/Private/UObject/LinkerManager.cpp
--- a/Engine/Source/Runtime/CoreUObject/Private/UObject/LinkerManager.cpp
+++ b/Engine/Source/Runtime/CoreUObject/Private/UObject/LinkerManager.cpp
@@ -194,14 +194,16 @@ void FLinkerManager::DeleteLinkers()
 	}
 
 	// Note that even though DeleteLinkers can only be called on the main thread,
-	// we store the IsDeleteingLinkers in TLS so that we're sure nothing on
+	// we store the IsDeletingLinkers in TLS so that we're sure nothing on
 	// another thread can delete linkers except FLinkerManager at the time
 	// we enter this loop.
 	FUObjectThreadContext& ThreadContext = FUObjectThreadContext::Get();
+	ThreadContext.IsDeletingLinkers = true;
 	for (FLinkerLoad* Linker : CleanupArray)
 	{
 		delete Linker;
 	}
+	ThreadContext.IsDeletingLinkers = false;
 }
 
 void FLinkerManager::RemoveLinker(FLinkerLoad* Linker)
diff --git a/Engine/Source/Runtime/CoreUObject/Private/UObject/Obj.cpp b/Engine/Source/Runtime/CoreUObject/Private/UObject/Obj.cpp
--- a/Engine/Source/Runtime/CoreUObject/Private/UObject/Obj.cpp
+++ b/Engine/Source/Runtime/CoreUObject/Private/UObject/Obj.cpp
@@ -3632,8 +3632,6 @@ void StaticUObjectInit()
 	UE_LOG(LogInit, Log, TEXT("Object subsystem initialized") );
 }
 
-COREUOBJECT_API void DeleteLoaders();
-
 //
 // Shut down the object manager.
 //
diff --git a/Engine/Source/Runtime/CoreUObject/Private/UObject/UObjectThreadContext.cpp b/Engine/Source/Runtime/CoreUObject/Private/UObject/UObjectThreadContext.cpp
--- a/Engine/Source/Runtime/CoreUObject/Private/UObject/UObjectThreadContext.cpp
+++ b/Engine/Source/Runtime/CoreUObject/Private/UObject/UObjectThreadContext.cpp
@@ -12,6 +12,7 @@ FUObjectThreadContext::FUObjectThreadContext()
 , ForcedExportCount(0)
 , ObjBeginLoadCount(0)
 , IsRoutingPostLoad(false)
+, IsDeletingLinkers(false)
 , IsInConstructor(0)
 , ConstructedObject(nullptr)
 , SerializedObject(nullptr)
diff --git a/Engine/Source/Runtime/CoreUObject/Public/UObject/Linker.h b/Engine/Source/Runtime/CoreUObject/Public/UObject/Linker.h
--- a/Engine/Source/Runtime/CoreUObject/Public/UObject/Linker.h
+++ b/Engine/Source/Runtime/CoreUObject/Public/UObject/Linker.h
@@ -2262,6 +2262,9 @@ typedef uint32 ELazyLoaderFlags;
 /** Resets linkers on packages after they have finished loading */
 COREUOBJECT_API void ResetLoaders( UObject* InOuter );
 
+/** Deletes all linkers that have finished loading */
+COREUOBJECT_API void DeleteLoaders();
+
 /**
  * Dissociates all linker import and forced export object references. This currently needs to 
  * happen as the referred objects might be destroyed at any time.
diff --git a/Engine/Source/Runtime/CoreUObject/Public/UObject/UObjectThreadContext.h b/Engine/Source/Runtime/CoreUObject/Public/UObject/UObjectThreadContext.h
--- a/Engine/Source/Runtime/CoreUObject/Public/UObject/UObjectThreadContext.h
+++ b/Engine/Source/Runtime/CoreUObject/Public/UObject/UObjectThreadContext.h
@@ -26,6 +26,8 @@ public:
 	TArray<FLinkerLoad*> DelayedLinkerClosePackages;
 	/** true when we are routing ConditionalPostLoad/PostLoad to objects										*/
 	bool IsRoutingPostLoad;
+	/** true when FLinkerManager deletes linkers */
+	bool IsDeletingLinkers;
 	/* Global flag so that FObjectFinders know if they are called from inside the UObject constructors or not. */
 	int32 IsInConstructor;
 	/* Object that is currently being constructed with ObjectInitializer */
diff --git a/Engine/Source/Runtime/Launch/Private/LaunchEngineLoop.cpp b/Engine/Source/Runtime/Launch/Private/LaunchEngineLoop.cpp
--- a/Engine/Source/Runtime/Launch/Private/LaunchEngineLoop.cpp
+++ b/Engine/Source/Runtime/Launch/Private/LaunchEngineLoop.cpp
@@ -2252,10 +2252,6 @@ bool FEngineLoop::ShouldUseIdleMode() const
 	return bIdleMode;
 }
 
-#if WITH_COREUOBJECT
-COREUOBJECT_API void DeleteLoaders();
-#endif
-
 void FEngineLoop::Tick()
 {
 	// Ensure we aren't starting a frame while loading or playing a loading movie
