Engine - Writing better error messages when we fail to resolve package names.

diff --git a/Engine/Source/Runtime/CoreUObject/Private/Misc/PackageName.cpp b/Engine/Source/Runtime/CoreUObject/Private/Misc/PackageName.cpp
--- a/Engine/Source/Runtime/CoreUObject/Private/Misc/PackageName.cpp
+++ b/Engine/Source/Runtime/CoreUObject/Private/Misc/PackageName.cpp
@@ -206,10 +206,10 @@ private:
 
 		FString RebasedGameDir = FString::Printf(TEXT("../../../%s/"), FApp::GetGameName());
 
-        GameContentPathRebased = RebasedGameDir / TEXT("Content/");
+		GameContentPathRebased = RebasedGameDir / TEXT("Content/");
 		GameConfigPathRebased  = RebasedGameDir / TEXT("Config/");
-        GameScriptPathRebased  = RebasedGameDir / TEXT("Script/");
-        GameSavedPathRebased   = RebasedGameDir / TEXT("Saved/");
+		GameScriptPathRebased  = RebasedGameDir / TEXT("Script/");
+		GameSavedPathRebased   = RebasedGameDir / TEXT("Saved/");
 		
 		ContentPathToRoot.Empty(11);
 		ContentPathToRoot.Emplace(EngineRootPath, EngineContentPath);
@@ -748,7 +748,16 @@ bool FPackageName::SearchForPackageOnDisk(const FString& PackageName, FString* O
 		}
 	}
 	float ThisTime = FPlatformTime::Seconds() - StartTime;
-	UE_LOG(LogPackageName, Log, TEXT("SearchForPackageOnDisk took %7.3fs to resolve %s"), ThisTime, *PackageName);
+
+	if ( bResult )
+	{
+		UE_LOG(LogPackageName, Log, TEXT("SearchForPackageOnDisk took %7.3fs to resolve %s."), ThisTime, *PackageName);
+	}
+	else
+	{
+		UE_LOG(LogPackageName, Log, TEXT("SearchForPackageOnDisk took %7.3fs, but failed to resolve %s."), ThisTime, *PackageName);
+	}
+
 	return bResult;
 }
 
