When queueing new async packages, also look for existing packages in loaded packages to process.

Added UPackage member to FPackageAsync to track both linker and the package.

Made sure async packages get deleted when cancelling async loading

diff --git a/Engine/Source/Runtime/CoreUObject/Private/Serialization/AsyncLoading.cpp b/Engine/Source/Runtime/CoreUObject/Private/Serialization/AsyncLoading.cpp
--- a/Engine/Source/Runtime/CoreUObject/Private/Serialization/AsyncLoading.cpp
+++ b/Engine/Source/Runtime/CoreUObject/Private/Serialization/AsyncLoading.cpp
@@ -312,6 +312,7 @@ void FAsyncLoadingThread::CancelAsyncLoadingInternal()
 		for (auto AsyncPackage : AsyncPackages)
 		{
 			AsyncPackage->Cancel();
+			delete AsyncPackage;
 		}
 
 		AsyncPackages.Empty();
@@ -325,12 +326,18 @@ void FAsyncLoadingThread::CancelAsyncLoadingInternal()
 		for (auto LoadedPackage : LoadedPackages)
 		{
 			LoadedPackage->Cancel();
+			delete LoadedPackage;
 		}
 		LoadedPackages.Empty();
-
+	}
+	{
+#if THREADSAFE_UOBJECTS
+		FScopeLock LoadedLock(&LoadedPackagesToProcessCritical);
+#endif
 		for (auto LoadedPackage : LoadedPackagesToProcess)
 		{
 			LoadedPackage->Cancel();
+			delete LoadedPackage;
 		}
 		LoadedPackagesToProcess.Empty();
 	}
@@ -358,6 +365,25 @@ void FAsyncLoadingThread::QueuePackage(const FAsyncPackageDesc& Package)
 	QueuedRequestsEvent->Trigger();
 }
 
+FAsyncPackage* FAsyncLoadingThread::FindExistingPackageAndAddCompletionCallback(FAsyncPackageDesc* PackageRequest, TArray<FAsyncPackage*>& PackageList)
+{
+	checkSlow(IsInAsyncLoadThread());
+	FAsyncPackage* Result = nullptr;
+	const int32 ExistingPackageIndex = FindPackageByName(PackageList, PackageRequest->Name);
+	if (ExistingPackageIndex != INDEX_NONE)
+	{
+		Result = PackageList[ExistingPackageIndex];
+		if (PackageRequest->PackageLoadedDelegate.IsBound())
+		{
+			const bool bInternalCallback = false;
+			Result->AddCompletionCallback(PackageRequest->PackageLoadedDelegate, bInternalCallback);
+		}
+		const int32 QueuedPackagesCount = QueuedPackagesCounter.Decrement();
+		check(QueuedPackagesCount >= 0);
+	}
+	return Result;
+}
+
 int32 FAsyncLoadingThread::CreateAsyncPackagesFromQueue()
 {
 	SCOPE_CYCLE_COUNTER(STAT_FAsyncPackage_CreateAsyncPackagesFromQueue);
@@ -375,45 +401,43 @@ int32 FAsyncLoadingThread::CreateAsyncPackagesFromQueue()
 		QueueCopy = QueuedPackages;
 		QueuedPackages.Empty();
 	}
+
 	for (auto PackageRequest : QueueCopy)
 	{
-		FAsyncPackage* Package = nullptr;
-		int ExistingPackageIndex = FindAsyncPackage(PackageRequest->Name);
-		if (ExistingPackageIndex != INDEX_NONE)
-		{
-			// This is a package that is currently being loaded.
-			Package = AsyncPackages[ExistingPackageIndex];
-		}
-		else
+		FAsyncPackage* Package = FindExistingPackageAndAddCompletionCallback(PackageRequest, AsyncPackages);
+
+		if (!Package)
 		{
 			// [BLOCKING] LoadedPackages are accessed on the main thread too, so lock to be able to add a completion callback
 #if THREADSAFE_UOBJECTS
 			FScopeLock LoadedLock(&LoadedPackagesCritical);
 #endif
-			ExistingPackageIndex = FindLoadedPackage(PackageRequest->Name);
-			if (ExistingPackageIndex != INDEX_NONE && PackageRequest->PackageLoadedDelegate.IsBound())
-			{
-				const bool bInternalCallback = false;
-				LoadedPackages[ExistingPackageIndex]->AddCompletionCallback(PackageRequest->PackageLoadedDelegate, bInternalCallback);
-				const int32 QueuedPackagesCount = QueuedPackagesCounter.Decrement();
-				check(QueuedPackagesCount >= 0);
-				continue;
-			}
+			Package = FindExistingPackageAndAddCompletionCallback(PackageRequest, LoadedPackages);
+		}
+
+		if (!Package)
+		{
+			// [BLOCKING] LoadedPackagesToProcess are modified on the main thread, so lock to be able to add a completion callback
+#if THREADSAFE_UOBJECTS
+			FScopeLock LoadedLock(&LoadedPackagesToProcessCritical);
+#endif
+			Package = FindExistingPackageAndAddCompletionCallback(PackageRequest, LoadedPackagesToProcess);
 		}
+
 		if (!Package)
 		{
 			// New package that needs to be loaded or a package has already been loaded long time ago
 			Package = new FAsyncPackage(*PackageRequest);
+			if (PackageRequest->PackageLoadedDelegate.IsBound())
+			{
+				const bool bInternalCallback = false;
+				Package->AddCompletionCallback(PackageRequest->PackageLoadedDelegate, bInternalCallback);
+			}
 
 			// Add to queue according to priority.
 			InsertPackage(Package);
-		}
-		if (PackageRequest->PackageLoadedDelegate.IsBound())
-		{
-			const bool bInternalCallback = false;
-			Package->AddCompletionCallback(PackageRequest->PackageLoadedDelegate, bInternalCallback);
-		}
-		{
+
+			// For all other cases this is handled in FindExistingPackageAndAddCompletionCallback
 			const int32 QueuedPackagesCount = QueuedPackagesCounter.Decrement();
 			check(QueuedPackagesCount >= 0);
 		}
@@ -520,7 +544,8 @@ EAsyncPackageState::Type FAsyncLoadingThread::ProcessLoadedPackages(bool bUseTim
 	double TickStartTime = FPlatformTime::Seconds();
 	{
 #if THREADSAFE_UOBJECTS
-		FScopeLock LoadedLock(&LoadedPackagesCritical);
+		FScopeLock LoadedPackagesLock(&LoadedPackagesCritical);
+		FScopeLock LoadedPackagesToProcessLock(&LoadedPackagesToProcessCritical);
 #endif
 		LoadedPackagesToProcess.Append(LoadedPackages);
 		LoadedPackages.Empty();
@@ -541,7 +566,19 @@ EAsyncPackageState::Type FAsyncLoadingThread::ProcessLoadedPackages(bool bUseTim
 			{
 				// Remove the package from the list before we trigger the callbacks, 
 				// this is to ensure we can re-enter FlushAsyncLoading from any of the callbacks
-				LoadedPackagesToProcess.RemoveAt(PackageIndex--);
+				{
+					FScopeLock LoadedLock(&LoadedPackagesToProcessCritical);					
+					LoadedPackagesToProcess.RemoveAt(PackageIndex--);
+					// Detach linker in mutex scope to make sure that if something requests this package
+					// before it's been deleted does not try to associate the new async package with the old linker
+					// while this async package is still bound to it.
+					Package->DetachLinker();
+					if (FPlatformProperties::RequiresCookedData())
+					{
+						// Emulates ResetLoaders on the package linker's linkerroot.
+						Package->ResetLoader();
+					}
+				}
 
 				// Incremented on the Async Thread, now decrement as we're done with this package				
 				const int32 NewAsyncLoadingCounterValue = AsyncLoadingCounter.Decrement();
@@ -552,12 +589,6 @@ EAsyncPackageState::Type FAsyncLoadingThread::ProcessLoadedPackages(bool bUseTim
 				const EAsyncLoadingResult::Type LoadingResult = Package->HasLoadFailed() ? EAsyncLoadingResult::Failed : EAsyncLoadingResult::Succeeded;
 				Package->CallCompletionCallbacks(bInternalCallbacks, LoadingResult);
 
-				if (FPlatformProperties::RequiresCookedData())
-				{
-					// Emulates ResetLoaders on the package linker's linkerroot.
-					Package->ResetLoader();
-				}
-
 				// We don't need the package anymore
 				delete Package;
 			}
@@ -783,7 +814,7 @@ float FAsyncLoadingThread::GetAsyncLoadPercentage(const FName& PackageName)
 	if (LoadPercentage < 0.0f)
 	{
 		checkSlow(IsInGameThread());
-		// No lock required as we're in the game thread and LoadedPackagesToProcess are only managed on the game thread
+		// No lock required as we're in the game thread and LoadedPackagesToProcess are only modified on the game thread
 		const int32 PackageIndex = FindPackageByName(LoadedPackagesToProcess, PackageName);
 		if (PackageIndex != INDEX_NONE)
 		{
@@ -825,6 +856,7 @@ int32 FAsyncPackage::PostLoadIndex = 0;
 FAsyncPackage::FAsyncPackage(const FAsyncPackageDesc& InDesc)
 : Desc(InDesc)
 , Linker(nullptr)
+, LinkerRoot(nullptr)
 , DependencyRefCount(0)
 , LoadImportIndex(0)
 , ImportIndex(0)
@@ -968,7 +1000,7 @@ void FAsyncPackage::EndAsyncLoad()
 	if (!bLoadHasFailed)
 	{
 		// Mark the package as loaded, if we succeeded
-		Linker->LinkerRoot->SetFlags(RF_WasLoaded);
+		LinkerRoot->SetFlags(RF_WasLoaded);
 	}
 }
 
@@ -1110,6 +1142,8 @@ EAsyncPackageState::Type FAsyncPackage::CreateLinker()
 		{
 			FGCScopeGuard GCGuard;
 			Package = CreatePackage(nullptr, *Desc.Name.ToString());
+			FAsyncObjectsReferencer::Get().AddObject(Package);
+			LinkerRoot = Package;
 		}
 		FScopeCycleCounterUObject ConstructorScope(Package, GET_STATID(STAT_FAsyncPackage_CreateLinker));
 
@@ -1152,7 +1186,7 @@ EAsyncPackageState::Type FAsyncPackage::CreateLinker()
 		}
 
 		// Associate this async package with the linker
-		check(Linker->AsyncRoot == nullptr);
+		check(Linker->AsyncRoot == nullptr || Linker->AsyncRoot == this);
 		Linker->AsyncRoot = this;
 
 		UE_LOG(LogStreaming, Verbose, TEXT("FAsyncPackage::CreateLinker for %s finished."), *Desc.NameToLoad.ToString());
@@ -1168,10 +1202,10 @@ EAsyncPackageState::Type FAsyncPackage::CreateLinker()
 EAsyncPackageState::Type FAsyncPackage::FinishLinker()
 {
 	EAsyncPackageState::Type Result = EAsyncPackageState::Complete;
-	if( !Linker->HasFinishedInitialization() )
+	if (Linker && !Linker->HasFinishedInitialization())
 	{
 		SCOPE_CYCLE_COUNTER(STAT_FAsyncPackage_FinishLinker);
-		LastObjectWorkWasPerformedOn	= Linker->LinkerRoot;
+		LastObjectWorkWasPerformedOn	= LinkerRoot;
 		LastTypeOfWorkPerformed			= TEXT("ticking linker");
 	
 		const float RemainingTimeLimit = TimeLimit - (float)(FPlatformTime::Seconds() - TickStartTime);
@@ -1265,7 +1299,8 @@ bool FAsyncPackage::AddUniqueLinkerDependencyPackage(int32 CurrentPackageIndex,
 {
 	if (ContainsDependencyPackage(PendingImportedPackages, PendingImport.GetPackageName()) == INDEX_NONE)
 	{
-		if (PendingImport.Linker == nullptr || !PendingImport.Linker->HasFinishedInitialization())
+		FLinkerLoad* PendingImportLinker = PendingImport.Linker;
+		if (PendingImportLinker == nullptr || !PendingImportLinker->HasFinishedInitialization())
 		{
 			AddImportDependency(CurrentPackageIndex, PendingImport.GetPackageName());
 			UE_LOG(LogStreaming, Verbose, TEXT("  Adding linker dependency %s"), *PendingImport.GetPackageName().ToString());
@@ -1310,7 +1345,7 @@ void FAsyncPackage::AddDependencyTree(int32 CurrentPackageIndex, FAsyncPackage&
 EAsyncPackageState::Type FAsyncPackage::LoadImports()
 {
 	SCOPE_CYCLE_COUNTER(STAT_FAsyncPackage_LoadImports);
-	LastObjectWorkWasPerformedOn	= Linker->LinkerRoot;
+	LastObjectWorkWasPerformedOn	= LinkerRoot;
 	LastTypeOfWorkPerformed			= TEXT("loading imports");
 	
 	// Index of this package in the async queue.
@@ -1351,7 +1386,8 @@ EAsyncPackageState::Type FAsyncPackage::LoadImports()
 			if (PendingAsyncPackageIndex != INDEX_NONE)
 			{
 				FAsyncPackage& PendingPackage = *FAsyncLoadingThread::Get().GetPackage(PendingAsyncPackageIndex);
-				if (PendingPackage.Linker == nullptr || !PendingPackage.Linker->HasFinishedInitialization())
+				FLinkerLoad* PendingPackageLinker = PendingPackage.Linker;
+				if (PendingPackageLinker == nullptr || !PendingPackageLinker->HasFinishedInitialization())
 				{
 					// Add this import to the dependency list.
 					AddUniqueLinkerDependencyPackage(AsyncQueueIndex, PendingPackage);
@@ -1673,17 +1709,17 @@ EAsyncPackageState::Type FAsyncPackage::PostLoadDeferredObjects(double InTickSta
 		}
 
 		// Mark package as having been fully loaded and update load time.
-		if (Linker && !bLoadHasFailed)
+		if (LinkerRoot && !bLoadHasFailed)
 		{
-			if (Linker->LinkerRoot)
+			LinkerRoot->AtomicallyClearFlags(RF_AsyncLoading);
+			LinkerRoot->MarkAsFullyLoaded();
+			LinkerRoot->SetLoadTime(FPlatformTime::Seconds() - LoadStartTime);
+
+			if (Linker)
 			{
-				Linker->LinkerRoot->AtomicallyClearFlags(RF_AsyncLoading);
-				Linker->LinkerRoot->MarkAsFullyLoaded();
-				Linker->LinkerRoot->SetLoadTime(FPlatformTime::Seconds() - LoadStartTime);
+				// give a hint to the IO system that we are done with this file for now
+				FIOSystem::Get().HintDoneWithFile(Linker->Filename);
 			}
-
-			// give a hint to the IO system that we are done with this file for now
-			FIOSystem::Get().HintDoneWithFile(Linker->Filename);
 		}
 	}
 
@@ -1749,12 +1785,12 @@ EAsyncPackageState::Type FAsyncPackage::FinishObjects()
 
 void FAsyncPackage::CallCompletionCallbacks(bool bInternal, EAsyncLoadingResult::Type LoadingResult)
 {
-	UPackage* LinkerRoot = (!bLoadHasFailed && Linker) ? Linker->LinkerRoot : nullptr;
+	UPackage* LoadedPackage = (!bLoadHasFailed) ? LinkerRoot : nullptr;
 	for (auto& CompletionCallback : CompletionCallbacks)
 	{
 		if (CompletionCallback.bIsInternal == bInternal)
 		{
-			CompletionCallback.Callback.ExecuteIfBound(Desc.Name, LinkerRoot, LoadingResult);
+			CompletionCallback.Callback.ExecuteIfBound(Desc.Name, LoadedPackage, LoadingResult);
 		}
 	}
 }
@@ -1767,19 +1803,19 @@ void FAsyncPackage::Cancel()
 	{
 		CompletionCallbacks[CallbackIndex].Callback.ExecuteIfBound(Desc.Name, nullptr, Result);
 	}
-	if (Linker)
+	if (LinkerRoot)
 	{
-		// give a hint to the IO system that we are done with this file for now
-		FIOSystem::Get().HintDoneWithFile(Linker->Filename);
-		Linker->FlushCache();
-		if (Linker->LinkerRoot)
+		if (Linker)
 		{
-			Linker->LinkerRoot->ClearFlags(RF_WasLoaded);
-			Linker->LinkerRoot->bHasBeenFullyLoaded = false;
-			Linker->LinkerRoot->Rename(*MakeUniqueObjectName(GetTransientPackage(), UPackage::StaticClass()).ToString(), nullptr, REN_DontCreateRedirectors | REN_DoNotDirty | REN_ForceNoResetLoaders | REN_NonTransactional);
+			// give a hint to the IO system that we are done with this file for now
+			FIOSystem::Get().HintDoneWithFile(Linker->Filename);
+			Linker->FlushCache();
 		}
+		LinkerRoot->ClearFlags(RF_WasLoaded);
+		LinkerRoot->bHasBeenFullyLoaded = false;
+		LinkerRoot->Rename(*MakeUniqueObjectName(GetTransientPackage(), UPackage::StaticClass()).ToString(), nullptr, REN_DontCreateRedirectors | REN_DoNotDirty | REN_ForceNoResetLoaders | REN_NonTransactional);
+		DetachLinker();
 	}
-
 	PreLoadIndex = 0;
 }
 
diff --git a/Engine/Source/Runtime/CoreUObject/Private/Serialization/AsyncLoadingThread.h b/Engine/Source/Runtime/CoreUObject/Private/Serialization/AsyncLoadingThread.h
--- a/Engine/Source/Runtime/CoreUObject/Private/Serialization/AsyncLoadingThread.h
+++ b/Engine/Source/Runtime/CoreUObject/Private/Serialization/AsyncLoadingThread.h
@@ -41,7 +41,13 @@ class FAsyncLoadingThread : public FRunnable
 #endif
 	/** [GAME THREAD] Event used to signal there's queued packages to stream */
 	TArray<FAsyncPackage*> LoadedPackagesToProcess;
-	
+#if THREADSAFE_UOBJECTS
+	/** [ASYNC/GAME THREAD] Critical section for LoadedPackagesToProcess list. 
+	 * Note this is only required for looking up existing packages on the async loading thread 
+	 */
+	FCriticalSection LoadedPackagesToProcessCritical;
+#endif
+
 	/** [ASYNC THREAD] Array of packages that are being preloaded */
 	TArray<FAsyncPackage*> AsyncPackages;
 #if THREADSAFE_UOBJECTS
@@ -296,6 +302,11 @@ private:
 	*/
 	int32 CreateAsyncPackagesFromQueue();
 
+	/**
+	* [ASYNC THREAD] Finds existing async package and adds the new request's completion callback to it.
+	*/
+	FAsyncPackage* FindExistingPackageAndAddCompletionCallback(FAsyncPackageDesc* PackageRequest, TArray<FAsyncPackage*>& PackageList);
+
 	/**
 	* [ASYNC THREAD] Adds a package to a list of packages that have finished loading on the async thread
 	*/
diff --git a/Engine/Source/Runtime/CoreUObject/Public/Serialization/AsyncLoading.h b/Engine/Source/Runtime/CoreUObject/Public/Serialization/AsyncLoading.h
--- a/Engine/Source/Runtime/CoreUObject/Public/Serialization/AsyncLoading.h
+++ b/Engine/Source/Runtime/CoreUObject/Public/Serialization/AsyncLoading.h
@@ -115,6 +115,8 @@ private:
 	FAsyncPackageDesc Desc;
 	/** Linker which is going to have its exports and imports loaded									*/
 	FLinkerLoad*				Linker;
+	/** Package which is going to have its exports and imports loaded									*/
+	UPackage*				LinkerRoot;
 	/** Call backs called when we finished loading this package											*/
 	TArray<FCompletionCallback>	CompletionCallbacks;
 	/** Pending Import packages - we wait until all of them have been fully loaded. */
