Removing transient text serialization warning.

diff --git a/Engine/Source/Runtime/CoreUObject/Private/UObject/UTextProperty.cpp b/Engine/Source/Runtime/CoreUObject/Private/UObject/UTextProperty.cpp
--- a/Engine/Source/Runtime/CoreUObject/Private/UObject/UTextProperty.cpp
+++ b/Engine/Source/Runtime/CoreUObject/Private/UObject/UTextProperty.cpp
@@ -42,14 +42,6 @@ bool UTextProperty::Identical( const void* A, const void* B, uint32 PortFlags )
 void UTextProperty::SerializeItem( FArchive& Ar, void* Value, void const* Defaults ) const
 {
 	const TCppType PropertyValue = GetPropertyValue(Value);
-	if ( Ar.IsSaving() && Ar.IsPersistent() && PropertyValue.IsTransient() )
-	{
-		// Convert to a string and back, to get it set as an invariant message instead of transient. If it stayed transient, error messages recursively expand
-		FString ErrorMessage = FText::Format( FText::SerializationFailureError, FText::FromString( FTextInspector::GetDisplayString(PropertyValue) ) ).ToString();
-		UE_LOG( LogProperty, Warning, TEXT("%s"), *ErrorMessage);
-		SetPropertyValue(Value, FText::FromString(ErrorMessage));
-	}
-
 	Ar << *GetPropertyValuePtr(Value);
 }
 
