[UE-17419] When super blueprints had been compiled we could incorrectly deferr a blueprint on load because its super blueprint was no longer tagged as RF_LoadCompleted. Checking the linker instead

diff --git a/Engine/Source/Runtime/CoreUObject/Private/UObject/UObjectGlobals.cpp b/Engine/Source/Runtime/CoreUObject/Private/UObject/UObjectGlobals.cpp
--- a/Engine/Source/Runtime/CoreUObject/Private/UObject/UObjectGlobals.cpp
+++ b/Engine/Source/Runtime/CoreUObject/Private/UObject/UObjectGlobals.cpp
@@ -2145,8 +2145,11 @@ FObjectInitializer::~FObjectInitializer()
 		UClass* ArchetypeClass = ObjectArchetype->GetClass();
 		// if this is a blueprint CDO that derives from another blueprint, and 
 		// that parent (archetype) CDO isn't fully serialized
-		if ( !Class->HasAnyFlags(RF_Native) && !ArchetypeClass->HasAnyFlags(RF_Native) && (ObjectArchetype->HasAnyFlags(RF_NeedLoad) ||
-			 !ObjectArchetype->HasAnyFlags(RF_LoadCompleted) || FDeferredObjInitializerTracker::IsCdoDeferred(ArchetypeClass)) )
+		if ( !Class->HasAnyFlags(RF_Native) && !ArchetypeClass->HasAnyFlags(RF_Native) && 
+			(	ObjectArchetype->HasAnyFlags(RF_NeedLoad) ||
+				(ArchetypeClass->GetLinker() && ArchetypeClass->GetLinker()->IsBlueprintFinalizationPending()) ||
+				FDeferredObjInitializerTracker::IsCdoDeferred(ArchetypeClass))
+			)
 		{
 			FLinkerLoad* ClassLinker = Class->GetLinker();
 			if ((ClassLinker != nullptr) && (ClassLinker->LoadFlags & LOAD_DeferDependencyLoads) != 0x00)
diff --git a/Engine/Source/Runtime/CoreUObject/Public/UObject/Linker.h b/Engine/Source/Runtime/CoreUObject/Public/UObject/Linker.h
--- a/Engine/Source/Runtime/CoreUObject/Public/UObject/Linker.h
+++ b/Engine/Source/Runtime/CoreUObject/Public/UObject/Linker.h
@@ -1946,6 +1946,16 @@ public:
 	 */
 	COREUOBJECT_API ELinkerStatus SerializeThumbnails( bool bForceEnableForCommandlet=false );
 
+	/**
+	* Query method to help handle recursive behavior. When this returns true,
+	* this linker is in the middle of, or is about to call FinalizeBlueprint()
+	* (for a blueprint class somewhere in the current callstack). Needed when
+	* we get to finalizing a sub-class before we've finished finalizing its
+	* super (so we know we need to finish finalizing the super first).
+	*
+	* @return True if FinalizeBlueprint() is currently being ran (or about to be ran) for an export (Blueprint) class.
+	*/
+	bool IsBlueprintFinalizationPending() const;
 private:
 	/**
 	 * Regenerates/Refreshes a blueprint class
@@ -2043,17 +2053,6 @@ private:
 	 */
 	void ResolveDeferredExports(UClass* LoadClass);
 
-	/**
-	 * Query method to help handle recursive behavior. When this returns true, 
-	 * this linker is in the middle of, or is about to call FinalizeBlueprint()
-	 * (for a blueprint class somewhere in the current callstack). Needed when 
-	 * we get to finalizing a sub-class before we've finished finalizing its 
-	 * super (so we know we need to finish finalizing the super first).
-	 * 
-	 * @return True if FinalizeBlueprint() is currently being ran (or about to be ran) for an export (Blueprint) class.
-	 */
-	bool IsBlueprintFinalizationPending() const;
-
 	/**
 	 * Sometimes we have to instantiate an export object that is of an imported 
 	 * type, and sometimes in those scenarios (thanks to cyclic dependencies) 
@@ -2085,7 +2084,7 @@ private:
 
 
 	void ResetDeferredLoadingState();
-public:
+
 	bool HasPerformedFullExportResolvePass();
 
 #if	USE_CIRCULAR_DEPENDENCY_LOAD_DEFERRING
