Fix for crash when restarting level after streaming levels in.

diff --git a/Engine/Source/Runtime/CoreUObject/Private/UObject/Package.cpp b/Engine/Source/Runtime/CoreUObject/Private/UObject/Package.cpp
--- a/Engine/Source/Runtime/CoreUObject/Private/UObject/Package.cpp
+++ b/Engine/Source/Runtime/CoreUObject/Private/UObject/Package.cpp
@@ -194,6 +194,7 @@ void UPackage::BeginDestroy()
 	// Detach linker if still attached
 	if (LinkerLoad)
 	{
+		LinkerLoad->Detach();
 		FLinkerManager::Get().RemoveLinker(LinkerLoad);
 		LinkerLoad = nullptr;
 	}
