Always try to find packages and objects in memory before attempting to load them.

diff --git a/Engine/Source/Runtime/CoreUObject/Private/UObject/UObjectGlobals.cpp b/Engine/Source/Runtime/CoreUObject/Private/UObject/UObjectGlobals.cpp
--- a/Engine/Source/Runtime/CoreUObject/Private/UObject/UObjectGlobals.cpp
+++ b/Engine/Source/Runtime/CoreUObject/Private/UObject/UObjectGlobals.cpp
@@ -672,7 +672,9 @@ bool ResolveName( UObject*& InPackage, FString& InOutName, bool Create, bool Thr
 		}
 		else if (!FPackageName::IsShortPackageName(PartialName))
 		{
-			if (!ScriptPackageName)
+			// Try to find the package in memory first, should be faster than attempting to load or create
+			InPackage = StaticFindObjectFast(UPackage::StaticClass(), InPackage, *PartialName);
+			if (!ScriptPackageName && !InPackage)
 			{
 				InPackage = LoadPackage(dynamic_cast<UPackage*>(InPackage), *PartialName, 0);
 			}
@@ -763,7 +765,9 @@ UObject* StaticLoadObjectInternal(UClass* ObjectClass, UObject* InOuter, const T
 	ResolveName(InOuter, StrName, true, true);
 	if (InOuter)
 	{
-		if (bAllowObjectReconciliation && ((FApp::IsGame() && !GIsEditor && !IsRunningCommandlet())
+		// If we have a full UObject name then attempt to find the object in memory first,
+		// unless we're currently inside of async loading path because the object may not be fully loaded yet.
+		if (bAllowObjectReconciliation && ((bContainsObjectName && !IsInAsyncLoadingThread())
 #if WITH_EDITOR
 			|| GIsImportingT3D
 #endif
