RF_IsWeaklyReferencedOnFrame flag added. GC cannot rmove object referenced form Persisten Frame. The reference is nulled first.

UE-13798 Make the shared state of the event graph use something akin to weak ptrs

diff --git a/Engine/Source/Runtime/CoreUObject/Private/UObject/GarbageCollection.cpp b/Engine/Source/Runtime/CoreUObject/Private/UObject/GarbageCollection.cpp
--- a/Engine/Source/Runtime/CoreUObject/Private/UObject/GarbageCollection.cpp
+++ b/Engine/Source/Runtime/CoreUObject/Private/UObject/GarbageCollection.cpp
@@ -295,7 +295,7 @@ void SerializeRootSet( FArchive& Ar, EObjectFlags KeepFlags )
  * @param ReferencingObject UObject which owns the reference (can be NULL)
  * @param bAllowReferenceElimination	Whether to allow NULL'ing the reference if RF_PendingKill is set
  */
-static FORCEINLINE void HandleObjectReference(TArray<UObject*>& ObjectsToSerialize, UObject* ReferencingObject, UObject*& Object, bool bAllowReferenceElimination, bool bDoHandle = true)
+static FORCEINLINE void HandleObjectReference(TArray<UObject*>& ObjectsToSerialize, UObject* ReferencingObject, UObject*& Object, bool bAllowReferenceElimination, bool bStrongReference = true)
 {
 	// Disregard NULL objects and perform very fast check to see whether object is part of permanent
 	// object pool and should therefore be disregarded. The check doesn't touch the object and is
@@ -312,8 +312,8 @@ static FORCEINLINE void HandleObjectReference(TArray<UObject*>& ObjectsToSeriali
 				Object = NULL;
 			}
 			// Add encountered object reference to list of to be serialized objects if it hasn't already been added.
-			else if (Object->HasAnyFlags(RF_Unreachable) && bDoHandle)
-			{				
+			else if (Object->HasAnyFlags(RF_Unreachable))
+			{	
 				if( GIsRunningParallelReachability )
 				{
 					// Mark it as reachable.
@@ -341,6 +341,11 @@ static FORCEINLINE void HandleObjectReference(TArray<UObject*>& ObjectsToSeriali
 					ObjectsToSerialize.Add( Object );
 				}
 			}
+
+			if (Object && bStrongReference)
+			{
+				Object->ClearFlags(RF_NoStrongReference);
+			}
 #if PERF_DETAILED_PER_CLASS_GC_STATS
 			GCurrentObjectRegularObjectRefs++;
 		}
@@ -613,7 +618,7 @@ public:
 				}
 				else
 				{
-					Object->SetFlags( RF_Unreachable );
+					Object->SetFlags(RF_Unreachable | RF_NoStrongReference);
 				}
 			}
 
@@ -1348,6 +1353,11 @@ void CollectGarbageInternal(EObjectFlags KeepFlags, bool bPerformFullPurge)
 				// Begin the object's asynchronous destruction.
 				Object->ConditionalBeginDestroy();
 			}
+			else if (Object->HasAnyFlags(RF_NoStrongReference))
+			{
+				Object->ClearFlags(RF_NoStrongReference);
+				Object->SetFlags(RF_PendingKill);
+			}
 		}
 		UE_LOG( LogGarbage, Log, TEXT( "%f ms for unhashing unreachable objects" ), (FPlatformTime::Seconds() - StartTime) * 1000 );
 	}
diff --git a/Engine/Source/Runtime/CoreUObject/Public/UObject/ObjectBase.h b/Engine/Source/Runtime/CoreUObject/Public/UObject/ObjectBase.h
--- a/Engine/Source/Runtime/CoreUObject/Public/UObject/ObjectBase.h
+++ b/Engine/Source/Runtime/CoreUObject/Public/UObject/ObjectBase.h
@@ -463,10 +463,11 @@ enum EObjectFlags
 	RF_InheritableComponentTemplate = 0x00400000, ///< Archetype of the object can be in its super class
 	RF_Async = 0x00800000, ///< Object exists only on a different thread than the game thread.
 	RF_StrongRefOnFrame			= 0x01000000,	///< References to this object from persistent function frame are handled as strong ones.
+	RF_NoStrongReference		= 0x02000000,  ///< The object is not referenced by any strong reference. The flag is used by GC.
 };
 
 	// Special all and none masks
-#define RF_AllFlags				(EObjectFlags)0x00ffffff	///< All flags, used mainly for error checking
+#define RF_AllFlags				(EObjectFlags)0x03ffffff	///< All flags, used mainly for error checking
 
 	// Predefined groups of the above
 #define RF_Load						((EObjectFlags)(RF_Public | RF_Standalone | RF_Native | RF_Transactional | RF_ClassDefaultObject | RF_ArchetypeObject | RF_DefaultSubObject | RF_TextExportTransient | RF_InheritableComponentTemplate)) // Flags to load from Unrealfiles.
diff --git a/Engine/Source/Runtime/Engine/Private/BlueprintGeneratedClass.cpp b/Engine/Source/Runtime/Engine/Private/BlueprintGeneratedClass.cpp
--- a/Engine/Source/Runtime/Engine/Private/BlueprintGeneratedClass.cpp
+++ b/Engine/Source/Runtime/Engine/Private/BlueprintGeneratedClass.cpp
@@ -688,7 +688,7 @@ public:
 protected:
 	virtual FArchive& operator<<(UObject*& Object) override
 	{
-		const bool bWeakRef = false; //Object ? !Object->HasAnyFlags(RF_StrongRefOnFrame) : false;
+		const bool bWeakRef = Object ? !Object->HasAnyFlags(RF_StrongRefOnFrame) : false;
 		Collector.SetShouldHandleAsWeakRef(bWeakRef); 
 		return FSimpleObjectReferenceCollectorArchive::operator<<(Object);
 	}
