Basic support for flushing async loading for the specified package only.

diff --git a/Engine/Source/Runtime/CoreUObject/Private/Serialization/AsyncLoading.cpp b/Engine/Source/Runtime/CoreUObject/Private/Serialization/AsyncLoading.cpp
--- a/Engine/Source/Runtime/CoreUObject/Private/Serialization/AsyncLoading.cpp
+++ b/Engine/Source/Runtime/CoreUObject/Private/Serialization/AsyncLoading.cpp
@@ -386,6 +386,7 @@ FAsyncPackage* FAsyncLoadingThread::FindExistingPackageAndAddCompletionCallback(
 		{
 			const bool bInternalCallback = false;
 			Result->AddCompletionCallback(PackageRequest->PackageLoadedDelegate, bInternalCallback);
+			Result->AddRequestID(PackageRequest->RequestID);
 		}
 		const int32 QueuedPackagesCount = QueuedPackagesCounter.Decrement();
 		check(QueuedPackagesCount >= 0);
@@ -411,7 +412,7 @@ int32 FAsyncLoadingThread::CreateAsyncPackagesFromQueue()
 		QueuedPackages.Empty();
 	}
 
-	for (auto PackageRequest : QueueCopy)
+	for (FAsyncPackageDesc* PackageRequest : QueueCopy)
 	{
 		FAsyncPackage* Package = FindExistingPackageAndAddCompletionCallback(PackageRequest, AsyncPackages);
 
@@ -546,7 +547,7 @@ EAsyncPackageState::Type FAsyncLoadingThread::ProcessAsyncLoading(int32& OutPack
 }
 
 
-EAsyncPackageState::Type FAsyncLoadingThread::ProcessLoadedPackages(bool bUseTimeLimit, bool bUseFullTimeLimit, float TimeLimit)
+EAsyncPackageState::Type FAsyncLoadingThread::ProcessLoadedPackages(bool bUseTimeLimit, bool bUseFullTimeLimit, float TimeLimit, int32 WaitForRequestID /*= INDEX_NONE*/)
 {
 	EAsyncPackageState::Type Result = EAsyncPackageState::Complete;
 
@@ -567,7 +568,7 @@ EAsyncPackageState::Type FAsyncLoadingThread::ProcessLoadedPackages(bool bUseTim
 			break;
 		}
 
-		auto Package = LoadedPackagesToProcess[PackageIndex];
+		FAsyncPackage* Package = LoadedPackagesToProcess[PackageIndex];
 		if (Package->GetDependencyRefCount() == 0)
 		{
 			Result = Package->PostLoadDeferredObjects(TickStartTime, bUseTimeLimit, TimeLimit);
@@ -600,6 +601,12 @@ EAsyncPackageState::Type FAsyncLoadingThread::ProcessLoadedPackages(bool bUseTim
 
 				// We don't need the package anymore
 				delete Package;
+
+				if (WaitForRequestID != INDEX_NONE && !ContainsRequestID(WaitForRequestID))
+				{
+					// The only package we care about has finished loading, so we're good to exit
+					break;
+				}
 			}
 			else
 			{
@@ -619,7 +626,7 @@ EAsyncPackageState::Type FAsyncLoadingThread::ProcessLoadedPackages(bool bUseTim
 
 
 
-EAsyncPackageState::Type FAsyncLoadingThread::TickAsyncLoading(bool bUseTimeLimit, bool bUseFullTimeLimit, float TimeLimit)
+EAsyncPackageState::Type FAsyncLoadingThread::TickAsyncLoading(bool bUseTimeLimit, bool bUseFullTimeLimit, float TimeLimit, int32 WaitForRequestID /*= INDEX_NONE*/)
 {
 	const bool bLoadingSuspended = IsAsyncLoadingSuspended();
 	const bool bIsMultithreaded = FAsyncLoadingThread::IsMultithreaded();
@@ -639,7 +646,7 @@ EAsyncPackageState::Type FAsyncLoadingThread::TickAsyncLoading(bool bUseTimeLimi
 		if (Result != EAsyncPackageState::TimeOut && !IsTimeLimitExceeded(TickStartTime, bUseTimeLimit, TimeLimit, TEXT("Pre-ProcessLoadedPackages")))
 		{
 			double RemainingTimeLimit = FMath::Max(0.0, TimeLimit - TimeLimitUsedForPreloading);
-			Result = ProcessLoadedPackages(bUseTimeLimit, bUseFullTimeLimit, RemainingTimeLimit);
+			Result = ProcessLoadedPackages(bUseTimeLimit, bUseFullTimeLimit, RemainingTimeLimit, WaitForRequestID);
 		}
 
 		if (Result != EAsyncPackageState::TimeOut && !IsTimeLimitExceeded(TickStartTime, bUseTimeLimit, TimeLimit, TEXT("Pre-EmptyReferencedObjects")))
@@ -903,13 +910,24 @@ FAsyncPackage::FAsyncPackage(const FAsyncPackageDesc& InDesc)
 , FinishObjectsTime(0.0)
 #endif // PERF_TRACK_DETAILED_ASYNC_STATS
 {
+	AddRequestID(InDesc.RequestID);
 }
 
 FAsyncPackage::~FAsyncPackage()
 {
+	AsyncLoadingThread.RemovePendingRequests(RequestIDs);
 	DetachLinker();
 }
 
+void FAsyncPackage::AddRequestID(int32 Id)
+{
+	if (Id > 0)
+	{
+		RequestIDs.Add(Id);
+		AsyncLoadingThread.AddPendingRequest(Id);
+	}
+}
+
 /**
  * @return Time load begun. This is NOT the time the load was requested in the case of other pending requests.
  */
@@ -1269,7 +1287,7 @@ void FAsyncPackage::AddImportDependency(int32 CurrentPackageIndex, const FName&
 	int32 ExistingAsyncPackageIndex = FAsyncLoadingThread::Get().FindAsyncPackage(PendingImport);
 	if (ExistingAsyncPackageIndex == INDEX_NONE)
 	{
-		const FAsyncPackageDesc Info(PendingImport);
+		const FAsyncPackageDesc Info(INDEX_NONE, PendingImport);
 		PackageToStream = new FAsyncPackage(Info);
 
 		// If priority of the dependency is not set, inherit from parent.
@@ -1853,10 +1871,10 @@ void FAsyncPackage::UpdateLoadPercentage()
 	LoadPercentage = FMath::Max(NewLoadPercentage, LoadPercentage);
 }
 
-/*-----------------------------------------------------------------------------
-	UObject async (pre)loading.
------------------------------------------------------------------------------*/
-void LoadPackageAsync(const FString& InName, const FGuid* InGuid /*= nullptr*/, const TCHAR* InPackageToLoadFrom /*= nullptr*/, FLoadPackageAsyncDelegate InCompletionDelegate /*= FLoadPackageAsyncDelegate()*/, EPackageFlags InPackageFlags /*= PKG_None*/, int32 InPIEInstanceID /*= INDEX_NONE*/, uint32 InPackagePriority /*= 0*/)
+/** Global request ID counter */
+static FThreadSafeCounter GPackageRequestID;
+
+int32 LoadPackageAsync(const FString& InName, const FGuid* InGuid /*= nullptr*/, const TCHAR* InPackageToLoadFrom /*= nullptr*/, FLoadPackageAsyncDelegate InCompletionDelegate /*= FLoadPackageAsyncDelegate()*/, EPackageFlags InPackageFlags /*= PKG_None*/, int32 InPIEInstanceID /*= INDEX_NONE*/, uint32 InPackagePriority /*= 0*/)
 {
 	// The comments clearly state that it should be a package name but we also handle it being a filename as this function is not perf critical
 	// and LoadPackage handles having a filename being passed in as well.
@@ -1888,20 +1906,27 @@ void LoadPackageAsync(const FString& InName, const FGuid* InGuid /*= nullptr*/,
 		UE_LOG(LogStreaming, Fatal, TEXT("Async loading code requires long package names (%s)."), *PackageNameToLoad);
 	}
 
-	FAsyncPackageDesc PackageDesc(*PackageName, *PackageNameToLoad, InGuid ? *InGuid : FGuid(), InCompletionDelegate, InPackageFlags, InPIEInstanceID, InPackagePriority);
+	// Generate new request ID and add it immediately to the global request list (it needs to be there before we exit
+	// this function, otherwise it would be added when the packages are being processed on the async thread).
+	const int32 RequestID = GPackageRequestID.Increment();
+	FAsyncLoadingThread::Get().AddPendingRequest(RequestID);
+	// Add new package request
+	FAsyncPackageDesc PackageDesc(RequestID, *PackageName, *PackageNameToLoad, InGuid ? *InGuid : FGuid(), InCompletionDelegate, InPackageFlags, InPIEInstanceID, InPackagePriority);
 	FAsyncLoadingThread::Get().QueuePackage(PackageDesc);
+
+	return RequestID;
 }
 
-void LoadPackageAsync(const FString& PackageName, FLoadPackageAsyncDelegate CompletionDelegate, uint32 InPackagePriority /*= 0*/, EPackageFlags InPackageFlags /*= PKG_None*/)
+int32 LoadPackageAsync(const FString& PackageName, FLoadPackageAsyncDelegate CompletionDelegate, uint32 InPackagePriority /*= 0*/, EPackageFlags InPackageFlags /*= PKG_None*/)
 {
 	const FGuid* Guid = nullptr;
 	const TCHAR* PackageToLoadFrom = nullptr;
-	LoadPackageAsync(PackageName, Guid, PackageToLoadFrom, CompletionDelegate, InPackageFlags, -1, InPackagePriority );
+	return LoadPackageAsync(PackageName, Guid, PackageToLoadFrom, CompletionDelegate, InPackageFlags, -1, InPackagePriority );
 }
 
-void LoadPackageAsync(const FString& InName, const FGuid* InGuid, FName InType /* Unused, deprecated */, const TCHAR* InPackageToLoadFrom /*= nullptr*/, FLoadPackageAsyncDelegate InCompletionDelegate /*= FLoadPackageAsyncDelegate()*/, EPackageFlags InPackageFlags /*= PKG_None*/, int32 InPIEInstanceID /*= INDEX_NONE*/, uint32 InPackagePriority /*= 0*/)
+int32 LoadPackageAsync(const FString& InName, const FGuid* InGuid, FName InType /* Unused, deprecated */, const TCHAR* InPackageToLoadFrom /*= nullptr*/, FLoadPackageAsyncDelegate InCompletionDelegate /*= FLoadPackageAsyncDelegate()*/, EPackageFlags InPackageFlags /*= PKG_None*/, int32 InPIEInstanceID /*= INDEX_NONE*/, uint32 InPackagePriority /*= 0*/)
 {
-	LoadPackageAsync(InName, InGuid, InPackageToLoadFrom, InCompletionDelegate, InPackageFlags, InPIEInstanceID, InPackagePriority);
+	return LoadPackageAsync(InName, InGuid, InPackageToLoadFrom, InCompletionDelegate, InPackageFlags, InPIEInstanceID, InPackagePriority);
 }
 
 void CancelAsyncLoading()
@@ -1928,15 +1953,21 @@ bool IsInAsyncLoadingThreadCoreUObjectInternal()
 	return FAsyncLoadingThread::Get().IsInAsyncLoadThread();
 }
 
-void FlushAsyncLoading()
+void FlushAsyncLoading(int32 PackageID /* = INDEX_NONE */)
 {
 	if (IsAsyncLoading())
 	{
+		FAsyncLoadingThread& AsyncThread = FAsyncLoadingThread::Get();
 		// Flushing async loading while loading is suspend will result in infinite stall
-		UE_CLOG(FAsyncLoadingThread::Get().IsAsyncLoadingSuspended(), LogStreaming, Fatal, TEXT("Cannot Flush Async Loading while async loading is suspended."));
+		UE_CLOG(AsyncThread.IsAsyncLoadingSuspended(), LogStreaming, Fatal, TEXT("Cannot Flush Async Loading while async loading is suspended."));
 
 		SCOPE_CYCLE_COUNTER(STAT_FAsyncPackage_FlushAsyncLoadingGameThread);
 
+		if (PackageID != INDEX_NONE && !AsyncThread.ContainsRequestID(PackageID))
+		{
+			return;
+		}
+
 		// Disallow low priority requests like texture streaming while we are flushing streaming
 		// in order to avoid excessive seeking.
 		FIOSystem::Get().SetMinPriority( AIOP_Normal );
@@ -1947,15 +1978,20 @@ void FlushAsyncLoading()
 			SCOPE_CYCLE_COUNTER(STAT_FAsyncPackage_TickAsyncLoadingGameThread);
 			while (IsAsyncLoading())
 			{
-				FAsyncLoadingThread::Get().TickAsyncLoading(false, false, 0);
-				if (FAsyncLoadingThread::Get().IsMultithreaded())
+				EAsyncPackageState::Type Result = AsyncThread.TickAsyncLoading(false, false, 0, PackageID);
+				if (PackageID != INDEX_NONE && Result == EAsyncPackageState::Complete)
+				{
+					break;
+				}
+
+				if (AsyncThread.IsMultithreaded())
 				{
 					FPlatformProcess::SleepNoStats(0.0001f);
 				}
 			}
 		}
 
-		check( !IsAsyncLoading() );
+		check(PackageID != INDEX_NONE || !IsAsyncLoading());
 
 		// Reset min priority again.
 		FIOSystem::Get().SetMinPriority( AIOP_MIN );
diff --git a/Engine/Source/Runtime/CoreUObject/Private/Serialization/AsyncLoadingThread.h b/Engine/Source/Runtime/CoreUObject/Private/Serialization/AsyncLoadingThread.h
--- a/Engine/Source/Runtime/CoreUObject/Private/Serialization/AsyncLoadingThread.h
+++ b/Engine/Source/Runtime/CoreUObject/Private/Serialization/AsyncLoadingThread.h
@@ -55,6 +55,13 @@ class FAsyncLoadingThread : public FRunnable
 	FCriticalSection AsyncPackagesCritical;
 #endif
 
+	/** List of all pending package requests */
+	TSet<int32> PendingRequests;
+#if THREADSAFE_UOBJECTS
+	/** Synchronization object for PendingRequests list */
+	FCriticalSection PendingRequestsCritical;
+#endif
+
 	/** [ASYNC/GAME THREAD] Number of package load requests in the async loading queue */
 	FThreadSafeCounter QueuedPackagesCounter;
 	/** [ASYNC/GAME THREAD] Number of packages being loaded on the async thread and post loaded on the game thread */
@@ -264,7 +271,7 @@ public:
 	* @param TimeLimit Maximum amount of time that can be spent in this call [time-slicing].
 	* @return The current state of async loading
 	*/
-	EAsyncPackageState::Type TickAsyncLoading(bool bUseTimeLimit, bool bUseFullTimeLimit, float TimeLimit);
+	EAsyncPackageState::Type TickAsyncLoading(bool bUseTimeLimit, bool bUseFullTimeLimit, float TimeLimit, int32 WaitForRequestID = INDEX_NONE);
 
 	/**
 	* [ASYNC THREAD] Main thread loop
@@ -285,6 +292,45 @@ public:
 	 */
 	float GetAsyncLoadPercentage(const FName& PackageName);
 
+	/** 
+	 * [ASYNC/GAME THREAD] Checks if a request ID already is added to the loading queue
+	 */
+	bool ContainsRequestID(int32 RequestID)
+	{
+#if THREADSAFE_UOBJECTS
+		FScopeLock Lock(&PendingRequestsCritical);
+#endif
+		return PendingRequests.Contains(RequestID);
+	}
+
+	/** 
+	 * [ASYNC/GAME THREAD] Adds a request ID to the list of pending requests
+	 */
+	void AddPendingRequest(int32 RequestID)
+	{
+#if THREADSAFE_UOBJECTS
+		FScopeLock Lock(&PendingRequestsCritical);
+#endif
+		if (!PendingRequests.Contains(RequestID))
+		{
+			PendingRequests.Add(RequestID);
+		}
+	}
+
+	/** 
+	 * [ASYNC/GAME THREAD] Removes a request ID from the list of pending requests
+	 */
+	void RemovePendingRequests(TArray<int32>& RequestIDs)
+	{
+#if THREADSAFE_UOBJECTS
+		FScopeLock Lock(&PendingRequestsCritical);
+#endif
+		for (int32 ID : RequestIDs)
+		{
+			PendingRequests.Remove(ID);
+		}		
+	}
+
 private:
 
 	/**
@@ -293,9 +339,10 @@ private:
 	* @param bUseTimeLimit True if time limit should be used [time-slicing].
 	* @param bUseFullTimeLimit True if full time limit should be used [time-slicing].
 	* @param TimeLimit Maximum amount of time that can be spent in this call [time-slicing].
+	* @param WaitForRequestID If the package request ID is valid, exits as soon as the request is no longer in the qeueue.
 	* @return The current state of async loading
 	*/
-	EAsyncPackageState::Type ProcessLoadedPackages(bool bUseTimeLimit, bool bUseFullTimeLimit, float TimeLimit);
+	EAsyncPackageState::Type ProcessLoadedPackages(bool bUseTimeLimit, bool bUseFullTimeLimit, float TimeLimit, int32 WaitForRequestID = INDEX_NONE);
 
 	/**
 	* [ASYNC THREAD] Creates async packages from the queued requests
diff --git a/Engine/Source/Runtime/CoreUObject/Public/Serialization/AsyncLoading.h b/Engine/Source/Runtime/CoreUObject/Public/Serialization/AsyncLoading.h
--- a/Engine/Source/Runtime/CoreUObject/Public/Serialization/AsyncLoading.h
+++ b/Engine/Source/Runtime/CoreUObject/Public/Serialization/AsyncLoading.h
@@ -89,6 +89,9 @@ struct FAsyncPackage
 		return bLoadHasFailed;
 	}
 
+	/** Adds new request ID to the existing package */
+	void AddRequestID(int32 Id);
+
 	/**
 	* Cancel loading this package.
 	*/
@@ -163,6 +166,8 @@ private:
 	TArray<UObject*> DeferredPostLoadObjects;
 	/** Objects to be finalized on the game thread */
 	TArray<UObject*> DeferredFinalizeObjects;
+	/** List of all request handles */
+	TArray<int32> RequestIDs;
 	/** Cached async loading thread object this package was created by */
 	class FAsyncLoadingThread& AsyncLoadingThread;
 public:
diff --git a/Engine/Source/Runtime/CoreUObject/Public/Serialization/AsyncPackage.h b/Engine/Source/Runtime/CoreUObject/Public/Serialization/AsyncPackage.h
--- a/Engine/Source/Runtime/CoreUObject/Public/Serialization/AsyncPackage.h
+++ b/Engine/Source/Runtime/CoreUObject/Public/Serialization/AsyncPackage.h
@@ -8,6 +8,8 @@
 
 struct FAsyncPackageDesc
 {
+	/** Handle for the caller */
+	int32 RequestID;
 	/** Name of the UPackage to create. */
 	FName Name;
 	/** Name of the package to load. */
@@ -26,8 +28,9 @@ struct FAsyncPackageDesc
 #endif
 
 
-	FAsyncPackageDesc(const FName& InName, FName InPackageToLoadFrom = NAME_None, const FGuid& InGuid = FGuid(), FLoadPackageAsyncDelegate InCompletionDelegate = FLoadPackageAsyncDelegate(), EPackageFlags InPackageFlags = PKG_None, int32 InPIEInstanceID = INDEX_NONE, uint32 InPriority = 0)
-		: Name(InName)
+	FAsyncPackageDesc(int32 InRequestID, const FName& InName, FName InPackageToLoadFrom = NAME_None, const FGuid& InGuid = FGuid(), FLoadPackageAsyncDelegate InCompletionDelegate = FLoadPackageAsyncDelegate(), EPackageFlags InPackageFlags = PKG_None, int32 InPIEInstanceID = INDEX_NONE, uint32 InPriority = 0)
+		: RequestID(InRequestID)
+		, Name(InName)
 		, NameToLoad(InPackageToLoadFrom)
 		, Guid(InGuid)
 		, PackageLoadedDelegate(InCompletionDelegate)
diff --git a/Engine/Source/Runtime/CoreUObject/Public/UObject/UObjectGlobals.h b/Engine/Source/Runtime/CoreUObject/Public/UObject/UObjectGlobals.h
--- a/Engine/Source/Runtime/CoreUObject/Public/UObject/UObjectGlobals.h
+++ b/Engine/Source/Runtime/CoreUObject/Public/UObject/UObjectGlobals.h
@@ -291,9 +291,10 @@ DECLARE_DELEGATE_ThreeParams(FLoadPackageAsyncDelegate, const FName& /*PackageNa
  * @param	InFlags					Package flags
  * @param	InPIEInstanceID			PIE instance ID
  * @param	InPackagePriority		Loading priority
+ * @return Unique ID associated with this load request (the same package can be associated with multiple IDs).
  */
 DEPRECATED(4.9, "LoadPackageAsync override that takes package type parameter FName InType is deprecated.")
-COREUOBJECT_API void LoadPackageAsync(const FString& InName, const FGuid* InGuid, FName InType, const TCHAR* InPackageToLoadFrom = nullptr, FLoadPackageAsyncDelegate InCompletionDelegate = FLoadPackageAsyncDelegate(), EPackageFlags InPackageFlags = PKG_None, int32 InPIEInstanceID = INDEX_NONE, uint32 InPackagePriority = 0);
+COREUOBJECT_API int32 LoadPackageAsync(const FString& InName, const FGuid* InGuid, FName InType, const TCHAR* InPackageToLoadFrom = nullptr, FLoadPackageAsyncDelegate InCompletionDelegate = FLoadPackageAsyncDelegate(), EPackageFlags InPackageFlags = PKG_None, int32 InPIEInstanceID = INDEX_NONE, uint32 InPackagePriority = 0);
 
 /**
 * Asynchronously load a package and all contained objects that match context flags. Non- blocking.
@@ -305,8 +306,9 @@ COREUOBJECT_API void LoadPackageAsync(const FString& InName, const FGuid* InGuid
 * @param	InFlags					Package flags
 * @param	InPIEInstanceID			PIE instance ID
 * @param	InPackagePriority		Loading priority
+* @return Unique ID associated with this load request (the same package can be associated with multiple IDs).
 */
-COREUOBJECT_API void LoadPackageAsync(const FString& InName, const FGuid* InGuid = nullptr, const TCHAR* InPackageToLoadFrom = nullptr, FLoadPackageAsyncDelegate InCompletionDelegate = FLoadPackageAsyncDelegate(), EPackageFlags InPackageFlags = PKG_None, int32 InPIEInstanceID = INDEX_NONE, uint32 InPackagePriority = 0);
+COREUOBJECT_API int32 LoadPackageAsync(const FString& InName, const FGuid* InGuid = nullptr, const TCHAR* InPackageToLoadFrom = nullptr, FLoadPackageAsyncDelegate InCompletionDelegate = FLoadPackageAsyncDelegate(), EPackageFlags InPackageFlags = PKG_None, int32 InPIEInstanceID = INDEX_NONE, uint32 InPackagePriority = 0);
 
 /**
 * Asynchronously load a package and all contained objects that match context flags. Non- blocking.
@@ -314,8 +316,9 @@ COREUOBJECT_API void LoadPackageAsync(const FString& InName, const FGuid* InGuid
 * @param	InName					Name of package to load
 * @param	InCompletionDelegate	Delegate to be invoked when the packages has finished streaming
 * @param	InPackagePriority		Loading priority
+* @return Unique ID associated with this load request (the same package can be associated with multiple IDs).
 */
-COREUOBJECT_API void LoadPackageAsync(const FString& InName, FLoadPackageAsyncDelegate InCompletionDelegate, uint32 InPackagePriority = 0, EPackageFlags InPackageFlags = PKG_None);
+COREUOBJECT_API int32 LoadPackageAsync(const FString& InName, FLoadPackageAsyncDelegate InCompletionDelegate, uint32 InPackagePriority = 0, EPackageFlags InPackageFlags = PKG_None);
 
 /**
 * Cancels all async package loading requests.
@@ -409,8 +412,11 @@ COREUOBJECT_API bool IsReferenced( UObject*& Res, EObjectFlags KeepFlags, bool b
 
 /**
  * Blocks till all pending package/ linker requests are fulfilled.
+ *
+ * @param PackageID if the package associated with this request ID gets loaded, FlushAsyncLoading returns 
+ *        immediately without waiting for the remaining packages to finish loading.
  */
-COREUOBJECT_API void FlushAsyncLoading();
+COREUOBJECT_API void FlushAsyncLoading(int32 PackageID = INDEX_NONE);
 
 /**
 * [Deprecated] Blocks till all pending package/ linker requests are fulfilled.
