UE-17541 Copying struct value using Set node from null source crashes

diff --git a/Engine/Source/Runtime/CoreUObject/Private/UObject/Class.cpp b/Engine/Source/Runtime/CoreUObject/Private/UObject/Class.cpp
--- a/Engine/Source/Runtime/CoreUObject/Private/UObject/Class.cpp
+++ b/Engine/Source/Runtime/CoreUObject/Private/UObject/Class.cpp
@@ -2331,7 +2331,14 @@ void UScriptStruct::ClearScriptStruct(void* Dest, int32 ArrayDim) const
 			{
 				TheCppStructOps->Destruct(PropertyData);
 			}
-			TheCppStructOps->Construct(PropertyData);
+			if (TheCppStructOps->HasZeroConstructor())
+			{
+				FMemory::Memzero(PropertyData, Stride);
+			}
+			else
+			{
+				TheCppStructOps->Construct(PropertyData);
+			}
 		}
 		ClearedSize = TheCppStructOps->GetSize();
 		// here we want to make sure C++ and the property system agree on the size
