Added code to prevent multiple identical default subobjects from being constructed.

diff --git a/Engine/Source/Runtime/CoreUObject/Private/UObject/UObjectGlobals.cpp b/Engine/Source/Runtime/CoreUObject/Private/UObject/UObjectGlobals.cpp
--- a/Engine/Source/Runtime/CoreUObject/Private/UObject/UObjectGlobals.cpp
+++ b/Engine/Source/Runtime/CoreUObject/Private/UObject/UObjectGlobals.cpp
@@ -2963,6 +2963,10 @@ UObject* FObjectInitializer::CreateDefaultSubobject(UObject* Outer, FName Subobj
 			const EObjectFlags SubobjectFlags = Outer->GetMaskedFlags(RF_PropagateToSubObjects);
 			const bool bOwnerArchetypeIsNotNative = !Outer->GetArchetype()->GetClass()->HasAnyClassFlags(CLASS_Native | CLASS_Intrinsic);
 			const bool bOwnerTemplateIsNotCDO = ObjectArchetype != nullptr && ObjectArchetype != Outer->GetClass()->GetDefaultObject(false) && !Outer->HasAnyFlags(RF_ClassDefaultObject);
+			if (StaticFindObjectFast(OverrideClass, Outer, SubobjectFName) != nullptr)
+			{
+				UE_LOG(LogClass, Fatal, TEXT("Default subobject %s %s already exists for %s."), *OverrideClass->GetName(), *SubobjectFName.ToString(), *Outer->GetFullName());
+			}
 			Result = StaticConstructObject_Internal(OverrideClass, Outer, SubobjectFName, SubobjectFlags);
 			if (!bIsTransient && (bOwnerArchetypeIsNotNative || bOwnerTemplateIsNotCDO))
 			{
