Core: Added message log error for source name collisions in redirector ini sections

diff --git a/Engine/Source/Runtime/CoreUObject/Private/UObject/LinkerLoad.cpp b/Engine/Source/Runtime/CoreUObject/Private/UObject/LinkerLoad.cpp
--- a/Engine/Source/Runtime/CoreUObject/Private/UObject/LinkerLoad.cpp
+++ b/Engine/Source/Runtime/CoreUObject/Private/UObject/LinkerLoad.cpp
@@ -85,9 +85,12 @@ void FLinkerLoad::CreateActiveRedirectsMap(const FString& GEngineIniName)
 		FConfigSection* PackageRedirects = GConfig->GetSectionPrivate( TEXT("/Script/Engine.Engine"), false, true, GEngineIniName );
 		if (PackageRedirects)
 		{
+			FDeferredMessageLog RedirectErrors(NAME_LoadErrors);
+
+			static FName ActiveClassRedirectsKey(TEXT("ActiveClassRedirects"));
 			for( FConfigSection::TIterator It(*PackageRedirects); It; ++It )
 			{
-				if( It.Key() == TEXT("ActiveClassRedirects") )
+				if (It.Key() == ActiveClassRedirectsKey)
 				{
 					FName OldClassName = NAME_None;
 					FName NewClassName = NAME_None;
@@ -109,16 +112,32 @@ void FLinkerLoad::CreateActiveRedirectsMap(const FString& GEngineIniName)
 					if (NewSubobjName != NAME_None || OldSubobjName != NAME_None)
 					{
 						check(OldSubobjName != NAME_None && OldClassName != NAME_None );
+
+						if (SubobjectNameRedirects.Contains(OldSubobjName))
+						{
+							RedirectErrors.Error(FText::Format(LOCTEXT("SubobjRedirectorCollision", "{0} contains a collision with multiple redirectors for old suboject named {1} (old class {2})"), FText::FromName(ActiveClassRedirectsKey), FText::FromName(OldSubobjName), FText::FromName(OldClassName)));
+						}
+
 						SubobjectNameRedirects.Add(OldSubobjName, FSubobjectRedirect(OldClassName, NewSubobjName));
 					}
 					//instances only
 					else if( bInstanceOnly )
 					{
+						if (ObjectNameRedirectsInstanceOnly.Contains(OldClassName))
+						{
+							RedirectErrors.Error(FText::Format(LOCTEXT("InstanceOnlyRedirectorCollision", "{0} contains a collision with multiple instance-only redirectors for old class named {1}"), FText::FromName(ActiveClassRedirectsKey), FText::FromName(OldClassName)));
+						}
+
 						ObjectNameRedirectsInstanceOnly.Add(OldClassName,NewClassName);
 					}
 					//objects only on a per-object basis
 					else if( ObjectName != NAME_None )
 					{
+						if (ObjectNameRedirectsObjectOnly.Contains(OldClassName))
+						{
+							RedirectErrors.Error(FText::Format(LOCTEXT("ObjectOnlyRedirectorCollision", "{0} contains a collision with multiple object-only redirectors for old class named {1}"), FText::FromName(ActiveClassRedirectsKey), FText::FromName(OldClassName)));
+						}
+
 						ObjectNameRedirectsObjectOnly.Add(ObjectName, NewClassName);
 					}
 					//full redirect
@@ -126,10 +145,15 @@ void FLinkerLoad::CreateActiveRedirectsMap(const FString& GEngineIniName)
 					{
 						if (NewClassName.ToString().Find(TEXT("."), ESearchCase::CaseSensitive) != NewClassName.ToString().Find(TEXT("."), ESearchCase::CaseSensitive, ESearchDir::FromEnd))
 						{
-							UE_LOG(LogLinker, Error, TEXT("Currently we cannot rename nested objects for '%s'; if you want to leave the outer alone, just specify the name with no path"), *NewClassName.ToString());
+							RedirectErrors.Error(FText::Format(LOCTEXT("NestedRenameDisallowed", "{0} cannot contain a rename of nested objects for '{1}'; if you want to leave the outer alone, just specify the name with no path"), FText::FromName(ActiveClassRedirectsKey), FText::FromName(NewClassName)));
 						}
 						else
 						{
+							if (ObjectNameRedirects.Contains(OldClassName))
+							{
+								RedirectErrors.Error(FText::Format(LOCTEXT("RedirectorCollision", "{0} contains a collision with multiple redirectors for old class named {1}"), FText::FromName(ActiveClassRedirectsKey), FText::FromName(OldClassName)));
+							}
+
 							ObjectNameRedirects.Add(OldClassName,NewClassName);
 						}
 					}
@@ -1256,7 +1280,7 @@ FLinkerLoad::ELinkerStatus FLinkerLoad::FixupImportMap()
 				{
 					FObjectImport& Import = ImportMap[i];
 					{
-						FSubobjectRedirect *Redirect = SubobjectNameRedirects.Find(Import.ObjectName);
+						FSubobjectRedirect* Redirect = SubobjectNameRedirects.Find(Import.ObjectName);
 						if (Redirect)
 						{
 							if (Import.ClassName == Redirect->MatchClass)
@@ -3914,6 +3938,12 @@ UObject* FLinkerLoad::CreateImport( int32 Index )
 				FLinkerManager::Get().AddLoaderWithNewImports(this);
 			}
 		}
+
+		if (Import.XObject == nullptr)
+		{
+			const FString OuterName = Import.OuterIndex.IsNull() ? LinkerRoot->GetFullName() : GetFullImpExpName(Import.OuterIndex);
+			UE_LOG(LogLinker, Verbose, TEXT("Failed to resolve import named %s in %s"), *Import.ObjectName.ToString(), *OuterName);
+		}
 	}
 	return Import.XObject;
 }
