Translate tagged FStringAssetReference properties into TAssetPtrs during deserialization.

diff --git a/Engine/Source/Runtime/CoreUObject/Private/UObject/Class.cpp b/Engine/Source/Runtime/CoreUObject/Private/UObject/Class.cpp
--- a/Engine/Source/Runtime/CoreUObject/Private/UObject/Class.cpp
+++ b/Engine/Source/Runtime/CoreUObject/Private/UObject/Class.cpp
@@ -826,6 +826,7 @@ void UStruct::SerializeTaggedProperties(FArchive& Ar, uint8* Data, UStruct* Defa
 		{
 			FPropertyTag Tag;
 			Ar << Tag;
+
 			if( Tag.Name == NAME_None )
 			{
 				break;
@@ -1129,6 +1130,21 @@ void UStruct::SerializeTaggedProperties(FArchive& Ar, uint8* Data, UStruct* Defa
 					UE_LOG(LogClass, Warning, TEXT("SerializeFromMismatchedTag failed: Type mismatch in %s of %s - Previous (%s) Current(%s) for package:  %s"), *Tag.Name.ToString(), *GetName(), *Tag.Type.ToString(), *Property->GetID().ToString(), *Ar.GetArchiveName() );
 				}
 			}
+			else if (Tag.Type == NAME_StructProperty && Property->GetID() == NAME_AssetObjectProperty)
+			{
+				// This property used to be a FStringAssetReference but is now a TAssetPtr<Foo>
+				FStringAssetReference PreviousValue;
+				// explicitly call Serialize to ensure that the various delegates needed for cooking are fired
+				PreviousValue.Serialize(Ar);
+
+				// now copy the value into the object's address space
+				FAssetPtr PreviousValueAssetPtr;
+				PreviousValueAssetPtr = PreviousValue;
+				CastChecked<UAssetObjectProperty>(Property)->SetPropertyValue_InContainer(Data, PreviousValueAssetPtr, Tag.ArrayIndex);
+
+				AdvanceProperty = true;
+				continue;
+			}
 			else if( Tag.Type!=Property->GetID() )
 			{
 				UE_LOG(LogClass, Warning, TEXT("Type mismatch in %s of %s - Previous (%s) Current(%s) for package:  %s"), *Tag.Name.ToString(), *GetName(), *Tag.Type.ToString(), *Property->GetID().ToString(), *Ar.GetArchiveName() );
