Corrected config section used to save config values for PerObjectConfig objects that are contained in the transient package.

This brings the code in line with what is used in UObject::LoadConfig

diff --git a/Engine/Source/Runtime/CoreUObject/Private/UObject/Obj.cpp b/Engine/Source/Runtime/CoreUObject/Private/UObject/Obj.cpp
--- a/Engine/Source/Runtime/CoreUObject/Private/UObject/Obj.cpp
+++ b/Engine/Source/Runtime/CoreUObject/Private/UObject/Obj.cpp
@@ -1722,7 +1722,15 @@ void UObject::SaveConfig( uint64 Flags, const TCHAR* InFilename, FConfigCacheIni
 	if ( bPerObject == true )
 	{
 		FString PathNameString;
-		GetPathName(GetOutermost(), PathNameString);
+		UObject* Outermost = GetOutermost();
+		if ( Outermost == GetTransientPackage() )
+		{
+			PathNameString = GetName();
+		}
+		else
+		{
+			GetPathName(Outermost, PathNameString);
+		}
 		Section = PathNameString + TEXT(" ") + GetClass()->GetName();
 	}
 
