Clarify error message on failing to route PostInitProperites.

diff --git a/Engine/Source/Runtime/CoreUObject/Private/UObject/UObjectGlobals.cpp b/Engine/Source/Runtime/CoreUObject/Private/UObject/UObjectGlobals.cpp
--- a/Engine/Source/Runtime/CoreUObject/Private/UObject/UObjectGlobals.cpp
+++ b/Engine/Source/Runtime/CoreUObject/Private/UObject/UObjectGlobals.cpp
@@ -2305,7 +2305,7 @@ void FObjectInitializer::PostConstructInit()
 #if !(UE_BUILD_SHIPPING || UE_BUILD_TEST)
 	if (!FUObjectThreadContext::Get().PostInitPropertiesCheck.Num() || (FUObjectThreadContext::Get().PostInitPropertiesCheck.Pop() != Obj))
 	{
-		UE_LOG(LogUObjectGlobals, Fatal, TEXT("%s failed to route PostInitProperties"), *Obj->GetClass()->GetName() );
+		UE_LOG(LogUObjectGlobals, Fatal, TEXT("%s failed to route PostInitProperties. Call Super::PostInitProperties() in %s::PostInitProperties()."), *Obj->GetClass()->GetName(), *Obj->GetClass()->GetName());
 	}
 	// Check if all TSubobjectPtr properties have been initialized.
 	if (!Obj->HasAnyFlags(RF_NeedLoad))
