UEnum class exported with COREUOBJECT_API.
Bunch of UEnum member functions de-inlined.

diff --git a/Engine/Source/Runtime/CoreUObject/Private/UObject/Enum.cpp b/Engine/Source/Runtime/CoreUObject/Private/UObject/Enum.cpp
--- a/Engine/Source/Runtime/CoreUObject/Private/UObject/Enum.cpp
+++ b/Engine/Source/Runtime/CoreUObject/Private/UObject/Enum.cpp
@@ -104,6 +104,91 @@ FString UEnum::GenerateFullEnumName(const TCHAR* InEnumName) const
 	return (CppForm != ECppForm::Regular) ? GenerateFullEnumName(this, InEnumName) : InEnumName;
 }
 
+FName UEnum::GetNameByIndex(int8 Index) const
+{
+	if (Names.IsValidIndex(Index))
+	{
+		return Names[Index].Key;
+	}
+
+	return NAME_None;
+}
+
+int8 UEnum::GetValueByIndex(int8 Index) const
+{
+	check(Names.IsValidIndex(Index));
+	return Names[Index].Value;
+}
+
+FName UEnum::GetNameByValue(int8 InValue) const
+{
+	for (TPair<FName, int8> Kvp : Names)
+	{
+		if (Kvp.Value == InValue)
+		{
+			return Kvp.Key;
+		}
+	}
+
+	return NAME_None;
+}
+
+int8 UEnum::GetValueByName(FName InName)
+{
+	for (TPair<FName, int8> Kvp : Names)
+	{
+		if (Kvp.Key == InName)
+		{
+			return Kvp.Value;
+		}
+	}
+
+	return INDEX_NONE;
+}
+
+int32 UEnum::GetIndexByName(FName Name) const
+{
+	int32 Count = Names.Num();
+	for (int32 Counter = 0; Counter < Count; ++Counter)
+	{
+		if (Names[Counter].Key == Name)
+		{
+			return Counter;
+		}
+	}
+
+	return INDEX_NONE;
+}
+
+int8 UEnum::GetMaxEnumValue() const
+{
+	int8 MaxValue = -1;
+	if (Names.Num() > 0)
+	{
+		MaxValue = Names[0].Value;
+	}
+
+	for (const auto& Pair : Names)
+	{
+		if (Pair.Value > MaxValue)
+		{
+			MaxValue = Pair.Value;
+		}
+	}
+
+	return MaxValue;
+}
+
+FString UEnum::GenerateFullEnumName(const UEnum* InEnum, const TCHAR* InEnumName)
+{
+	if (InEnum->GetCppForm() == ECppForm::Regular || IsFullEnumName(InEnumName))
+	{
+		return InEnumName;
+	}
+
+	return FString::Printf(TEXT("%s::%s"), *InEnum->GetName(), InEnumName);
+}
+
 void UEnum::AddNamesToMasterList()
 {
 	for (TPair<FName, int8> Kvp : Names)
@@ -189,6 +274,40 @@ FString UEnum::GenerateEnumPrefix() const
 	return Prefix;
 }
 
+FString UEnum::GetEnumName(int32 InIndex) const
+{
+	if (Names.IsValidIndex(InIndex))
+	{
+		if (CppForm == ECppForm::Regular)
+		{
+			return GetNameByIndex(InIndex).ToString();
+		}
+
+		// Strip the namespace from the name.
+		FString EnumName(GetNameByIndex(InIndex).ToString());
+		int32 ScopeIndex = EnumName.Find(TEXT("::"), ESearchCase::CaseSensitive);
+		if (ScopeIndex != INDEX_NONE)
+		{
+			return EnumName.Mid(ScopeIndex + 2);
+		}
+	}
+	return FName(NAME_None).ToString();
+}
+
+FText UEnum::GetEnumText(int32 InIndex) const
+{
+#if WITH_EDITOR
+	//@todo These values should be properly localized [9/24/2013 justin.sargent]
+	FText LocalizedDisplayName = GetDisplayNameText(InIndex);
+	if(!LocalizedDisplayName.IsEmpty())
+	{
+		return LocalizedDisplayName;
+	}
+#endif
+
+	return FText::FromString( GetEnumName(InIndex) );
+}
+
 int32 UEnum::FindEnumIndex(FName InName) const
 {
 	int32 EnumIndex = GetIndexByName(InName);
@@ -518,6 +637,39 @@ void UEnum::RemoveMetaData( const TCHAR* Key, int32 NameIndex/*=INDEX_NONE*/) co
 
 #endif
 
+int32 UEnum::LookupEnumName(FName TestName, UEnum** FoundEnum)
+{
+	UEnum* TheEnum = AllEnumNames.FindRef(TestName);
+	if (FoundEnum != nullptr)
+	{
+		*FoundEnum = TheEnum;
+	}
+	return (TheEnum != nullptr) ? TheEnum->GetIndexByName(TestName) : INDEX_NONE;
+}
+
+int32 UEnum::LookupEnumNameSlow(const TCHAR* InTestShortName, UEnum** FoundEnum)
+{
+	int32 EnumIndex = LookupEnumName(InTestShortName, FoundEnum);
+	if (EnumIndex == INDEX_NONE)
+	{
+		FString TestShortName = FString(TEXT("::")) + InTestShortName;
+		UEnum* TheEnum = NULL;
+		for (TMap<FName, UEnum*>::TIterator It(AllEnumNames); It; ++It)
+		{
+			if (It.Key().ToString().Contains(TestShortName) )
+			{
+				TheEnum = It.Value();
+			}
+		}
+		if (FoundEnum != NULL)
+		{
+			*FoundEnum = TheEnum;
+		}
+		EnumIndex = (TheEnum != NULL) ? TheEnum->FindEnumIndex(InTestShortName) : INDEX_NONE;
+	}
+	return EnumIndex;
+}
+
 int32 UEnum::ParseEnum(const TCHAR*& Str)
 {
 	FString Token;
diff --git a/Engine/Source/Runtime/CoreUObject/Public/UObject/Class.h b/Engine/Source/Runtime/CoreUObject/Public/UObject/Class.h
--- a/Engine/Source/Runtime/CoreUObject/Public/UObject/Class.h
+++ b/Engine/Source/Runtime/CoreUObject/Public/UObject/Class.h
@@ -1335,9 +1335,9 @@ public:
 //
 // Reflection data for an enumeration.
 //
-class UEnum : public UField
+class COREUOBJECT_API UEnum : public UField
 {
-	DECLARE_CASTED_CLASS_INTRINSIC_WITH_API(UEnum,UField,0,CoreUObject,CASTCLASS_UEnum,COREUOBJECT_API)
+	DECLARE_CASTED_CLASS_INTRINSIC_WITH_API(UEnum,UField,0,CoreUObject,CASTCLASS_UEnum,NO_API)
 
 public:
 	enum class ECppForm
@@ -1351,84 +1351,22 @@ public:
 	FString CppType;
 
 	/** Gets enum name by index in Names array. Returns NAME_None if Index is not valid. */
-	FName GetNameByIndex(int8 Index) const
-	{
-		if (Names.IsValidIndex(Index))
-		{
-			return Names[Index].Key;
-		}
-		
-		return NAME_None;
-	}
+	FName GetNameByIndex(int8 Index) const;
 
 	/** Gets enum value by index in Names array. */
-	int8 GetValueByIndex(int8 Index) const
-	{
-		check(Names.IsValidIndex(Index));
-		return Names[Index].Value;
-	}
+	int8 GetValueByIndex(int8 Index) const;
 
 	/** Gets enum name by value. Returns NAME_None if value is not found. */
-	FName GetNameByValue(int8 InValue) const
-	{
-		for (TPair<FName, int8> Kvp : Names)
-		{
-			if (Kvp.Value == InValue)
-			{
-				return Kvp.Key;
-			}
-		}
-		
-		return NAME_None;
-	}
+	FName GetNameByValue(int8 InValue) const;
 
 	/** Gets enum value by name. Returns INDEX_NONE when name is not found. */
-	int8 GetValueByName(FName InName)
-	{
-		for (TPair<FName, int8> Kvp : Names)
-		{
-			if (Kvp.Key == InName)
-			{
-				return Kvp.Value;
-			}
-		}
-
-		return INDEX_NONE;
-	}
+	int8 GetValueByName(FName InName);
 
 	/** Gets index of name in enum. Returns INDEX_NONE when name is not found. */
-	int32 GetIndexByName(FName Name) const
-	{
-		int32 Count = Names.Num();
-		for (int32 Counter = 0; Counter < Count; ++Counter)
-		{
-			if (Names[Counter].Key == Name)
-			{
-				return Counter;
-			}
-		}
-		return INDEX_NONE;
-	}
+	int32 GetIndexByName(FName Name) const;
 
 	/** Gets max value of Enum. Defaults to zero if there are no entries. */
-	int8 GetMaxEnumValue() const
-	{
-		int8 MaxValue = -1;
-		if (Names.Num() > 0)
-		{
-			MaxValue = Names[0].Value;
-		}
-
-		for (const auto& Pair : Names)
-		{
-			if (Pair.Value > MaxValue)
-			{
-				MaxValue = Pair.Value;
-			}
-		}
-
-		return MaxValue;
-	}
+	int8 GetMaxEnumValue() const;
 
 protected:
 	// Variables.
@@ -1439,28 +1377,28 @@ protected:
 	ECppForm CppForm;
 
 	/** global list of all value names used by all enums in memory, used for property text import */
-	COREUOBJECT_API static TMap<FName, UEnum*> AllEnumNames;
+	static TMap<FName, UEnum*> AllEnumNames;
 
 protected: 
 	
 	/** adds the Names in this enum to the master AllEnumNames list */
-	COREUOBJECT_API void AddNamesToMasterList();
+	void AddNamesToMasterList();
 
 public:
 
 	/** removes the Names in this enum from the master AllEnumNames list */
-	COREUOBJECT_API void RemoveNamesFromMasterList();
+	void RemoveNamesFromMasterList();
 
 	// UObject interface.
-	COREUOBJECT_API virtual void Serialize(FArchive& Ar) override;
+	virtual void Serialize(FArchive& Ar) override;
 	// End of UObject interface.
 
-	COREUOBJECT_API ~UEnum();
+	~UEnum();
 
 	/*
 	 *	Try to update an out-of-date enum index after an enum's change
 	 */
-	COREUOBJECT_API virtual int32 ResolveEnumerator(FArchive& Ar, int32 EnumeratorIndex) const;
+	virtual int32 ResolveEnumerator(FArchive& Ar, int32 EnumeratorIndex) const;
 
 	/**
 	 * Returns the type of enum: whether it's a regular enum, namespaced enum or C++11 enum class.
@@ -1490,15 +1428,7 @@ public:
 	 * @param InEnumName Enum name.
 	 * @return Full enum name.
 	 */
-	COREUOBJECT_API static FString GenerateFullEnumName(const UEnum* InEnum, const TCHAR* InEnumName)
-	{
-		if (InEnum->GetCppForm() == ECppForm::Regular || IsFullEnumName(InEnumName))
-		{
-			return InEnumName;
-		}
-
-		return FString::Printf(TEXT("%s::%s"), *InEnum->GetName(), InEnumName);
-	}
+	static FString GenerateFullEnumName(const UEnum* InEnum, const TCHAR* InEnumName);
 
 	/**
 	 * Generates full enum name give enum name.
@@ -1506,46 +1436,17 @@ public:
 	 * @param InEnumName Enum name.
 	 * @return Full enum name.
 	 */
-	COREUOBJECT_API virtual FString GenerateFullEnumName(const TCHAR* InEnumName) const;
+	virtual FString GenerateFullEnumName(const TCHAR* InEnumName) const;
 
 	/** searches the list of all enum value names for the specified name
 	 * @return the value the specified name represents if found, otherwise INDEX_NONE
 	 */
-	static int32 LookupEnumName(FName TestName, UEnum** FoundEnum = nullptr)
-	{
-		UEnum* TheEnum = AllEnumNames.FindRef(TestName);
-		if (FoundEnum != nullptr)
-		{
-			*FoundEnum = TheEnum;
-		}
-		return (TheEnum != nullptr) ? TheEnum->GetIndexByName(TestName) : INDEX_NONE;
-	}
+	static int32 LookupEnumName(FName TestName, UEnum** FoundEnum = nullptr);
 
 	/** searches the list of all enum value names for the specified name
 	 * @return the value the specified name represents if found, otherwise INDEX_NONE
 	 */
-	static int32 LookupEnumNameSlow(const TCHAR* InTestShortName, UEnum** FoundEnum = NULL)
-	{
-		int32 EnumIndex = LookupEnumName(InTestShortName, FoundEnum);
-		if (EnumIndex == INDEX_NONE)
-		{
-			FString TestShortName = FString(TEXT("::")) + InTestShortName;
-			UEnum* TheEnum = NULL;
-			for (TMap<FName, UEnum*>::TIterator It(AllEnumNames); It; ++It)
-			{
-				if (It.Key().ToString().Contains(TestShortName) )
-				{
-					TheEnum = It.Value();
-				}
-			}
-			if (FoundEnum != NULL)
-			{
-				*FoundEnum = TheEnum;
-			}
-			EnumIndex = (TheEnum != NULL) ? TheEnum->FindEnumIndex(InTestShortName) : INDEX_NONE;
-		}
-		return EnumIndex;
-	}
+	static int32 LookupEnumNameSlow(const TCHAR* InTestShortName, UEnum** FoundEnum = nullptr);
 
 	/** parses the passed in string for a name, then searches for that name in any Enum (in any package)
 	 * @param Str	pointer to string to parse; if we successfully find an enum, this pointer is advanced past the name found
@@ -1560,7 +1461,7 @@ public:
 	 * @param InCppForm The form of enum.
 	 * @return	true unless the MAX enum already exists and isn't the last enum.
 	 */
-	COREUOBJECT_API virtual bool SetEnums(TArray<TPair<FName, int8>>& InNames, ECppForm InCppForm);
+	virtual bool SetEnums(TArray<TPair<FName, int8>>& InNames, ECppForm InCppForm);
 
 	/**
 	 * @return	The enum name at the specified Index.
@@ -1577,46 +1478,17 @@ public:
 	/**
 	 * @return	The short enum name at the specified Index.
 	 */
-	FString GetEnumName(int32 InIndex) const
-	{
-		if (Names.IsValidIndex(InIndex))
-		{
-			if (CppForm == ECppForm::Regular)
-			{
-				return GetNameByIndex(InIndex).ToString();
-			}
-
-			// Strip the namespace from the name.
-			FString EnumName(GetNameByIndex(InIndex).ToString());
-			int32 ScopeIndex = EnumName.Find(TEXT("::"), ESearchCase::CaseSensitive);
-			if (ScopeIndex != INDEX_NONE)
-			{
-				return EnumName.Mid(ScopeIndex + 2);
-			}
-		}
-		return FName(NAME_None).ToString();
-	}
+	FString GetEnumName(int32 InIndex) const;
 
 	/**
 	 * @return	The enum string at the specified index.
 	 */
-	COREUOBJECT_API virtual FText GetEnumText(int32 InIndex) const
-	{
-#if WITH_EDITOR
-		//@todo These values should be properly localized [9/24/2013 justin.sargent]
-		FText LocalizedDisplayName = GetDisplayNameText(InIndex);
-		if(!LocalizedDisplayName.IsEmpty())
-		{
-			return LocalizedDisplayName;
-		}
-#endif
+	virtual FText GetEnumText(int32 InIndex) const;
 
-		return FText::FromString( GetEnumName(InIndex) );
-	}
 	/**
 	 * @return	The index of the specified name, if it exists in the enum names list.
 	 */
-	COREUOBJECT_API int32 FindEnumIndex(FName InName) const;
+	int32 FindEnumIndex(FName InName) const;
 
 	/**
 	 * @return	 The number of enum names.
@@ -1632,7 +1504,7 @@ public:
 	 * @return	the longest common prefix between all items in the enum.  If a common prefix
 	 *			cannot be found, returns the full name of the enum.
 	 */
-	COREUOBJECT_API FString GenerateEnumPrefix() const;
+	FString GenerateEnumPrefix() const;
 
 #if WITH_EDITOR
 	/**
@@ -1642,7 +1514,7 @@ public:
 	 *
 	 * @return The display name for this object.
 	 */
-	COREUOBJECT_API FText GetDisplayNameText(int32 NameIndex=INDEX_NONE) const;
+	FText GetDisplayNameText(int32 NameIndex=INDEX_NONE) const;
 
 	/**
 	 * Finds the localized tooltip or native tooltip as a fallback.
@@ -1651,7 +1523,7 @@ public:
 	 *
 	 * @return The tooltip for this object.
 	 */
-	COREUOBJECT_API FText GetToolTipText(int32 NameIndex=INDEX_NONE) const;
+	FText GetToolTipText(int32 NameIndex=INDEX_NONE) const;
 
 	/**
 	 * Wrapper method for easily determining whether this enum has metadata associated with it.
@@ -1661,7 +1533,7 @@ public:
 	 *
 	 * @return true if the specified key exists in the list of metadata for this enum, even if the value of that key is empty
 	 */
-	COREUOBJECT_API bool HasMetaData( const TCHAR* Key, int32 NameIndex=INDEX_NONE ) const;
+	bool HasMetaData( const TCHAR* Key, int32 NameIndex=INDEX_NONE ) const;
 
 	/**
 	 * Return the metadata value associated with the specified key.
@@ -1671,7 +1543,7 @@ public:
 	 *
 	 * @return	the value for the key specified, or an empty string if the key wasn't found or had no value.
 	 */
-	COREUOBJECT_API const FString& GetMetaData( const TCHAR* Key, int32 NameIndex=INDEX_NONE ) const;
+	const FString& GetMetaData( const TCHAR* Key, int32 NameIndex=INDEX_NONE ) const;
 
 	/**
 	 * Set the metadata value associated with the specified key.
@@ -1681,7 +1553,7 @@ public:
 	 * @param	InValue		Value of the metadata for the key
 	 *
 	 */
-	COREUOBJECT_API void SetMetaData( const TCHAR* Key, const TCHAR* InValue, int32 NameIndex=INDEX_NONE) const;
+	void SetMetaData( const TCHAR* Key, const TCHAR* InValue, int32 NameIndex=INDEX_NONE) const;
 	
 	/**
 	 * Remove given key meta data
@@ -1690,7 +1562,7 @@ public:
 	 * @param	NameIndex	if specified, will search the metadata linked for that enum value; otherwise, searches the metadata for the enum itself
 	 *
 	 */
-	COREUOBJECT_API void RemoveMetaData( const TCHAR* Key, int32 NameIndex=INDEX_NONE ) const;
+	void RemoveMetaData( const TCHAR* Key, int32 NameIndex=INDEX_NONE ) const;
 #endif
 	/**
 	 * Find the enum and entry value from EnumRedirects
@@ -1699,7 +1571,7 @@ public:
 	 * @param	EnumEntryName	Name of the entry of the enum
 	 *
 	 */
-	static COREUOBJECT_API int32 FindEnumRedirects(const UEnum* Enum, FName EnumEntryName);
+	static int32 FindEnumRedirects(const UEnum* Enum, FName EnumEntryName);
 
 
 	/**
