Dependency preloading stuff
 - Generate a new cooker file open log for unrealpak to use if it can't find a game file order
 - Add dependency data into the "mini" asset registry
 - Load mini asset registry at runtime and use it to load package dependencies BEFORE opening the requested package. Hopefully, all in the same order as the cooker file log

diff --git a/Engine/Source/Editor/UnrealEd/Private/Commandlets/ChunkManifestGenerator.cpp b/Engine/Source/Editor/UnrealEd/Private/Commandlets/ChunkManifestGenerator.cpp
--- a/Engine/Source/Editor/UnrealEd/Private/Commandlets/ChunkManifestGenerator.cpp
+++ b/Engine/Source/Editor/UnrealEd/Private/Commandlets/ChunkManifestGenerator.cpp
@@ -268,9 +268,10 @@ bool FChunkManifestGenerator::LoadAssetRegistry(const FString& SandboxPath, cons
 		FArchive* AssetRegistryReader = &FileContents;
 
 		TMap<FName, FAssetData*> SavedAssetRegistryData;
+		TArray<FDependsNode*> DependencyData;
 		if (AssetRegistryReader)
 		{
-			AssetRegistry.LoadRegistryData(*AssetRegistryReader, SavedAssetRegistryData);
+			AssetRegistry.LoadRegistryData(*AssetRegistryReader, SavedAssetRegistryData, DependencyData);
 		}
 		for (auto& LoadedAssetData : AssetRegistryData)
 		{
@@ -468,6 +469,30 @@ void FChunkManifestGenerator::BuildChunkManifest(const TArray<FName>& CookedPack
 
 }
 
+void FChunkManifestGenerator::AddAssetToFileOrderRecursive(FAssetData* InAsset, TArray<FName>& OutFileOrder, TArray<FName>& OutEncounteredNames, const TMap<FName, FAssetData*>& InAssets, const TArray<FName>& InMapList)
+{
+	if (!OutEncounteredNames.Contains(InAsset->PackageName))
+	{
+		OutEncounteredNames.Add(InAsset->PackageName);
+
+		TArray<FName> Dependencies;
+		AssetRegistry.GetDependencies(InAsset->PackageName, Dependencies);
+
+		for (auto DependencyName : Dependencies)
+		{
+			if (InAssets.Contains(DependencyName) && !OutFileOrder.Contains(DependencyName))
+			{
+				if (!InMapList.Contains(DependencyName))
+				{
+					auto Dependency = InAssets[DependencyName];
+					AddAssetToFileOrderRecursive(Dependency, OutFileOrder, OutEncounteredNames, InAssets, InMapList);
+				}
+			}
+		}
+
+		OutFileOrder.Add(InAsset->PackageName);
+	}
+}
 
 bool FChunkManifestGenerator::SaveAssetRegistry(const FString& SandboxPath, const TArray<FName>* IgnorePackageList)
 {
@@ -482,9 +507,10 @@ bool FChunkManifestGenerator::SaveAssetRegistry(const FString& SandboxPath, cons
 			IgnorePackageSet.Add(IgnorePackage);
 		}
 	}
-	
+
 
 	// Create asset registry data
+	TArray<FName> MapList;
 	FArrayWriter SerializedAssetRegistry;
 	SerializedAssetRegistry.SetFilterEditorOnly(true);
 	TMap<FName, FAssetData*> GeneratedAssetRegistryData;
@@ -498,17 +524,31 @@ bool FChunkManifestGenerator::SaveAssetRegistry(const FString& SandboxPath, cons
 		// Add only assets that have actually been cooked and belong to any chunk
 		if (AssetData.ChunkIDs.Num() > 0)
 		{
-			GeneratedAssetRegistryData.Add(AssetData.ObjectPath, &AssetData);
+			GeneratedAssetRegistryData.Add(AssetData.PackageName, &AssetData);
+
+			if (ContainsMap(AssetData.PackageName))
+			{
+				MapList.Add(AssetData.PackageName);
+			}
 		}
 	}
-	AssetRegistry.SaveRegistryData(SerializedAssetRegistry, GeneratedAssetRegistryData);
+
+	AssetRegistry.SaveRegistryData(SerializedAssetRegistry, GeneratedAssetRegistryData, &MapList);
 	UE_LOG(LogChunkManifestGenerator, Display, TEXT("Generated asset registry num assets %d, size is %5.2fkb"), GeneratedAssetRegistryData.Num(), (float)SerializedAssetRegistry.Num() / 1024.f);
 
+	auto CookerFileOrderString = CreateCookerFileOrderString(GeneratedAssetRegistryData, MapList);
+
 	// Save the generated registry for each platform
 	for (auto Platform : Platforms)
 	{
 		FString PlatformSandboxPath = SandboxPath.Replace(TEXT("[Platform]"), *Platform->PlatformName());
 		FFileHelper::SaveArrayToFile(SerializedAssetRegistry, *PlatformSandboxPath);
+
+		if (CookerFileOrderString.Len())
+		{
+			auto OpenOrderFilename = FString::Printf(TEXT("%sBuild/%s/FileOpenOrder/CookerOpenOrder.log"), *FPaths::GameDir(), *Platform->PlatformName());
+			FFileHelper::SaveStringToFile(CookerFileOrderString, *OpenOrderFilename);
+		}
 	}
 
 	UE_LOG(LogChunkManifestGenerator, Display, TEXT("Done saving asset registry."));
@@ -516,6 +556,86 @@ bool FChunkManifestGenerator::SaveAssetRegistry(const FString& SandboxPath, cons
 	return true;
 }
 
+/** Helper function which reroots a sandbox path to the staging area directory which UnrealPak expects */
+inline void ConvertFilenameToPakFormat(FString& InOutPath)
+{
+	auto GameDir = FPaths::GameDir();
+	auto EngineDir = FPaths::EngineDir();
+	auto GameName = FApp::GetGameName();
+
+	if (InOutPath.Contains(GameDir))
+	{
+		FPaths::MakePathRelativeTo(InOutPath, *GameDir);
+		InOutPath = FString::Printf(TEXT("../../../%s/%s"), GameName, *InOutPath);
+	}
+	else if (InOutPath.Contains(EngineDir))
+	{
+		FPaths::MakePathRelativeTo(InOutPath, *EngineDir);
+		InOutPath = FPaths::Combine(TEXT("../../../Engine/"), *InOutPath);
+	}
+}
+
+FString FChunkManifestGenerator::CreateCookerFileOrderString(const TMap<FName, FAssetData*>& InAssetData, const TArray<FName>& InMaps)
+{
+	FString FileOrderString;
+	TArray<FAssetData*> TopLevelNodes;
+
+	for (auto Asset : InAssetData)
+	{
+		auto PackageName = Asset.Value->PackageName;
+		TArray<FName> Referencers;
+		AssetRegistry.GetReferencers(PackageName, Referencers);
+
+		bool bIsTopLevel = true;
+		bool bIsMap = InMaps.Contains(PackageName);
+
+		if (!bIsMap && Referencers.Num() > 0)
+		{
+			for (auto ReferencerName : Referencers)
+			{
+				if (InAssetData.Contains(ReferencerName))
+				{
+					bIsTopLevel = false;
+					break;
+				}
+			}
+		}
+
+		if (bIsTopLevel)
+		{
+			if (bIsMap)
+			{
+				TopLevelNodes.Insert(Asset.Value, 0);
+			}
+			else
+			{
+				TopLevelNodes.Insert(Asset.Value, TopLevelNodes.Num());
+			}
+		}
+	}
+
+	TArray<FName> FileOrder;
+	TArray<FName> EncounteredNames;
+	for (auto Asset : TopLevelNodes)
+	{
+		AddAssetToFileOrderRecursive(Asset, FileOrder, EncounteredNames, InAssetData, InMaps);
+	}
+
+	int32 CurrentIndex = 0;
+	for (auto PackageName : FileOrder)
+	{
+		auto Asset = InAssetData[PackageName];
+		bool bIsMap = InMaps.Contains(Asset->PackageName);
+		auto Filename = FPackageName::LongPackageNameToFilename(Asset->PackageName.ToString(), bIsMap ? FPackageName::GetMapPackageExtension() : FPackageName::GetAssetPackageExtension());
+
+		ConvertFilenameToPakFormat(Filename);
+		auto Line = FString::Printf(TEXT("\"%s\" %i\n"), *Filename, CurrentIndex++);
+		FileOrderString.Append(Line);
+	}
+
+	return FileOrderString;
+}
+
 typedef TSharedRef< TJsonWriter<TCHAR, TPrettyJsonPrintPolicy<TCHAR> > > JsonWriter;
 typedef TSharedRef< TJsonReader<TCHAR> > JsonReader;
 
diff --git a/Engine/Source/Editor/UnrealEd/Private/Commandlets/ChunkManifestGenerator.h b/Engine/Source/Editor/UnrealEd/Private/Commandlets/ChunkManifestGenerator.h
--- a/Engine/Source/Editor/UnrealEd/Private/Commandlets/ChunkManifestGenerator.h
+++ b/Engine/Source/Editor/UnrealEd/Private/Commandlets/ChunkManifestGenerator.h
@@ -291,4 +291,22 @@ public:
 	 */
 	bool SaveCookedPackageAssetRegistry( const FString& SandboxPath, const bool Append );
 
+	/**
+	* Follows an assets dependency chain to build up a list of package names in the same order as the runtime would attempt to load them
+	* 
+	* @param InAsset - The asset to (potentially) add to the file order
+	* @param OutFileOrder - Output array which collects the package names
+	* @param OutEncounteredArray - Temporary collection of package names we've seen. Similar to OutFileOrder but updated BEFORE following dependencies so as to avoid circular references
+	* @param InAssets - The source asset list. Used to distinguish between dependencies on other packages and internal objects
+	* @param InAssetTypeMap - Lookup table of package name to type. Used to determine if a dependency is a streamed level so as not to follow dependencies
+	*/
+	void AddAssetToFileOrderRecursive(FAssetData* InAsset, TArray<FName>& OutFileOrder, TArray<FName>& OutEncounteredNames, const TMap<FName, FAssetData*>& InAssets, const TArray<FName>& InMapList);
+
+	/**
+	* Build a file order string which represents the order in which files would be loaded at runtime. 
+	* 
+	* @param InAssetData - Assets data for those assets which were cooked on this run. 
+	* @param InTopLevelAssets - Names of map assets
+	*/
+	FString CreateCookerFileOrderString(const TMap<FName, FAssetData*>& InAssetData, const TArray<FName>& InMaps);
 };
diff --git a/Engine/Source/Programs/AutomationTool/Scripts/CopyBuildToStagingDirectory.Automation.cs b/Engine/Source/Programs/AutomationTool/Scripts/CopyBuildToStagingDirectory.Automation.cs
--- a/Engine/Source/Programs/AutomationTool/Scripts/CopyBuildToStagingDirectory.Automation.cs
+++ b/Engine/Source/Programs/AutomationTool/Scripts/CopyBuildToStagingDirectory.Automation.cs
@@ -706,11 +706,17 @@ public partial class Project : CommandUtils
 		var OutputLocation = CombinePaths(SC.RuntimeRootDir, OutputRelativeLocation);
 		// Add input file to controll order of file within the pak
 		var PakOrderFileLocationBase = CombinePaths(SC.ProjectRoot, "Build", SC.FinalCookPlatform, "FileOpenOrder");
-		var PakOrderFileLocation = CombinePaths(PakOrderFileLocationBase, "GameOpenOrder.log");
-		if (!FileExists_NoExceptions(PakOrderFileLocation))
+
+		var OrderFileLocations = new string[] { "GameOpenOrder.log", "CookerOpenOrder.log", "EditorOpenOrder.log" };
+		string PakOrderFileLocation = null;
+		foreach (var Location in OrderFileLocations)
 		{
-			// Use a fall back, it doesn't matter if this file exists or not. GameOpenOrder.log is preferred however
-			PakOrderFileLocation = CombinePaths(PakOrderFileLocationBase, "EditorOpenOrder.log");
+			PakOrderFileLocation = CombinePaths(PakOrderFileLocationBase, Location);
+
+			if (FileExists_NoExceptions(PakOrderFileLocation))
+			{
+				break;
+			}
 		}
 
 		bool bCopiedExistingPak = false;
diff --git a/Engine/Source/Runtime/AssetRegistry/Private/AssetRegistry.cpp b/Engine/Source/Runtime/AssetRegistry/Private/AssetRegistry.cpp
--- a/Engine/Source/Runtime/AssetRegistry/Private/AssetRegistry.cpp
+++ b/Engine/Source/Runtime/AssetRegistry/Private/AssetRegistry.cpp
@@ -34,6 +34,7 @@ EChunkProgressReportingType::Type GetChunkAvailabilityProgressType(EAssetAvailab
 FAssetRegistry::FAssetRegistry()
 	: PathTreeRoot(FString())
 	, PreallocatedAssetDataBuffer(NULL)
+	, PreallocatedDependsNodeDataBuffer(NULL)
 {
 	const double StartupStartTime = FPlatformTime::Seconds();
 
@@ -130,6 +131,7 @@ FAssetRegistry::~FAssetRegistry()
 	if (PreallocatedAssetDataBuffer != NULL)
 	{
 		delete [] PreallocatedAssetDataBuffer;
+		PreallocatedAssetDataBuffer = nullptr;
 		NumAssets = 0;
 	}
 	else
@@ -149,13 +151,22 @@ FAssetRegistry::~FAssetRegistry()
 	// Make sure we have deleted all our allocated FAssetData objects
 	ensure(NumAssets == 0);
 
-	// Delete all depends nodes in the cache
-	for (TMap<FName, FDependsNode*>::TConstIterator DependsIt(CachedDependsNodes); DependsIt; ++DependsIt)
+	if (PreallocatedDependsNodeDataBuffer != NULL)
 	{
-		if ( DependsIt.Value() )
+		delete[] PreallocatedDependsNodeDataBuffer;
+		PreallocatedDependsNodeDataBuffer = nullptr;
+		NumDependsNodes = 0;
+	}
+	else
+	{
+		// Delete all depends nodes in the cache
+		for (TMap<FName, FDependsNode*>::TConstIterator DependsIt(CachedDependsNodes); DependsIt; ++DependsIt)
 		{
-			delete DependsIt.Value();
-			NumDependsNodes--;
+			if (DependsIt.Value())
+			{
+				delete DependsIt.Value();
+				NumDependsNodes--;
+			}
 		}
 	}
 
@@ -1320,29 +1331,159 @@ void FAssetRegistry::Serialize(FArchive& Ar)
 
 			AddAssetPath(NewAssetData->PackagePath.ToString());
 		}
+
+		int32 LocalNumDependsNodes = 0;
+		Ar << LocalNumDependsNodes;
+
+		PreallocatedDependsNodeDataBuffer = new FDependsNode[LocalNumDependsNodes];
+
+		for (int32 DependsNodeIndex = 0; DependsNodeIndex < LocalNumDependsNodes; DependsNodeIndex++)
+		{
+			auto NewDependsNodeData = &PreallocatedDependsNodeDataBuffer[DependsNodeIndex];
+			int32 AssetIndex = 0;
+			Ar << AssetIndex;
+			NewDependsNodeData->PackageName = PreallocatedAssetDataBuffer[AssetIndex].PackageName;
+
+			CachedDependsNodes.Add(NewDependsNodeData->PackageName, NewDependsNodeData);
+		}
+
+		for (int32 DependsNodeIndex = 0; DependsNodeIndex < LocalNumDependsNodes; DependsNodeIndex++)
+		{
+			auto NewDependsNodeData = &PreallocatedDependsNodeDataBuffer[DependsNodeIndex];
+			int32 LocalNumDependencies = 0;
+			int32 LocalNumReferencers = 0;
+			Ar << LocalNumDependencies;
+			Ar << LocalNumReferencers;
+
+			for (int32 DependencyIndex = 0; DependencyIndex < LocalNumDependencies; ++DependencyIndex)
+			{
+				int32 Index = 0;
+				Ar << Index;
+				NewDependsNodeData->Dependencies.Add(&PreallocatedDependsNodeDataBuffer[Index]);
+			}
+
+			for (int32 ReferencerIndex = 0; ReferencerIndex < LocalNumReferencers; ++ReferencerIndex)
+			{
+				int32 Index = 0;
+				Ar << Index;
+				NewDependsNodeData->Referencers.Add(&PreallocatedDependsNodeDataBuffer[Index]);
+			}
+		}
+	}
+}
+
+void AddDependsNodesRecursive(FDependsNode* InDependsNode, TArray<FDependsNode*>& OutArray, TMap<FName, FAssetData*>& InAssets)
+{
+	if (!OutArray.Contains(InDependsNode))
+	{
+		if (InAssets.Contains(InDependsNode->PackageName))
+		{
+			OutArray.Add(InDependsNode);
+
+			for (auto Dependency : InDependsNode->Dependencies)
+			{
+				AddDependsNodesRecursive(Dependency, OutArray, InAssets);
+			}
+
+			for (auto Referencer : InDependsNode->Referencers)
+			{
+				AddDependsNodesRecursive(Referencer, OutArray, InAssets);
+			}
+		}
 	}
 }
 
-void FAssetRegistry::SaveRegistryData(FArchive& Ar, TMap<FName, FAssetData*>& Data )
+void FAssetRegistry::SaveRegistryData(FArchive& Ar, TMap<FName, FAssetData*>& Data, TArray<FName>* InMaps /* = nullptr */)
 {
 	// serialize number of objects
 	int32 AssetCount = Data.Num();
 	Ar << AssetCount;
 
+	TArray<FName> DependencyNames;
+	TArray<FDependsNode*> Dependencies;
+	TMap<FName, int32> AssetIndexMap;
+	AssetIndexMap.Reserve(Data.Num());
+
 	// save out by walking the TMap
 	for (TMap<FName, FAssetData*>::TIterator It(Data); It; ++It)
 	{
 		Ar << *It.Value();
+		AssetIndexMap.Add(It.Value()->PackageName, AssetIndexMap.Num());
+
+		AddDependsNodesRecursive(FindDependsNode(It.Value()->PackageName), Dependencies, Data);
+	}
+
+	int32 DependsNodeCount = Dependencies.Num();
+	Ar << DependsNodeCount;
+
+	TMap<FDependsNode*, int32> DependencyIndexMap;
+	DependencyIndexMap.Reserve(DependsNodeCount);
+
+	for (auto DependentNode : Dependencies)
+	{
+		int32 AssetIndex = AssetIndexMap[DependentNode->PackageName];
+		Ar << AssetIndex;
+		DependencyIndexMap.Add(DependentNode, DependencyIndexMap.Num());
+	}
+
+	for (auto DependentNode : Dependencies)
+	{
+		auto DependencyCount = 0;
+		auto ReferencerCount = 0;
+
+		for (auto Dependency : DependentNode->Dependencies)
+		{
+			bool bIsMap = InMaps && InMaps->Contains(Dependency->PackageName);
+		
+			if (!bIsMap && DependencyIndexMap.Contains(Dependency))
+			{
+				DependencyCount++;
+			}
+		}
+
+		for (auto Referencer : DependentNode->Referencers)
+		{
+			if (DependencyIndexMap.Contains(Referencer))
+			{
+				ReferencerCount++;
+			}
+		}
+
+		Ar << DependencyCount;
+		Ar << ReferencerCount;
+
+		for (auto Dependency : DependentNode->Dependencies)
+		{
+			bool bIsMap = InMaps && InMaps->Contains(Dependency->PackageName);
+
+			if (!bIsMap && DependencyIndexMap.Contains(Dependency))
+			{
+				int32 Index = DependencyIndexMap[Dependency];
+				Ar << Index;
+			}
+		}
+
+		for (auto Referencer : DependentNode->Referencers)
+		{
+			if (DependencyIndexMap.Contains(Referencer))
+			{
+				int32 Index = DependencyIndexMap[Referencer];
+				Ar << Index;
+			}
+		}
 	}
 }
 
-void FAssetRegistry::LoadRegistryData(FArchive& Ar, TMap<FName, FAssetData*>& Data )
+void FAssetRegistry::LoadRegistryData(FArchive& Ar, TMap<FName, FAssetData*>& Data, TArray<FDependsNode*>& OutDependencyData)
 {
 	check(Ar.IsLoading());
 	// serialize number of objects
 	int AssetCount = 0;
 	Ar << AssetCount;
 
+	Data.Reserve(AssetCount);
+	TMap<int32, FName> AssetIndexMap;
+
 	for (int32 AssetIndex = 0; AssetIndex < AssetCount; AssetIndex++)
 	{
 		// make a new asset data object
@@ -1351,7 +1492,55 @@ void FAssetRegistry::LoadRegistryData(FArchive& Ar, TMap<FName, FAssetData*>& Da
 		// load it
 		Ar << *NewAssetData;
 
-		Data.Add(NewAssetData->ObjectPath, NewAssetData);
+		Data.Add(NewAssetData->PackageName, NewAssetData);
+		AssetIndexMap.Add(AssetIndexMap.Num(), NewAssetData->PackageName);
+	}
+
+	if (Ar.TotalSize() > Ar.Tell())
+	{
+		int DependsNodeCount = 0;
+		Ar << DependsNodeCount;
+
+		OutDependencyData.Reserve(DependsNodeCount);
+
+		for (int32 DependsNodeIndex = 0; DependsNodeIndex < DependsNodeCount; DependsNodeIndex++)
+		{
+			int32 AssetIndex;
+			Ar << AssetIndex;
+
+			auto PackageName = AssetIndexMap[AssetIndex];
+			auto NewNode = new FDependsNode(PackageName);
+			OutDependencyData.Add(NewNode);
+		}
+
+		for (int32 DependsNodeIndex = 0; DependsNodeIndex < DependsNodeCount; DependsNodeIndex++)
+		{
+			int32 NumDependencies = 0;
+			Ar << NumDependencies;
+
+			int32 NumReferencers = 0;
+			Ar << NumReferencers;
+
+			auto DependsNode = OutDependencyData[DependsNodeIndex];
+			DependsNode->Dependencies.Reserve(NumDependencies);
+			DependsNode->Referencers.Reserve(NumReferencers);
+
+			for (int32 i = 0; i < NumDependencies; ++i)
+			{
+				int32 DependencyIndex = 0;
+				Ar << DependencyIndex;
+
+				DependsNode->Dependencies.Add(OutDependencyData[DependencyIndex]);
+			}
+
+			for (int32 i = 0; i < NumReferencers; ++i)
+			{
+				int32 Index = 0;
+				Ar << Index;
+
+				DependsNode->Referencers.Add(OutDependencyData[Index]);
+			}
+		}
 	}
 }
 
@@ -1540,12 +1729,25 @@ void FAssetRegistry::DependencyDataGathered(const double TickStartTime, TArray<F
 	DependsResults.RemoveAt(0, ResultIdx);
 }
 
+FDependsNode* FAssetRegistry::FindDependsNode(FName ObjectName)
+{
+	FDependsNode** FoundNode = CachedDependsNodes.Find(ObjectName);
+	if (FoundNode)
+	{
+		return *FoundNode;
+	}
+	else
+	{
+		return nullptr;
+	}
+}
+
 FDependsNode* FAssetRegistry::CreateOrFindDependsNode(FName ObjectPath)
 {
-	FDependsNode** FoundNode = CachedDependsNodes.Find(ObjectPath);
+	FDependsNode* FoundNode = FindDependsNode(ObjectPath);
 	if ( FoundNode )
 	{
-		return *FoundNode;
+		return FoundNode;
 	}
 
 	FDependsNode* NewNode = new FDependsNode(ObjectPath);
diff --git a/Engine/Source/Runtime/AssetRegistry/Private/AssetRegistry.h b/Engine/Source/Runtime/AssetRegistry/Private/AssetRegistry.h
--- a/Engine/Source/Runtime/AssetRegistry/Private/AssetRegistry.h
+++ b/Engine/Source/Runtime/AssetRegistry/Private/AssetRegistry.h
@@ -38,8 +38,8 @@ public:
 	virtual void ScanPathsSynchronous(const TArray<FString>& InPaths, bool bForceRescan = false) override;
 	virtual void PrioritizeSearchPath(const FString& PathToPrioritize) override;
 	virtual void Serialize(FArchive& Ar) override;
-	virtual void SaveRegistryData(FArchive& Ar, TMap<FName, FAssetData*>& Data) override;
-	virtual void LoadRegistryData(FArchive& Ar, TMap<FName, FAssetData*>& Data) override;
+	virtual void SaveRegistryData(FArchive& Ar, TMap<FName, FAssetData*>& Data, TArray<FName>* InMaps = nullptr) override;
+	virtual void LoadRegistryData(FArchive& Ar, TMap<FName, FAssetData*>& Data, TArray<FDependsNode*>& OutDependencyData) override;
 
 	DECLARE_DERIVED_EVENT( FAssetRegistry, IAssetRegistry::FPathAddedEvent, FPathAddedEvent);
 	virtual FPathAddedEvent& OnPathAdded() override { return PathAddedEvent; }
@@ -93,6 +93,9 @@ private:
 	/** Called every tick to when data is retrieved by the background dependency search */
 	void DependencyDataGathered(const double TickStartTime, TArray<FPackageDependencyData>& DependsResults);
 
+	/** Finds an existing node for the given package and returns it, or returns null if one isn't found */
+	FDependsNode* FindDependsNode(FName ObjectName);
+
 	/** Creates a node in the CachedDependsNodes map or finds the existing node and returns it */
 	FDependsNode* CreateOrFindDependsNode(FName ObjectName);
 
@@ -249,6 +252,7 @@ private:
 
 	/** When loading a registry from disk, we can allocate all the FAssetData objects in one chunk, to save on 10s of thousands of heap allocations */
 	FAssetData* PreallocatedAssetDataBuffer;
+	FDependsNode* PreallocatedDependsNodeDataBuffer;
 
 	/** A set used to ignore repeated requests to synchronously scan the same folder multiple times */
 	TSet<FString> SynchronouslyScannedPaths;
diff --git a/Engine/Source/Runtime/AssetRegistry/Private/DependsNode.h b/Engine/Source/Runtime/AssetRegistry/Private/DependsNode.h
--- a/Engine/Source/Runtime/AssetRegistry/Private/DependsNode.h
+++ b/Engine/Source/Runtime/AssetRegistry/Private/DependsNode.h
@@ -6,6 +6,7 @@
 class FDependsNode
 {
 public:
+	FDependsNode() { PackageName = NAME_None; }
 	FDependsNode(FName InPackageName);
 
 	/** Prints the dependencies and referencers for this node to the log */
diff --git a/Engine/Source/Runtime/AssetRegistry/Public/AssetRegistryModule.h b/Engine/Source/Runtime/AssetRegistry/Public/AssetRegistryModule.h
--- a/Engine/Source/Runtime/AssetRegistry/Public/AssetRegistryModule.h
+++ b/Engine/Source/Runtime/AssetRegistry/Public/AssetRegistryModule.h
@@ -5,6 +5,7 @@
 
 #include "ModuleInterface.h"
 #include "IAssetRegistry.h"
+#include "AssetRegistryInterface.h"
 
 namespace AssetRegistryConstants
 {
@@ -14,7 +15,7 @@ namespace AssetRegistryConstants
 /**
  * Asset registry module
  */
-class FAssetRegistryModule : public IModuleInterface
+class FAssetRegistryModule : public IAssetRegistryInterface
 {
 
 public:
@@ -65,6 +66,12 @@ public:
 		}
 	}
 
+	/** Access the dependent package names for a given source package */
+	void GetDependencies(FName InPackageName, TArray<FName>& OutDependencies) override
+	{
+		Get().GetDependencies(InPackageName, OutDependencies);
+	}
+
 private:
 	IAssetRegistry* AssetRegistry;
 	class FAssetRegistryConsoleCommands* ConsoleCommands;
diff --git a/Engine/Source/Runtime/AssetRegistry/Public/IAssetRegistry.h b/Engine/Source/Runtime/AssetRegistry/Public/IAssetRegistry.h
--- a/Engine/Source/Runtime/AssetRegistry/Public/IAssetRegistry.h
+++ b/Engine/Source/Runtime/AssetRegistry/Public/IAssetRegistry.h
@@ -26,6 +26,7 @@ namespace EAssetAvailabilityProgressReportingType
 	};
 }
 
+class FDependsNode;
 
 class IAssetRegistry
 {
@@ -223,9 +224,9 @@ public:
 	virtual void Serialize(FArchive& Ar) = 0;
 
 	/** Serialize raw registry data to a file, skipping editor only data */
-	virtual void SaveRegistryData(FArchive& Ar, TMap<FName, FAssetData*>& Data) = 0;
+	virtual void SaveRegistryData(FArchive& Ar, TMap<FName, FAssetData*>& Data, TArray<FName>* InMaps = nullptr) = 0;
 
 
 	/** Serialize registry data from a file */
-	virtual void LoadRegistryData(FArchive& Ar, TMap<FName, FAssetData*>& Data) = 0;
+	virtual void LoadRegistryData(FArchive& Ar, TMap<FName, FAssetData*>& Data, TArray<FDependsNode*>& OutDependencyData) = 0;
 };
diff --git a/Engine/Source/Runtime/CoreUObject/Private/Serialization/AsyncLoading.cpp b/Engine/Source/Runtime/CoreUObject/Private/Serialization/AsyncLoading.cpp
--- a/Engine/Source/Runtime/CoreUObject/Private/Serialization/AsyncLoading.cpp
+++ b/Engine/Source/Runtime/CoreUObject/Private/Serialization/AsyncLoading.cpp
@@ -11,6 +11,7 @@
 #include "UObject/UObjectThreadContext.h"
 #include "UObject/LinkerManager.h"
 #include "Serialization/AsyncLoadingThread.h"
+#include "AssetRegistryInterface.h"
 
 /*-----------------------------------------------------------------------------
 	Async loading stats.
@@ -49,6 +50,9 @@ DECLARE_FLOAT_ACCUMULATOR_STAT( TEXT( "Async package precache wait time" ), STAT
 /** Returns true if we're inside a FGCScopeLock */
 bool IsGarbageCollectionLocked();
 
+/** Global request ID counter */
+static FThreadSafeCounter GPackageRequestID;
+
 /** 
  * Keeps a reference to all objects created during async load until streaming has finished 
  *
@@ -243,6 +247,16 @@ static FAutoConsoleVariableRef CVarTimeLimitExceededMinTime(
 	ECVF_Default
 	);
 
+static int32 GPreloadPackageDependencies = 0;
+static FAutoConsoleVariableRef CVarPreloadPackageDependencies(
+	TEXT("s.PreloadPackageDependencies"),
+	GPreloadPackageDependencies,
+	TEXT("Enables preloading of package dependencies based on data from the asset registry\n") \
+	TEXT("0 - Do not preload dependencies. Can cause more seeks but uses less memory [default].\n") \
+	TEXT("1 - Preload package dependencies. Faster but requires asset registry data to be loaded into memory\n"),
+	ECVF_Default
+	);
+
 static FORCEINLINE bool IsTimeLimitExceeded(double InTickStartTime, bool bUseTimeLimit, float InTimeLimit, const TCHAR* InLastTypeOfWorkPerformed = nullptr, UObject* InLastObjectWorkWasPerformedOn = nullptr)
 {
 	bool bTimeLimitExceeded = false;
@@ -394,6 +408,83 @@ FAsyncPackage* FAsyncLoadingThread::FindExistingPackageAndAddCompletionCallback(
 	return Result;
 }
 
+void FAsyncLoadingThread::ProcessAsyncPackageRequest(FAsyncPackageDesc* InRequest, FAsyncPackage* InRootPackage, TMap<FName, int32>& InDependencyTracker)
+{
+	check(!InDependencyTracker.Contains(InRequest->Name));
+	FAsyncPackage* Package = FindExistingPackageAndAddCompletionCallback(InRequest, AsyncPackages);
+
+	if (!Package)
+	{
+		// [BLOCKING] LoadedPackages are accessed on the main thread too, so lock to be able to add a completion callback
+#if THREADSAFE_UOBJECTS
+		FScopeLock LoadedLock(&LoadedPackagesCritical);
+#endif
+		Package = FindExistingPackageAndAddCompletionCallback(InRequest, LoadedPackages);
+	}
+
+	if (!Package)
+	{
+		// [BLOCKING] LoadedPackagesToProcess are modified on the main thread, so lock to be able to add a completion callback
+#if THREADSAFE_UOBJECTS
+		FScopeLock LoadedLock(&LoadedPackagesToProcessCritical);
+#endif
+		Package = FindExistingPackageAndAddCompletionCallback(InRequest, LoadedPackagesToProcess);
+	}
+
+	if (!Package)
+	{
+		// New package that needs to be loaded or a package has already been loaded long time ago
+		Package = new FAsyncPackage(*InRequest);
+		if (InRequest->PackageLoadedDelegate.IsBound())
+		{
+			const bool bInternalCallback = false;
+			Package->AddCompletionCallback(InRequest->PackageLoadedDelegate, bInternalCallback);
+		}
+		Package->SetDependencyRootPackage(InRootPackage);
+
+#if !WITH_EDITOR
+		if (GPreloadPackageDependencies != 0)
+		{
+			auto AssetRegistry = IAssetRegistryInterface::GetPtr();
+			if (AssetRegistry)
+			{
+				TArray<FName> Dependencies;
+				AssetRegistry->GetDependencies(Package->GetPackageName(), Dependencies);
+
+				int32 PIEInstanceID = INDEX_NONE;
+#if WITH_EDITOR
+				PIEInstanceID = InRequest->PIEInstanceID;
+#endif
+
+				if (InRootPackage == nullptr)
+				{
+					InRootPackage = Package;
+				}
+
+				for (auto DependencyName : Dependencies)
+				{
+					if (!InDependencyTracker.Contains(DependencyName))
+					{
+						QueuedPackagesCounter.Increment();
+						const int32 RequestID = GPackageRequestID.Increment();
+						FAsyncLoadingThread::Get().AddPendingRequest(RequestID);
+						FAsyncPackageDesc DependencyPackageRequest(RequestID, DependencyName, NAME_None, FGuid(), FLoadPackageAsyncDelegate(), InRequest->PackageFlags, PIEInstanceID, InRequest->Priority);
+						InDependencyTracker.Add(InRequest->Name);
+						ProcessAsyncPackageRequest(&DependencyPackageRequest, InRootPackage, InDependencyTracker);
+					}
+				}
+			}
+		}
+#endif
+		// Add to queue according to priority.
+		InsertPackage(Package);
+
+		// For all other cases this is handled in FindExistingPackageAndAddCompletionCallback
+		const int32 QueuedPackagesCount = QueuedPackagesCounter.Decrement();
+		check(QueuedPackagesCount >= 0);
+	}
+}
+
 int32 FAsyncLoadingThread::CreateAsyncPackagesFromQueue()
 {
 	SCOPE_CYCLE_COUNTER(STAT_FAsyncPackage_CreateAsyncPackagesFromQueue);
@@ -412,46 +503,20 @@ int32 FAsyncLoadingThread::CreateAsyncPackagesFromQueue()
 		QueuedPackages.Empty();
 	}
 
-	for (FAsyncPackageDesc* PackageRequest : QueueCopy)
+	if (QueueCopy.Num() > 0)
 	{
-		FAsyncPackage* Package = FindExistingPackageAndAddCompletionCallback(PackageRequest, AsyncPackages);
-
-		if (!Package)
+		double Timer = 0;
 		{
-			// [BLOCKING] LoadedPackages are accessed on the main thread too, so lock to be able to add a completion callback
-#if THREADSAFE_UOBJECTS
-			FScopeLock LoadedLock(&LoadedPackagesCritical);
-#endif
-			Package = FindExistingPackageAndAddCompletionCallback(PackageRequest, LoadedPackages);
-		}
-
-		if (!Package)
-		{
-			// [BLOCKING] LoadedPackagesToProcess are modified on the main thread, so lock to be able to add a completion callback
-#if THREADSAFE_UOBJECTS
-			FScopeLock LoadedLock(&LoadedPackagesToProcessCritical);
-#endif
-			Package = FindExistingPackageAndAddCompletionCallback(PackageRequest, LoadedPackagesToProcess);
-		}
-
-		if (!Package)
-		{
-			// New package that needs to be loaded or a package has already been loaded long time ago
-			Package = new FAsyncPackage(*PackageRequest);
-			if (PackageRequest->PackageLoadedDelegate.IsBound())
+			SCOPE_SECONDS_COUNTER(Timer);
+			DependencyTracker.Empty();
+			for (auto PackageRequest : QueueCopy)
 			{
-				const bool bInternalCallback = false;
-				Package->AddCompletionCallback(PackageRequest->PackageLoadedDelegate, bInternalCallback);
+				ProcessAsyncPackageRequest(PackageRequest, nullptr, DependencyTracker);
 			}
-
-			// Add to queue according to priority.
-			InsertPackage(Package);
-
-			// For all other cases this is handled in FindExistingPackageAndAddCompletionCallback
-			const int32 QueuedPackagesCount = QueuedPackagesCounter.Decrement();
-			check(QueuedPackagesCount >= 0);
 		}
+		UE_LOG(LogStreaming, Display, TEXT("Async package requests inserted in %fms"), Timer * 1000.0);
 	}
+
 	NumCreated = QueueCopy.Num();
 
 	return NumCreated;
@@ -474,7 +539,9 @@ void FAsyncLoadingThread::InsertPackage(FAsyncPackage* Package)
 		// Insert new package keeping descending priority order in AsyncPackages
 		auto InsertIndex = AsyncPackages.IndexOfByPredicate([Package](const FAsyncPackage* Element)
 		{
-			return Element->GetPriority() <= Package->GetPriority();
+			// NOTE: Switched to < because we want newly inserted packages to go AFTER all existing packages of the 
+			//		 same priority. If this is too slow, we could probably do "LastIndexOfByPredicate" where we iterate backwards??
+			return Element->GetPriority() < Package->GetPriority();
 		});
 
 		InsertIndex = InsertIndex == INDEX_NONE ? AsyncPackages.Num() : InsertIndex;
@@ -873,6 +940,7 @@ FAsyncPackage::FAsyncPackage(const FAsyncPackageDesc& InDesc)
 : Desc(InDesc)
 , Linker(nullptr)
 , LinkerRoot(nullptr)
+, DependencyRootPackage(nullptr)
 , DependencyRefCount(0)
 , LoadImportIndex(0)
 , ImportIndex(0)
@@ -1062,6 +1130,17 @@ EAsyncPackageState::Type FAsyncPackage::Tick(bool InbUseTimeLimit, bool InbUseFu
 	if (LoadStartTime == 0.0)
 	{
 		LoadStartTime = TickStartTime;
+
+		// If we are a dependency of another package, we need to tell that package when its first dependent started loading,
+		// otherwise because that package loads last it'll not include the entire load time of all its dependencies
+		if (DependencyRootPackage)
+		{
+			// Only the first dependent needs to register the start time
+			if (DependencyRootPackage->GetLoadStartTime() == 0.0)
+			{
+				DependencyRootPackage->LoadStartTime = TickStartTime;
+			}
+		}
 	}
 
 	// Make sure we finish our work if there's no time limit. The loop is required as PostLoad
@@ -1871,11 +1950,19 @@ void FAsyncPackage::UpdateLoadPercentage()
 	LoadPercentage = FMath::Max(NewLoadPercentage, LoadPercentage);
 }
 
-/** Global request ID counter */
-static FThreadSafeCounter GPackageRequestID;
-
 int32 LoadPackageAsync(const FString& InName, const FGuid* InGuid /*= nullptr*/, const TCHAR* InPackageToLoadFrom /*= nullptr*/, FLoadPackageAsyncDelegate InCompletionDelegate /*= FLoadPackageAsyncDelegate()*/, EPackageFlags InPackageFlags /*= PKG_None*/, int32 InPIEInstanceID /*= INDEX_NONE*/, uint32 InPackagePriority /*= 0*/)
 {
+#if !WITH_EDITOR
+	if (GPreloadPackageDependencies)
+	{
+		// If dependency preloading is enabled, we need to force the asset registry module to be loaded on the game thread
+		// as it will potentiall be used on the async loading thread, which isn't allowed to load modules.
+		// We could do this at init time, but doing it here allows us to not load the module at all if preloading is
+		// disabled.
+		IAssetRegistryInterface::GetPtr();
+	}
+#endif
+
 	// The comments clearly state that it should be a package name but we also handle it being a filename as this function is not perf critical
 	// and LoadPackage handles having a filename being passed in as well.
 	FString PackageName;
diff --git a/Engine/Source/Runtime/CoreUObject/Private/Serialization/AsyncLoadingThread.h b/Engine/Source/Runtime/CoreUObject/Private/Serialization/AsyncLoadingThread.h
--- a/Engine/Source/Runtime/CoreUObject/Private/Serialization/AsyncLoadingThread.h
+++ b/Engine/Source/Runtime/CoreUObject/Private/Serialization/AsyncLoadingThread.h
@@ -74,6 +74,9 @@ class FAsyncLoadingThread : public FRunnable
 	/** Async loading thread ID */
 	static uint32 AsyncLoadingThreadID;
 
+	/** Helper map for tracking the dependency packages that have been requested while loading a package */
+	TMap<FName, int32> DependencyTracker;
+
 #if LOOKING_FOR_PERF_ISSUES
 	/** Thread safe counter used to accumulate cycles spent on blocking. Using stats may generate to many stats messages. */
 	static FThreadSafeCounter BlockingCycles;
@@ -349,6 +352,12 @@ private:
 	*/
 	int32 CreateAsyncPackagesFromQueue();
 
+	/**
+	* [ASYNC THREAD] Internal helper function for processing a package load request. If dependency preloading is enabled, 
+	* it will call itself recursively for all the package dependencies
+	*/
+	void ProcessAsyncPackageRequest(FAsyncPackageDesc* InRequest, FAsyncPackage* InRootPackage, TMap<FName, int32>& InDependencyTracker);
+
 	/**
 	* [ASYNC THREAD] Finds existing async package and adds the new request's completion callback to it.
 	*/
diff --git a/Engine/Source/Runtime/CoreUObject/Private/UObject/UObjectGlobals.cpp b/Engine/Source/Runtime/CoreUObject/Private/UObject/UObjectGlobals.cpp
--- a/Engine/Source/Runtime/CoreUObject/Private/UObject/UObjectGlobals.cpp
+++ b/Engine/Source/Runtime/CoreUObject/Private/UObject/UObjectGlobals.cpp
@@ -11,6 +11,7 @@
 #include "UObject/TlsObjectInitializers.h"
 #include "UObject/UObjectThreadContext.h"
 #include "BlueprintSupport.h" // for FDeferredObjInitializerTracker
+#include "AssetRegistryInterface.h"
 
 DEFINE_LOG_CATEGORY(LogUObjectGlobals);
 
@@ -956,6 +957,36 @@ UPackage* LoadPackageInternal(UPackage* InOuter, const TCHAR* InLongPackageName,
 			return Result;
 		}
 
+#if !WITH_EDITOR
+		static auto CVarPreloadDependencies = IConsoleManager::Get().FindConsoleVariable(TEXT("s.PreloadPackageDependencies"));
+		if (CVarPreloadDependencies && CVarPreloadDependencies->GetInt() != 0)
+		{
+			if (Result->PackageFlags & PKG_ProcessingDependencies)
+			{
+				// We've currently already processing the dependencies of this package, so there is a circular dependency
+				EndLoad();
+				return nullptr;
+			}
+
+			auto AssetRegistry = IAssetRegistryInterface::GetPtr();
+
+			if (AssetRegistry)
+			{
+				TArray<FName> PackageDependencies;
+				FName PackageName(InLongPackageName);
+
+				AssetRegistry->GetDependencies(PackageName, PackageDependencies);
+
+				Result->PackageFlags |= PKG_ProcessingDependencies;
+				for (auto Dependency : PackageDependencies)
+				{
+					LoadPackage(InOuter, *Dependency.ToString(), LoadFlags);
+				}
+				Result->PackageFlags &= ~PKG_ProcessingDependencies;
+			}
+		}
+#endif
+
 		// If we are loading a package for diff'ing, set the package flag
 		if(LoadFlags & LOAD_ForDiff)
 		{
diff --git a/Engine/Source/Runtime/CoreUObject/Public/Misc/AssetRegistryInterface.h b/Engine/Source/Runtime/CoreUObject/Public/Misc/AssetRegistryInterface.h
new file mode 100644
--- /dev/null
+++ b/Engine/Source/Runtime/CoreUObject/Public/Misc/AssetRegistryInterface.h
@@ -0,0 +1,33 @@
+// Copyright 1998-2015 Epic Games, Inc. All Rights Reserved.
+
+#pragma once
+
+#include "ModuleManager.h"
+#include "CoreUObject.h"
+
+/**
+* HotReload module interface
+*/
+class IAssetRegistryInterface : public IModuleInterface
+{
+public:
+	/**
+	 * Tries to gets a pointer to the active AssetRegistryInterface implementation. 
+	 */
+	static inline IAssetRegistryInterface* GetPtr()
+	{
+		static FName AssetRegistry("AssetRegistry");
+		auto Module = FModuleManager::GetModulePtr<IAssetRegistryInterface>(AssetRegistry);
+		if (Module == nullptr)
+		{
+			Module = &FModuleManager::LoadModuleChecked<IAssetRegistryInterface>(AssetRegistry);
+		}
+		return reinterpret_cast<IAssetRegistryInterface*>(Module);
+	}
+
+	/**
+	* Lookup dependencies for the given package name and fill OutDependencies with direct dependencies
+	*/
+	virtual void GetDependencies(FName InPackageName, TArray<FName>& OutDependencies) = 0;
+};
+
diff --git a/Engine/Source/Runtime/CoreUObject/Public/Serialization/AsyncLoading.h b/Engine/Source/Runtime/CoreUObject/Public/Serialization/AsyncLoading.h
--- a/Engine/Source/Runtime/CoreUObject/Public/Serialization/AsyncLoading.h
+++ b/Engine/Source/Runtime/CoreUObject/Public/Serialization/AsyncLoading.h
@@ -97,6 +97,14 @@ struct FAsyncPackage
 	*/
 	void Cancel();
 
+	/**
+	* Set the package that spawned this package as a dependency.
+	*/
+	void SetDependencyRootPackage(FAsyncPackage* InDependencyRootPackage)
+	{
+		DependencyRootPackage = InDependencyRootPackage;
+	}
+
 private:
 
 	struct FCompletionCallback
@@ -126,6 +134,8 @@ private:
 	TArray<FAsyncPackage*> PendingImportedPackages;
 	/** Referenced imports - list of packages we need until we finish loading this package. */
 	TArray<FAsyncPackage*> ReferencedImports;
+	/** Root package if this package was loaded as a dependency of another. NULL otherwise */
+	FAsyncPackage* DependencyRootPackage;
 	/** Number of references to this package from other packages in the dependency tree. */
 	FThreadSafeCounter	DependencyRefCount;
 	/** Current index into linkers import table used to spread creation over several frames				*/
diff --git a/Engine/Source/Runtime/CoreUObject/Public/UObject/ObjectBase.h b/Engine/Source/Runtime/CoreUObject/Public/UObject/ObjectBase.h
--- a/Engine/Source/Runtime/CoreUObject/Public/UObject/ObjectBase.h
+++ b/Engine/Source/Runtime/CoreUObject/Public/UObject/ObjectBase.h
@@ -107,7 +107,7 @@ enum EPackageFlags
 	PKG_DisallowLazyLoading			= 0x00080000,	// Set if the archive serializing this package cannot use lazy loading
 	PKG_PlayInEditor				= 0x00100000,	// Set if the package was created for the purpose of PIE
 	PKG_ContainsScript				= 0x00200000,	// Package is allowed to contain UClass objects
-//	PKG_Unused						= 0x00400000,
+	PKG_ProcessingDependencies		= 0x00400000,
 //	PKG_Unused						= 0x00800000,
 //	PKG_Unused						= 0x01000000,	
 	PKG_StoreCompressed				= 0x02000000,	// Package is being stored compressed, requires archive support for compression
