fixed typo: EX_SwithValue -> EX_SwitchValue

diff --git a/Engine/Source/Editor/KismetCompiler/Private/KismetCompilerVMBackend.cpp b/Engine/Source/Editor/KismetCompiler/Private/KismetCompilerVMBackend.cpp
--- a/Engine/Source/Editor/KismetCompiler/Private/KismetCompilerVMBackend.cpp
+++ b/Engine/Source/Editor/KismetCompiler/Private/KismetCompilerVMBackend.cpp
@@ -1258,7 +1258,7 @@ public:
 			ensure(false);
 		}
 
-		Writer << EX_SwithValue;
+		Writer << EX_SwitchValue;
 		// number of cases (without default)
 		const int32 TermsPerCase = 2;
 		uint16 NumCases = ((Statement.RHS.Num() - 2) / TermsPerCase);
diff --git a/Engine/Source/Editor/UnrealEd/Private/ScriptDisassembler.cpp b/Engine/Source/Editor/UnrealEd/Private/ScriptDisassembler.cpp
--- a/Engine/Source/Editor/UnrealEd/Private/ScriptDisassembler.cpp
+++ b/Engine/Source/Editor/UnrealEd/Private/ScriptDisassembler.cpp
@@ -815,7 +815,7 @@ void FKismetBytecodeDisassembler::ProcessCommon(int32& ScriptIndex, EExprToken O
 			Ar.Logf(TEXT("%s $%X: .. debug site .."), *Indents, (int32)Opcode);
 			break;
 		}
-	case EX_SwithValue:
+	case EX_SwitchValue:
 		{
 			const auto NumCases = ReadWORD(ScriptIndex);
 			const auto AfterSkip = ReadSkipCount(ScriptIndex);
diff --git a/Engine/Source/Runtime/CoreUObject/Private/UObject/ScriptCore.cpp b/Engine/Source/Runtime/CoreUObject/Private/UObject/ScriptCore.cpp
--- a/Engine/Source/Runtime/CoreUObject/Private/UObject/ScriptCore.cpp
+++ b/Engine/Source/Runtime/CoreUObject/Private/UObject/ScriptCore.cpp
@@ -1463,7 +1463,7 @@ void UObject::execLetValueOnPersistentFrame(FFrame& Stack, RESULT_DECL)
 }
 IMPLEMENT_VM_FUNCTION(EX_LetValueOnPersistentFrame, execLetValueOnPersistentFrame);
 
-void UObject::execSwithValue(FFrame& Stack, RESULT_DECL)
+void UObject::execSwitchValue(FFrame& Stack, RESULT_DECL)
 {
 	const int32 NumCases = Stack.ReadWord();
 	const CodeSkipSizeType OffsetToEnd = Stack.ReadCodeSkipCount();
@@ -1517,7 +1517,7 @@ void UObject::execSwithValue(FFrame& Stack, RESULT_DECL)
 		Stack.Step(Stack.Object, RESULT_PARAM);
 	}
 }
-IMPLEMENT_VM_FUNCTION(EX_SwithValue, execSwithValue);
+IMPLEMENT_VM_FUNCTION(EX_SwitchValue, execSwitchValue);
 
 void UObject::execLet(FFrame& Stack, RESULT_DECL)
 {
diff --git a/Engine/Source/Runtime/CoreUObject/Public/UObject/Script.h b/Engine/Source/Runtime/CoreUObject/Public/UObject/Script.h
--- a/Engine/Source/Runtime/CoreUObject/Public/UObject/Script.h
+++ b/Engine/Source/Runtime/CoreUObject/Public/UObject/Script.h
@@ -240,7 +240,7 @@ enum EExprToken
 	EX_EndArrayConst		= 0x66,
 	EX_AssetConst			= 0x67,
 	EX_CallMath				= 0x68, // static pure function from on local call space
-	EX_SwithValue			= 0x69,
+	EX_SwitchValue			= 0x69,
 	EX_Max					= 0x100,
 };
 
diff --git a/Engine/Source/Runtime/CoreUObject/Public/UObject/ScriptSerialization.h b/Engine/Source/Runtime/CoreUObject/Public/UObject/ScriptSerialization.h
--- a/Engine/Source/Runtime/CoreUObject/Public/UObject/ScriptSerialization.h
+++ b/Engine/Source/Runtime/CoreUObject/Public/UObject/ScriptSerialization.h
@@ -420,7 +420,7 @@
 			SerializeExpr( iCode, Ar ); 
 			break;
 		}
-		case EX_SwithValue:
+		case EX_SwitchValue:
 		{
 			XFER(uint16); // number of cases, without default one
 			const uint16 NumCases = *(uint16*)(&Script[iCode - sizeof(uint16)]);
diff --git a/Engine/Source/Runtime/CoreUObject/Public/UObject/UObject.h b/Engine/Source/Runtime/CoreUObject/Public/UObject/UObject.h
--- a/Engine/Source/Runtime/CoreUObject/Public/UObject/UObject.h
+++ b/Engine/Source/Runtime/CoreUObject/Public/UObject/UObject.h
@@ -1144,7 +1144,7 @@ public:
 
 	DECLARE_FUNCTION(execCallMathFunction);
 
-	DECLARE_FUNCTION(execSwithValue);
+	DECLARE_FUNCTION(execSwitchValue);
 
 	// -- K2 support functions
 	struct Object_eventExecuteUbergraph_Parms
