Core: Changed "CreateExport: Failed to load Outer..." message to be emitted to the message log so it is better surfaced to the user

diff --git a/Engine/Source/Runtime/CoreUObject/Private/UObject/LinkerLoad.cpp b/Engine/Source/Runtime/CoreUObject/Private/UObject/LinkerLoad.cpp
--- a/Engine/Source/Runtime/CoreUObject/Private/UObject/LinkerLoad.cpp
+++ b/Engine/Source/Runtime/CoreUObject/Private/UObject/LinkerLoad.cpp
@@ -3520,16 +3520,22 @@ UObject* FLinkerLoad::CreateExport( int32 Index )
 			Export.bExportLoadFailed = true;
 
 			// otherwise, return NULL and let the calling code determine what to do
-			FString OuterName = Export.OuterIndex.IsNull() ? LinkerRoot->GetFullName() : GetFullImpExpName(Export.OuterIndex);
+			const FString OuterName = Export.OuterIndex.IsNull() ? LinkerRoot->GetFullName() : GetFullImpExpName(Export.OuterIndex);
+
+			FFormatNamedArguments Arguments;
+			Arguments.Add(TEXT("ObjectName"), FText::FromName(Export.ObjectName));
+			Arguments.Add(TEXT("OuterName"), FText::FromString(OuterName));
+
 			if (ParentRedirector)
 			{
-				UE_LOG(LogLinker, Warning, TEXT("CreateExport: Failed to load Outer for resource because it is a redirector '%s': %s"), *Export.ObjectName.ToString(), *OuterName);
+				LoadErrors.Warning(FText::Format(LOCTEXT("CreateExportFailedToLoadOuterIsRedirector", "CreateExport: Failed to load Outer for resource because it is a redirector '{ObjectName}': {OuterName}"), Arguments));
 			}
 			else
 			{
-				UE_LOG(LogLinker, Warning, TEXT("CreateExport: Failed to load Outer for resource '%s': %s"), *Export.ObjectName.ToString(), *OuterName);
+				LoadErrors.Warning(FText::Format(LOCTEXT("CreateExportFailedToLoadOuter", "CreateExport: Failed to load Outer for resource '{ObjectName}': {OuterName}"), Arguments));
 			}
-			return NULL;
+
+			return nullptr;
 		}
 
 		// Find the Archetype object for the one we are loading.
