ScriptCore.cpp change made it's way into the backout of changelist 2610237 and was causing clang errors.

diff --git a/Engine/Source/Runtime/CoreUObject/Private/UObject/ScriptCore.cpp b/Engine/Source/Runtime/CoreUObject/Private/UObject/ScriptCore.cpp
--- a/Engine/Source/Runtime/CoreUObject/Private/UObject/ScriptCore.cpp
+++ b/Engine/Source/Runtime/CoreUObject/Private/UObject/ScriptCore.cpp
@@ -3,7 +3,6 @@
 /*=============================================================================
 	ScriptCore.cpp: Kismet VM execution and support code.
 =============================================================================*/
-#pragma optimize("", off)
 #include "CoreUObjectPrivate.h"
 #include "MallocProfiler.h"
 #include "HotReloadInterface.h"
@@ -2151,7 +2150,7 @@ IMPLEMENT_VM_FUNCTION( EX_IntConstByte, execIntConstByte );
 void UObject::execDynamicCast( FFrame& Stack, RESULT_DECL )
 {
 	// Get "to cast to" class for the dynamic actor class
-	UClass* Class = (UClass *)Stack.ReadObject();
+	UClass* ClassPtr = (UClass *)Stack.ReadObject();
 
 	// Compile object expression.
 	UObject* Castee = NULL;
@@ -2159,12 +2158,12 @@ void UObject::execDynamicCast( FFrame& Stack, RESULT_DECL )
 	//*(UObject**)RESULT_PARAM = (Castee && Castee->IsA(Class)) ? Castee : NULL;
 	*(UObject**)RESULT_PARAM = NULL; // default value
 
-	if (Class)
+	if (ClassPtr)
 	{
 		// if we were passed in a null value
 		if( Castee == NULL )
 		{
-			if( Class->HasAnyClassFlags(CLASS_Interface) )
+			if(ClassPtr->HasAnyClassFlags(CLASS_Interface) )
 			{
 				((FScriptInterface*)RESULT_PARAM)->SetObject(NULL);
 			}
@@ -2177,17 +2176,17 @@ void UObject::execDynamicCast( FFrame& Stack, RESULT_DECL )
 
 		// check to see if the Castee is an implemented interface by looking up the
 		// class hierarchy and seeing if any class in said hierarchy implements the interface
-		if( Class->HasAnyClassFlags(CLASS_Interface) )
+		if(ClassPtr->HasAnyClassFlags(CLASS_Interface) )
 		{
-			if ( Castee->GetClass()->ImplementsInterface(Class) )
+			if ( Castee->GetClass()->ImplementsInterface(ClassPtr) )
 			{
 				// interface property type - convert to FScriptInterface
 				((FScriptInterface*)RESULT_PARAM)->SetObject(Castee);
-				((FScriptInterface*)RESULT_PARAM)->SetInterface(Castee->GetInterfaceAddress(Class));
+				((FScriptInterface*)RESULT_PARAM)->SetInterface(Castee->GetInterfaceAddress(ClassPtr));
 			}
 		}
 		// check to see if the Castee is a castable class
-		else if( Castee->IsA(Class) )
+		else if( Castee->IsA(ClassPtr) )
 		{
 			*(UObject**)RESULT_PARAM = Castee;
 		}
