Removed "#pragma optimize("", off)" from ScriptCore.cpp.
Clang doesn't like it.

diff --git a/Engine/Source/Runtime/CoreUObject/Private/UObject/ScriptCore.cpp b/Engine/Source/Runtime/CoreUObject/Private/UObject/ScriptCore.cpp
--- a/Engine/Source/Runtime/CoreUObject/Private/UObject/ScriptCore.cpp
+++ b/Engine/Source/Runtime/CoreUObject/Private/UObject/ScriptCore.cpp
@@ -3,7 +3,7 @@
 /*=============================================================================
 	ScriptCore.cpp: Kismet VM execution and support code.
 =============================================================================*/
-#pragma optimize("", off)
+
 #include "CoreUObjectPrivate.h"
 #include "MallocProfiler.h"
 #include "HotReloadInterface.h"
