VS2015 - Fixing a host of compiler issues when using VS2015 toolset

Lots of shadow variable issues in header files
Type narrowing mismatches between integer types.
Fixed some variable names which conflicted with private members, this works around a bug in VC14

Commiting some of Jeff's code changes from 2605975

diff --git a/Engine/Source/Runtime/Core/Private/HAL/ThreadingBase.cpp b/Engine/Source/Runtime/Core/Private/HAL/ThreadingBase.cpp
--- a/Engine/Source/Runtime/Core/Private/HAL/ThreadingBase.cpp
+++ b/Engine/Source/Runtime/Core/Private/HAL/ThreadingBase.cpp
@@ -129,7 +129,7 @@ public:
 		FSingleThreadManager::Get().RemoveThread(this);
 	}
 
-	virtual bool CreateInternal(FRunnable* InRunnable, const TCHAR* ThreadName,
+	virtual bool CreateInternal(FRunnable* InRunnable, const TCHAR* InThreadName,
 		uint32 InStackSize,
 		EThreadPriority InThreadPri, uint64 InThreadAffinityMask) override
 
diff --git a/Engine/Source/Runtime/Core/Private/Windows/WindowsApplication.cpp b/Engine/Source/Runtime/Core/Private/Windows/WindowsApplication.cpp
--- a/Engine/Source/Runtime/Core/Private/Windows/WindowsApplication.cpp
+++ b/Engine/Source/Runtime/Core/Private/Windows/WindowsApplication.cpp
@@ -1896,15 +1896,15 @@ HRESULT FWindowsApplication::OnOLEDrop( const HWND HWnd, const FDragDropOLEData&
 }
 
 
-void FWindowsApplication::AddMessageHandler(IWindowsMessageHandler& MessageHandler)
+void FWindowsApplication::AddMessageHandler(IWindowsMessageHandler& InMessageHandler)
 {
-	WindowsApplication->MessageHandlers.AddUnique(&MessageHandler);
+	WindowsApplication->MessageHandlers.AddUnique(&InMessageHandler);
 }
 
 
-void FWindowsApplication::RemoveMessageHandler(IWindowsMessageHandler& MessageHandler)
+void FWindowsApplication::RemoveMessageHandler(IWindowsMessageHandler& InMessageHandler)
 {
-	WindowsApplication->MessageHandlers.RemoveSwap(&MessageHandler);
+	WindowsApplication->MessageHandlers.RemoveSwap(&InMessageHandler);
 }
 
 
diff --git a/Engine/Source/Runtime/Core/Private/Windows/WindowsPlatformProcess.cpp b/Engine/Source/Runtime/Core/Private/Windows/WindowsPlatformProcess.cpp
--- a/Engine/Source/Runtime/Core/Private/Windows/WindowsPlatformProcess.cpp
+++ b/Engine/Source/Runtime/Core/Private/Windows/WindowsPlatformProcess.cpp
@@ -333,12 +333,12 @@ FProcHandle FWindowsPlatformProcess::CreateProc( const TCHAR* URL, const TCHAR*
 	STARTUPINFO StartupInfo = {
 		sizeof(STARTUPINFO),
 		NULL, NULL, NULL,
-		CW_USEDEFAULT,
-		CW_USEDEFAULT,
-		CW_USEDEFAULT,
-		CW_USEDEFAULT,
-		0, 0, 0,
-		dwFlags,
+		(::DWORD)CW_USEDEFAULT,
+		(::DWORD)CW_USEDEFAULT,
+		(::DWORD)CW_USEDEFAULT,
+		(::DWORD)CW_USEDEFAULT,
+		(::DWORD)0, (::DWORD)0, (::DWORD)0,
+		(::DWORD)dwFlags,
 		ShowWindowFlags,
 		0, NULL,
 		::GetStdHandle(ProcessConstants::WIN_STD_INPUT_HANDLE),
@@ -350,7 +350,7 @@ FProcHandle FWindowsPlatformProcess::CreateProc( const TCHAR* URL, const TCHAR*
 	FString CommandLine = FString::Printf(TEXT("\"%s\" %s"), URL, Parms);
 	PROCESS_INFORMATION ProcInfo;
 
-	if (!CreateProcess(NULL, CommandLine.GetCharArray().GetData(), &Attr, &Attr, true, CreateFlags, NULL, OptionalWorkingDirectory, &StartupInfo, &ProcInfo))
+	if (!CreateProcess(NULL, CommandLine.GetCharArray().GetData(), &Attr, &Attr, true, (::DWORD)CreateFlags, NULL, OptionalWorkingDirectory, &StartupInfo, &ProcInfo))
 	{
 		UE_LOG(LogWindows, Warning, TEXT("CreateProc failed (%u) %s %s"), ::GetLastError(), URL, Parms);
 		if (OutProcessID != nullptr)
@@ -606,9 +606,9 @@ bool FWindowsPlatformProcess::ExecProcess( const TCHAR* URL, const TCHAR* Params
 
 	bool bSuccess = false;
 	STARTUPINFO StartupInfo = { sizeof(STARTUPINFO), NULL, NULL, NULL,
-		CW_USEDEFAULT, CW_USEDEFAULT, CW_USEDEFAULT, CW_USEDEFAULT,
-		0, 0, 0, dwFlags, ShowWindowFlags, 0, NULL,
-		::GetStdHandle(ProcessConstants::WIN_STD_INPUT_HANDLE), WritablePipes[0], WritablePipes[1] };
+		(::DWORD)CW_USEDEFAULT, (::DWORD)CW_USEDEFAULT, (::DWORD)CW_USEDEFAULT, (::DWORD)CW_USEDEFAULT,
+		(::DWORD)0, (::DWORD)0, (::DWORD)0, (::DWORD)dwFlags, ShowWindowFlags, 0, NULL,
+		::GetStdHandle((::DWORD)ProcessConstants::WIN_STD_INPUT_HANDLE), WritablePipes[0], WritablePipes[1] };
 	if (CreateProcess(NULL, CommandLine.GetCharArray().GetData(), &Attr, &Attr, true, CreateFlags,
 		NULL, NULL, &StartupInfo, &ProcInfo))
 	{
diff --git a/Engine/Source/Runtime/Core/Public/Delegates/DelegateSignatureImpl_Variadics.inl b/Engine/Source/Runtime/Core/Public/Delegates/DelegateSignatureImpl_Variadics.inl
--- a/Engine/Source/Runtime/Core/Public/Delegates/DelegateSignatureImpl_Variadics.inl
+++ b/Engine/Source/Runtime/Core/Public/Delegates/DelegateSignatureImpl_Variadics.inl
@@ -431,13 +431,13 @@ public:
 	 */
 	inline RetValType Execute(ParamTypes... Params) const
 	{
-		TDelegateInstanceInterface* DelegateInstance = (TDelegateInstanceInterface*)GetDelegateInstance();
+		TDelegateInstanceInterface* LocalDelegateInstance = (TDelegateInstanceInterface*)GetDelegateInstance();
 
 		// If this assert goes off, Execute() was called before a function was bound to the delegate.
 		// Consider using ExecuteIfSafe() instead.
-		checkSlow(DelegateInstance != nullptr);
+		checkSlow(LocalDelegateInstance != nullptr);
 
-		return DelegateInstance->Execute(Params...);
+		return LocalDelegateInstance->Execute(Params...);
 	}
 
 	/**
@@ -458,17 +458,17 @@ public:
 	 */
 	bool DEPRECATED_Compare( const TBaseDelegate& Other ) const
 	{
-		TDelegateInstanceInterface* DelegateInstance = (TDelegateInstanceInterface*)GetDelegateInstance();
+		TDelegateInstanceInterface* ThisDelegateInstance = (TDelegateInstanceInterface*)GetDelegateInstance();
 		TDelegateInstanceInterface* OtherInstance = (TDelegateInstanceInterface*)Other.GetDelegateInstance();
 		
 		// The function these delegates point to must be the same
-		if ((DelegateInstance != nullptr) && (OtherInstance != nullptr))
+		if ((ThisDelegateInstance != nullptr) && (OtherInstance != nullptr))
 		{
-			return DelegateInstance->IsSameFunction(*OtherInstance);
+			return ThisDelegateInstance->IsSameFunction(*OtherInstance);
 		}
 		
 		// If neither delegate is initialized to anything yet, then we treat them as equal
-		if ((DelegateInstance == nullptr) && (OtherInstance == nullptr))
+		if ((ThisDelegateInstance == nullptr) && (OtherInstance == nullptr))
 		{
 			return true;
 		}
@@ -1032,13 +1032,13 @@ public:
 
 		LockInvocationList();
 		{
-			const TArray<IDelegateInstance*>& InvocationList = GetInvocationList();
+			const TArray<IDelegateInstance*>& LocalInvocationList = GetInvocationList();
 
 			// call bound functions in reverse order, so we ignore any instances that may be added by callees
-			for (int32 InvocationListIndex = InvocationList.Num() - 1; InvocationListIndex >= 0; --InvocationListIndex)
+			for (int32 InvocationListIndex = LocalInvocationList.Num() - 1; InvocationListIndex >= 0; --InvocationListIndex)
 			{
 				// this down-cast is OK! allows for managing invocation list in the base class without requiring virtual functions
-				TDelegateInstanceInterface* DelegateInstanceInterface = (TDelegateInstanceInterface*)InvocationList[InvocationListIndex];
+				TDelegateInstanceInterface* DelegateInstanceInterface = (TDelegateInstanceInterface*)LocalInvocationList[InvocationListIndex];
 
 				if ((DelegateInstanceInterface == nullptr) || !DelegateInstanceInterface->ExecuteIfSafe(Params...))
 				{
@@ -1065,13 +1065,13 @@ protected:
 	 */
 	void RemoveDelegateInstance( FDelegateHandle Handle )
 	{
-		const TArray<IDelegateInstance*>& InvocationList = GetInvocationList();
+		const TArray<IDelegateInstance*>& LocalInvocationList = GetInvocationList();
 
-		for (int32 InvocationListIndex = 0; InvocationListIndex < InvocationList.Num(); ++InvocationListIndex)
+		for (int32 InvocationListIndex = 0; InvocationListIndex < LocalInvocationList.Num(); ++InvocationListIndex)
 		{
 			// InvocationList is const, so we const_cast to be able to null the entry
 			// TODO: This is horrible, can we get the base class to do it?
-			IDelegateInstance*& DelegateInstanceRef = const_cast<IDelegateInstance*&>(InvocationList[InvocationListIndex]);
+			IDelegateInstance*& DelegateInstanceRef = const_cast<IDelegateInstance*&>(LocalInvocationList[InvocationListIndex]);
 
 			if ((DelegateInstanceRef != nullptr) && DelegateInstanceRef->GetHandle() == Handle)
 			{
@@ -1095,16 +1095,16 @@ protected:
 	 */
 	void DEPRECATED_RemoveDelegateInstance( const TDelegateInstanceInterface* InDelegateInstance )
 	{
-		const TArray<IDelegateInstance*>& InvocationList = GetInvocationList();
+		const TArray<IDelegateInstance*>& LocalInvocationList = GetInvocationList();
 
 		// NOTE: We assume that this method is never called with a nullptr object, in which case the
 		//       the following algorithm would break down (it would remove the first found instance
 		//       of a matching function binding, which is not necessarily the instance we wish to remove).
 
-		for (int32 InvocationListIndex = 0; InvocationListIndex < InvocationList.Num(); ++InvocationListIndex)
+		for (int32 InvocationListIndex = 0; InvocationListIndex < LocalInvocationList.Num(); ++InvocationListIndex)
 		{
 			// this down-cast is OK! allows for managing invocation list in the base class without requiring virtual functions
-			TDelegateInstanceInterface*& DelegateInstanceRef = (TDelegateInstanceInterface*&)InvocationList[InvocationListIndex];
+			TDelegateInstanceInterface*& DelegateInstanceRef = (TDelegateInstanceInterface*&)LocalInvocationList[InvocationListIndex];
 
 			// NOTE: We must do a deep compare here, not just compare delegate pointers, because multiple
 			//       delegate pointers can refer to the exact same object and method
diff --git a/Engine/Source/Runtime/Core/Public/HAL/IPlatformFileProfilerWrapper.h b/Engine/Source/Runtime/Core/Public/HAL/IPlatformFileProfilerWrapper.h
--- a/Engine/Source/Runtime/Core/Public/HAL/IPlatformFileProfilerWrapper.h
+++ b/Engine/Source/Runtime/Core/Public/HAL/IPlatformFileProfilerWrapper.h
@@ -333,7 +333,7 @@ public:
 	{
 		StatsType* FileStat = CreateStat( Filename );
 		FProfiledFileStatsOp* OpStat = FileStat->CreateOpStat( FProfiledFileStatsOp::EOpType::GetTimeStamp );
-		double StartTime = FPlatformTime::Seconds();
+		double OpStartTime = FPlatformTime::Seconds();
 		FDateTime Result = LowerLevel->GetTimeStamp(Filename);
 		OpStat->Duration += FPlatformTime::Seconds() * 1000 - OpStat->LastOpTime;
 		return Result;
@@ -342,7 +342,7 @@ public:
 	{
 		StatsType* FileStat = CreateStat( Filename );
 		FProfiledFileStatsOp* OpStat = FileStat->CreateOpStat( FProfiledFileStatsOp::EOpType::SetTimeStamp );
-		double StartTime = FPlatformTime::Seconds();
+		double OpStartTime = FPlatformTime::Seconds();
 		LowerLevel->SetTimeStamp(Filename, DateTime);
 		OpStat->Duration += FPlatformTime::Seconds() * 1000 - OpStat->LastOpTime;
 	}
@@ -350,7 +350,7 @@ public:
 	{
 		StatsType* FileStat = CreateStat( Filename );
 		FProfiledFileStatsOp* OpStat = FileStat->CreateOpStat( FProfiledFileStatsOp::EOpType::GetTimeStamp );
-		double StartTime = FPlatformTime::Seconds();
+		double OpStartTime = FPlatformTime::Seconds();
 		FDateTime Result = LowerLevel->GetAccessTimeStamp(Filename);
 		OpStat->Duration += FPlatformTime::Seconds() * 1000 - OpStat->LastOpTime;
 		return Result;
@@ -359,7 +359,7 @@ public:
 	{
 		StatsType* FileStat = CreateStat(Filename);
 		FProfiledFileStatsOp* OpStat = FileStat->CreateOpStat( FProfiledFileStatsOp::EOpType::GetFilenameOnDisk );
-		double StartTime = FPlatformTime::Seconds();
+		double OpStartTime = FPlatformTime::Seconds();
 		FString Result = LowerLevel->GetFilenameOnDisk(Filename);
 		OpStat->Duration += FPlatformTime::Seconds() * 1000 - OpStat->LastOpTime;
 		return Result;
diff --git a/Engine/Source/Runtime/Core/Public/HAL/ThreadingBase.h b/Engine/Source/Runtime/Core/Public/HAL/ThreadingBase.h
--- a/Engine/Source/Runtime/Core/Public/HAL/ThreadingBase.h
+++ b/Engine/Source/Runtime/Core/Public/HAL/ThreadingBase.h
@@ -378,7 +378,7 @@ protected:
 	 * @param InThreadPri Tells the thread whether it needs to adjust its priority or not. Defaults to normal priority
 	 * @return True if the thread and all of its initialization was successful, false otherwise
 	 */
-	virtual bool CreateInternal( FRunnable* InRunnable, const TCHAR* ThreadName,
+	virtual bool CreateInternal( FRunnable* InRunnable, const TCHAR* InThreadName,
 		uint32 InStackSize = 0,
 		EThreadPriority InThreadPri = TPri_Normal, uint64 InThreadAffinityMask = 0 ) = 0;
 
diff --git a/Engine/Source/Runtime/Core/Public/HTML5/HTML5PlatformRunnableThread.h b/Engine/Source/Runtime/Core/Public/HTML5/HTML5PlatformRunnableThread.h
--- a/Engine/Source/Runtime/Core/Public/HTML5/HTML5PlatformRunnableThread.h
+++ b/Engine/Source/Runtime/Core/Public/HTML5/HTML5PlatformRunnableThread.h
@@ -47,7 +47,7 @@ public:
 
 protected:
 
-	virtual bool CreateInternal (FRunnable* InRunnable, const TCHAR* ThreadName,
+	virtual bool CreateInternal (FRunnable* InRunnable, const TCHAR* InThreadName,
 		uint32 InStackSize = 0,
 		EThreadPriority InThreadPri = TPri_Normal, uint64 InThreadAffinityMask = 0) 
 	{
diff --git a/Engine/Source/Runtime/Core/Public/Math/Transform.h b/Engine/Source/Runtime/Core/Public/Math/Transform.h
--- a/Engine/Source/Runtime/Core/Public/Math/Transform.h
+++ b/Engine/Source/Runtime/Core/Public/Math/Transform.h
@@ -489,7 +489,7 @@ public:
 	*/
 	FORCEINLINE void operator*=(const FQuat& Other);
 
-	FORCEINLINE void ScaleTranslation(const FVector& Scale3D);
+	FORCEINLINE void ScaleTranslation(const FVector& InScale3D);
 	FORCEINLINE void ScaleTranslation(const float& Scale);
 	FORCEINLINE void RemoveScaling(float Tolerance=SMALL_NUMBER);
 	FORCEINLINE float GetMaximumAxisScale() const;
diff --git a/Engine/Source/Runtime/Core/Public/Math/TransformVectorized.h b/Engine/Source/Runtime/Core/Public/Math/TransformVectorized.h
--- a/Engine/Source/Runtime/Core/Public/Math/TransformVectorized.h
+++ b/Engine/Source/Runtime/Core/Public/Math/TransformVectorized.h
@@ -488,7 +488,7 @@ public:
 	FORCEINLINE FTransform		operator*(const FQuat& Other) const;
 	FORCEINLINE void			operator*=(const FQuat& Other);
 
-	FORCEINLINE void ScaleTranslation(const FVector& Scale3D);
+	FORCEINLINE void ScaleTranslation(const FVector& InScale3D);
 	FORCEINLINE void ScaleTranslation(const float& Scale);
 	FORCEINLINE void RemoveScaling(float Tolerance=SMALL_NUMBER);
 	FORCEINLINE float GetMaximumAxisScale() const;
diff --git a/Engine/Source/Runtime/Core/Public/Math/UnrealMatrix.h b/Engine/Source/Runtime/Core/Public/Math/UnrealMatrix.h
--- a/Engine/Source/Runtime/Core/Public/Math/UnrealMatrix.h
+++ b/Engine/Source/Runtime/Core/Public/Math/UnrealMatrix.h
@@ -502,11 +502,11 @@ inline float FMatrix::GetMaximumAxisScale() const
 	return FMath::Sqrt(MaxRowScaleSquared);
 }
 
-inline void FMatrix::ScaleTranslation(const FVector& Scale3D)
+inline void FMatrix::ScaleTranslation(const FVector& InScale3D)
 {
-	M[3][0] *= Scale3D.X;
-	M[3][1] *= Scale3D.Y;
-	M[3][2] *= Scale3D.Z;
+	M[3][0] *= InScale3D.X;
+	M[3][1] *= InScale3D.Y;
+	M[3][2] *= InScale3D.Z;
 }
 
 // GetOrigin
diff --git a/Engine/Source/Runtime/Core/Public/Misc/TypeContainer.h b/Engine/Source/Runtime/Core/Public/Misc/TypeContainer.h
--- a/Engine/Source/Runtime/Core/Public/Misc/TypeContainer.h
+++ b/Engine/Source/Runtime/Core/Public/Misc/TypeContainer.h
@@ -325,11 +325,11 @@ public:
 		AddProvider(TNameOf<R>::GetName(), Provider);
 	}
 
-#ifdef __clang__
+#if _MSC_VER >= 1900 || __clang__
 	/**
 	 * Register a factory function for the specified class.
 	 *
-	 * This is a Clang specific overload for handling raw function pointers.
+	 * This is a Clang/VS2015-specific overload for handling raw function pointers.
 	 *
 	 * @param R The type of class to register the instance for.
 	 * @param P Additional parameters that the factory function requires.
diff --git a/Engine/Source/Runtime/Core/Public/Windows/AllowWindowsPlatformTypes.h b/Engine/Source/Runtime/Core/Public/Windows/AllowWindowsPlatformTypes.h
--- a/Engine/Source/Runtime/Core/Public/Windows/AllowWindowsPlatformTypes.h
+++ b/Engine/Source/Runtime/Core/Public/Windows/AllowWindowsPlatformTypes.h
@@ -6,6 +6,9 @@
 	#error Nesting AllowWindowsPLatformTypes.h is not allowed!
 #endif
 
+#pragma warning( push )
+#pragma warning( disable : 4459 )
+
 #define INT ::INT
 #define UINT ::UINT
 #define DWORD ::DWORD
diff --git a/Engine/Source/Runtime/Core/Public/Windows/HideWindowsPlatformTypes.h b/Engine/Source/Runtime/Core/Public/Windows/HideWindowsPlatformTypes.h
--- a/Engine/Source/Runtime/Core/Public/Windows/HideWindowsPlatformTypes.h
+++ b/Engine/Source/Runtime/Core/Public/Windows/HideWindowsPlatformTypes.h
@@ -18,3 +18,5 @@
 #ifdef FALSE
 	#undef FALSE
 #endif
+
+#pragma warning( pop )
diff --git a/Engine/Source/Runtime/Core/Public/Windows/WindowsApplication.h b/Engine/Source/Runtime/Core/Public/Windows/WindowsApplication.h
--- a/Engine/Source/Runtime/Core/Public/Windows/WindowsApplication.h
+++ b/Engine/Source/Runtime/Core/Public/Windows/WindowsApplication.h
@@ -315,7 +315,7 @@ public:
 	 * @param MessageHandler The message handler to register.
 	 * @see RemoveMessageHandler
 	 */
-	virtual void AddMessageHandler(IWindowsMessageHandler& MessageHandler);
+	virtual void AddMessageHandler(IWindowsMessageHandler& InMessageHandler);
 
 	/**
 	 * Removes a Windows message handler with the application instance.
@@ -323,7 +323,7 @@ public:
 	 * @param MessageHandler The message handler to register.
 	 * @see AddMessageHandler
 	 */
-	virtual void RemoveMessageHandler(IWindowsMessageHandler& MessageHandler);
+	virtual void RemoveMessageHandler(IWindowsMessageHandler& InMessageHandler);
 
 public:
 
diff --git a/Engine/Source/Runtime/Core/Public/Windows/WindowsPlatformCompilerPreSetup.h b/Engine/Source/Runtime/Core/Public/Windows/WindowsPlatformCompilerPreSetup.h
--- a/Engine/Source/Runtime/Core/Public/Windows/WindowsPlatformCompilerPreSetup.h
+++ b/Engine/Source/Runtime/Core/Public/Windows/WindowsPlatformCompilerPreSetup.h
@@ -71,9 +71,17 @@
 	#endif // PRAGMA_POP
 #else
 	// VC++
+	// 4456 - declaration of 'LocalVariable' hides previous local declaration
+	// 4457 - declaration of 'LocalVariable' hides function parameter
+	// 4458 - declaration of 'LocalVariable' hides class member
+	// 4459 - declaration of 'LocalVariable' hides global declaration
 	#ifndef PRAGMA_DISABLE_SHADOW_VARIABLE_WARNINGS
 		#define PRAGMA_DISABLE_SHADOW_VARIABLE_WARNINGS \
-			__pragma (warning(push))
+			__pragma (warning(push)) \
+			__pragma (warning(disable:4456)) \
+			__pragma (warning(disable:4457)) \
+			__pragma (warning(disable:4458)) \
+			__pragma (warning(disable:4459))
 	#endif // PRAGMA_DISABLE_SHADOW_VARIABLE_WARNINGS
 
 	#ifndef PRAGMA_ENABLE_SHADOW_VARIABLE_WARNINGS
diff --git a/Engine/Source/Runtime/Core/Public/Windows/WindowsPlatformCompilerSetup.h b/Engine/Source/Runtime/Core/Public/Windows/WindowsPlatformCompilerSetup.h
--- a/Engine/Source/Runtime/Core/Public/Windows/WindowsPlatformCompilerSetup.h
+++ b/Engine/Source/Runtime/Core/Public/Windows/WindowsPlatformCompilerSetup.h
@@ -9,15 +9,15 @@
 /**
  * Future-proofing the min version check so we keep bumping it whenever we upgrade.
  */ 
-#if _MSC_VER > 1800
+#if _MSC_VER > 1900
 	#pragma message("Detected compiler newer than Visual Studio 2013, please update min version checking in WindowsPlatformCompilerSetup.h")
 #endif
 
 
 /**
- * We require at least Visual Studio 2012 to compile
+ * We require at least Visual Studio 2013 to compile
  */
-static_assert(_MSC_VER >= 1700, "Visual Studio 2012 or newer required.");
+static_assert(_MSC_VER >= 1800, "Visual Studio 2013 or later is required to compile on Windows platform");
 
 #pragma warning (error:      4001 4002 4003 4004      4006 4007 4008 4009 4010 4011 4012 4013 4014 4015 4016 4017 4018 4019 4020 4021 4022 4023 4024 4025 4026 4027 4028 4029 4030 4031 4032 4033 4034 4035 4036 4037 4038 4039 4040 4041 4042 4043 4044 4045 4046 4047 4048 4049 4050 4051 4052 4053 4054 4055 4056 4057 4058 4059 4060           4063 4064 4065 4066 4067 4068 4069 4070 4071 4072 4073 4074 4075 4076 4077 4078 4079 4080 4081 4082 4083 4084 4085 4086 4087 4088 4089 4090 4091 4092 4093 4094 4095 4096 4097 4098 4099)
 #pragma warning (error: 4100 4101 4102 4103 4104 4105 4106 4107 4108 4109 4110 4111 4112 4113 4114 4115 4116 4117 4118 4119 4120 4121 4122 4123 4124 4125 4126 4127 4128 4129 4130 4131 4132 4133 4134 4135 4136 4137 4138 4139 4140 4141 4142 4143 4144 4145 4146 4147 4148 4149 4150 4151 4152 4153 4154 4155 4156 4157 4158 4159 4160 4161 4162 4163 4164      4166 4167 4168 4169 4170 4171 4172 4173 4174 4175 4176 4177 4178 4179 4180 4181 4182 4183 4184 4185 4186 4187 4188 4189 4190 4191 4192 4193 4194 4195 4196 4197 4198 4199)
@@ -85,6 +85,9 @@ static_assert(_MSC_VER >= 1700, "Visual Studio 2012 or newer required.");
 // 4988 - 'symbol': variable declared outside class/function scope
 #pragma warning (default : 4996)
 
+// @todo UWP: Disabled because DbgHelp.h has some anonymous typedefs in it (not allowed in Visual Studio 2015).  We should probably just wrap that header.
+#pragma warning(disable : 4091) // 'typedef ': ignored on left of 'type declaration' when no variable is declared
+
 // Unwanted VC++ level 4 warnings to disable.
 #pragma warning(disable : 4100) // unreferenced formal parameter										
 #pragma warning(disable : 4127) // Conditional expression is constant									
diff --git a/Engine/Source/Runtime/CoreUObject/Private/UObject/Class.cpp b/Engine/Source/Runtime/CoreUObject/Private/UObject/Class.cpp
--- a/Engine/Source/Runtime/CoreUObject/Private/UObject/Class.cpp
+++ b/Engine/Source/Runtime/CoreUObject/Private/UObject/Class.cpp
@@ -20,6 +20,12 @@ DECLARE_LOG_CATEGORY_EXTERN(LogScriptSerialization, Log, All);
 DEFINE_LOG_CATEGORY(LogScriptSerialization);
 DEFINE_LOG_CATEGORY(LogClass);
 
+#if _MSC_VER == 1900
+	#ifdef PRAGMA_DISABLE_SHADOW_VARIABLE_WARNINGS
+		PRAGMA_DISABLE_SHADOW_VARIABLE_WARNINGS
+	#endif
+#endif
+
 //////////////////////////////////////////////////////////////////////////
 
 /**
@@ -4461,4 +4467,11 @@ UDelegateFunction::UDelegateFunction(UFunction* InSuperFunction, uint32 InFuncti
 IMPLEMENT_CORE_INTRINSIC_CLASS(UDelegateFunction, UFunction,
 	{
 	}
-);
\ No newline at end of file
+);
+
+
+#if _MSC_VER == 1900
+	#ifdef PRAGMA_ENABLE_SHADOW_VARIABLE_WARNINGS
+		PRAGMA_ENABLE_SHADOW_VARIABLE_WARNINGS
+	#endif
+#endif
\ No newline at end of file
diff --git a/Engine/Source/Runtime/CoreUObject/Private/UObject/CoreNet.cpp b/Engine/Source/Runtime/CoreUObject/Private/UObject/CoreNet.cpp
--- a/Engine/Source/Runtime/CoreUObject/Private/UObject/CoreNet.cpp
+++ b/Engine/Source/Runtime/CoreUObject/Private/UObject/CoreNet.cpp
@@ -245,7 +245,7 @@ void FClassNetCacheMgr::ClearClassNetCache()
 	UPackageMap implementation.
 -----------------------------------------------------------------------------*/
 
-bool UPackageMap::SerializeName(FArchive& Ar, FName& Name)
+bool UPackageMap::SerializeName(FArchive& Ar, FName& InName)
 {
 	if (Ar.IsLoading())
 	{
@@ -256,7 +256,7 @@ bool UPackageMap::SerializeName(FArchive& Ar, FName& Name)
 			// replicated by hardcoded index
 			uint32 NameIndex;
 			Ar.SerializeInt(NameIndex, MAX_NETWORKED_HARDCODED_NAME + 1);
-			Name = EName(NameIndex);
+			InName = EName(NameIndex);
 			// hardcoded names never have a Number
 		}
 		else
@@ -265,25 +265,25 @@ bool UPackageMap::SerializeName(FArchive& Ar, FName& Name)
 			FString InString;
 			int32 InNumber;
 			Ar << InString << InNumber;
-			Name = FName(*InString, InNumber);
+			InName = FName(*InString, InNumber);
 		}
 	}
 	else if (Ar.IsSaving())
 	{
-		uint8 bHardcoded = Name.GetComparisonIndex() <= MAX_NETWORKED_HARDCODED_NAME;
+		uint8 bHardcoded = InName.GetComparisonIndex() <= MAX_NETWORKED_HARDCODED_NAME;
 		Ar.SerializeBits(&bHardcoded, 1);
 		if (bHardcoded)
 		{
 			// send by hardcoded index
-			checkSlow(Name.GetNumber() <= 0); // hardcoded names should never have a Number
-			uint32 NameIndex = uint32(Name.GetComparisonIndex());
+			checkSlow(InName.GetNumber() <= 0); // hardcoded names should never have a Number
+			uint32 NameIndex = uint32(InName.GetComparisonIndex());
 			Ar.SerializeInt(NameIndex, MAX_NETWORKED_HARDCODED_NAME + 1);
 		}
 		else
 		{
 			// send by string
-			FString OutString = Name.GetPlainNameString();
-			int32 OutNumber = Name.GetNumber();
+			FString OutString = InName.GetPlainNameString();
+			int32 OutNumber = InName.GetNumber();
 			Ar << OutString << OutNumber;
 		}
 	}
diff --git a/Engine/Source/Runtime/CoreUObject/Private/UObject/Obj.cpp b/Engine/Source/Runtime/CoreUObject/Private/UObject/Obj.cpp
--- a/Engine/Source/Runtime/CoreUObject/Private/UObject/Obj.cpp
+++ b/Engine/Source/Runtime/CoreUObject/Private/UObject/Obj.cpp
@@ -733,17 +733,17 @@ void UObject::PostLoadSubobjects( FObjectInstancingGraph* OuterInstanceGraph/*=N
 {
 	if( GetClass()->HasAnyClassFlags(CLASS_HasInstancedReference) )
 	{
-		UObject* Outer = GetOuter();
+		UObject* ObjOuter = GetOuter();
 		// make sure our Outer has already called ConditionalPostLoadSubobjects
-		if ( Outer != NULL && Outer->HasAnyFlags(RF_NeedPostLoadSubobjects) )
+		if (ObjOuter != NULL && ObjOuter->HasAnyFlags(RF_NeedPostLoadSubobjects) )
 		{
-			if ( Outer->HasAnyFlags(RF_NeedPostLoad) )
+			if (ObjOuter->HasAnyFlags(RF_NeedPostLoad) )
 			{
-				Outer->ConditionalPostLoad();
+				ObjOuter->ConditionalPostLoad();
 			}
 			else
 			{
-				Outer->ConditionalPostLoadSubobjects();
+				ObjOuter->ConditionalPostLoadSubobjects();
 			}
 			if ( !HasAnyFlags(RF_NeedPostLoadSubobjects) )
 			{
@@ -857,23 +857,23 @@ bool UObject::IsSelected() const
 void UObject::Serialize( FArchive& Ar )
 {
 	// These three items are very special items from a serialization standpoint. They aren't actually serialized.
-	UClass *Class = GetClass();
+	UClass *ObjClass = GetClass();
 	UObject* LoadOuter = GetOuter();
 	FName LoadName = GetFName();
 
 	// Make sure this object's class's data is loaded.
-	if( Class->HasAnyFlags(RF_NeedLoad) )
+	if(ObjClass->HasAnyFlags(RF_NeedLoad) )
 	{
-		Ar.Preload( Class );
+		Ar.Preload(ObjClass);
 
 		// make sure this object's template data is loaded - the only objects
 		// this should actually affect are those that don't have any defaults
 		// to serialize.  for objects with defaults that actually require loading
 		// the class default object should be serialized in FLinkerLoad::Preload, before
 		// we've hit this code.
-		if ( !HasAnyFlags(RF_ClassDefaultObject) && Class->GetDefaultsCount() > 0 )
+		if ( !HasAnyFlags(RF_ClassDefaultObject) && ObjClass->GetDefaultsCount() > 0 )
 		{
-			Ar.Preload(Class->GetDefaultObject());
+			Ar.Preload(ObjClass->GetDefaultObject());
 		}
 	}
 
@@ -887,7 +887,7 @@ void UObject::Serialize( FArchive& Ar )
 		}
 		if ( !Ar.IsIgnoringClassRef() )
 		{
-			Ar << Class;
+			Ar << ObjClass;
 		}
 		//@todo UE4 - This seems to be required and it should not be. Seems to be related to the texture streamer.
 		FLinkerLoad* LinkerLoad = GetLinker();
@@ -923,7 +923,7 @@ void UObject::Serialize( FArchive& Ar )
 
 	// Serialize object properties which are defined in the class.
 	// Handle derived UClass objects (exact UClass objects are native only and shouldn't be touched)
-	if (Class != UClass::StaticClass())
+	if (ObjClass != UClass::StaticClass())
 	{
 		SerializeScriptProperties(Ar);
 	}
@@ -974,7 +974,7 @@ void UObject::SerializeScriptProperties( FArchive& Ar ) const
 		Ar.StartSerializingDefaults();
 	}
 
-	UClass *Class = GetClass();
+	UClass *ObjClass = GetClass();
 
 	if( (Ar.IsLoading() || Ar.IsSaving()) && !Ar.WantBinaryPropertySerialization() )
 	{
@@ -985,16 +985,16 @@ void UObject::SerializeScriptProperties( FArchive& Ar ) const
 #else 
 		const bool bBreakSerializationRecursion = false;
 #endif
-		Class->SerializeTaggedProperties(Ar, (uint8*)this, HasAnyFlags(RF_ClassDefaultObject) ? Class->GetSuperClass() : Class, (uint8*)DiffObject, bBreakSerializationRecursion ? this : NULL);
+		ObjClass->SerializeTaggedProperties(Ar, (uint8*)this, HasAnyFlags(RF_ClassDefaultObject) ? ObjClass->GetSuperClass() : ObjClass, (uint8*)DiffObject, bBreakSerializationRecursion ? this : NULL);
 	}
 	else if ( Ar.GetPortFlags() != 0 )
 	{
 		UObject* DiffObject = GetArchetype();
-		Class->SerializeBinEx( Ar, const_cast<UObject *>(this), DiffObject, DiffObject ? DiffObject->GetClass() : NULL );
+		ObjClass->SerializeBinEx( Ar, const_cast<UObject *>(this), DiffObject, DiffObject ? DiffObject->GetClass() : NULL );
 	}
 	else
 	{
-		Class->SerializeBin( Ar, const_cast<UObject *>(this) );
+		ObjClass->SerializeBin( Ar, const_cast<UObject *>(this) );
 	}
 
 	if( HasAnyFlags(RF_ClassDefaultObject) )
@@ -1129,9 +1129,9 @@ bool UObject::CheckDefaultSubobjectsInternal()
 	bool Result = true;	
 
 	CompCheck(this);
-	UClass* Class = GetClass();
+	UClass* ObjClass = GetClass();
 
-	if (Class != UFunction::StaticClass() && Class->GetName() != TEXT("EdGraphPin"))
+	if (ObjClass != UFunction::StaticClass() && ObjClass->GetName() != TEXT("EdGraphPin"))
 	{
 		// Check for references to default subobjects of other objects.
 		// There should never be a pointer to a subobject from outside of the outer (chain) it belongs to.
@@ -1146,21 +1146,21 @@ bool UObject::CheckDefaultSubobjectsInternal()
 	}
 
 #if 0 // usually overkill, but valid tests
-	if (!HasAnyFlags(RF_ClassDefaultObject) && Class->HasAnyClassFlags(CLASS_HasInstancedReference))
+	if (!HasAnyFlags(RF_ClassDefaultObject) && ObjClass->HasAnyClassFlags(CLASS_HasInstancedReference))
 	{
 		UObject *Archetype = GetArchetype();
 		CompCheck(this != Archetype);
 		Archetype->CheckDefaultSubobjects();
-		if (Archetype != Class->GetDefaultObject())
+		if (Archetype != ObjClass->GetDefaultObject())
 		{
-			Class->GetDefaultObject()->CheckDefaultSubobjects();
+			ObjClass->GetDefaultObject()->CheckDefaultSubobjects();
 		}
 	}
 #endif
 
 	if (HasAnyFlags(RF_ClassDefaultObject))
 	{
-		CompCheck(GetFName() == Class->GetDefaultObjectName());
+		CompCheck(GetFName() == ObjClass->GetDefaultObjectName());
 	}
 
 
@@ -1947,18 +1947,18 @@ void UObject::UpdateSinglePropertyInConfigFile(const UProperty* InProperty, cons
 
 void UObject::InstanceSubobjectTemplates( FObjectInstancingGraph* InstanceGraph )
 {
-	UClass *Class = GetClass();
-	if ( Class->HasAnyClassFlags(CLASS_HasInstancedReference) )
+	UClass *ObjClass = GetClass();
+	if (ObjClass->HasAnyClassFlags(CLASS_HasInstancedReference) )
 	{
 		UObject *Archetype = GetArchetype();
 		if (InstanceGraph)
 		{
-			Class->InstanceSubobjectTemplates( this, Archetype, Archetype ? Archetype->GetClass() : NULL, this, InstanceGraph );
+			ObjClass->InstanceSubobjectTemplates( this, Archetype, Archetype ? Archetype->GetClass() : NULL, this, InstanceGraph );
 		}
 		else
 		{
 			FObjectInstancingGraph TempInstanceGraph(this);
-			Class->InstanceSubobjectTemplates( this, Archetype, Archetype ? Archetype->GetClass() : NULL, this, &TempInstanceGraph );
+			ObjClass->InstanceSubobjectTemplates( this, Archetype, Archetype ? Archetype->GetClass() : NULL, this, &TempInstanceGraph );
 		}
 	}
 	CheckDefaultSubobjects();
diff --git a/Engine/Source/Runtime/CoreUObject/Private/UObject/Property.cpp b/Engine/Source/Runtime/CoreUObject/Private/UObject/Property.cpp
--- a/Engine/Source/Runtime/CoreUObject/Private/UObject/Property.cpp
+++ b/Engine/Source/Runtime/CoreUObject/Private/UObject/Property.cpp
@@ -1263,11 +1263,11 @@ IMPLEMENT_CORE_INTRINSIC_CLASS(UNumericProperty, UProperty,
 }
 );
 
-UProperty* UStruct::FindPropertyByName(FName Name) const
+UProperty* UStruct::FindPropertyByName(FName InName) const
 {
 	for (UProperty* Property = PropertyLink; Property != NULL; Property = Property->PropertyLinkNext)
 	{
-		if (Property->GetFName() == Name)
+		if (Property->GetFName() == InName)
 		{
 			return Property;
 		}
diff --git a/Engine/Source/Runtime/CoreUObject/Private/UObject/UObjectBaseUtility.cpp b/Engine/Source/Runtime/CoreUObject/Private/UObject/UObjectBaseUtility.cpp
--- a/Engine/Source/Runtime/CoreUObject/Private/UObject/UObjectBaseUtility.cpp
+++ b/Engine/Source/Runtime/CoreUObject/Private/UObject/UObjectBaseUtility.cpp
@@ -36,14 +36,14 @@ void UObjectBaseUtility::GetPathName( const UObject* StopOuter, FString& ResultS
 {
 	if( this != StopOuter && this != NULL )
 	{
-		UObject* Outer = GetOuter();
-		if ( Outer && Outer != StopOuter )
+		UObject* ObjOuter = GetOuter();
+		if (ObjOuter && ObjOuter != StopOuter )
 		{
-			Outer->GetPathName( StopOuter, ResultString );
+			ObjOuter->GetPathName( StopOuter, ResultString );
 
 			// SUBOBJECT_DELIMITER is used to indicate that this object's outer is not a UPackage
-			if (Outer->GetClass() != UPackage::StaticClass()
-			&&	Outer->GetOuter()->GetClass() == UPackage::StaticClass())
+			if (ObjOuter->GetClass() != UPackage::StaticClass()
+			&& ObjOuter->GetOuter()->GetClass() == UPackage::StaticClass())
 			{
 				ResultString += SUBOBJECT_DELIMITER;
 			}
@@ -115,12 +115,12 @@ UPackage* UObjectBaseUtility::GetOutermost() const
 	UObject* Top = (UObject*)this;
 	for (;;)
 	{
-		UObject* Outer = Top->GetOuter();
-		if (!Outer)
+		UObject* CurrentOuter = Top->GetOuter();
+		if (!CurrentOuter)
 		{
 			return CastChecked<UPackage>(Top);
 		}
-		Top = Outer;
+		Top = CurrentOuter;
 	}
 }
 
diff --git a/Engine/Source/Runtime/CoreUObject/Public/Serialization/AsyncLoading.h b/Engine/Source/Runtime/CoreUObject/Public/Serialization/AsyncLoading.h
--- a/Engine/Source/Runtime/CoreUObject/Public/Serialization/AsyncLoading.h
+++ b/Engine/Source/Runtime/CoreUObject/Public/Serialization/AsyncLoading.h
@@ -228,7 +228,7 @@ public:
 	*
 	* @return true if we finished calling PostLoad on all loaded objects and no new ones were created, false otherwise
 	*/
-	EAsyncPackageState::Type PostLoadDeferredObjects(double InTickStartTime, bool bUseTimeLimit, float& InOutTimeLimit);
+	EAsyncPackageState::Type PostLoadDeferredObjects(double InTickStartTime, bool bInUseTimeLimit, float& InOutTimeLimit);
 
 private:
 	/**
diff --git a/Engine/Source/Runtime/CoreUObject/Public/Serialization/BulkData.h b/Engine/Source/Runtime/CoreUObject/Public/Serialization/BulkData.h
--- a/Engine/Source/Runtime/CoreUObject/Public/Serialization/BulkData.h
+++ b/Engine/Source/Runtime/CoreUObject/Public/Serialization/BulkData.h
@@ -575,7 +575,7 @@ public:
 		}
 		Formats.Empty();
 	}
-	COREUOBJECT_API void Serialize(FArchive& Ar, UObject* Owner, const TArray<FName>* FormatsToSave = NULL, bool bSingleUse = true, uint32 Alignment = DEFAULT_ALIGNMENT);
+	COREUOBJECT_API void Serialize(FArchive& Ar, UObject* Owner, const TArray<FName>* FormatsToSave = NULL, bool bSingleUse = true, uint32 InAlignment = DEFAULT_ALIGNMENT);
 };
 
 #endif
diff --git a/Engine/Source/Runtime/CoreUObject/Public/UObject/Class.h b/Engine/Source/Runtime/CoreUObject/Public/UObject/Class.h
--- a/Engine/Source/Runtime/CoreUObject/Public/UObject/Class.h
+++ b/Engine/Source/Runtime/CoreUObject/Public/UObject/Class.h
@@ -282,7 +282,7 @@ public:
 	// UField interface.
 	virtual void AddCppProperty(UProperty* Property) override;
 
-	UProperty* FindPropertyByName(FName Name) const;
+	UProperty* FindPropertyByName(FName InName) const;
 
 	/**
 	 * Creates new copies of components
@@ -337,7 +337,7 @@ public:
 
 #if WITH_EDITOR
 private:
-	virtual UProperty* CustomFindProperty(const FName Name) const { return NULL; };
+	virtual UProperty* CustomFindProperty(const FName InName) const { return NULL; };
 #endif // WITH_EDITOR
 public:
 	virtual EExprToken SerializeExpr(int32& iCode, FArchive& Ar);
@@ -409,9 +409,9 @@ public:
 		Children = Child;
 	}
 
-	virtual FString PropertyNameToDisplayName(FName Name) const 
+	virtual FString PropertyNameToDisplayName(FName InName) const 
 	{ 
-		return Name.ToString(); 
+		return InName.ToString();
 	}
 
 #if WITH_EDITOR
@@ -1363,7 +1363,7 @@ public:
 	int8 GetValueByName(FName InName);
 
 	/** Gets index of name in enum. Returns INDEX_NONE when name is not found. */
-	int32 GetIndexByName(FName Name) const;
+	int32 GetIndexByName(FName InName) const;
 
 	/** Gets max value of Enum. Defaults to zero if there are no entries. */
 	int8 GetMaxEnumValue() const;
@@ -1956,7 +1956,7 @@ public:
 	virtual void Serialize(FArchive& Ar) override;
 	virtual void PostLoad() override;
 	virtual void FinishDestroy() override;
-	virtual void DeferredRegister(UClass *UClassStaticClass,const TCHAR* PackageName,const TCHAR* Name) override;
+	virtual void DeferredRegister(UClass *UClassStaticClass,const TCHAR* PackageName,const TCHAR* InName) override;
 	virtual bool Rename(const TCHAR* NewName = NULL, UObject* NewOuter = NULL, ERenameFlags Flags = REN_None) override;
 	virtual void TagSubobjects(EObjectFlags NewFlags) override;
 	virtual void PostInitProperties() override;
diff --git a/Engine/Source/Runtime/CoreUObject/Public/UObject/CoreNet.h b/Engine/Source/Runtime/CoreUObject/Public/UObject/CoreNet.h
--- a/Engine/Source/Runtime/CoreUObject/Public/UObject/CoreNet.h
+++ b/Engine/Source/Runtime/CoreUObject/Public/UObject/CoreNet.h
@@ -122,13 +122,13 @@ class COREUOBJECT_API UPackageMap : public UObject
 {
 	DECLARE_CLASS_INTRINSIC( UPackageMap, UObject, CLASS_Transient | CLASS_Abstract | 0, CoreUObject );
 
-	virtual bool		WriteObject( FArchive & Ar, UObject* Outer, FNetworkGUID NetGUID, FString ObjName ) { return false; }
+	virtual bool		WriteObject( FArchive & Ar, UObject* InOuter, FNetworkGUID NetGUID, FString ObjName ) { return false; }
 
 	// @todo document
-	virtual bool		SerializeObject( FArchive& Ar, UClass* Class, UObject*& Obj, FNetworkGUID *OutNetGUID = NULL ) { return false; }
+	virtual bool		SerializeObject( FArchive& Ar, UClass* InClass, UObject*& Obj, FNetworkGUID *OutNetGUID = NULL ) { return false; }
 
 	// @todo document
-	virtual bool		SerializeName( FArchive& Ar, FName& Name );
+	virtual bool		SerializeName( FArchive& Ar, FName& InName );
 
 	virtual UObject*	ResolvePathAndAssignNetGUID( const FNetworkGUID& NetGUID, const FString& PathName ) { return NULL; }
 
diff --git a/Engine/Source/Runtime/CoreUObject/Public/UObject/LazyObjectPtr.h b/Engine/Source/Runtime/CoreUObject/Public/UObject/LazyObjectPtr.h
--- a/Engine/Source/Runtime/CoreUObject/Public/UObject/LazyObjectPtr.h
+++ b/Engine/Source/Runtime/CoreUObject/Public/UObject/LazyObjectPtr.h
@@ -182,9 +182,9 @@ public:
 	 * Copy from a unique object identifier
 	 * @param ObjectID Object identifier to create a weak pointer to
 	 */
-	FORCEINLINE void operator=(const FUniqueObjectGuid &ObjectID)
+	FORCEINLINE void operator=(const FUniqueObjectGuid &InObjectID)
 	{
-		TPersistentObjectPtr<FUniqueObjectGuid>::operator=(ObjectID);
+		TPersistentObjectPtr<FUniqueObjectGuid>::operator=(InObjectID);
 	}
 };
 
@@ -292,9 +292,9 @@ public:
 	 * because the object corresponding to this ID may not even be loaded!
 	 * @param ObjectID Object identifier to create a lazy pointer to
 	 */
-	FORCEINLINE void operator=(const FUniqueObjectGuid& ObjectID)
+	FORCEINLINE void operator=(const FUniqueObjectGuid& InObjectID)
 	{
-		FLazyObjectPtr::operator=(ObjectID);
+		FLazyObjectPtr::operator=(InObjectID);
 	}
 
 	/**
diff --git a/Engine/Source/Runtime/CoreUObject/Public/UObject/Linker.h b/Engine/Source/Runtime/CoreUObject/Public/UObject/Linker.h
--- a/Engine/Source/Runtime/CoreUObject/Public/UObject/Linker.h
+++ b/Engine/Source/Runtime/CoreUObject/Public/UObject/Linker.h
@@ -1587,7 +1587,7 @@ public:
 	 * @param Checked		Whether or not a failure will throw an error
 	 * @return The created object, or (UObject*)-1 if this is just verifying
 	 */
-	UObject* Create( UClass* ObjectClass, FName ObjectName, UObject* Outer, uint32 LoadFlags, bool Checked );
+	UObject* Create( UClass* ObjectClass, FName ObjectName, UObject* Outer, uint32 InLoadFlags, bool Checked );
 
 	/**
 	 * Serialize the object data for the specified object from the unreal package file.  Loads any
diff --git a/Engine/Source/Runtime/CoreUObject/Public/UObject/ReferenceChainSearch.h b/Engine/Source/Runtime/CoreUObject/Public/UObject/ReferenceChainSearch.h
--- a/Engine/Source/Runtime/CoreUObject/Public/UObject/ReferenceChainSearch.h
+++ b/Engine/Source/Runtime/CoreUObject/Public/UObject/ReferenceChainSearch.h
@@ -93,7 +93,7 @@ public:
 		{
 		}
 
-		virtual void HandleObjectReference( UObject*& InObject, const UObject* ReferencingObject, const UProperty* ReferencingProperty ) override;
+		virtual void HandleObjectReference( UObject*& InObject, const UObject* InReferencingObject, const UProperty* ReferencingProperty ) override;
 
 		virtual bool IsIgnoringArchetypeRef() const override	{ return false; }
 		virtual bool IsIgnoringTransient() const override		{ return false; }
diff --git a/Engine/Source/Runtime/CoreUObject/Public/UObject/UObjectBaseUtility.h b/Engine/Source/Runtime/CoreUObject/Public/UObject/UObjectBaseUtility.h
--- a/Engine/Source/Runtime/CoreUObject/Public/UObject/UObjectBaseUtility.h
+++ b/Engine/Source/Runtime/CoreUObject/Public/UObject/UObjectBaseUtility.h
@@ -6,6 +6,12 @@
 
 #pragma once
 
+#if _MSC_VER == 1900
+	#ifdef PRAGMA_DISABLE_SHADOW_VARIABLE_WARNINGS
+		PRAGMA_DISABLE_SHADOW_VARIABLE_WARNINGS
+	#endif
+#endif
+
 #include "Templates/PointerIsConvertibleFromTo.h"
 
 class COREUOBJECT_API UObjectBaseUtility : public UObjectBase
@@ -529,3 +535,9 @@ struct FScopeCycleCounterUObject
 	}
 };
 #endif
+
+#if _MSC_VER == 1900
+	#ifdef PRAGMA_ENABLE_SHADOW_VARIABLE_WARNINGS
+		PRAGMA_ENABLE_SHADOW_VARIABLE_WARNINGS
+	#endif
+#endif
\ No newline at end of file
diff --git a/Engine/Source/Runtime/Engine/Classes/AI/Navigation/RecastNavMesh.h b/Engine/Source/Runtime/Engine/Classes/AI/Navigation/RecastNavMesh.h
--- a/Engine/Source/Runtime/Engine/Classes/AI/Navigation/RecastNavMesh.h
+++ b/Engine/Source/Runtime/Engine/Classes/AI/Navigation/RecastNavMesh.h
@@ -764,7 +764,7 @@ public:
 	void RemoveTileCacheLayers(int32 TileX, int32 TileY);
 	
 	/**  */
-	void AddTileCacheLayers(int32 TileX, int32 TileY, const TArray<FNavMeshTileData>& Layers);
+	void AddTileCacheLayers(int32 TileX, int32 TileY, const TArray<FNavMeshTileData>& InLayers);
 	
 	/**  */
 	TArray<FNavMeshTileData> GetTileCacheLayers(int32 TileX, int32 TileY) const;
diff --git a/Engine/Source/Runtime/Engine/Classes/Animation/AnimInstance.h b/Engine/Source/Runtime/Engine/Classes/Animation/AnimInstance.h
--- a/Engine/Source/Runtime/Engine/Classes/Animation/AnimInstance.h
+++ b/Engine/Source/Runtime/Engine/Classes/Animation/AnimInstance.h
@@ -513,9 +513,9 @@ public:
 	* AnimMontage native C++ interface
 	********************************************************************************************* */
 public:	
-	void Montage_SetEndDelegate(FOnMontageEnded & OnMontageEnded, UAnimMontage * Montage = NULL);
+	void Montage_SetEndDelegate(FOnMontageEnded & InOnMontageEnded, UAnimMontage * Montage = NULL);
 	
-	void Montage_SetBlendingOutDelegate(FOnMontageBlendingOutStarted & OnMontageBlendingOut, UAnimMontage* Montage = NULL);
+	void Montage_SetBlendingOutDelegate(FOnMontageBlendingOutStarted & InOnMontageBlendingOut, UAnimMontage* Montage = NULL);
 	
 	/** Get pointer to BlendingOutStarted delegate for Montage.
 	If Montage reference is NULL, it will pick the first active montage found. */
diff --git a/Engine/Source/Runtime/Engine/Classes/Animation/AnimSequence.h b/Engine/Source/Runtime/Engine/Classes/Animation/AnimSequence.h
--- a/Engine/Source/Runtime/Engine/Classes/Animation/AnimSequence.h
+++ b/Engine/Source/Runtime/Engine/Classes/Animation/AnimSequence.h
@@ -593,7 +593,7 @@ private:
 	 * @param	RequiredBones		BoneContainer
 	 * @param	ExtractionContext	Extraction Context to access root motion extraction information.
 	 */
-	void ResetRootBoneForRootMotion(FTransform& BoneTransform, const FBoneContainer & RequiredBones, ERootMotionRootLock::Type RootMotionRootLock) const;
+	void ResetRootBoneForRootMotion(FTransform& BoneTransform, const FBoneContainer & RequiredBones, ERootMotionRootLock::Type InRootMotionRootLock) const;
 
 	/**
 	* Retarget a single bone transform, to apply right after extraction.
diff --git a/Engine/Source/Runtime/Engine/Classes/Animation/Skeleton.h b/Engine/Source/Runtime/Engine/Classes/Animation/Skeleton.h
--- a/Engine/Source/Runtime/Engine/Classes/Animation/Skeleton.h
+++ b/Engine/Source/Runtime/Engine/Classes/Animation/Skeleton.h
@@ -614,7 +614,7 @@ public:
 	 *
 	 * @param Name	Name of pose to update
 	 */
-	ENGINE_API void UpdateRetargetSource( const FName Name );
+	ENGINE_API void UpdateRetargetSource( const FName InName );
 #endif
 protected:
 	/** 
diff --git a/Engine/Source/Runtime/Engine/Classes/Camera/CameraShake.h b/Engine/Source/Runtime/Engine/Classes/Camera/CameraShake.h
--- a/Engine/Source/Runtime/Engine/Classes/Camera/CameraShake.h
+++ b/Engine/Source/Runtime/Engine/Classes/Camera/CameraShake.h
@@ -244,7 +244,7 @@ public:
 	// Native API
 
 	virtual void UpdateAndApplyCameraShake(float DeltaTime, float Alpha, FMinimalViewInfo& InOutPOV);
-	virtual void PlayShake(class APlayerCameraManager* Camera, float Scale, ECameraAnimPlaySpace::Type PlaySpace, FRotator UserPlaySpaceRot = FRotator::ZeroRotator);
+	virtual void PlayShake(class APlayerCameraManager* Camera, float Scale, ECameraAnimPlaySpace::Type InPlaySpace, FRotator UserPlaySpaceRot = FRotator::ZeroRotator);
 	virtual bool IsFinished() const;
 	virtual void StopShake();
 
diff --git a/Engine/Source/Runtime/Engine/Classes/Components/ChildActorComponent.h b/Engine/Source/Runtime/Engine/Classes/Components/ChildActorComponent.h
--- a/Engine/Source/Runtime/Engine/Classes/Components/ChildActorComponent.h
+++ b/Engine/Source/Runtime/Engine/Classes/Components/ChildActorComponent.h
@@ -14,7 +14,7 @@ class UChildActorComponent : public USceneComponent
 	GENERATED_UCLASS_BODY()
 
 	UFUNCTION(BlueprintCallable, Category=ChildActorComponent)
-	ENGINE_API void SetChildActorClass(TSubclassOf<AActor> Class);
+	ENGINE_API void SetChildActorClass(TSubclassOf<AActor> InClass);
 
 	TSubclassOf<AActor> GetChildActorClass() const { return ChildActorClass; }
 
diff --git a/Engine/Source/Runtime/Engine/Classes/Components/DestructibleComponent.h b/Engine/Source/Runtime/Engine/Classes/Components/DestructibleComponent.h
--- a/Engine/Source/Runtime/Engine/Classes/Components/DestructibleComponent.h
+++ b/Engine/Source/Runtime/Engine/Classes/Components/DestructibleComponent.h
@@ -172,7 +172,7 @@ public:
 	 * @param ChunkIndex - Which chunk to affect.  ChunkIndex must lie in the range: 0 <= ChunkIndex < ((DestructibleMesh*)USkeletalMesh)->ApexDestructibleAsset->chunkCount().
 	 * @param bVisible - If true, the chunk will be made visible.  Otherwise, the chunk is made invisible.
 	 */
-	void SetChunkVisible( int32 ChunkIndex, bool bVisible );
+	void SetChunkVisible( int32 ChunkIndex, bool bInVisible );
 
 #if WITH_APEX
 	/** This method takes a collection of active actors and updates the chunks in one pass. Saves a lot of duplicate work instead of calling each individual chunk
diff --git a/Engine/Source/Runtime/Engine/Classes/Components/InstancedStaticMeshComponent.h b/Engine/Source/Runtime/Engine/Classes/Components/InstancedStaticMeshComponent.h
--- a/Engine/Source/Runtime/Engine/Classes/Components/InstancedStaticMeshComponent.h
+++ b/Engine/Source/Runtime/Engine/Classes/Components/InstancedStaticMeshComponent.h
@@ -237,7 +237,7 @@ protected:
 	virtual void GetNavigationPerInstanceTransforms(const FBox& AreaBox, TArray<FTransform>& InstanceData) const;
 
 	/** Initializes the body instance for the specified instance of the static mesh*/
-	void InitInstanceBody(int32 InstanceIdx, FBodyInstance* BodyInstance);
+	void InitInstanceBody(int32 InstanceIdx, FBodyInstance* InBodyInstance);
 
 	/** Number of pending lightmaps still to be calculated (Apply()'d) */
 	UPROPERTY(Transient, DuplicateTransient, TextExportTransient)
diff --git a/Engine/Source/Runtime/Engine/Classes/Components/LightComponent.h b/Engine/Source/Runtime/Engine/Classes/Components/LightComponent.h
--- a/Engine/Source/Runtime/Engine/Classes/Components/LightComponent.h
+++ b/Engine/Source/Runtime/Engine/Classes/Components/LightComponent.h
@@ -250,7 +250,7 @@ public:
 	 * @param Bounds - The bounding volume to test.
 	 * @return True if the light affects the bounding volume
 	 */
-	virtual bool AffectsBounds(const FBoxSphereBounds& Bounds) const;
+	virtual bool AffectsBounds(const FBoxSphereBounds& InBounds) const;
 
 	/**
 	 * Return the world-space bounding box of the light's influence.
diff --git a/Engine/Source/Runtime/Engine/Classes/Components/LineBatchComponent.h b/Engine/Source/Runtime/Engine/Classes/Components/LineBatchComponent.h
--- a/Engine/Source/Runtime/Engine/Classes/Components/LineBatchComponent.h
+++ b/Engine/Source/Runtime/Engine/Classes/Components/LineBatchComponent.h
@@ -126,7 +126,7 @@ class ULineBatchComponent : public UPrimitiveComponent
 	void DrawLines(const TArray<FBatchedLine>& InLines);
 
 	/** Draw a box */
-	ENGINE_API void DrawBox(const FBox& Box, const FMatrix& TM, const FColor& Color, uint8 DepthPriorityGroup);
+	ENGINE_API void DrawBox(const FBox& Box, const FMatrix& TM, const FColor& Color, uint8 InDepthPriorityGroup);
 
 	/** Draw an arrow */
 	ENGINE_API void DrawDirectionalArrow(const FMatrix& ArrowToWorld,FColor InColor,float Length,float ArrowSize,uint8 DepthPriority);
diff --git a/Engine/Source/Runtime/Engine/Classes/Components/PointLightComponent.h b/Engine/Source/Runtime/Engine/Classes/Components/PointLightComponent.h
--- a/Engine/Source/Runtime/Engine/Classes/Components/PointLightComponent.h
+++ b/Engine/Source/Runtime/Engine/Classes/Components/PointLightComponent.h
@@ -78,7 +78,7 @@ protected:
 public:
 
 	// ULightComponent interface.
-	virtual bool AffectsBounds(const FBoxSphereBounds& Bounds) const override;
+	virtual bool AffectsBounds(const FBoxSphereBounds& InBounds) const override;
 	virtual FVector4 GetLightPosition() const override;
 	virtual FBox GetBoundingBox() const override;
 	virtual FSphere GetBoundingSphere() const override;
diff --git a/Engine/Source/Runtime/Engine/Classes/Components/PrimitiveComponent.h b/Engine/Source/Runtime/Engine/Classes/Components/PrimitiveComponent.h
--- a/Engine/Source/Runtime/Engine/Classes/Components/PrimitiveComponent.h
+++ b/Engine/Source/Runtime/Engine/Classes/Components/PrimitiveComponent.h
@@ -584,7 +584,7 @@ public:
 
 	/** Returns list of components this component is overlapping. */
 	UFUNCTION(BlueprintCallable, Category="Collision", meta=(UnsafeDuringActorConstruction="true"))
-	void GetOverlappingComponents(TArray<UPrimitiveComponent*>& OverlappingComponents) const;
+	void GetOverlappingComponents(TArray<UPrimitiveComponent*>& InOverlappingComponents) const;
 
 	/** Returns list of components this component is overlapping. */
 	UFUNCTION(BlueprintCallable, Category="Collision")
@@ -615,13 +615,13 @@ public:
 	 *  @param	ObjectQueryParams	List of object types it's looking for. When this enters, we do object query with component shape
 	 *  @return true if OutOverlaps contains any blocking results
 	 */
-	bool ComponentOverlapMulti(TArray<struct FOverlapResult>& OutOverlaps, const class UWorld* World, const FVector& Pos, const FQuat& Rot, ECollisionChannel TestChannel, const struct FComponentQueryParams& Params, const struct FCollisionObjectQueryParams& ObjectQueryParams = FCollisionObjectQueryParams::DefaultObjectQueryParam) const;
-	bool ComponentOverlapMulti(TArray<struct FOverlapResult>& OutOverlaps, const class UWorld* World, const FVector& Pos, const FRotator& Rot, ECollisionChannel TestChannel, const struct FComponentQueryParams& Params, const struct FCollisionObjectQueryParams& ObjectQueryParams = FCollisionObjectQueryParams::DefaultObjectQueryParam) const;
+	bool ComponentOverlapMulti(TArray<struct FOverlapResult>& OutOverlaps, const class UWorld* InWorld, const FVector& Pos, const FQuat& Rot, ECollisionChannel TestChannel, const struct FComponentQueryParams& Params, const struct FCollisionObjectQueryParams& ObjectQueryParams = FCollisionObjectQueryParams::DefaultObjectQueryParam) const;
+	bool ComponentOverlapMulti(TArray<struct FOverlapResult>& OutOverlaps, const class UWorld* InWorld, const FVector& Pos, const FRotator& Rot, ECollisionChannel TestChannel, const struct FComponentQueryParams& Params, const struct FCollisionObjectQueryParams& ObjectQueryParams = FCollisionObjectQueryParams::DefaultObjectQueryParam) const;
 
 protected:
 
 	// Override this method for custom behavior.
-	virtual bool ComponentOverlapMultiImpl(TArray<struct FOverlapResult>& OutOverlaps, const class UWorld* World, const FVector& Pos, const FQuat& Rot, ECollisionChannel TestChannel, const struct FComponentQueryParams& Params, const struct FCollisionObjectQueryParams& ObjectQueryParams = FCollisionObjectQueryParams::DefaultObjectQueryParam) const;
+	virtual bool ComponentOverlapMultiImpl(TArray<struct FOverlapResult>& OutOverlaps, const class UWorld* InWorld, const FVector& Pos, const FQuat& Rot, ECollisionChannel TestChannel, const struct FComponentQueryParams& Params, const struct FCollisionObjectQueryParams& ObjectQueryParams = FCollisionObjectQueryParams::DefaultObjectQueryParam) const;
 
 public:
 
@@ -1457,7 +1457,7 @@ public:
 	 * @param Owner: AActor that owns this component
 	 * @param BlockingHit: FHitResult that generated the blocking hit.
 	 */
-	void DispatchBlockingHit(AActor& Owner, FHitResult const& BlockingHit);
+	void DispatchBlockingHit(AActor& OutOwner, FHitResult const& BlockingHit);
 
 	/**
 	 * Dispatch notification for wake events and propagate to any welded bodies
@@ -1856,14 +1856,14 @@ public:
 //////////////////////////////////////////////////////////////////////////
 // PrimitiveComponent inlines
 
-FORCEINLINE_DEBUGGABLE bool UPrimitiveComponent::ComponentOverlapMulti(TArray<struct FOverlapResult>& OutOverlaps, const class UWorld* World, const FVector& Pos, const FQuat& Rot, ECollisionChannel TestChannel, const struct FComponentQueryParams& Params, const struct FCollisionObjectQueryParams& ObjectQueryParams) const
+FORCEINLINE_DEBUGGABLE bool UPrimitiveComponent::ComponentOverlapMulti(TArray<struct FOverlapResult>& OutOverlaps, const class UWorld* InWorld, const FVector& Pos, const FQuat& Rot, ECollisionChannel TestChannel, const struct FComponentQueryParams& Params, const struct FCollisionObjectQueryParams& ObjectQueryParams) const
 {
-	return ComponentOverlapMultiImpl(OutOverlaps, World, Pos, Rot, TestChannel, Params, ObjectQueryParams);
+	return ComponentOverlapMultiImpl(OutOverlaps, InWorld, Pos, Rot, TestChannel, Params, ObjectQueryParams);
 }
 
-FORCEINLINE_DEBUGGABLE bool UPrimitiveComponent::ComponentOverlapMulti(TArray<struct FOverlapResult>& OutOverlaps, const class UWorld* World, const FVector& Pos, const FRotator& Rot, ECollisionChannel TestChannel, const struct FComponentQueryParams& Params, const struct FCollisionObjectQueryParams& ObjectQueryParams) const
+FORCEINLINE_DEBUGGABLE bool UPrimitiveComponent::ComponentOverlapMulti(TArray<struct FOverlapResult>& OutOverlaps, const class UWorld* InWorld, const FVector& Pos, const FRotator& Rot, ECollisionChannel TestChannel, const struct FComponentQueryParams& Params, const struct FCollisionObjectQueryParams& ObjectQueryParams) const
 {
-	return ComponentOverlapMultiImpl(OutOverlaps, World, Pos, Rot.Quaternion(), TestChannel, Params, ObjectQueryParams);
+	return ComponentOverlapMultiImpl(OutOverlaps, InWorld, Pos, Rot.Quaternion(), TestChannel, Params, ObjectQueryParams);
 }
 
 FORCEINLINE_DEBUGGABLE bool UPrimitiveComponent::ComponentOverlapComponent(class UPrimitiveComponent* PrimComp, const FVector Pos, const FQuat& Rot, const FCollisionQueryParams& Params)
diff --git a/Engine/Source/Runtime/Engine/Classes/Components/SkeletalMeshComponent.h b/Engine/Source/Runtime/Engine/Classes/Components/SkeletalMeshComponent.h
--- a/Engine/Source/Runtime/Engine/Classes/Components/SkeletalMeshComponent.h
+++ b/Engine/Source/Runtime/Engine/Classes/Components/SkeletalMeshComponent.h
@@ -858,7 +858,7 @@ protected:
 	 *	@param	ObjectQueryParams	List of object types it's looking for. When this enters, we do object query with component shape
 	 *  @return TRUE if OutOverlaps contains any blocking results
 	 */
-	virtual bool ComponentOverlapMultiImpl(TArray<struct FOverlapResult>& OutOverlaps, const class UWorld* World, const FVector& Pos, const FQuat& Rot, ECollisionChannel TestChannel, const struct FComponentQueryParams& Params, const struct FCollisionObjectQueryParams& ObjectQueryParams = FCollisionObjectQueryParams::DefaultObjectQueryParam) const override;
+	virtual bool ComponentOverlapMultiImpl(TArray<struct FOverlapResult>& OutOverlaps, const class UWorld* InWorld, const FVector& Pos, const FQuat& Rot, ECollisionChannel TestChannel, const struct FComponentQueryParams& Params, const struct FCollisionObjectQueryParams& ObjectQueryParams = FCollisionObjectQueryParams::DefaultObjectQueryParam) const override;
 	
 	virtual bool ComponentOverlapComponentImpl(class UPrimitiveComponent* PrimComp, const FVector Pos, const FQuat& Quat, const FCollisionQueryParams& Params) override;
 
diff --git a/Engine/Source/Runtime/Engine/Classes/Components/SkinnedMeshComponent.h b/Engine/Source/Runtime/Engine/Classes/Components/SkinnedMeshComponent.h
--- a/Engine/Source/Runtime/Engine/Classes/Components/SkinnedMeshComponent.h
+++ b/Engine/Source/Runtime/Engine/Classes/Components/SkinnedMeshComponent.h
@@ -906,7 +906,7 @@ public:
 	* 
 	* @param bRecentlyRendered : true if at least one SkinnedMeshComponent on this Actor has been rendered in the last second.
 	* @param MaxDistanceFactor : Largest SkinnedMeshComponent of this Actor drawn on screen. */
-	void AnimUpdateRateSetParams(uint8 UpdateRateShift, float DeltaTime, const bool & bRecentlyRendered, const float& MaxDistanceFactor, const bool & bPlayingRootMotion);
+	void AnimUpdateRateSetParams(uint8 UpdateRateShift, float DeltaTime, const bool & bInRecentlyRendered, const float& InMaxDistanceFactor, const bool & bPlayingRootMotion);
 
 	virtual bool IsPlayingRootMotion(){ return false; }
 
diff --git a/Engine/Source/Runtime/Engine/Classes/Components/SpotLightComponent.h b/Engine/Source/Runtime/Engine/Classes/Components/SpotLightComponent.h
--- a/Engine/Source/Runtime/Engine/Classes/Components/SpotLightComponent.h
+++ b/Engine/Source/Runtime/Engine/Classes/Components/SpotLightComponent.h
@@ -36,7 +36,7 @@ class ENGINE_API USpotLightComponent : public UPointLightComponent
 
 	// ULightComponent interface.
 	virtual FSphere GetBoundingSphere() const override;
-	virtual bool AffectsBounds(const FBoxSphereBounds& Bounds) const override;
+	virtual bool AffectsBounds(const FBoxSphereBounds& InBounds) const override;
 	virtual ELightComponentType GetLightType() const override;
 	virtual FLightSceneProxy* CreateSceneProxy() const override;
 #if WITH_EDITOR
diff --git a/Engine/Source/Runtime/Engine/Classes/Engine/BlueprintGeneratedClass.h b/Engine/Source/Runtime/Engine/Classes/Engine/BlueprintGeneratedClass.h
--- a/Engine/Source/Runtime/Engine/Classes/Engine/BlueprintGeneratedClass.h
+++ b/Engine/Source/Runtime/Engine/Classes/Engine/BlueprintGeneratedClass.h
@@ -498,7 +498,7 @@ public:
 	void UnbindDynamicDelegatesForProperty(UObject* InInstance, const UObjectProperty* InObjectProperty);
 	
 	// Finds the desired dynamic binding object for this blueprint generated class
-	UDynamicBlueprintBinding* GetDynamicBindingObject(UClass* Class) const;
+	UDynamicBlueprintBinding* GetDynamicBindingObject(UClass* InClass) const;
 
 	/** called to gather blueprint replicated properties */
 	virtual void GetLifetimeBlueprintReplicationList(TArray<class FLifetimeProperty>& OutLifetimeProps) const;
diff --git a/Engine/Source/Runtime/Engine/Classes/Engine/GameInstance.h b/Engine/Source/Runtime/Engine/Classes/Engine/GameInstance.h
--- a/Engine/Source/Runtime/Engine/Classes/Engine/GameInstance.h
+++ b/Engine/Source/Runtime/Engine/Classes/Engine/GameInstance.h
@@ -177,13 +177,13 @@ public:
 	inline FTimerManager& GetTimerManager() const { return *TimerManager; }
 
 	/** Start recording a replay with the given custom name and friendly name. */
-	virtual void StartRecordingReplay(const FString& Name, const FString& FriendlyName);
+	virtual void StartRecordingReplay(const FString& InName, const FString& FriendlyName);
 
 	/** Stop recording a replay if one is currently in progress */
 	virtual void StopRecordingReplay();
 
 	/** Start playing back a previously recorded replay. */
-	virtual void PlayReplay(const FString& Name);
+	virtual void PlayReplay(const FString& InName);
 
 	/**
 	 * Adds a join-in-progress user to the set of users associated with the currently recording replay (if any)
diff --git a/Engine/Source/Runtime/Engine/Classes/Engine/GameViewportClient.h b/Engine/Source/Runtime/Engine/Classes/Engine/GameViewportClient.h
--- a/Engine/Source/Runtime/Engine/Classes/Engine/GameViewportClient.h
+++ b/Engine/Source/Runtime/Engine/Classes/Engine/GameViewportClient.h
@@ -467,7 +467,7 @@ public:
 	 * @param	FailureType	the type of error
 	 * @param	ErrorString	additional string detailing the error
 	 */
-	virtual void PeekNetworkFailureMessages(UWorld *World, UNetDriver *NetDriver, enum ENetworkFailure::Type FailureType, const FString& ErrorString);
+	virtual void PeekNetworkFailureMessages(UWorld *InWorld, UNetDriver *NetDriver, enum ENetworkFailure::Type FailureType, const FString& ErrorString);
 
 	/** Make sure all navigation objects have appropriate path rendering components set.  Called when EngineShowFlags.Navigation is set. */
 	virtual void VerifyPathRenderingComponents();
diff --git a/Engine/Source/Runtime/Engine/Classes/Engine/NetDriver.h b/Engine/Source/Runtime/Engine/Classes/Engine/NetDriver.h
--- a/Engine/Source/Runtime/Engine/Classes/Engine/NetDriver.h
+++ b/Engine/Source/Runtime/Engine/Classes/Engine/NetDriver.h
@@ -287,7 +287,7 @@ public:
 	TMap< TWeakObjectPtr< UObject >, TSharedPtr< FRepLayout > >					RepLayoutMap;
 
 	/** Creates if necessary, and returns a FRepLayout that maps to the passed in UClass */
-	TSharedPtr< FRepLayout >	GetObjectClassRepLayout( UClass * Class );
+	TSharedPtr< FRepLayout >	GetObjectClassRepLayout( UClass * InClass );
 
 	/** Creates if necessary, and returns a FRepLayout that maps to the passed in UFunction */
 	TSharedPtr<FRepLayout>		GetFunctionRepLayout( UFunction * Function );
diff --git a/Engine/Source/Runtime/Engine/Classes/Engine/PackageMapClient.h b/Engine/Source/Runtime/Engine/Classes/Engine/PackageMapClient.h
--- a/Engine/Source/Runtime/Engine/Classes/Engine/PackageMapClient.h
+++ b/Engine/Source/Runtime/Engine/Classes/Engine/PackageMapClient.h
@@ -123,10 +123,10 @@ public:
 
 	
 	// UPackageMap Interface
-	virtual bool SerializeObject( FArchive& Ar, UClass* Class, UObject*& Obj, FNetworkGUID *OutNetGUID = NULL ) override;
+	virtual bool SerializeObject( FArchive& Ar, UClass* InClass, UObject*& Obj, FNetworkGUID *OutNetGUID = NULL ) override;
 	virtual bool SerializeNewActor( FArchive& Ar, class UActorChannel *Channel, class AActor*& Actor) override;
 	
-	virtual bool WriteObject( FArchive& Ar, UObject* Outer, FNetworkGUID NetGUID, FString ObjName ) override;
+	virtual bool WriteObject( FArchive& Ar, UObject* InOuter, FNetworkGUID NetGUID, FString ObjName ) override;
 
 	// UPackageMapClient Connection specific methods
 
diff --git a/Engine/Source/Runtime/Engine/Classes/Engine/Polys.h b/Engine/Source/Runtime/Engine/Classes/Engine/Polys.h
--- a/Engine/Source/Runtime/Engine/Classes/Engine/Polys.h
+++ b/Engine/Source/Runtime/Engine/Classes/Engine/Polys.h
@@ -97,7 +97,7 @@ public:
 	/**
 	 * Split with plane. Meant to be numerically stable.
 	 */
-	ENGINE_API int32 SplitWithPlane(const FVector &Base,const FVector &Normal,FPoly *FrontPoly,FPoly *BackPoly,int32 VeryPrecise) const;
+	ENGINE_API int32 SplitWithPlane(const FVector &InBase,const FVector &InNormal,FPoly *FrontPoly,FPoly *BackPoly,int32 VeryPrecise) const;
 
 	/** Split with a Bsp node. */
 	ENGINE_API int32 SplitWithNode(const UModel *Model,int32 iNode,FPoly *FrontPoly,FPoly *BackPoly,int32 VeryPrecise) const;
@@ -109,7 +109,7 @@ public:
 	ENGINE_API int32 SplitWithPlaneFast(const FPlane& Plane,FPoly *FrontPoly,FPoly *BackPoly) const;
 
 	/** Split a poly and keep only the front half. Returns number of vertices, 0 if clipped away. */
-	ENGINE_API int32 Split(const FVector &Normal, const FVector &Base );
+	ENGINE_API int32 Split(const FVector &InNormal, const FVector &InBase );
 
 	/** Remove colinear vertices and check convexity.  Returns 1 if convex, 0 if nonconvex or collapsed. */
 	ENGINE_API int32 RemoveColinears();
diff --git a/Engine/Source/Runtime/Engine/Classes/Engine/TextureRenderTarget2D.h b/Engine/Source/Runtime/Engine/Classes/Engine/TextureRenderTarget2D.h
--- a/Engine/Source/Runtime/Engine/Classes/Engine/TextureRenderTarget2D.h
+++ b/Engine/Source/Runtime/Engine/Classes/Engine/TextureRenderTarget2D.h
@@ -72,7 +72,7 @@ class UTextureRenderTarget2D : public UTextureRenderTarget
 	 * @param AlphaOverride - If specified, the values here will become the alpha values in the resulting texture
 	 * @return New UTexture2D object.
 	 */
-	ENGINE_API UTexture2D* ConstructTexture2D(UObject* Outer, const FString& NewTexName, EObjectFlags ObjectFlags, uint32 Flags=CTF_Default, TArray<uint8>* AlphaOverride=NULL);
+	ENGINE_API UTexture2D* ConstructTexture2D(UObject* InOuter, const FString& NewTexName, EObjectFlags InObjectFlags, uint32 Flags=CTF_Default, TArray<uint8>* AlphaOverride=NULL);
 
 	/**
 	 * Updates (resolves) the render target texture immediately.
diff --git a/Engine/Source/Runtime/Engine/Classes/Engine/TextureRenderTargetCube.h b/Engine/Source/Runtime/Engine/Classes/Engine/TextureRenderTargetCube.h
--- a/Engine/Source/Runtime/Engine/Classes/Engine/TextureRenderTargetCube.h
+++ b/Engine/Source/Runtime/Engine/Classes/Engine/TextureRenderTargetCube.h
@@ -56,7 +56,7 @@ class UTextureRenderTargetCube : public UTextureRenderTarget
 	* @param	Flags			Various control flags for operation (see EObjectFlags)
 	* @return					New UTextureCube object.
 	*/
-	ENGINE_API class UTextureCube* ConstructTextureCube(UObject* Outer, const FString& NewTexName, EObjectFlags InFlags);
+	ENGINE_API class UTextureCube* ConstructTextureCube(UObject* InOuter, const FString& NewTexName, EObjectFlags InFlags);
 
 	// Begin UTexture interface.
 	virtual float GetSurfaceWidth() const  override { return SizeX; }
diff --git a/Engine/Source/Runtime/Engine/Classes/Engine/UserDefinedStruct.h b/Engine/Source/Runtime/Engine/Classes/Engine/UserDefinedStruct.h
--- a/Engine/Source/Runtime/Engine/Classes/Engine/UserDefinedStruct.h
+++ b/Engine/Source/Runtime/Engine/Classes/Engine/UserDefinedStruct.h
@@ -57,7 +57,7 @@ public:
 #endif	// WITH_EDITOR
 
 	virtual void SerializeTaggedProperties(FArchive& Ar, uint8* Data, UStruct* DefaultsStruct, uint8* Defaults, const UObject* BreakRecursionIfFullyLoad=NULL) const override;
-	virtual FString PropertyNameToDisplayName(FName Name) const override;
+	virtual FString PropertyNameToDisplayName(FName InName) const override;
 
 	// UScriptStruct interface.
 	virtual void RecursivelyPreload() override;
diff --git a/Engine/Source/Runtime/Engine/Classes/Engine/World.h b/Engine/Source/Runtime/Engine/Classes/Engine/World.h
--- a/Engine/Source/Runtime/Engine/Classes/Engine/World.h
+++ b/Engine/Source/Runtime/Engine/Classes/Engine/World.h
@@ -2787,7 +2787,7 @@ public:
 	 */
 	void RemoveActor( AActor* Actor, bool bShouldModifyLevel );
 
-	AActor* SpawnActor( UClass* Class, FVector const* Location=NULL, FRotator const* Rotation=NULL, const FActorSpawnParameters& SpawnParameters = FActorSpawnParameters() );
+	AActor* SpawnActor( UClass* InClass, FVector const* Location=NULL, FRotator const* Rotation=NULL, const FActorSpawnParameters& SpawnParameters = FActorSpawnParameters() );
 
 	/** Templated version of SpawnActor that allows you to specify a class type via the template type */
 	template< class T >
@@ -2891,7 +2891,7 @@ public:
 	 * @param InNetPlayerIndex (optional) - the NetPlayerIndex to set on the PlayerController
 	 * @return the PlayerController that was spawned (may fail and return NULL)
 	 */
-	APlayerController* SpawnPlayActor(class UPlayer* Player, ENetRole RemoteRole, const FURL& URL, const TSharedPtr<const FUniqueNetId>& UniqueId, FString& Error, uint8 InNetPlayerIndex = 0);
+	APlayerController* SpawnPlayActor(class UPlayer* Player, ENetRole RemoteRole, const FURL& InURL, const TSharedPtr<const FUniqueNetId>& UniqueId, FString& Error, uint8 InNetPlayerIndex = 0);
 
 	/** Try to find an acceptable position to place TestActor as close to possible to PlaceLocation.  Expects PlaceLocation to be a valid location inside the level. */
 	bool FindTeleportSpot( AActor* TestActor, FVector& PlaceLocation, FRotator PlaceRotation );
@@ -3088,7 +3088,7 @@ public:
 	 * @param bAbsolute whether we are using relative or absolute travel
 	 * @param bShouldSkipGameNotify whether to notify the clients/game or not
 	 */
-	virtual void ServerTravel(const FString& URL, bool bAbsolute = false, bool bShouldSkipGameNotify = false);
+	virtual void ServerTravel(const FString& InURL, bool bAbsolute = false, bool bShouldSkipGameNotify = false);
 
 	/** seamlessly travels to the given URL by first loading the entry level in the background,
 	 * switching to it, and then loading the specified level. Does not disrupt network communication or disconnect clients.
@@ -3102,7 +3102,7 @@ public:
 	 * @param MapPackageGuid (opt) - the GUID of the map package to travel to - this is used to find the file when it has been auto-downloaded,
 	 * 				so it is only needed for clients
 	 */
-	void SeamlessTravel(const FString& URL, bool bAbsolute = false, FGuid MapPackageGuid = FGuid());
+	void SeamlessTravel(const FString& InURL, bool bAbsolute = false, FGuid MapPackageGuid = FGuid());
 
 	/** @return whether we're currently in a seamless transition */
 	bool IsInSeamlessTravel();
diff --git a/Engine/Source/Runtime/Engine/Classes/GameFramework/Character.h b/Engine/Source/Runtime/Engine/Classes/GameFramework/Character.h
--- a/Engine/Source/Runtime/Engine/Classes/GameFramework/Character.h
+++ b/Engine/Source/Runtime/Engine/Classes/GameFramework/Character.h
@@ -387,7 +387,7 @@ public:
 	virtual void PawnClientRestart() override;
 	virtual void PossessedBy(AController* NewController) override;
 	virtual void UnPossessed() override;
-	virtual void SetupPlayerInputComponent(class UInputComponent* InputComponent) override;
+	virtual void SetupPlayerInputComponent(class UInputComponent* InInputComponent) override;
 	virtual void DisplayDebug(class UCanvas* Canvas, const FDebugDisplayInfo& DebugDisplay, float& YL, float& YPos) override;
 	virtual void RecalculateBaseEyeHeight() override;
 	virtual void UpdateNavigationRelevance() override;
diff --git a/Engine/Source/Runtime/Engine/Classes/GameFramework/CharacterMovementComponent.h b/Engine/Source/Runtime/Engine/Classes/GameFramework/CharacterMovementComponent.h
--- a/Engine/Source/Runtime/Engine/Classes/GameFramework/CharacterMovementComponent.h
+++ b/Engine/Source/Runtime/Engine/Classes/GameFramework/CharacterMovementComponent.h
@@ -1277,7 +1277,7 @@ public:
 	 *  Revert to previous position OldLocation, return to being based on OldBase.
 	 *  if bFailMove, stop movement and notify controller
 	 */	
-	void RevertMove(const FVector& OldLocation, UPrimitiveComponent* OldBase, const FVector& OldBaseLocation, const FFindFloorResult& OldFloor, bool bFailMove);
+	void RevertMove(const FVector& OldLocation, UPrimitiveComponent* OldBase, const FVector& InOldBaseLocation, const FFindFloorResult& OldFloor, bool bFailMove);
 
 	/** Perform rotation over deltaTime */
 	virtual void PhysicsRotation(float DeltaTime);
@@ -1972,7 +1972,7 @@ public:
 	virtual void Clear();
 
 	/** Called to set up this saved move (when initially created) to make a predictive correction. */
-	virtual void SetMoveFor(ACharacter* C, float DeltaTime, FVector const& NewAccel, class FNetworkPredictionData_Client_Character & ClientData);
+	virtual void SetMoveFor(ACharacter* C, float InDeltaTime, FVector const& NewAccel, class FNetworkPredictionData_Client_Character & ClientData);
 
 	/** Set the properties describing the position, etc. of the moved pawn at the start of the move. */
 	virtual void SetInitialPosition(ACharacter* C);
diff --git a/Engine/Source/Runtime/Engine/Classes/GameFramework/Controller.h b/Engine/Source/Runtime/Engine/Classes/GameFramework/Controller.h
--- a/Engine/Source/Runtime/Engine/Classes/GameFramework/Controller.h
+++ b/Engine/Source/Runtime/Engine/Classes/GameFramework/Controller.h
@@ -106,7 +106,7 @@ protected:
 	virtual void AddPawnTickDependency(APawn* NewPawn);
 
 	/** Remove dependency that makes us tick before the given Pawn.	 */
-	virtual void RemovePawnTickDependency(APawn* OldPawn);
+	virtual void RemovePawnTickDependency(APawn* InOldPawn);
 
 public:
 
diff --git a/Engine/Source/Runtime/Engine/Classes/GameFramework/DefaultPawn.h b/Engine/Source/Runtime/Engine/Classes/GameFramework/DefaultPawn.h
--- a/Engine/Source/Runtime/Engine/Classes/GameFramework/DefaultPawn.h
+++ b/Engine/Source/Runtime/Engine/Classes/GameFramework/DefaultPawn.h
@@ -24,7 +24,7 @@ class ENGINE_API ADefaultPawn : public APawn
 
 	// Begin Pawn overrides
 	virtual UPawnMovementComponent* GetMovementComponent() const override;
-	virtual void SetupPlayerInputComponent(UInputComponent* InputComponent) override;
+	virtual void SetupPlayerInputComponent(UInputComponent* InInputComponent) override;
 	virtual void UpdateNavigationRelevance() override;
 	// End Pawn overrides
 
diff --git a/Engine/Source/Runtime/Engine/Classes/GameFramework/GameMode.h b/Engine/Source/Runtime/Engine/Classes/GameFramework/GameMode.h
--- a/Engine/Source/Runtime/Engine/Classes/GameFramework/GameMode.h
+++ b/Engine/Source/Runtime/Engine/Classes/GameFramework/GameMode.h
@@ -461,7 +461,7 @@ public:
 	 *
 	 * @return a new player controller for the logged in player, NULL if login failed for any reason
 	 */
-	virtual APlayerController* Login(class UPlayer* NewPlayer, ENetRole RemoteRole, const FString& Portal, const FString& Options, const TSharedPtr<const FUniqueNetId>& UniqueId, FString& ErrorMessage);
+	virtual APlayerController* Login(class UPlayer* NewPlayer, ENetRole InRemoteRole, const FString& Portal, const FString& Options, const TSharedPtr<const FUniqueNetId>& UniqueId, FString& ErrorMessage);
 
 	/** Called after a successful login.  This is the first place it is safe to call replicated functions on the PlayerAController. */
 	virtual void PostLogin( APlayerController* NewPlayer );
@@ -479,7 +479,7 @@ public:
 	 *
 	 * @return PlayerController for the player, NULL if there is any reason this player shouldn't exist or due to some error
 	 */
-	virtual APlayerController* SpawnPlayerController(ENetRole RemoteRole, FVector const& SpawnLocation, FRotator const& SpawnRotation);
+	virtual APlayerController* SpawnPlayerController(ENetRole InRemoteRole, FVector const& SpawnLocation, FRotator const& SpawnRotation);
 
 	/** @Returns true if NewPlayerController may only join the server as a spectator. */
 	UFUNCTION(BlueprintNativeEvent, Category="Game")
diff --git a/Engine/Source/Runtime/Engine/Classes/GameFramework/GameSession.h b/Engine/Source/Runtime/Engine/Classes/GameFramework/GameSession.h
--- a/Engine/Source/Runtime/Engine/Classes/GameFramework/GameSession.h
+++ b/Engine/Source/Runtime/Engine/Classes/GameFramework/GameSession.h
@@ -155,7 +155,7 @@ class ENGINE_API AGameSession : public AInfo
 	 * @param bJoinViaPresence anyone who can see you can join the game
 	 * @param bJoinViaPresenceFriendsOnly can only friends actively join your game 
 	 */
-	virtual void UpdateSessionJoinability(FName SessionName, bool bPublicSearchable, bool bAllowInvites, bool bJoinViaPresence, bool bJoinViaPresenceFriendsOnly);
+	virtual void UpdateSessionJoinability(FName InSessionName, bool bPublicSearchable, bool bAllowInvites, bool bJoinViaPresence, bool bJoinViaPresenceFriendsOnly);
 
 	/**
 	 * Travel to a session URL (as client) for a given session
@@ -165,7 +165,7 @@ class ENGINE_API AGameSession : public AInfo
 	 *
 	 * @return true if successful, false otherwise
 	 */
-	virtual bool TravelToSession(int32 ControllerId, FName SessionName);
+	virtual bool TravelToSession(int32 ControllerId, FName InSessionName);
 
     /**
      * Does the session require push to talk
diff --git a/Engine/Source/Runtime/Engine/Classes/GameFramework/HUD.h b/Engine/Source/Runtime/Engine/Classes/GameFramework/HUD.h
--- a/Engine/Source/Runtime/Engine/Classes/GameFramework/HUD.h
+++ b/Engine/Source/Runtime/Engine/Classes/GameFramework/HUD.h
@@ -388,7 +388,7 @@ public:
 	 * @param Priority			The priority of the box used for layering. Larger values are considered first.  Equal values are considered in the order they were added.
 	 */
 	UFUNCTION(BlueprintCallable, Category=HUD)
-	void AddHitBox(FVector2D Position, FVector2D Size, FName Name, bool bConsumesInput, int32 Priority = 0);
+	void AddHitBox(FVector2D Position, FVector2D Size, FName InName, bool bConsumesInput, int32 Priority = 0);
 
 protected:
 	/** Returns the PlayerController for this HUD's player.	 */
diff --git a/Engine/Source/Runtime/Engine/Classes/GameFramework/NavMovementComponent.h b/Engine/Source/Runtime/Engine/Classes/GameFramework/NavMovementComponent.h
--- a/Engine/Source/Runtime/Engine/Classes/GameFramework/NavMovementComponent.h
+++ b/Engine/Source/Runtime/Engine/Classes/GameFramework/NavMovementComponent.h
@@ -59,7 +59,7 @@ public:
 	void SetUpdateNavAgentWithOwnersCollisions(bool bUpdateWithOwner);
 	FORCEINLINE bool ShouldUpdateNavAgentWithOwnersCollision() const { return bUpdateNavAgentWithOwnersCollision != 0; }
 	
-	void UpdateNavAgent(const AActor& Owner);
+	void UpdateNavAgent(const AActor& InOwner);
 	void UpdateNavAgent(const UCapsuleComponent& CapsuleComponent);
 
 	/** @returns location of controlled actor - meaning center of collision bounding box */
@@ -141,7 +141,7 @@ public:
 	FNavAgentProperties* GetNavAgentProperties();
 
 	DEPRECATED(4.8, "This function is deprecated. Please use UpdateNavAgent version that's accepring a reference instead.")
-	void UpdateNavAgent(AActor* Owner);
+	void UpdateNavAgent(AActor* InOwner);
 	DEPRECATED(4.8, "This function is deprecated. Please use UpdateNavAgent version that's accepring a reference instead.")
 	void UpdateNavAgent(UCapsuleComponent* CapsuleComponent);
 };
diff --git a/Engine/Source/Runtime/Engine/Classes/GameFramework/Pawn.h b/Engine/Source/Runtime/Engine/Classes/GameFramework/Pawn.h
--- a/Engine/Source/Runtime/Engine/Classes/GameFramework/Pawn.h
+++ b/Engine/Source/Runtime/Engine/Classes/GameFramework/Pawn.h
@@ -327,7 +327,7 @@ protected:
 	virtual void DestroyPlayerInputComponent();
 
 	/** Allows a Pawn to set up custom input bindings. Called upon possession by a PlayerController, using the InputComponent created by CreatePlayerInputComponent(). */
-	virtual void SetupPlayerInputComponent(UInputComponent* InputComponent) { /* No bindings by default.*/ }
+	virtual void SetupPlayerInputComponent(UInputComponent* InInputComponent) { /* No bindings by default.*/ }
 
 public:
 	/**
diff --git a/Engine/Source/Runtime/Engine/Classes/GameFramework/PlayerController.h b/Engine/Source/Runtime/Engine/Classes/GameFramework/PlayerController.h
--- a/Engine/Source/Runtime/Engine/Classes/GameFramework/PlayerController.h
+++ b/Engine/Source/Runtime/Engine/Classes/GameFramework/PlayerController.h
@@ -1180,7 +1180,7 @@ public:
 	virtual bool InputMotion(const FVector& Tilt, const FVector& RotationRate, const FVector& Gravity, const FVector& Acceleration);
 
 	/** Associate a new UPlayer with this PlayerController. */
-	virtual void SetPlayer(UPlayer* Player);
+	virtual void SetPlayer(UPlayer* InPlayer);
 
 	/** Returns the ULocalPlayer for this controller if it exists, or null otherwise */
 	class ULocalPlayer* GetLocalPlayer() const;
diff --git a/Engine/Source/Runtime/Engine/Classes/Materials/MaterialExpression.h b/Engine/Source/Runtime/Engine/Classes/Materials/MaterialExpression.h
--- a/Engine/Source/Runtime/Engine/Classes/Materials/MaterialExpression.h
+++ b/Engine/Source/Runtime/Engine/Classes/Materials/MaterialExpression.h
@@ -285,7 +285,7 @@ class ENGINE_API UMaterialExpression : public UObject
 	/**
 	 * Connects the specified output to the passed material for previewing. 
 	 */
-	void ConnectToPreviewMaterial(UMaterial* Material, int32 OutputIndex);
+	void ConnectToPreviewMaterial(UMaterial* InMaterial, int32 OutputIndex);
 
 	/**
 	 * Connects the specified input expression to the specified output of this expression.
diff --git a/Engine/Source/Runtime/Engine/Classes/Materials/MaterialExpressionTextureSample.h b/Engine/Source/Runtime/Engine/Classes/Materials/MaterialExpressionTextureSample.h
--- a/Engine/Source/Runtime/Engine/Classes/Materials/MaterialExpressionTextureSample.h
+++ b/Engine/Source/Runtime/Engine/Classes/Materials/MaterialExpressionTextureSample.h
@@ -78,8 +78,5 @@ class ENGINE_API UMaterialExpressionTextureSample : public UMaterialExpressionTe
 #endif
 	// End UMaterialExpression Interface
 
-	void UpdateTextureResource(class UTexture* Texture);
-};
-
-
-
+	void UpdateTextureResource(class UTexture* InTexture);
+};
\ No newline at end of file
diff --git a/Engine/Source/Runtime/Engine/Classes/Particles/ParticleSystemComponent.h b/Engine/Source/Runtime/Engine/Classes/Particles/ParticleSystemComponent.h
--- a/Engine/Source/Runtime/Engine/Classes/Particles/ParticleSystemComponent.h
+++ b/Engine/Source/Runtime/Engine/Classes/Particles/ParticleSystemComponent.h
@@ -836,14 +836,14 @@ public:
 	* @param Name - The slot name of the material to replace.  If invalid, the material is unchanged and NULL is returned.
 	*/
 	UFUNCTION(BlueprintCallable, Category = "Rendering|Material")
-	virtual class UMaterialInstanceDynamic* CreateNamedDynamicMaterialInstance(FName Name, class UMaterialInterface* SourceMaterial = NULL);
+	virtual class UMaterialInstanceDynamic* CreateNamedDynamicMaterialInstance(FName InName, class UMaterialInterface* SourceMaterial = NULL);
 
 	/** Returns a named material. If this named material is not found, returns NULL. */
 	UFUNCTION(BlueprintCallable, Category = "Rendering|Material")
-	virtual class UMaterialInterface* GetNamedMaterial(FName Name) const;
+	virtual class UMaterialInterface* GetNamedMaterial(FName InName) const;
 
 	/** Returns the index into the EmitterMaterials array for this named. If there are no named material slots or this material is not found, INDEX_NONE is returned. */
-	virtual int32 GetNamedMaterialIndex(FName Name) const;
+	virtual int32 GetNamedMaterialIndex(FName InName) const;
 
 protected:
 
diff --git a/Engine/Source/Runtime/Engine/Classes/PhysicsEngine/ConstraintInstance.h b/Engine/Source/Runtime/Engine/Classes/PhysicsEngine/ConstraintInstance.h
--- a/Engine/Source/Runtime/Engine/Classes/PhysicsEngine/ConstraintInstance.h
+++ b/Engine/Source/Runtime/Engine/Classes/PhysicsEngine/ConstraintInstance.h
@@ -373,32 +373,32 @@ public:
 	/** Sets the LinearX Motion Type
 	*	@param MotionType	New Motion Type
 	*/
-	void SetLinearXLimit(ELinearConstraintMotion ConstraintType, float LinearLimitSize);
+	void SetLinearXLimit(ELinearConstraintMotion ConstraintType, float InLinearLimitSize);
 
 	/** Sets the LinearY Motion Type
 	*	@param MotionType	New Motion Type
 	*/
-	void SetLinearYLimit(ELinearConstraintMotion ConstraintType, float LinearLimitSize);
+	void SetLinearYLimit(ELinearConstraintMotion ConstraintType, float InLinearLimitSize);
 
 	/** Sets the LinearZ Motion Type
 	*	@param MotionType	New Motion Type
 	*/
-	void SetLinearZLimit(ELinearConstraintMotion ConstraintType, float LinearLimitSize);
+	void SetLinearZLimit(ELinearConstraintMotion ConstraintType, float InLinearLimitSize);
 
 	/** Sets the Angular Swing1 Motion Type
 	*	@param MotionType	New Motion Type
 	*/
-	void SetAngularSwing1Limit(EAngularConstraintMotion MotionType, float Swing1LimitAngle);
+	void SetAngularSwing1Limit(EAngularConstraintMotion MotionType, float InSwing1LimitAngle);
 
 	/** Sets the Angular Swing2 Motion Type
 	*	@param MotionType	New Motion Type
 	*/
-	void SetAngularSwing2Limit(EAngularConstraintMotion MotionType, float Swing2LimitAngle);
+	void SetAngularSwing2Limit(EAngularConstraintMotion MotionType, float InSwing2LimitAngle);
 
 	/** Sets the Angular Twist Motion Type
 	*	@param MotionType	New Motion Type
 	*/
-	void SetAngularTwistLimit(EAngularConstraintMotion MotionType, float TwistLimitAngle);
+	void SetAngularTwistLimit(EAngularConstraintMotion MotionType, float InTwistLimitAngle);
 
 #if WITH_PHYSX
 	FPhysxUserData PhysxUserData;
@@ -457,8 +457,8 @@ public:
 
 	void SetLinearPositionDrive(bool bEnableXDrive, bool bEnableYDrive, bool bEnableZDrive);
 	void SetLinearVelocityDrive(bool bEnableXDrive, bool bEnableYDrive, bool bEnableZDrive);
-	void SetAngularPositionDrive(bool bEnableSwingDrive, bool bEnableTwistDrive);
-	void SetAngularVelocityDrive(bool bEnableSwingDrive, bool bEnableTwistDrive);
+	void SetAngularPositionDrive(bool bInEnableSwingDrive, bool bInEnableTwistDrive);
+	void SetAngularVelocityDrive(bool bInEnableSwingDrive, bool bInEnableTwistDrive);
 
 	void SetLinearPositionTarget(const FVector& InPosTarget);
 	void SetLinearVelocityTarget(const FVector& InVelTarget);
diff --git a/Engine/Source/Runtime/Engine/Classes/Sound/SoundBase.h b/Engine/Source/Runtime/Engine/Classes/Sound/SoundBase.h
--- a/Engine/Source/Runtime/Engine/Classes/Sound/SoundBase.h
+++ b/Engine/Source/Runtime/Engine/Classes/Sound/SoundBase.h
@@ -84,7 +84,7 @@ public:
 	 * @param Location				Location to check against
 	 * @param AttenuationSettings	Optional Attenuation override if not using settings from the sound
 	 */
-	bool IsAudibleSimple( class FAudioDevice* AudioDevice, const FVector Location, USoundAttenuation* AttenuationSettings = NULL );
+	bool IsAudibleSimple( class FAudioDevice* AudioDevice, const FVector Location, USoundAttenuation* InAttenuationSettings = NULL );
 
 	/** 
 	 * Returns the farthest distance at which the sound could be heard
diff --git a/Engine/Source/Runtime/Engine/Private/AudioDerivedData.cpp b/Engine/Source/Runtime/Engine/Private/AudioDerivedData.cpp
--- a/Engine/Source/Runtime/Engine/Private/AudioDerivedData.cpp
+++ b/Engine/Source/Runtime/Engine/Private/AudioDerivedData.cpp
@@ -638,7 +638,6 @@ public:
 	explicit FAudioStatusMessageContext( const FText& InMessage )
 	 : FScopedSlowTask(1, InMessage, GIsEditor && !IsRunningCommandlet())
 	{
-		DEFINE_LOG_CATEGORY_STATIC(LogAudioDerivedData, Log, All);
 		UE_LOG(LogAudioDerivedData, Display, TEXT("%s"), *InMessage.ToString());
 	}
 };
diff --git a/Engine/Source/Runtime/Engine/Private/Character.cpp b/Engine/Source/Runtime/Engine/Private/Character.cpp
--- a/Engine/Source/Runtime/Engine/Private/Character.cpp
+++ b/Engine/Source/Runtime/Engine/Private/Character.cpp
@@ -92,8 +92,8 @@ ACharacter::ACharacter(const FObjectInitializer& ObjectInitializer)
 		Mesh->PrimaryComponentTick.TickGroup = TG_PrePhysics;
 		Mesh->bChartDistanceFactor = true;
 		Mesh->AttachParent = CapsuleComponent;
-		static FName CollisionProfileName(TEXT("CharacterMesh"));
-		Mesh->SetCollisionProfileName(CollisionProfileName);
+		static FName MeshCollisionProfileName(TEXT("CharacterMesh"));
+		Mesh->SetCollisionProfileName(MeshCollisionProfileName);
 		Mesh->bGenerateOverlapEvents = false;
 		Mesh->bCanEverAffectNavigation = false;
 	}
diff --git a/Engine/Source/Runtime/Engine/Private/TileRendering.cpp b/Engine/Source/Runtime/Engine/Private/TileRendering.cpp
--- a/Engine/Source/Runtime/Engine/Private/TileRendering.cpp
+++ b/Engine/Source/Runtime/Engine/Private/TileRendering.cpp
@@ -319,8 +319,8 @@ bool FCanvasTileRendererItem::Render_GameThread(const FCanvas* Canvas)
 	{
 		View,
 		Data,
-		Canvas->IsHitTesting(),
-		bNeedsToSwitchVerticalAxis,
+		Canvas->IsHitTesting() ? uint32(1) : uint32(0),
+		bNeedsToSwitchVerticalAxis ? uint32(1) : uint32(0),
 		Canvas->GetAllowedModes()
 	};
 	ENQUEUE_UNIQUE_RENDER_COMMAND_ONEPARAMETER(
diff --git a/Engine/Source/Runtime/Engine/Private/TriangleRendering.cpp b/Engine/Source/Runtime/Engine/Private/TriangleRendering.cpp
--- a/Engine/Source/Runtime/Engine/Private/TriangleRendering.cpp
+++ b/Engine/Source/Runtime/Engine/Private/TriangleRendering.cpp
@@ -288,8 +288,8 @@ bool FCanvasTriangleRendererItem::Render_GameThread(const FCanvas* Canvas)
 	{
 		View,
 		Data,
-		bNeedsToSwitchVerticalAxis,
-		Canvas->IsHitTesting(),
+		(uint32)bNeedsToSwitchVerticalAxis,
+		(uint32)Canvas->IsHitTesting(),
 		Canvas->GetAllowedModes()
 	};
 	ENQUEUE_UNIQUE_RENDER_COMMAND_ONEPARAMETER(
diff --git a/Engine/Source/Runtime/Engine/Private/UnrealClient.cpp b/Engine/Source/Runtime/Engine/Private/UnrealClient.cpp
--- a/Engine/Source/Runtime/Engine/Private/UnrealClient.cpp
+++ b/Engine/Source/Runtime/Engine/Private/UnrealClient.cpp
@@ -1005,7 +1005,7 @@ void FViewport::Draw( bool bShouldPresent /*= true */)
 				SetRequiresVsync(bLockToVsync);
 
 				//@todo UE4: If Slate controls this viewport, we should not present
-				FEndDrawingCommandParams Params = { this, bLockToVsync, GInputLatencyTimer.GameThreadTrigger, PresentAndStopMovieDelay > 0 ? 0 : (uint32)bShouldPresent };
+				FEndDrawingCommandParams Params = { this, (uint32)bLockToVsync, (uint32)GInputLatencyTimer.GameThreadTrigger, (uint32)(PresentAndStopMovieDelay > 0 ? 0 : bShouldPresent) };
 				ENQUEUE_UNIQUE_RENDER_COMMAND_ONEPARAMETER(
 					EndDrawingCommand,
 					FEndDrawingCommandParams,Parameters,Params,
diff --git a/Engine/Source/Runtime/Engine/Private/UserInterface/Canvas.cpp b/Engine/Source/Runtime/Engine/Private/UserInterface/Canvas.cpp
--- a/Engine/Source/Runtime/Engine/Private/UserInterface/Canvas.cpp
+++ b/Engine/Source/Runtime/Engine/Private/UserInterface/Canvas.cpp
@@ -438,8 +438,8 @@ bool FCanvasBatchedElementRenderItem::Render_GameThread(const FCanvas* Canvas)
 		FBatchedDrawParameters DrawParameters =
 		{
 			Data,
-			Canvas->IsHitTesting(),
-			bNeedsToSwitchVerticalAxis,
+			(uint32)(Canvas->IsHitTesting() ? 1 : 0),
+			(uint32)(bNeedsToSwitchVerticalAxis ? 1 : 0),
 			(uint32)CanvasRenderTarget->GetSizeXY().X,
 			(uint32)CanvasRenderTarget->GetSizeXY().Y,
 			Gamma,
diff --git a/Engine/Source/Runtime/Engine/Private/UserInterface/Console.cpp b/Engine/Source/Runtime/Engine/Private/UserInterface/Console.cpp
--- a/Engine/Source/Runtime/Engine/Private/UserInterface/Console.cpp
+++ b/Engine/Source/Runtime/Engine/Private/UserInterface/Console.cpp
@@ -529,8 +529,8 @@ void UConsole::OutputText(const FString& Text)
 
 void UConsole::StartTyping(const FString& Text)
 {
-	static const FName NAME_Typing = FName(TEXT("Typing"));
-	FakeGotoState(NAME_Typing);
+	static const FName TypingName = FName(TEXT("Typing"));
+	FakeGotoState(TypingName);
 	SetInputText(Text);
 	SetCursorPos(Text.Len());
 }
@@ -1358,9 +1358,9 @@ void UConsole::FakeGotoState(FName NextStateName)
 void UConsole::Serialize( const TCHAR* V, ELogVerbosity::Type Verbosity, const class FName& Category )
 {
 	// e.g. UE_LOG(LogConsoleResponse, Display, TEXT("Test"));
-	static const FName LogConsoleResponse = FName("LogConsoleResponse");
+	static const FName ConsoleResponseLog = FName("LogConsoleResponse");
 
-	if (Category == LogConsoleResponse)
+	if (Category == ConsoleResponseLog)
 	{
 		// log all LogConsoleResponse
 		OutputText(V);
diff --git a/Engine/Source/Runtime/Engine/Public/AI/NavigationModifier.h b/Engine/Source/Runtime/Engine/Public/AI/NavigationModifier.h
--- a/Engine/Source/Runtime/Engine/Public/AI/NavigationModifier.h
+++ b/Engine/Source/Runtime/Engine/Public/AI/NavigationModifier.h
@@ -102,8 +102,8 @@ protected:
 	/** if set, area shape will be extended by agent's height to cover area underneath like regular colliding geometry */
 	uint8 bIncludeAgentHeight : 1;
 
-	void Init(const TSubclassOf<UNavArea> AreaClass);
-	void SetConvex(const FVector* Points, const int32 FirstIndex, const int32 LastIndex, ENavigationCoordSystem::Type CoordType, const FTransform& LocalToWorld);
+	void Init(const TSubclassOf<UNavArea> InAreaClass);
+	void SetConvex(const FVector* InPoints, const int32 FirstIndex, const int32 LastIndex, ENavigationCoordSystem::Type CoordType, const FTransform& LocalToWorld);
 	void SetBox(const FBox& Box, const FTransform& LocalToWorld);
 };
 
@@ -193,7 +193,7 @@ struct ENGINE_API FCustomLinkNavModifier : public FNavigationModifier
 	FTransform LocalToWorld;
 
 	FCustomLinkNavModifier() : LinkDefinitionClass(NULL) {}
-	void Set(TSubclassOf<UNavLinkDefinition> LinkDefinitionClass, const FTransform& LocalToWorld);
+	void Set(TSubclassOf<UNavLinkDefinition> LinkDefinitionClass, const FTransform& InLocalToWorld);
 	FORCEINLINE const TSubclassOf<UNavLinkDefinition> GetNavLinkClass() const { return LinkDefinitionClass; }
 
 protected:
diff --git a/Engine/Source/Runtime/Engine/Public/ActiveSound.h b/Engine/Source/Runtime/Engine/Public/ActiveSound.h
--- a/Engine/Source/Runtime/Engine/Public/ActiveSound.h
+++ b/Engine/Source/Runtime/Engine/Public/ActiveSound.h
@@ -185,7 +185,7 @@ public:
 #endif
 
 	// Updates the wave instances to be played.
-	void UpdateWaveInstances( FAudioDevice* AudioDevice, TArray<FWaveInstance*> &WaveInstances, const float DeltaTime );
+	void UpdateWaveInstances( FAudioDevice* AudioDevice, TArray<FWaveInstance*> &OutWaveInstances, const float DeltaTime );
 
 	void Stop(FAudioDevice* AudioDevice);
 
diff --git a/Engine/Source/Runtime/Engine/Public/PrimitiveSceneProxy.h b/Engine/Source/Runtime/Engine/Public/PrimitiveSceneProxy.h
--- a/Engine/Source/Runtime/Engine/Public/PrimitiveSceneProxy.h
+++ b/Engine/Source/Runtime/Engine/Public/PrimitiveSceneProxy.h
@@ -745,7 +745,7 @@ private:
 #endif
 
 	/** Updates the proxy's actor position, called from the game thread. */
-	ENGINE_API void UpdateActorPosition(FVector ActorPosition);
+	ENGINE_API void UpdateActorPosition(FVector InActorPosition);
 
 	/**
 	 * Updates the primitive proxy's cached transforms, and calls OnUpdateTransform to notify it of the change.
@@ -754,7 +754,7 @@ private:
 	 * @param InBounds - The new bounds of the primitive.
 	 * @param InLocalBounds - The local space bounds of the primitive.
 	 */
-	ENGINE_API void SetTransform(const FMatrix& InLocalToWorld, const FBoxSphereBounds& InBounds, const FBoxSphereBounds& InLocalBounds, FVector ActorPosition);
+	ENGINE_API void SetTransform(const FMatrix& InLocalToWorld, const FBoxSphereBounds& InBounds, const FBoxSphereBounds& InLocalBounds, FVector InActorPosition);
 
 	/**
 	 * Either updates the uniform buffer or defers it until it becomes visible depending on a cvar
diff --git a/Engine/Source/Runtime/Engine/Public/SceneManagement.h b/Engine/Source/Runtime/Engine/Public/SceneManagement.h
--- a/Engine/Source/Runtime/Engine/Public/SceneManagement.h
+++ b/Engine/Source/Runtime/Engine/Public/SceneManagement.h
@@ -1302,7 +1302,7 @@ public:
 		return 0;
 	}
 
-	void DrawBatchedElements(FRHICommandList& RHICmdList, const FSceneView& View, FTexture2DRHIRef DepthTexture, EBlendModeFilter::Type Filter) const;
+	void DrawBatchedElements(FRHICommandList& RHICmdList, const FSceneView& InView, FTexture2DRHIRef DepthTexture, EBlendModeFilter::Type Filter) const;
 
 	/** The batched simple elements. */
 	FBatchedElements BatchedElements;
diff --git a/Engine/Source/Runtime/Engine/Public/SceneView.h b/Engine/Source/Runtime/Engine/Public/SceneView.h
--- a/Engine/Source/Runtime/Engine/Public/SceneView.h
+++ b/Engine/Source/Runtime/Engine/Public/SceneView.h
@@ -702,7 +702,7 @@ public:
 	void UpdateViewMatrix();
 
 	/** Setup defaults and depending on view position (postprocess volumes) */
-	void StartFinalPostprocessSettings(FVector ViewLocation);
+	void StartFinalPostprocessSettings(FVector InViewLocation);
 
 	/**
 	 * custom layers can be combined with the existing settings
diff --git a/Engine/Source/Runtime/Engine/Public/SkeletalMeshTypes.h b/Engine/Source/Runtime/Engine/Public/SkeletalMeshTypes.h
--- a/Engine/Source/Runtime/Engine/Public/SkeletalMeshTypes.h
+++ b/Engine/Source/Runtime/Engine/Public/SkeletalMeshTypes.h
@@ -1242,7 +1242,7 @@ private:
 	/**
 	 * Resizes the vertex data storage & updates the cached info.
 	 */
-	void ResizeData(int32 NumVertices);
+	void ResizeData(int32 InNumVertices);
 
 	/**
 	 * Update the cached 'VertexData' information.
@@ -1830,7 +1830,7 @@ protected:
 	
 	void GetDynamicElementsSection(const TArray<const FSceneView*>& Views, const FSceneViewFamily& ViewFamily, uint32 VisibilityMap, 
 		const FStaticLODModel& LODModel, const int32 LODIndex, const FSkelMeshSection& Section, const FSkelMeshChunk& Chunk, 
-		const FSectionElementInfo& SectionElementInfo, const FTwoVectors& CustomLeftRightVectors, bool bSelectable, FMeshElementCollector& Collector ) const;
+		const FSectionElementInfo& SectionElementInfo, const FTwoVectors& CustomLeftRightVectors, bool bInSelectable, FMeshElementCollector& Collector ) const;
 
-	void GetMeshElementsConditionallySelectable(const TArray<const FSceneView*>& Views, const FSceneViewFamily& ViewFamily, bool bSelectable, uint32 VisibilityMap, FMeshElementCollector& Collector) const;
+	void GetMeshElementsConditionallySelectable(const TArray<const FSceneView*>& Views, const FSceneViewFamily& ViewFamily, bool bInSelectable, uint32 VisibilityMap, FMeshElementCollector& Collector) const;
 };
diff --git a/Engine/Source/Runtime/OpenGLDrv/Private/Windows/OpenGLWindows.cpp b/Engine/Source/Runtime/OpenGLDrv/Private/Windows/OpenGLWindows.cpp
--- a/Engine/Source/Runtime/OpenGLDrv/Private/Windows/OpenGLWindows.cpp
+++ b/Engine/Source/Runtime/OpenGLDrv/Private/Windows/OpenGLWindows.cpp
@@ -454,9 +454,9 @@ bool PlatformBlitToViewport( FPlatformOpenGLDevice* Device, const FOpenGLViewpor
 		{
 			Device->TargetDirty = false;
 			glDisable(GL_FRAMEBUFFER_SRGB);
-			bool bPresent = Viewport.GetCustomPresent()->Present(SyncInterval);
+			bool bShouldPresent = Viewport.GetCustomPresent()->Present(SyncInterval);
 			glEnable(GL_FRAMEBUFFER_SRGB);
-			if (!bPresent)
+			if (!bShouldPresent)
 			{
 				return false;
 			}
diff --git a/Engine/Source/Runtime/RHI/Public/RHI.h b/Engine/Source/Runtime/RHI/Public/RHI.h
--- a/Engine/Source/Runtime/RHI/Public/RHI.h
+++ b/Engine/Source/Runtime/RHI/Public/RHI.h
@@ -845,21 +845,21 @@ struct FResolveParams
 struct FUpdateTextureRegion2D
 {
 	/** offset in texture */
-	int32 DestX;
-	int32 DestY;
+	uint32 DestX;
+	uint32 DestY;
 	
 	/** offset in source image data */
 	int32 SrcX;
 	int32 SrcY;
 	
 	/** size of region to copy */
-	int32 Width;
-	int32 Height;
+	uint32 Width;
+	uint32 Height;
 
 	FUpdateTextureRegion2D()
 	{}
 
-	FUpdateTextureRegion2D(int32 InDestX, int32 InDestY, int32 InSrcX, int32 InSrcY, int32 InWidth, int32 InHeight)
+	FUpdateTextureRegion2D(uint32 InDestX, uint32 InDestY, int32 InSrcX, int32 InSrcY, uint32 InWidth, uint32 InHeight)
 	:	DestX(InDestX)
 	,	DestY(InDestY)
 	,	SrcX(InSrcX)
@@ -884,9 +884,9 @@ struct FUpdateTextureRegion2D
 struct FUpdateTextureRegion3D
 {
 	/** offset in texture */
-	int32 DestX;
-	int32 DestY;
-	int32 DestZ;
+	uint32 DestX;
+	uint32 DestY;
+	uint32 DestZ;
 
 	/** offset in source image data */
 	int32 SrcX;
@@ -894,14 +894,14 @@ struct FUpdateTextureRegion3D
 	int32 SrcZ;
 
 	/** size of region to copy */
-	int32 Width;
-	int32 Height;
-	int32 Depth;
+	uint32 Width;
+	uint32 Height;
+	uint32 Depth;
 
 	FUpdateTextureRegion3D()
 	{}
 
-	FUpdateTextureRegion3D(int32 InDestX, int32 InDestY, int32 InDestZ, int32 InSrcX, int32 InSrcY, int32 InSrcZ, int32 InWidth, int32 InHeight, int32 InDepth)
+	FUpdateTextureRegion3D(uint32 InDestX, uint32 InDestY, uint32 InDestZ, int32 InSrcX, int32 InSrcY, int32 InSrcZ, uint32 InWidth, uint32 InHeight, uint32 InDepth)
 	:	DestX(InDestX)
 	,	DestY(InDestY)
 	,	DestZ(InDestZ)
diff --git a/Engine/Source/Runtime/Renderer/Private/PostProcess/SceneRenderTargets.cpp b/Engine/Source/Runtime/Renderer/Private/PostProcess/SceneRenderTargets.cpp
--- a/Engine/Source/Runtime/Renderer/Private/PostProcess/SceneRenderTargets.cpp
+++ b/Engine/Source/Runtime/Renderer/Private/PostProcess/SceneRenderTargets.cpp
@@ -406,8 +406,8 @@ void FSceneRenderTargets::Allocate(const FSceneViewFamily& ViewFamily)
 
 	uint32 Mobile32bpp = !IsMobileHDR() || IsMobileHDR32bpp();
 
-	static const auto CVarMobileMSAA = IConsoleManager::Get().FindTConsoleVariableDataInt(TEXT("r.MobileMSAA"));
-	int32 MobileMSAA = GShaderPlatformForFeatureLevel[NewFeatureLevel] == SP_OPENGL_ES2_IOS ? 1 : CVarMobileMSAA->GetValueOnRenderThread();
+	static const auto MobileMSAACVar = IConsoleManager::Get().FindTConsoleVariableDataInt(TEXT("r.MobileMSAA"));
+	int32 MobileMSAA = GShaderPlatformForFeatureLevel[NewFeatureLevel] == SP_OPENGL_ES2_IOS ? 1 : MobileMSAACVar->GetValueOnRenderThread();
 
 	bool bLightPropagationVolume = UseLightPropagationVolumeRT(NewFeatureLevel);
 
diff --git a/Engine/Source/Runtime/Renderer/Private/ScenePrivate.h b/Engine/Source/Runtime/Renderer/Private/ScenePrivate.h
--- a/Engine/Source/Runtime/Renderer/Private/ScenePrivate.h
+++ b/Engine/Source/Runtime/Renderer/Private/ScenePrivate.h
@@ -1639,7 +1639,7 @@ public:
 	virtual void GetRelevantLights( UPrimitiveComponent* Primitive, TArray<const ULightComponent*>* RelevantLights ) const override;
 
 	/** Sets the precomputed visibility handler for the scene, or NULL to clear the current one. */
-	virtual void SetPrecomputedVisibility(const FPrecomputedVisibilityHandler* PrecomputedVisibilityHandler) override;
+	virtual void SetPrecomputedVisibility(const FPrecomputedVisibilityHandler* InPrecomputedVisibilityHandler) override;
 
 	/** Sets shader maps on the specified materials without blocking. */
 	virtual void SetShaderMapsOnMaterialResources(const TMap<FMaterial*, class FMaterialShaderMap*>& MaterialsToUpdate) override;
diff --git a/Engine/Source/Runtime/Renderer/Private/SceneRendering.h b/Engine/Source/Runtime/Renderer/Private/SceneRendering.h
--- a/Engine/Source/Runtime/Renderer/Private/SceneRendering.h
+++ b/Engine/Source/Runtime/Renderer/Private/SceneRendering.h
@@ -821,9 +821,9 @@ protected:
 	void PostVisibilityFrameSetup();
 
 	void GatherDynamicMeshElements(
-		TArray<FViewInfo>& Views, 
-		const FScene* Scene, 
-		const FSceneViewFamily& ViewFamily, 
+		TArray<FViewInfo>& InViews, 
+		const FScene* InScene, 
+		const FSceneViewFamily& InViewFamily, 
 		const FPrimitiveViewMasks& HasDynamicMeshElementsMasks, 
 		const FPrimitiveViewMasks& HasDynamicEditorMeshElementsMasks, 
 		FMeshElementCollector& Collector);
diff --git a/Engine/Source/Runtime/Renderer/Private/StaticMeshDrawList.h b/Engine/Source/Runtime/Renderer/Private/StaticMeshDrawList.h
--- a/Engine/Source/Runtime/Renderer/Private/StaticMeshDrawList.h
+++ b/Engine/Source/Runtime/Renderer/Private/StaticMeshDrawList.h
@@ -333,7 +333,7 @@ public:
 	void SortFrontToBack(FVector ViewPosition);
 
 	/** Builds a list of primitives that use the given materials in this static draw list. */
-	void GetUsedPrimitivesBasedOnMaterials(ERHIFeatureLevel::Type FeatureLevel, const TArray<const FMaterial*>& Materials, TArray<FPrimitiveSceneInfo*>& PrimitivesToUpdate);
+	void GetUsedPrimitivesBasedOnMaterials(ERHIFeatureLevel::Type InFeatureLevel, const TArray<const FMaterial*>& Materials, TArray<FPrimitiveSceneInfo*>& PrimitivesToUpdate);
 
 	/**
 	 * Shifts all meshes bounds by an arbitrary delta.
diff --git a/Engine/Source/Runtime/Slate/Private/Framework/Application/SlateApplication.cpp b/Engine/Source/Runtime/Slate/Private/Framework/Application/SlateApplication.cpp
--- a/Engine/Source/Runtime/Slate/Private/Framework/Application/SlateApplication.cpp
+++ b/Engine/Source/Runtime/Slate/Private/Framework/Application/SlateApplication.cpp
@@ -5497,10 +5497,10 @@ TSharedRef<FSlateApplication> FSlateApplication::InitializeAsStandaloneApplicati
 }
 
 
-TSharedRef<FSlateApplication> FSlateApplication::InitializeAsStandaloneApplication(const TSharedRef< class FSlateRenderer >& PlatformRenderer, const TSharedRef<class GenericApplication>& PlatformApplication)
+TSharedRef<FSlateApplication> FSlateApplication::InitializeAsStandaloneApplication(const TSharedRef< class FSlateRenderer >& PlatformRenderer, const TSharedRef<class GenericApplication>& InPlatformApplication)
 {
 	// create the platform slate application (what FSlateApplication::Get() returns)
-	TSharedRef<FSlateApplication> Slate = FSlateApplication::Create(PlatformApplication);
+	TSharedRef<FSlateApplication> Slate = FSlateApplication::Create(InPlatformApplication);
 
 	// initialize renderer
 	FSlateApplication::Get().InitializeRenderer(PlatformRenderer);
diff --git a/Engine/Source/Runtime/Slate/Private/Framework/MultiBox/MultiBoxBuilder.cpp b/Engine/Source/Runtime/Slate/Private/Framework/MultiBox/MultiBoxBuilder.cpp
--- a/Engine/Source/Runtime/Slate/Private/Framework/MultiBox/MultiBoxBuilder.cpp
+++ b/Engine/Source/Runtime/Slate/Private/Framework/MultiBox/MultiBoxBuilder.cpp
@@ -316,8 +316,8 @@ void FMenuBarBuilder::AddPullDownMenu( const FText& InMenuLabel, const FText& In
 	const bool bIsSubMenu = false;
 	const bool bOpenSubMenuOnClick = false;
 	// Pulldown menus always close all menus not just themselves
-	const bool bCloseSelfOnly = false;
-	TSharedRef< FMenuEntryBlock > NewMenuEntryBlock( new FMenuEntryBlock( InExtensionHook, InMenuLabel, InToolTip, InPullDownMenu, ExtenderStack.Top(), bIsSubMenu, bOpenSubMenuOnClick, CommandListStack.Last(), bCloseSelfOnly ) );
+	const bool bShouldCloseSelfOnly = false;
+	TSharedRef< FMenuEntryBlock > NewMenuEntryBlock( new FMenuEntryBlock( InExtensionHook, InMenuLabel, InToolTip, InPullDownMenu, ExtenderStack.Top(), bIsSubMenu, bOpenSubMenuOnClick, CommandListStack.Last(), bShouldCloseSelfOnly) );
 	NewMenuEntryBlock->SetTutorialHighlightName(GenerateTutorialIdentfierName(TutorialHighlightName, InTutorialHighlightName, nullptr, MultiBox->GetBlocks().Num()));
 
 	MultiBox->AddMultiBlock( NewMenuEntryBlock );
@@ -396,17 +396,17 @@ void FToolBarBuilder::AddComboButton( const FUIAction& InAction, const FOnGetCon
 	MultiBox->AddMultiBlock( NewToolBarComboButtonBlock );
 }
 
-void FToolBarBuilder::AddWidget( TSharedRef<SWidget> InWidget, FName TutorialHighlightName )
+void FToolBarBuilder::AddWidget( TSharedRef<SWidget> InWidget, FName InTutorialHighlightName )
 {
 	ApplySectionBeginning();
 
 	// If tutorial name specified, wrap in tutorial wrapper
-	const FName WrapperName = GenerateTutorialIdentfierName(TutorialHighlightName, NAME_None, nullptr, MultiBox->GetBlocks().Num());
+	const FName WrapperName = GenerateTutorialIdentfierName(InTutorialHighlightName, NAME_None, nullptr, MultiBox->GetBlocks().Num());
 
 	TSharedRef<SWidget> ChildWidget = InWidget;
 	InWidget = 
 		SNew( SBox )
-		.AddMetaData<FTagMetaData>(FTagMetaData(TutorialHighlightName))
+		.AddMetaData<FTagMetaData>(FTagMetaData(InTutorialHighlightName))
 		[
 			ChildWidget
 		];
diff --git a/Engine/Source/Runtime/Slate/Private/Framework/MultiBox/SEditableTextBlock.cpp b/Engine/Source/Runtime/Slate/Private/Framework/MultiBox/SEditableTextBlock.cpp
--- a/Engine/Source/Runtime/Slate/Private/Framework/MultiBox/SEditableTextBlock.cpp
+++ b/Engine/Source/Runtime/Slate/Private/Framework/MultiBox/SEditableTextBlock.cpp
@@ -50,7 +50,7 @@ void SEditableTextBlock::BuildMultiBlockWidget( const ISlateStyle* StyleSet, con
 	const TSharedPtr< const FUICommandInfo >& UICommand = EditableTextBlock->GetAction();
 
 	const TAttribute<FText>& Label = !EditableTextBlock->LabelOverride.IsBound() && EditableTextBlock->LabelOverride.Get().IsEmpty() && UICommand.IsValid() ? UICommand->GetLabel() : EditableTextBlock->LabelOverride;
-	const TAttribute<FText>& ToolTip = !EditableTextBlock->ToolTipOverride.IsBound() && EditableTextBlock->ToolTipOverride.Get().IsEmpty() && UICommand.IsValid() ? UICommand->GetDescription() : EditableTextBlock->ToolTipOverride;
+	const TAttribute<FText>& TextBlockToolTip = !EditableTextBlock->ToolTipOverride.IsBound() && EditableTextBlock->ToolTipOverride.Get().IsEmpty() && UICommand.IsValid() ? UICommand->GetDescription() : EditableTextBlock->ToolTipOverride;
 	const bool bHasLabel = !Label.Get().IsEmpty();
 
 	// See if the action is valid and if so we will use the actions icon if we dont override it later
@@ -111,7 +111,7 @@ void SEditableTextBlock::BuildMultiBlockWidget( const ISlateStyle* StyleSet, con
 					SNew( STextBlock )
 					.TextStyle( StyleSet, ISlateStyle::Join(StyleName, ".Label") )
 					.Text( Label )
-					.ToolTipText( ToolTip )
+					.ToolTipText( TextBlockToolTip )
 				]
 		]
 
@@ -129,7 +129,7 @@ void SEditableTextBlock::BuildMultiBlockWidget( const ISlateStyle* StyleSet, con
 			.MinDesiredWidth( MultiBoxConstants::EditableTextMinWidth )
 			.OnTextChanged( EditableTextBlock->OnTextChanged )
 			.OnTextCommitted( EditableTextBlock->OnTextCommitted )
-			.ToolTipText( ToolTip )
+			.ToolTipText( TextBlockToolTip )
 		]
 	];
 
diff --git a/Engine/Source/Runtime/Slate/Private/Framework/MultiBox/SMenuEntryBlock.cpp b/Engine/Source/Runtime/Slate/Private/Framework/MultiBox/SMenuEntryBlock.cpp
--- a/Engine/Source/Runtime/Slate/Private/Framework/MultiBox/SMenuEntryBlock.cpp
+++ b/Engine/Source/Runtime/Slate/Private/Framework/MultiBox/SMenuEntryBlock.cpp
@@ -181,7 +181,7 @@ void SMenuEntryBlock::Construct( const FArguments& InArgs )
 TSharedRef< SWidget>  SMenuEntryBlock::BuildMenuBarWidget( const FMenuEntryBuildParams& InBuildParams )
 {
 	const TAttribute<FText>& Label = InBuildParams.Label;
-	const TAttribute<FText>& ToolTip = InBuildParams.ToolTip;
+	const TAttribute<FText>& EntryToolTip = InBuildParams.ToolTip;
 
 	check( OwnerMultiBoxWidget.IsValid() );
 
@@ -222,7 +222,7 @@ TSharedRef< SWidget>  SMenuEntryBlock::BuildMenuBarWidget( const FMenuEntryBuild
 				.ClickMethod( EButtonClickMethod::MouseDown )
 
 				// Pass along the block's tool-tip string
-				.ToolTipText( this, &SMenuEntryBlock::GetFilteredToolTipText, ToolTip )
+				.ToolTipText( this, &SMenuEntryBlock::GetFilteredToolTipText, EntryToolTip)
 
 				// Add horizontal padding between the edge of the button and the content.  Also add a bit of vertical
 				// padding to push the text down from the top of the menu bar a bit.
@@ -280,7 +280,7 @@ EVisibility SMenuEntryBlock::GetVisibility() const
 TSharedRef< SWidget > SMenuEntryBlock::BuildMenuEntryWidget( const FMenuEntryBuildParams& InBuildParams )
 {
 	const TAttribute<FText>& Label = InBuildParams.Label;
-	const TAttribute<FText>& ToolTip = InBuildParams.ToolTip;
+	const TAttribute<FText>& EntryToolTip = InBuildParams.ToolTip;
 	const TSharedPtr< const FMenuEntryBlock > MenuEntryBlock = InBuildParams.MenuEntryBlock;
 	const TSharedPtr< const FMultiBox > MultiBox = InBuildParams.MultiBox;
 	const TSharedPtr< const FUICommandInfo >& UICommand = InBuildParams.UICommand;
@@ -462,7 +462,7 @@ TSharedRef< SWidget > SMenuEntryBlock::BuildMenuEntryWidget( const FMenuEntryBui
 		.ClickMethod( ButtonClickMethod )
 
 		// Pass along the block's tool-tip string
-		.ToolTip( FMultiBoxSettings::ToolTipConstructor.Execute( ToolTip, nullptr, UICommand ) )
+		.ToolTip( FMultiBoxSettings::ToolTipConstructor.Execute(EntryToolTip, nullptr, UICommand ) )
 
 		.ContentPadding(FMargin(0, 2))
 
@@ -568,7 +568,7 @@ private:
 TSharedRef< SWidget> SMenuEntryBlock::BuildSubMenuWidget( const FMenuEntryBuildParams& InBuildParams )
 {
 	const TAttribute<FText>& Label = InBuildParams.Label;
-	const TAttribute<FText>& ToolTip = InBuildParams.ToolTip;
+	const TAttribute<FText>& EntryToolTip = InBuildParams.ToolTip;
 
 	const TSharedPtr< const FMenuEntryBlock > MenuEntryBlock = InBuildParams.MenuEntryBlock;
 	const TSharedPtr< const FMultiBox > MultiBox = InBuildParams.MultiBox;
@@ -699,7 +699,7 @@ TSharedRef< SWidget> SMenuEntryBlock::BuildSubMenuWidget( const FMenuEntryBuildP
 			// Create a button
 			SNew( SSubMenuButton )
 			// Pass along the block's tool-tip string
-			.ToolTipText( ToolTip )
+			.ToolTipText( EntryToolTip )
 			// Style to use
 			.ButtonStyle( &StyleSet->GetWidgetStyle<FButtonStyle>( ISlateStyle::Join( StyleName, ".Button" ) ) )
 			// Allow the button to change its state depending on the state of the submenu
diff --git a/Engine/Source/Runtime/Slate/Private/Framework/MultiBox/SToolBarButtonBlock.cpp b/Engine/Source/Runtime/Slate/Private/Framework/MultiBox/SToolBarButtonBlock.cpp
--- a/Engine/Source/Runtime/Slate/Private/Framework/MultiBox/SToolBarButtonBlock.cpp
+++ b/Engine/Source/Runtime/Slate/Private/Framework/MultiBox/SToolBarButtonBlock.cpp
@@ -31,12 +31,12 @@ FToolBarButtonBlock::FToolBarButtonBlock( const TAttribute<FText>& InLabel, cons
 
 void FToolBarButtonBlock::CreateMenuEntry(FMenuBuilder& MenuBuilder) const
 {
-	TSharedPtr<const FUICommandInfo> Action = GetAction();
-	TSharedPtr<const FUICommandList> ActionList = GetActionList();
-	if (Action.IsValid() && ActionList.IsValid())
+	TSharedPtr<const FUICommandInfo> MenuEntryAction = GetAction();
+	TSharedPtr<const FUICommandList> MenuEntryActionList = GetActionList();
+	if (MenuEntryAction.IsValid() && MenuEntryActionList.IsValid())
 	{
-		MenuBuilder.PushCommandList(ActionList.ToSharedRef());
-		MenuBuilder.AddMenuEntry(Action);				
+		MenuBuilder.PushCommandList(MenuEntryActionList.ToSharedRef());
+		MenuBuilder.AddMenuEntry(MenuEntryAction);
 		MenuBuilder.PopCommandList();
 	}
 	else if ( LabelOverride.IsSet() )
diff --git a/Engine/Source/Runtime/Slate/Private/Framework/Text/SlateTextLayout.cpp b/Engine/Source/Runtime/Slate/Private/Framework/Text/SlateTextLayout.cpp
--- a/Engine/Source/Runtime/Slate/Private/Framework/Text/SlateTextLayout.cpp
+++ b/Engine/Source/Runtime/Slate/Private/Framework/Text/SlateTextLayout.cpp
@@ -147,10 +147,10 @@ const FTextBlockStyle& FSlateTextLayout::GetDefaultTextStyle() const
 void FSlateTextLayout::AggregateChildren()
 {
 	Children.Empty();
-	const TArray< FLineModel >& LineModels = GetLineModels();
-	for (int32 LineModelIndex = 0; LineModelIndex < LineModels.Num(); LineModelIndex++)
+	const TArray< FLineModel >& LayoutLineModels = GetLineModels();
+	for (int32 LineModelIndex = 0; LineModelIndex < LayoutLineModels.Num(); LineModelIndex++)
 	{
-		const FLineModel& LineModel = LineModels[ LineModelIndex ];
+		const FLineModel& LineModel = LayoutLineModels[ LineModelIndex ];
 		for (int32 RunIndex = 0; RunIndex < LineModel.Runs.Num(); RunIndex++)
 		{
 			const FRunModel& LineRun = LineModel.Runs[ RunIndex ];
diff --git a/Engine/Source/Runtime/Slate/Private/Widgets/Input/SEditableTextBox.cpp b/Engine/Source/Runtime/Slate/Private/Widgets/Input/SEditableTextBox.cpp
--- a/Engine/Source/Runtime/Slate/Private/Widgets/Input/SEditableTextBox.cpp
+++ b/Engine/Source/Runtime/Slate/Private/Widgets/Input/SEditableTextBox.cpp
@@ -21,14 +21,14 @@ void SEditableTextBox::Construct( const FArguments& InArgs )
 
 	TAttribute<FMargin> Padding = PaddingOverride.IsSet() ? PaddingOverride : InArgs._Style->Padding;
 	TAttribute<FSlateFontInfo> Font = FontOverride.IsSet() ? FontOverride : InArgs._Style->Font;
-	TAttribute<FSlateColor> ForegroundColor = ForegroundColorOverride.IsSet() ? ForegroundColorOverride : InArgs._Style->ForegroundColor;
+	TAttribute<FSlateColor> BorderForegroundColor = ForegroundColorOverride.IsSet() ? ForegroundColorOverride : InArgs._Style->ForegroundColor;
 	TAttribute<FSlateColor> BackgroundColor = BackgroundColorOverride.IsSet() ? BackgroundColorOverride : InArgs._Style->BackgroundColor;
 	ReadOnlyForegroundColor = ReadOnlyForegroundColorOverride.IsSet() ? ReadOnlyForegroundColorOverride : InArgs._Style->ReadOnlyForegroundColor;
 
 	SBorder::Construct( SBorder::FArguments()
 		.BorderImage( this, &SEditableTextBox::GetBorderImage )
 		.BorderBackgroundColor( BackgroundColor )
-		.ForegroundColor( ForegroundColor )
+		.ForegroundColor( BorderForegroundColor )
 		.Padding( 0 )
 		[
 			SAssignNew( Box, SHorizontalBox)
diff --git a/Engine/Source/Runtime/Slate/Private/Widgets/Input/SMultiLineEditableTextBox.cpp b/Engine/Source/Runtime/Slate/Private/Widgets/Input/SMultiLineEditableTextBox.cpp
--- a/Engine/Source/Runtime/Slate/Private/Widgets/Input/SMultiLineEditableTextBox.cpp
+++ b/Engine/Source/Runtime/Slate/Private/Widgets/Input/SMultiLineEditableTextBox.cpp
@@ -49,7 +49,7 @@ void SMultiLineEditableTextBox::Construct( const FArguments& InArgs )
 	TAttribute<FMargin> HScrollBarPadding = InArgs._HScrollBarPadding.IsSet() ? InArgs._HScrollBarPadding : InArgs._Style->HScrollBarPadding;
 	TAttribute<FMargin> VScrollBarPadding = InArgs._VScrollBarPadding.IsSet() ? InArgs._VScrollBarPadding : InArgs._Style->VScrollBarPadding;
 	TAttribute<FSlateFontInfo> Font = InArgs._Font.IsSet() ? InArgs._Font : InArgs._Style->Font;
-	TAttribute<FSlateColor> ForegroundColor = InArgs._ForegroundColor.IsSet() ? InArgs._ForegroundColor : InArgs._Style->ForegroundColor;
+	TAttribute<FSlateColor> BorderForegroundColor = InArgs._ForegroundColor.IsSet() ? InArgs._ForegroundColor : InArgs._Style->ForegroundColor;
 	TAttribute<FSlateColor> BackgroundColor = InArgs._BackgroundColor.IsSet() ? InArgs._BackgroundColor : InArgs._Style->BackgroundColor;
 	ReadOnlyForegroundColor = InArgs._ReadOnlyForegroundColor.IsSet() ? InArgs._ReadOnlyForegroundColor : InArgs._Style->ReadOnlyForegroundColor;
 
@@ -80,7 +80,7 @@ void SMultiLineEditableTextBox::Construct( const FArguments& InArgs )
 	SBorder::Construct( SBorder::FArguments()
 		.BorderImage( this, &SMultiLineEditableTextBox::GetBorderImage )
 		.BorderBackgroundColor( BackgroundColor )
-		.ForegroundColor( ForegroundColor )
+		.ForegroundColor( BorderForegroundColor )
 		.Padding( Padding )
 		[
 			SAssignNew( Box, SHorizontalBox)
diff --git a/Engine/Source/Runtime/Slate/Private/Widgets/Input/SSlider.cpp b/Engine/Source/Runtime/Slate/Private/Widgets/Input/SSlider.cpp
--- a/Engine/Source/Runtime/Slate/Private/Widgets/Input/SSlider.cpp
+++ b/Engine/Source/Runtime/Slate/Private/Widgets/Input/SSlider.cpp
@@ -53,12 +53,12 @@ int32 SSlider::OnPaint( const FPaintArgs& Args, const FGeometry& AllottedGeometr
 	if (Orientation == Orient_Vertical)
 	{
 		// Do this by translating along -X by the width of the geometry, then rotating 90 degreess CCW (left-hand coords)
-		FSlateRenderTransform RenderTransform = TransformCast<FSlateRenderTransform>(Concatenate(Inverse(FVector2D(AllottedWidth, 0)), FQuat2D(FMath::DegreesToRadians(-90.0f))));
+		FSlateRenderTransform SlateRenderTransform = TransformCast<FSlateRenderTransform>(Concatenate(Inverse(FVector2D(AllottedWidth, 0)), FQuat2D(FMath::DegreesToRadians(-90.0f))));
 		// create a child geometry matching this one, but with the render transform.
 		SliderGeometry = AllottedGeometry.MakeChild(
 			FVector2D(AllottedWidth, AllottedHeight), 
 			FSlateLayoutTransform(), 
-			RenderTransform, FVector2D::ZeroVector);
+			SlateRenderTransform, FVector2D::ZeroVector);
 		// The clipping rect is already given properly in window space. But we do not support layout rotations, so our local space rendering cannot
 		// get the clipping rect into local space properly for the local space clipping we do in the shader.
 		// Thus, we transform the clip coords into local space manually, UNDO the render transform so it will clip properly,
@@ -66,7 +66,7 @@ int32 SSlider::OnPaint( const FPaintArgs& Args, const FGeometry& AllottedGeometr
 		RotatedClippingRect = TransformRect(
 			Concatenate(
 				Inverse(SliderGeometry.GetAccumulatedLayoutTransform()), 
-				Inverse(RenderTransform), 
+				Inverse(SlateRenderTransform),
 				SliderGeometry.GetAccumulatedLayoutTransform()), 
 			MyClippingRect);
 	}
diff --git a/Engine/Source/Runtime/Slate/Private/Widgets/Layout/SExpandableArea.cpp b/Engine/Source/Runtime/Slate/Private/Widgets/Layout/SExpandableArea.cpp
--- a/Engine/Source/Runtime/Slate/Private/Widgets/Layout/SExpandableArea.cpp
+++ b/Engine/Source/Runtime/Slate/Private/Widgets/Layout/SExpandableArea.cpp
@@ -176,14 +176,14 @@ FVector2D SExpandableArea::ComputeDesiredSize( float ) const
 	EVisibility ChildVisibility = ChildSlot.GetWidget()->GetVisibility();
 	if ( ChildVisibility != EVisibility::Collapsed )
 	{
-		FVector2D DesiredSize = ChildSlot.GetWidget()->GetDesiredSize() + ChildSlot.SlotPadding.Get().GetDesiredSize();
+		FVector2D SlotWidgetDesiredSize = ChildSlot.GetWidget()->GetDesiredSize() + ChildSlot.SlotPadding.Get().GetDesiredSize();
 		
 		// Only clamp if the user specified a max height
 		if( MaxHeight > 0.0f )
 		{
-			DesiredSize.Y = FMath::Min( MaxHeight, DesiredSize.Y );
+			SlotWidgetDesiredSize.Y = FMath::Min( MaxHeight, SlotWidgetDesiredSize.Y );
 		}
-		return DesiredSize;
+		return SlotWidgetDesiredSize;
 	}
 
 	return FVector2D::ZeroVector;
diff --git a/Engine/Source/Runtime/Slate/Private/Widgets/Layout/SFxWidget.cpp b/Engine/Source/Runtime/Slate/Private/Widgets/Layout/SFxWidget.cpp
--- a/Engine/Source/Runtime/Slate/Private/Widgets/Layout/SFxWidget.cpp
+++ b/Engine/Source/Runtime/Slate/Private/Widgets/Layout/SFxWidget.cpp
@@ -61,9 +61,9 @@ int32 SFxWidget::OnPaint( const FPaintArgs& Args, const FGeometry& AllottedGeome
 	const FVector2D ScaleOrigin = RenderScaleOrigin.Get() * AllottedGeometry.Size;
 	const FVector2D Offset = VisualOffset.Get() * AllottedGeometry.Size;
 	// create the render transform as a scale around ScaleOrigin and offset it by Offset.
-	const auto RenderTransform = Concatenate(Inverse(ScaleOrigin), RenderScale.Get(), ScaleOrigin, Offset);
+	const auto SlateRenderTransform = Concatenate(Inverse(ScaleOrigin), RenderScale.Get(), ScaleOrigin, Offset);
 	// This will append the render transform to the layout transform, and we only use it for rendering.
-	FGeometry ModifiedGeometry = AllottedGeometry.MakeChild(AllottedGeometry.Size, RenderTransform);
+	FGeometry ModifiedGeometry = AllottedGeometry.MakeChild(AllottedGeometry.Size, SlateRenderTransform);
 	
 	FArrangedChildren ArrangedChildren(EVisibility::Visible);
 	this->ArrangeChildren(ModifiedGeometry, ArrangedChildren);
diff --git a/Engine/Source/Runtime/Slate/Private/Widgets/Layout/SScaleBox.cpp b/Engine/Source/Runtime/Slate/Private/Widgets/Layout/SScaleBox.cpp
--- a/Engine/Source/Runtime/Slate/Private/Widgets/Layout/SScaleBox.cpp
+++ b/Engine/Source/Runtime/Slate/Private/Widgets/Layout/SScaleBox.cpp
@@ -32,27 +32,27 @@ void SScaleBox::OnArrangeChildren( const FGeometry& AllottedGeometry, FArrangedC
 	if ( ArrangedChildren.Accepts(ChildVisibility) )
 	{
 		const FVector2D AreaSize = AllottedGeometry.GetLocalSize();
-		FVector2D DesiredSize = ChildSlot.GetWidget()->GetDesiredSize();
+		FVector2D SlotWidgetDesiredSize = ChildSlot.GetWidget()->GetDesiredSize();
 
 		float FinalScale = 1;
 
 		const EStretch::Type CurrentStretch = Stretch.Get();
 		const EStretchDirection::Type CurrentStretchDirection = StretchDirection.Get();
 
-		if ( DesiredSize.X != 0 && DesiredSize.Y != 0 )
+		if (SlotWidgetDesiredSize.X != 0 && SlotWidgetDesiredSize.Y != 0 )
 		{
 			switch ( CurrentStretch )
 			{
 			case EStretch::None:
 				break;
 			case EStretch::Fill:
-				DesiredSize = AreaSize;
+				SlotWidgetDesiredSize = AreaSize;
 				break;
 			case EStretch::ScaleToFit:
-				FinalScale = FMath::Min(AreaSize.X / DesiredSize.X, AreaSize.Y / DesiredSize.Y);
+				FinalScale = FMath::Min(AreaSize.X / SlotWidgetDesiredSize.X, AreaSize.Y / SlotWidgetDesiredSize.Y);
 				break;
 			case EStretch::ScaleToFill:
-				FinalScale = FMath::Max(AreaSize.X / DesiredSize.X, AreaSize.Y / DesiredSize.Y);
+				FinalScale = FMath::Max(AreaSize.X / SlotWidgetDesiredSize.X, AreaSize.Y / SlotWidgetDesiredSize.Y);
 				break;
 			case EStretch::UserSpecified:
 				FinalScale = UserSpecifiedScale.Get(1.0f);
@@ -87,21 +87,21 @@ void SScaleBox::OnArrangeChildren( const FGeometry& AllottedGeometry, FArrangedC
 			// but scale the inverse of the scale we're applying.
 			if ( ChildSlot.HAlignment == HAlign_Fill )
 			{
-				DesiredSize.X = AreaSize.X / FinalScale;
+				SlotWidgetDesiredSize.X = AreaSize.X / FinalScale;
 			}
 
 			// If the layout vertically is fill, then we need the desired size to be the whole size of the widget, 
 			// but scale the inverse of the scale we're applying.
 			if ( ChildSlot.VAlignment == VAlign_Fill )
 			{
-				DesiredSize.Y = AreaSize.Y / FinalScale;
+				SlotWidgetDesiredSize.Y = AreaSize.Y / FinalScale;
 			}
 		}
 
 		ArrangedChildren.AddWidget(ChildVisibility, AllottedGeometry.MakeChild(
 			ChildSlot.GetWidget(),
 			FinalOffset,
-			DesiredSize,
+			SlotWidgetDesiredSize,
 			FinalScale
 		) );
 	}
diff --git a/Engine/Source/Runtime/Slate/Public/Framework/Application/SlateApplication.h b/Engine/Source/Runtime/Slate/Public/Framework/Application/SlateApplication.h
--- a/Engine/Source/Runtime/Slate/Public/Framework/Application/SlateApplication.h
+++ b/Engine/Source/Runtime/Slate/Public/Framework/Application/SlateApplication.h
@@ -142,7 +142,7 @@ public:
 
 	static TSharedRef<FSlateApplication> InitializeAsStandaloneApplication(const TSharedRef< class FSlateRenderer >& PlatformRenderer);
 	
-	static TSharedRef<FSlateApplication> InitializeAsStandaloneApplication(const TSharedRef< class FSlateRenderer >& PlatformRenderer, const TSharedRef<class GenericApplication>& PlatformApplication);
+	static TSharedRef<FSlateApplication> InitializeAsStandaloneApplication(const TSharedRef< class FSlateRenderer >& PlatformRenderer, const TSharedRef<class GenericApplication>& InPlatformApplication);
 
 	/**
 	 * Returns true if a Slate application instance is currently initialized and ready
diff --git a/Engine/Source/Runtime/Slate/Public/Widgets/Input/SComboBox.h b/Engine/Source/Runtime/Slate/Public/Widgets/Input/SComboBox.h
--- a/Engine/Source/Runtime/Slate/Public/Widgets/Input/SComboBox.h
+++ b/Engine/Source/Runtime/Slate/Public/Widgets/Input/SComboBox.h
@@ -151,7 +151,7 @@ public:
 
 		OptionsSource = InArgs._OptionsSource;
 
-		TSharedRef<SWidget> MenuContent = 
+		TSharedRef<SWidget> ComboBoxMenuContent = 
 			SNew(SBox)
 			.MaxDesiredHeight(InArgs._MaxListHeight)
 			[
@@ -182,7 +182,7 @@ public:
 			]
 			.MenuContent()
 			[
-				MenuContent
+				ComboBoxMenuContent
 			]
 			.HasDownArrow( InArgs._HasDownArrow )
 			.ContentPadding( InArgs._ContentPadding )
diff --git a/Engine/Source/Runtime/Slate/Public/Widgets/Layout/SConstraintCanvas.cpp b/Engine/Source/Runtime/Slate/Public/Widgets/Layout/SConstraintCanvas.cpp
--- a/Engine/Source/Runtime/Slate/Public/Widgets/Layout/SConstraintCanvas.cpp
+++ b/Engine/Source/Runtime/Slate/Public/Widgets/Layout/SConstraintCanvas.cpp
@@ -182,7 +182,7 @@ int32 SConstraintCanvas::OnPaint( const FPaintArgs& Args, const FGeometry& Allot
 
 FVector2D SConstraintCanvas::ComputeDesiredSize( float ) const
 {
-	FVector2D DesiredSize(0,0);
+	FVector2D FinalDesiredSize(0,0);
 
 	// Arrange the children now in their proper z-order.
 	for ( int32 ChildIndex = 0; ChildIndex < Children.Num(); ++ChildIndex )
@@ -207,12 +207,12 @@ FVector2D SConstraintCanvas::ComputeDesiredSize( float ) const
 			const bool bIsDockedHorizontally = ( Anchors.Minimum.X == Anchors.Maximum.X ) && ( Anchors.Minimum.X == 0 || Anchors.Minimum.X == 1 );
 			const bool bIsDockedVertically = ( Anchors.Minimum.Y == Anchors.Maximum.Y ) && ( Anchors.Minimum.Y == 0 || Anchors.Minimum.Y == 1 );
 
-			DesiredSize.X = FMath::Max(DesiredSize.X, Size.X + ( bIsDockedHorizontally ? FMath::Abs(Offset.Left) : 0.0f ));
-			DesiredSize.Y = FMath::Max(DesiredSize.Y, Size.Y + ( bIsDockedVertically ? FMath::Abs(Offset.Top) : 0.0f ));
+			FinalDesiredSize.X = FMath::Max(FinalDesiredSize.X, Size.X + ( bIsDockedHorizontally ? FMath::Abs(Offset.Left) : 0.0f ));
+			FinalDesiredSize.Y = FMath::Max(FinalDesiredSize.Y, Size.Y + ( bIsDockedVertically ? FMath::Abs(Offset.Top) : 0.0f ));
 		}
 	}
 
-	return DesiredSize;
+	return FinalDesiredSize;
 }
 
 FChildren* SConstraintCanvas::GetChildren()
diff --git a/Engine/Source/Runtime/Slate/Public/Widgets/Views/SListView.h b/Engine/Source/Runtime/Slate/Public/Widgets/Views/SListView.h
--- a/Engine/Source/Runtime/Slate/Public/Widgets/Views/SListView.h
+++ b/Engine/Source/Runtime/Slate/Public/Widgets/Views/SListView.h
@@ -1233,8 +1233,8 @@ protected:
 
 					if ( ScrollByAmountInSlateUnits > 0 )
 					{
-						FVector2D DesiredSize = RowWidget->AsWidget()->GetDesiredSize();
-						const float RemainingHeight = DesiredSize.Y * ( 1.0 - FMath::Fractional( NewScrollOffset ) );
+						FVector2D WidgetDesiredSize = RowWidget->AsWidget()->GetDesiredSize();
+						const float RemainingHeight = WidgetDesiredSize.Y * ( 1.0 - FMath::Fractional( NewScrollOffset ) );
 
 						if ( AbsScrollByAmount > RemainingHeight )
 						{
@@ -1257,13 +1257,13 @@ protected:
 						}
 						else
 						{
-							NewScrollOffset = (int32)NewScrollOffset + ( 1.0f - ( ( RemainingHeight - AbsScrollByAmount ) / DesiredSize.Y ) );
+							NewScrollOffset = (int32)NewScrollOffset + ( 1.0f - ( ( RemainingHeight - AbsScrollByAmount ) / WidgetDesiredSize.Y ) );
 							break;
 						}
 					}
 					else
 					{
-						FVector2D DesiredSize = RowWidget->AsWidget()->GetDesiredSize();
+						FVector2D WidgetDesiredSize = RowWidget->AsWidget()->GetDesiredSize();
 
 						float Fractional = FMath::Fractional( NewScrollOffset );
 						if ( Fractional == 0 )
@@ -1272,7 +1272,7 @@ protected:
 							--NewScrollOffset;
 						}
 
-						const float PrecedingHeight = DesiredSize.Y * Fractional;
+						const float PrecedingHeight = WidgetDesiredSize.Y * Fractional;
 
 						if ( AbsScrollByAmount > PrecedingHeight )
 						{
@@ -1295,7 +1295,7 @@ protected:
 						}
 						else
 						{
-							NewScrollOffset = (int32)NewScrollOffset + ( ( PrecedingHeight - AbsScrollByAmount ) / DesiredSize.Y );
+							NewScrollOffset = (int32)NewScrollOffset + ( ( PrecedingHeight - AbsScrollByAmount ) / WidgetDesiredSize.Y );
 							break;
 						}
 					}
diff --git a/Engine/Source/Runtime/Slate/Public/Widgets/Views/STableRow.h b/Engine/Source/Runtime/Slate/Public/Widgets/Views/STableRow.h
--- a/Engine/Source/Runtime/Slate/Public/Widgets/Views/STableRow.h
+++ b/Engine/Source/Runtime/Slate/Public/Widgets/Views/STableRow.h
@@ -986,7 +986,7 @@ public:
 	 * @param ColumnName    A unique ID for a column in this TableView; see SHeaderRow::FColumn for more info.
 	 * @return a widget to represent the contents of a cell in this row of a TableView. 
 	 */
-	virtual TSharedRef<SWidget> GenerateWidgetForColumn( const FName& ColumnName ) = 0;
+	virtual TSharedRef<SWidget> GenerateWidgetForColumn( const FName& InColumnName ) = 0;
 
 	/** Use this to construct the superclass; e.g. FSuperRowType::Construct( FTableRowArgs(), OwnerTableView ) */
 	typedef SMultiColumnTableRow< ItemType > FSuperRowType;
diff --git a/Engine/Source/Runtime/Slate/Public/Widgets/Views/STreeView.h b/Engine/Source/Runtime/Slate/Public/Widgets/Views/STreeView.h
--- a/Engine/Source/Runtime/Slate/Public/Widgets/Views/STreeView.h
+++ b/Engine/Source/Runtime/Slate/Public/Widgets/Views/STreeView.h
@@ -420,7 +420,7 @@ public:
 	 * @return true if we encountered expanded children; false otherwise.
 	 */
 	bool PopulateLinearizedItems(
-		const TArray<ItemType>& ItemsSource,
+		const TArray<ItemType>& InItemsSource,
 		TArray< ItemType >& InLinearizedItems,
 		TArray< FItemInfo >& NewDenseItemInfos,
 		int32 TreeLevel,
@@ -429,13 +429,13 @@ public:
 		bool bAddingItems )
 	{
 		bool bSawExpandedItems = false;
-		for ( int32 ItemIndex = 0; ItemIndex < ItemsSource.Num(); ++ItemIndex )
+		for ( int32 ItemIndex = 0; ItemIndex < InItemsSource.Num(); ++ItemIndex )
 		{
-			const ItemType& CurItem = ItemsSource[ItemIndex];
+			const ItemType& CurItem = InItemsSource[ItemIndex];
 
 			// Find this items children.
 			TArray<ItemType> ChildItems;
-			OnGetChildren.Execute( ItemsSource[ItemIndex], ChildItems );
+			OnGetChildren.Execute(InItemsSource[ItemIndex], ChildItems );
 
 			const bool bHasChildren = ChildItems.Num() > 0;
 
diff --git a/Engine/Source/Runtime/Sockets/Public/Sockets.h b/Engine/Source/Runtime/Sockets/Public/Sockets.h
--- a/Engine/Source/Runtime/Sockets/Public/Sockets.h
+++ b/Engine/Source/Runtime/Sockets/Public/Sockets.h
@@ -103,7 +103,7 @@ public:
 	 * @param		SocketDescription debug description of socket
 	 * @return		The new (heap-allocated) socket, or NULL if unsuccessful.
 	 */
-	virtual class FSocket* Accept(const FString& SocketDescription) = 0;
+	virtual class FSocket* Accept(const FString& InSocketDescription) = 0;
 
 	/**
 	 * Accepts a connection that is pending
@@ -113,7 +113,7 @@ public:
 	 *
 	 * @return		The new (heap-allocated) socket, or NULL if unsuccessful.
 	 */
-	virtual class FSocket* Accept(FInternetAddr& OutAddr, const FString& SocketDescription) = 0;
+	virtual class FSocket* Accept(FInternetAddr& OutAddr, const FString& InSocketDescription) = 0;
 
 	/**
 	 * Sends a buffer to a network byte ordered address
diff --git a/Engine/Source/Runtime/Windows/D3D11RHI/Private/D3D11Commands.cpp b/Engine/Source/Runtime/Windows/D3D11RHI/Private/D3D11Commands.cpp
--- a/Engine/Source/Runtime/Windows/D3D11RHI/Private/D3D11Commands.cpp
+++ b/Engine/Source/Runtime/Windows/D3D11RHI/Private/D3D11Commands.cpp
@@ -11,7 +11,14 @@
 		#define XM_NO_OPERATOR_OVERLOADS 1		// @todo clang: These xnamath operators don't compile correctly under clang
 		#define _XM_NO_INTRINSICS_ 1	// @todo clang: Clang has issues with __m128 intrinsics in xnamathvector.inl
 	#endif
+	#if _MSC_VER == 1900
+		#pragma warning(push)
+		#pragma warning(disable:4838)
+	#endif // _MSC_VER == 1900
 	#include <xnamath.h>
+	#if _MSC_VER == 1900
+		#pragma warning(pop)
+	#endif // _MSC_VER == 1900
 #include "HideWindowsPlatformTypes.h"
 #endif
 #include "D3D11RHIPrivateUtil.h"
@@ -176,7 +183,7 @@ void FD3D11DynamicRHI::RHISetViewport(uint32 MinX,uint32 MinY,float MinZ,uint32
 	check(MaxX <= (uint32)D3D11_VIEWPORT_BOUNDS_MAX);
 	check(MaxY <= (uint32)D3D11_VIEWPORT_BOUNDS_MAX);
 
-	D3D11_VIEWPORT Viewport = { MinX, MinY, MaxX - MinX, MaxY - MinY, MinZ, MaxZ };
+	D3D11_VIEWPORT Viewport = { (float)MinX, (float)MinY, (float)MaxX - MinX, (float)MaxY - MinY, MinZ, MaxZ };
 	//avoid setting a 0 extent viewport, which the debug runtime doesn't like
 	if (Viewport.Width > 0 && Viewport.Height > 0)
 	{
diff --git a/Engine/Source/Runtime/Windows/D3D11RHI/Private/D3D11ConstantBuffer.cpp b/Engine/Source/Runtime/Windows/D3D11RHI/Private/D3D11ConstantBuffer.cpp
--- a/Engine/Source/Runtime/Windows/D3D11RHI/Private/D3D11ConstantBuffer.cpp
+++ b/Engine/Source/Runtime/Windows/D3D11RHI/Private/D3D11ConstantBuffer.cpp
@@ -16,7 +16,7 @@
 const uint32 GConstantBufferSizes[MAX_CONSTANT_BUFFER_SLOTS] = 
 {
 	// CBs must be a multiple of 16
-	Align(MAX_GLOBAL_CONSTANT_BUFFER_SIZE, 16),
+	(uint32)Align(MAX_GLOBAL_CONSTANT_BUFFER_SIZE, 16),
 };
 
 // New circular buffer system for faster constant uploads.  Avoids CopyResource and speeds things up considerably
diff --git a/Engine/Source/Runtime/Windows/XAudio2/Private/XAudio2Source.cpp b/Engine/Source/Runtime/Windows/XAudio2/Private/XAudio2Source.cpp
--- a/Engine/Source/Runtime/Windows/XAudio2/Private/XAudio2Source.cpp
+++ b/Engine/Source/Runtime/Windows/XAudio2/Private/XAudio2Source.cpp
@@ -329,7 +329,7 @@ bool FXAudio2SoundSource::CreateSource( void )
 {
 	SCOPE_CYCLE_COUNTER( STAT_AudioSourceCreateTime );
 
-	int32 NumSends = 0;
+	uint32 NumSends = 0;
 
 #if XAUDIO2_SUPPORTS_SENDLIST
 	// Create a source that goes to the spatialisation code and reverb effect
