Fix another shadow variable.

diff --git a/Engine/Source/Runtime/CoreUObject/Private/UObject/Enum.cpp b/Engine/Source/Runtime/CoreUObject/Private/UObject/Enum.cpp
--- a/Engine/Source/Runtime/CoreUObject/Private/UObject/Enum.cpp
+++ b/Engine/Source/Runtime/CoreUObject/Private/UObject/Enum.cpp
@@ -146,12 +146,12 @@ int8 UEnum::GetValueByName(FName InName)
 	return INDEX_NONE;
 }
 
-int32 UEnum::GetIndexByName(FName Name) const
+int32 UEnum::GetIndexByName(FName InName) const
 {
 	int32 Count = Names.Num();
 	for (int32 Counter = 0; Counter < Count; ++Counter)
 	{
-		if (Names[Counter].Key == Name)
+		if (Names[Counter].Key == InName)
 		{
 			return Counter;
 		}
