Better way of checking if a default subobject has already been created with Create*DefaultSubobject function.

diff --git a/Engine/Source/Runtime/CoreUObject/Private/UObject/UObjectGlobals.cpp b/Engine/Source/Runtime/CoreUObject/Private/UObject/UObjectGlobals.cpp
--- a/Engine/Source/Runtime/CoreUObject/Private/UObject/UObjectGlobals.cpp
+++ b/Engine/Source/Runtime/CoreUObject/Private/UObject/UObjectGlobals.cpp
@@ -2994,10 +2994,18 @@ UObject* FObjectInitializer::CreateDefaultSubobject(UObject* Outer, FName Subobj
 			const EObjectFlags SubobjectFlags = Outer->GetMaskedFlags(RF_PropagateToSubObjects);
 			const bool bOwnerArchetypeIsNotNative = !Outer->GetArchetype()->GetClass()->HasAnyClassFlags(CLASS_Native | CLASS_Intrinsic);
 			const bool bOwnerTemplateIsNotCDO = ObjectArchetype != nullptr && ObjectArchetype != Outer->GetClass()->GetDefaultObject(false) && !Outer->HasAnyFlags(RF_ClassDefaultObject);
-			//if (StaticFindObjectFast(OverrideClass, Outer, SubobjectFName) != nullptr)
-			//{
-			//	UE_LOG(LogClass, Fatal, TEXT("Default subobject %s %s already exists for %s."), *OverrideClass->GetName(), *SubobjectFName.ToString(), *Outer->GetFullName());
-			//}
+#if !UE_BUILD_SHIPPING
+			// Guard against constructing the same subobject multiple times.
+			// We only need to check the name as ConstructObject would fail anyway if an object of the same name but different class already existed.
+			if (ConstructedSubobjects.Find(SubobjectFName) != INDEX_NONE)
+			{
+				UE_LOG(LogClass, Fatal, TEXT("Default subobject %s %s already exists for %s."), *OverrideClass->GetName(), *SubobjectFName.ToString(), *Outer->GetFullName());
+			}
+			else
+			{
+				ConstructedSubobjects.Add(SubobjectFName);
+			}
+#endif
 			Result = StaticConstructObject_Internal(OverrideClass, Outer, SubobjectFName, SubobjectFlags);
 			if (!bIsTransient && (bOwnerArchetypeIsNotNative || bOwnerTemplateIsNotCDO))
 			{
diff --git a/Engine/Source/Runtime/CoreUObject/Public/UObject/UObjectGlobals.h b/Engine/Source/Runtime/CoreUObject/Public/UObject/UObjectGlobals.h
--- a/Engine/Source/Runtime/CoreUObject/Public/UObject/UObjectGlobals.h
+++ b/Engine/Source/Runtime/CoreUObject/Public/UObject/UObjectGlobals.h
@@ -1007,6 +1007,10 @@ private:
 	mutable FOverrides ComponentOverrides;
 	/**  List of component classes to intialize after the C++ constructors **/
 	mutable FSubobjectsToInit ComponentInits;
+#if !UE_BUILD_SHIPPING
+	/** List of all subobject names constructed for this object */
+	mutable TArray<FName, TInlineAllocator<8>> ConstructedSubobjects;
+#endif
 	/**  Previously constructed object in the callstack */
 	UObject* LastConstructedObject;
 
