Support TMap properties in the reference chain.

diff --git a/Engine/Source/Runtime/CoreUObject/Private/UObject/GarbageCollection.cpp b/Engine/Source/Runtime/CoreUObject/Private/UObject/GarbageCollection.cpp
--- a/Engine/Source/Runtime/CoreUObject/Private/UObject/GarbageCollection.cpp
+++ b/Engine/Source/Runtime/CoreUObject/Private/UObject/GarbageCollection.cpp
@@ -1738,7 +1738,7 @@ void UMapProperty::EmitReferenceInfo(UClass& OwnerClass, int32 BaseOffset)
 {
 	if (ContainsObjectReference())
 	{
-		OwnerClass.ReferenceTokenStream.EmitReferenceInfo(FGCReferenceInfo(GCRT_AddTMapReferencedObjects, BaseOffset + GetOffset_ForGC()));
+		OwnerClass.EmitObjectReference(BaseOffset + GetOffset_ForGC(), GetFName(), GCRT_AddTMapReferencedObjects);
 		OwnerClass.ReferenceTokenStream.EmitPointer((const void*)this);
 	}
 }
diff --git a/Engine/Source/Runtime/CoreUObject/Private/UObject/ReferenceChainSearch.cpp b/Engine/Source/Runtime/CoreUObject/Private/UObject/ReferenceChainSearch.cpp
--- a/Engine/Source/Runtime/CoreUObject/Private/UObject/ReferenceChainSearch.cpp
+++ b/Engine/Source/Runtime/CoreUObject/Private/UObject/ReferenceChainSearch.cpp
@@ -15,6 +15,7 @@ namespace Internal
 			case FReferenceChainSearch::EReferenceType::ArrayProperty: return TEXT("Array");
 			case FReferenceChainSearch::EReferenceType::StructARO: return TEXT("StructARO");
 			case FReferenceChainSearch::EReferenceType::ARO: return TEXT("ARO");
+			case FReferenceChainSearch::EReferenceType::MapProperty: return TEXT("Map");
 			default: return TEXT("Invalid");
 		}
 	}
@@ -623,6 +624,20 @@ bool FReferenceChainSearch::ProcessObject( UObject* CurrentObject )
 				AddToReferenceList(ReferenceList, ReferenceCollector.References[i]);
 			}
 		}
+		else if( REFERENCE_INFO.Type == GCRT_AddTMapReferencedObjects )
+		{
+			void*         Map         = StackEntryData + REFERENCE_INFO.Offset;
+			UMapProperty* MapProperty = (UMapProperty*)TokenStream->ReadPointer( TokenStreamIndex );
+			TokenReturnCount = REFERENCE_INFO.ReturnCount;
+			FFindReferencerCollector ReferenceCollector(this, EReferenceType::MapProperty, (void*)MapProperty, CurrentObject);
+			FSimpleObjectReferenceCollectorArchive CollectorArchive(CurrentObject, ReferenceCollector);
+			MapProperty->SerializeItem(CollectorArchive, Map, nullptr);
+
+			for (const FReferenceChainLink& Ref : ReferenceCollector.References)
+			{
+				AddToReferenceList(ReferenceList, Ref);
+			}
+		}
 		else if( REFERENCE_INFO.Type == GCRT_EndOfStream )
 		{
 			// Break out of loop.
diff --git a/Engine/Source/Runtime/CoreUObject/Public/UObject/ReferenceChainSearch.h b/Engine/Source/Runtime/CoreUObject/Public/UObject/ReferenceChainSearch.h
--- a/Engine/Source/Runtime/CoreUObject/Public/UObject/ReferenceChainSearch.h
+++ b/Engine/Source/Runtime/CoreUObject/Public/UObject/ReferenceChainSearch.h
@@ -21,6 +21,8 @@ public:
 			StructARO,
 			/** Object is added to the references list by an AddReferencedObject call */
 			ARO,
+			/** Object is referenced by a map */
+			MapProperty,
 		};
 	};
 
@@ -55,7 +57,7 @@ public:
 #endif // !(UE_BUILD_TEST || UE_BUILD_SHIPPING)
 
 		/** Returns whether this link is a property reference or not */
-		FORCEINLINE bool IsProperty() const { return ReferenceType == EReferenceType::Property || ReferenceType == EReferenceType::ArrayProperty; }
+		FORCEINLINE bool IsProperty() const { return ReferenceType == EReferenceType::Property || ReferenceType == EReferenceType::ArrayProperty || ReferenceType == EReferenceType::MapProperty; }
 
 		FORCEINLINE bool operator == (const FReferenceChainLink& Rhs) const { return ReferencedBy == Rhs.ReferencedBy && ReferencedObj == Rhs.ReferencedObj; }
 
