Don't flag packages as play in editor if they already existed

diff --git a/Engine/Source/Runtime/CoreUObject/Private/UObject/Linker.cpp b/Engine/Source/Runtime/CoreUObject/Private/UObject/Linker.cpp
--- a/Engine/Source/Runtime/CoreUObject/Private/UObject/Linker.cpp
+++ b/Engine/Source/Runtime/CoreUObject/Private/UObject/Linker.cpp
@@ -494,7 +494,8 @@ FLinkerLoad* GetPackageLinker
 			return nullptr;
 		}
 
-		if (UPackage* ExistingPackage = FindObject<UPackage>(nullptr, *PackageName))
+		UPackage* ExistingPackage = FindObject<UPackage>(nullptr, *PackageName);
+		if (ExistingPackage)
 		{
 			if (!ExistingPackage->GetOuter() && (ExistingPackage->PackageFlags & PKG_InMemoryOnly))
 			{
@@ -515,8 +516,8 @@ FLinkerLoad* GetPackageLinker
 		}
 
 		// Create the package with the provided long package name.
-		UPackage* FilenamePkg = CreatePackage(nullptr, *PackageName);
-		if (FilenamePkg && (LoadFlags & LOAD_PackageForPIE))
+		UPackage* FilenamePkg = (ExistingPackage ? ExistingPackage : CreatePackage(nullptr, *PackageName));
+		if (FilenamePkg != ExistingPackage && (LoadFlags & LOAD_PackageForPIE))
 		{
 			FilenamePkg->PackageFlags |= PKG_PlayInEditor;
 		}
