Needless conversion from FName to FString to FText, FText::FromString takes FString by value for some reason...

diff --git a/Engine/Source/Runtime/CoreUObject/Private/UObject/LinkerLoad.cpp b/Engine/Source/Runtime/CoreUObject/Private/UObject/LinkerLoad.cpp
--- a/Engine/Source/Runtime/CoreUObject/Private/UObject/LinkerLoad.cpp
+++ b/Engine/Source/Runtime/CoreUObject/Private/UObject/LinkerLoad.cpp
@@ -2288,7 +2288,7 @@ bool FLinkerLoad::VerifyImportInner(const int32 ImportIndex, FString& WarningSuf
 	FObjectImport& Import = ImportMap[ImportIndex];
 
 #if WITH_EDITOR
-	FScopedSlowTask SlowTask(100, FText::Format(NSLOCTEXT("Core", "VerifyPackage_Scope", "Verifying '{0}'"), FText::FromString(Import.ObjectName.ToString())), ShouldReportProgress());
+	FScopedSlowTask SlowTask(100, FText::Format(NSLOCTEXT("Core", "VerifyPackage_Scope", "Verifying '{0}'"), FText::FromName(Import.ObjectName)), ShouldReportProgress());
 #endif
 
 	if
