PR #1348: [CoreUObject] Improvement in garbage collection performance by implementing three micro-optimizations. (Contributed by [REDACTED_PERSON])

diff --git a/Engine/Source/Runtime/CoreUObject/Private/UObject/GarbageCollection.cpp b/Engine/Source/Runtime/CoreUObject/Private/UObject/GarbageCollection.cpp
--- a/Engine/Source/Runtime/CoreUObject/Private/UObject/GarbageCollection.cpp
+++ b/Engine/Source/Runtime/CoreUObject/Private/UObject/GarbageCollection.cpp
@@ -295,68 +295,68 @@ void SerializeRootSet( FArchive& Ar, EObjectFlags KeepFlags )
  * @param ReferencingObject UObject which owns the reference (can be NULL)
  * @param bAllowReferenceElimination	Whether to allow NULL'ing the reference if RF_PendingKill is set
  */
-static FORCEINLINE void HandleObjectReference(TArray<UObject*>& ObjectsToSerialize, UObject* ReferencingObject, UObject*& Object, bool bAllowReferenceElimination, bool bStrongReference = true)
+static FORCEINLINE void HandleObjectReference(TArray<UObject*>& ObjectsToSerialize, const UObject * const ReferencingObject, UObject*& Object, const bool bAllowReferenceElimination, const bool bStrongReference = true)
 {
 	// Disregard NULL objects and perform very fast check to see whether object is part of permanent
 	// object pool and should therefore be disregarded. The check doesn't touch the object and is
 	// cache friendly as it's just a pointer compare against to globals.
-	if( Object )
+	const bool IsInPermanentPool = GUObjectAllocator.ResidesInPermanentPool(Object);
+
+#if PERF_DETAILED_PER_CLASS_GC_STATS
+	if(IsInPermanentPool)
 	{
-		if( !GUObjectAllocator.ResidesInPermanentPool(Object) )
-		{
-			UObject* ObjectToAdd = Object->HasAnyFlags( RF_Unreachable ) ? Object : NULL;
-			// Remove references to pending kill objects if we're allowed to do so.
-			if( Object->HasAnyFlags( RF_PendingKill ) && bAllowReferenceElimination )
-			{
-				// Null out reference.
-				Object = NULL;
-			}
-			// Add encountered object reference to list of to be serialized objects if it hasn't already been added.
-			else if (Object->HasAnyFlags(RF_Unreachable))
-			{	
-				if( GIsRunningParallelReachability )
-				{
-					// Mark it as reachable.
-					if (Object->ThisThreadAtomicallyClearedRFUnreachable())
-					{
-						// Add it to the list of objects to serialize.
-						ObjectsToSerialize.Add( Object );
-					}
-				}
-				else if ( ObjectToAdd )
-				{
-#if ENABLE_GC_DEBUG_OUTPUT
-					// this message is to help track down culprits behind "Object in PIE world still referenced" errors
-					if ( GIsEditor && !GIsPlayInEditorWorld && ReferencingObject != NULL && !ReferencingObject->RootPackageHasAnyFlags(PKG_PlayInEditor) && Object->RootPackageHasAnyFlags(PKG_PlayInEditor) )
-					{
-						UE_LOG(LogGarbage, Warning, TEXT("GC detected illegal reference to PIE object from content [possibly via [todo]]:"));
-						UE_LOG(LogGarbage, Warning, TEXT("      PIE object: %s"), *Object->GetFullName());
-						UE_LOG(LogGarbage, Warning, TEXT("  NON-PIE object: %s"), *ReferencingObject->GetFullName());
-					}
+		GCurrentObjectDisregardedObjectRefs++;
+	}
 #endif
+	if(Object == nullptr || IsInPermanentPool)
+	{
+		return;
+	}
 
-					// Mark it as reachable.
-					Object->ClearFlags( RF_Unreachable );
-					// Add it to the list of objects to serialize.
-					ObjectsToSerialize.Add( Object );
-				}
-			}
-
-			if (Object && bStrongReference)
+	// Remove references to pending kill objects if we're allowed to do so.
+	if( Object->HasAnyFlags( RF_PendingKill ) && bAllowReferenceElimination )
+	{
+		// Null out reference.
+		Object = NULL;
+	}
+	// Add encountered object reference to list of to be serialized objects if it hasn't already been added.
+	else if( Object->HasAnyFlags( RF_Unreachable ) )
+	{				
+		if( GIsRunningParallelReachability )
+		{
+			// Mark it as reachable.
+			if (Object->ThisThreadAtomicallyClearedRFUnreachable())
 			{
-				Object->ClearFlags(RF_NoStrongReference);
+				// Add it to the list of objects to serialize.
+				ObjectsToSerialize.Add( Object );
 			}
-#if PERF_DETAILED_PER_CLASS_GC_STATS
-			GCurrentObjectRegularObjectRefs++;
 		}
 		else
 		{
-			GCurrentObjectDisregardedObjectRefs++;
-		}
-#else
-		}
+#if ENABLE_GC_DEBUG_OUTPUT
+			// this message is to help track down culprits behind "Object in PIE world still referenced" errors
+			if ( GIsEditor && !GIsPlayInEditorWorld && ReferencingObject != NULL && !ReferencingObject->RootPackageHasAnyFlags(PKG_PlayInEditor) && Object->RootPackageHasAnyFlags(PKG_PlayInEditor) )
+			{
+				UE_LOG(LogGarbage, Warning, TEXT("GC detected illegal reference to PIE object from content [possibly via [todo]]:"));
+				UE_LOG(LogGarbage, Warning, TEXT("      PIE object: %s"), *Object->GetFullName());
+				UE_LOG(LogGarbage, Warning, TEXT("  NON-PIE object: %s"), *ReferencingObject->GetFullName());
+			}
 #endif
+
+			// Mark it as reachable.
+			Object->ClearFlags( RF_Unreachable );
+			// Add it to the list of objects to serialize.
+			ObjectsToSerialize.Add( Object );
+		}
+	}
+
+	if (Object && bStrongReference)
+	{
+		Object->ClearFlags(RF_NoStrongReference);
 	}
+#if PERF_DETAILED_PER_CLASS_GC_STATS
+	GCurrentObjectRegularObjectRefs++;
+#endif
 }
 
 static FORCEINLINE void HandleTokenStreamObjectReference(TArray<UObject*>& ObjectsToSerialize, UObject* ReferencingObject, UObject*& Object, const int32 TokenIndex, bool bAllowReferenceElimination)
@@ -710,11 +710,10 @@ public:
 #endif
 				CurrentObject = ObjectsToSerialize[CurrentIndex++];
 
-				if (CurrentIndex != ObjectsToSerialize.Num())
-				{
-					UObject* NextObject = ObjectsToSerialize.GetData()[CurrentIndex]; // special syntax avoiding out of bounds checking
-					FPlatformMisc::PrefetchBlock(NextObject, NextObject->GetClass()->GetPropertiesSize());
-				}
+				const UObject * const NextObject = ObjectsToSerialize.GetData()[CurrentIndex]; // special syntax avoiding out of bounds checking
+				// Prefetch the next object assuming that the property size of the next object is the same as the current one.
+				// This allows us to avoid a branch here.
+				FPlatformMisc::PrefetchBlock(NextObject, CurrentObject->GetClass()->GetPropertiesSize());
 
 				//@todo rtgc: we need to handle object references in struct defaults
 
@@ -790,10 +789,9 @@ public:
 						// We're dealing with an array of object references.
 						TArray<UObject*>& ObjectArray = *((TArray<UObject*>*)(StackEntryData + REFERENCE_INFO.Offset));
 						TokenReturnCount = REFERENCE_INFO.ReturnCount;
-						for( int32 ObjectIndex=0; ObjectIndex<ObjectArray.Num(); ObjectIndex++ )
+						for( int32 ObjectIndex = 0, ObjectNum = ObjectArray.Num(); ObjectIndex < ObjectNum; ++ObjectIndex )
 						{
-							UObject*& Object = ObjectArray[ObjectIndex];
-							HandleTokenStreamObjectReference(NewObjectsToSerialize, CurrentObject, Object, ReferenceTokenStreamIndex, true);
+							HandleTokenStreamObjectReference(NewObjectsToSerialize, CurrentObject, ObjectArray[ObjectIndex], ReferenceTokenStreamIndex, true);
 						}
 					}
 					else if( REFERENCE_INFO.Type == GCRT_ArrayStruct )
@@ -1248,9 +1246,9 @@ void CollectGarbageInternal(EObjectFlags KeepFlags, bool bPerformFullPurge)
 
 	// Flush streaming before GC if requested
 	if (GFlushStreamingOnGC)
-	{
-		FlushAsyncLoading();
-	}
+		{
+			FlushAsyncLoading();
+		}
 
 	// Route callbacks so we can ensure that we are e.g. not in the middle of loading something by flushing
 	// the async loading, etc...
@@ -1341,25 +1339,25 @@ void CollectGarbageInternal(EObjectFlags KeepFlags, bool bPerformFullPurge)
 	{
 		DECLARE_SCOPE_CYCLE_COUNTER( TEXT( "CollectGarbageInternal.UnhashUnreachable" ), STAT_CollectGarbageInternal_UnhashUnreachable, STATGROUP_GC );
 
-		// Unhash all unreachable objects.
-		const double StartTime = FPlatformTime::Seconds();
-		for (FRawObjectIterator It( true ); It; ++It)
-		{
-			//@todo UE4 - A prefetch was removed here. Re-add it. It wasn't right anyway, since it was ten items ahead and the consoles on have 8 prefetch slots
+	// Unhash all unreachable objects.
+	const double StartTime = FPlatformTime::Seconds();
+	for ( FRawObjectIterator It(true); It; ++It )
+	{
+		//@todo UE4 - A prefetch was removed here. Re-add it. It wasn't right anyway, since it was ten items ahead and the consoles on have 8 prefetch slots
 
-			UObject* Object = *It;
-			if (Object->HasAnyFlags( RF_Unreachable ))
-			{
-				// Begin the object's asynchronous destruction.
-				Object->ConditionalBeginDestroy();
-			}
+		UObject* Object = *It;
+		if( Object->HasAnyFlags( RF_Unreachable ) )
+		{
+			// Begin the object's asynchronous destruction.
+			Object->ConditionalBeginDestroy();
+		}
 			else if (Object->HasAnyFlags(RF_NoStrongReference))
 			{
 				Object->ClearFlags(RF_NoStrongReference);
 				Object->SetFlags(RF_PendingKill);
 			}
-		}
-		UE_LOG( LogGarbage, Log, TEXT( "%f ms for unhashing unreachable objects" ), (FPlatformTime::Seconds() - StartTime) * 1000 );
+	}
+	UE_LOG(LogGarbage, Log, TEXT("%f ms for unhashing unreachable objects"), (FPlatformTime::Seconds() - StartTime) * 1000 );
 	}
 
 	// Set flag to indicate that we are relying on a purge to be performed.
