Added cooking stats.
Added multithreaded handling of cooking output.
Added support for launch on recompiling changed global shaders when cook in editor.

diff --git a/Engine/Source/Developer/DerivedDataCache/DerivedDataCache.Build.cs b/Engine/Source/Developer/DerivedDataCache/DerivedDataCache.Build.cs
--- a/Engine/Source/Developer/DerivedDataCache/DerivedDataCache.Build.cs
+++ b/Engine/Source/Developer/DerivedDataCache/DerivedDataCache.Build.cs
@@ -8,6 +8,7 @@ public class DerivedDataCache : ModuleRules
 	public DerivedDataCache(TargetInfo Target)
 	{
 		PrivateDependencyModuleNames.Add("Core");
+        PrivateDependencyModuleNames.Add("CookingStats");
 
 		// Internal (NotForLicensees) module
 		var DDCUtilsModule = Path.Combine("Developer", "NotForLicensees", "DDCUtils", "DDCUtils.Build.cs");
diff --git a/Engine/Source/Developer/DerivedDataCache/Private/DDCStatsHelper.cpp b/Engine/Source/Developer/DerivedDataCache/Private/DDCStatsHelper.cpp
new file mode 100644
--- /dev/null
+++ b/Engine/Source/Developer/DerivedDataCache/Private/DDCStatsHelper.cpp
@@ -0,0 +1,111 @@
+#include "Core.h"
+
+#include "DDCStatsHelper.h"
+#include "CookingStatsModule.h"
+#include "ModuleManager.h"
+
+
+uint32 GetValidTLSSlot()
+{
+	uint32 TLSSlot = FPlatformTLS::AllocTlsSlot();
+	check(FPlatformTLS::IsValidTlsSlot(TLSSlot));
+	FPlatformTLS::SetTlsValue(TLSSlot, nullptr);
+	return TLSSlot;
+}
+
+ICookingStats* CookingStats = nullptr;
+
+struct FDDCStatsTLSStore
+{
+
+	FDDCStatsTLSStore() : TransactionGuid(NAME_None), CurrentIndex(0), RootScope(NULL) { }
+
+
+	FName GenerateNewTransactionName()
+	{
+		++CurrentIndex;
+		if (CurrentIndex <= 1)
+		{
+			FString TransactionGuidString = TEXT("DDCTransactionId");
+			TransactionGuidString += FGuid::NewGuid().ToString();
+			TransactionGuid = FName(*TransactionGuidString);
+			CurrentIndex = 1;
+		}
+		check(TransactionGuid != NAME_None);
+
+		return FName(TransactionGuid, CurrentIndex);
+	}
+
+	FName TransactionGuid;
+	int32 CurrentIndex;
+	FDDCScopeStatHelper *RootScope;
+};
+
+uint32 CookStatsFDDCStatsTLSStore = GetValidTLSSlot();
+
+FDDCScopeStatHelper::FDDCScopeStatHelper(const TCHAR* CacheKey, const FName& FunctionName ) : bHasParent(false)
+{
+	FDDCStatsTLSStore* TLSStore = (FDDCStatsTLSStore*)FPlatformTLS::GetTlsValue(CookStatsFDDCStatsTLSStore);
+	if (TLSStore == nullptr)
+	{
+		TLSStore = new FDDCStatsTLSStore();
+		FPlatformTLS::SetTlsValue(CookStatsFDDCStatsTLSStore, TLSStore);
+	}
+
+	TransactionGuid = TLSStore->GenerateNewTransactionName();
+
+	if (CookingStats == nullptr)
+	{
+		FCookingStatsModule& CookingStatsModule = FModuleManager::LoadModuleChecked<FCookingStatsModule>(TEXT("CookingStats"));
+		CookingStats = &CookingStatsModule.Get();
+	}
+	FString Temp;
+	AddTag(FunctionName, Temp);
+
+	if (TLSStore->RootScope == nullptr)
+	{
+		TLSStore->RootScope = this;
+		bHasParent = false;
+
+		static const FName NAME_CacheKey = FName(TEXT("CacheKey"));
+		AddTag(NAME_CacheKey, FString(CacheKey));
+	}
+	else
+	{
+		static const FName NAME_Parent = FName(TEXT("Parent"));
+		AddTag(NAME_Parent, TLSStore->RootScope->TransactionGuid.ToString());
+		
+		bHasParent = true;
+	}
+
+
+	static const FName NAME_StartTime(TEXT("StartTime"));
+	AddTag(NAME_StartTime, FDateTime::Now().ToString());
+
+	StartTime = FPlatformTime::Seconds();
+
+	// if there is a root scope then we want to put out stuff under a child key of the parent and link to the parent scope
+}
+
+
+FDDCScopeStatHelper::~FDDCScopeStatHelper()
+{
+	double EndTime = FPlatformTime::Seconds();
+	float DurationMS = (EndTime - StartTime) * 1000.0f;
+
+	static const FName NAME_Duration = FName(TEXT("Duration"));
+	AddTag(NAME_Duration, FString::Printf(TEXT("%fms"), DurationMS));
+
+	FDDCStatsTLSStore* TLSStore = (FDDCStatsTLSStore*)FPlatformTLS::GetTlsValue(CookStatsFDDCStatsTLSStore);
+	check(TLSStore);
+
+	if (TLSStore->RootScope == this)
+	{
+		TLSStore->RootScope = nullptr;
+	}
+}
+
+void FDDCScopeStatHelper::AddTag(const FName& Tag, const FString& Value)
+{
+	CookingStats->AddTagValue(TransactionGuid, Tag, Value);
+}
\ No newline at end of file
diff --git a/Engine/Source/Developer/DerivedDataCache/Private/DDCStatsHelper.h b/Engine/Source/Developer/DerivedDataCache/Private/DDCStatsHelper.h
new file mode 100644
--- /dev/null
+++ b/Engine/Source/Developer/DerivedDataCache/Private/DDCStatsHelper.h
@@ -0,0 +1,19 @@
+#pragma once
+
+
+class FDDCScopeStatHelper
+{
+private:
+	FName TransactionGuid;
+	double StartTime;
+	bool bHasParent;
+public:
+	FDDCScopeStatHelper(const TCHAR* CacheKey, const FName& FunctionName);
+
+	~FDDCScopeStatHelper();
+
+	void AddTag(const FName& Tag, const FString& Value);
+
+	bool HasParent() const { return bHasParent;  }
+};
+
diff --git a/Engine/Source/Developer/DerivedDataCache/Private/DerivedDataBackends.cpp b/Engine/Source/Developer/DerivedDataCache/Private/DerivedDataBackends.cpp
--- a/Engine/Source/Developer/DerivedDataCache/Private/DerivedDataBackends.cpp
+++ b/Engine/Source/Developer/DerivedDataCache/Private/DerivedDataBackends.cpp
@@ -4,6 +4,7 @@
 #include "SecureHash.h"
 #include "DerivedDataBackendInterface.h"
 #include "DerivedDataCacheInterface.h"
+#include "DDCStatsHelper.h"
 #include "MemoryDerivedDataBackend.h"
 #include "DerivedDataBackendAsyncPutWrapper.h"
 #include "PakFileDerivedDataBackend.h"
diff --git a/Engine/Source/Developer/DerivedDataCache/Private/FileSystemDerivedDataBackend.cpp b/Engine/Source/Developer/DerivedDataCache/Private/FileSystemDerivedDataBackend.cpp
--- a/Engine/Source/Developer/DerivedDataCache/Private/FileSystemDerivedDataBackend.cpp
+++ b/Engine/Source/Developer/DerivedDataCache/Private/FileSystemDerivedDataBackend.cpp
@@ -6,6 +6,8 @@
 #include "DerivedDataBackendInterface.h"
 #include "DDCCleanup.h"
 
+#include "DDCStatsHelper.h"
+
 #define MAX_BACKEND_KEY_LENGTH (120)
 #define MAX_BACKEND_NUMBERED_SUBFOLDER_LENGTH (9)
 #if PLATFORM_LINUX	// PATH_MAX on Linux is 4096 (getconf PATH_MAX /, also see limits.h), so this value can be larger (note that it is still arbitrary).
@@ -16,6 +18,9 @@
 #endif // PLATFORM_LINUX
 #define MAX_CACHE_EXTENTION_LEN (4)
 
+
+
+
 /** 
  * Cache server that uses the OS filesystem
  * The entire API should be callable from any thread (except the singleton can be assumed to be called at least once before concurrent access).
@@ -130,6 +135,8 @@ public:
 	 */
 	virtual bool CachedDataProbablyExists(const TCHAR* CacheKey) override
 	{
+		// FDDCScopeStatHelper Stat(FString(CacheKey), FString(TEXT("FileSystemDerivedDataBackend::CachedDataProbablyExists")));
+
 		check(!bFailed);
 		FString Filename = BuildFilename(CacheKey);
 		FDateTime TimeStamp = IFileManager::Get().GetTimeStamp(*Filename);
@@ -154,6 +161,11 @@ public:
 	 */
 	virtual bool GetCachedData(const TCHAR* CacheKey,TArray<uint8>& Data) override
 	{
+		static FName NAME_GetCachedData(TEXT("GetCachedData"));
+		FDDCScopeStatHelper Stat(CacheKey, NAME_GetCachedData);
+		static FName NAME_FileDDCPath(TEXT("FileDDCPath"));
+		Stat.AddTag(NAME_FileDDCPath, CachePath);
+		
 		check(!bFailed);
 		FString Filename = BuildFilename(CacheKey);
 
diff --git a/Engine/Source/Developer/DerivedDataCache/Private/HierarchicalDerivedDataBackend.h b/Engine/Source/Developer/DerivedDataCache/Private/HierarchicalDerivedDataBackend.h
--- a/Engine/Source/Developer/DerivedDataCache/Private/HierarchicalDerivedDataBackend.h
+++ b/Engine/Source/Developer/DerivedDataCache/Private/HierarchicalDerivedDataBackend.h
@@ -2,6 +2,9 @@
 
 #pragma once
 
+#include "DDCStatsHelper.h"
+
+
 /** 
  * A backend wrapper that implements a cache hierarchy of backends. 
 **/
@@ -80,6 +83,13 @@ public:
 	 */
 	virtual bool GetCachedData(const TCHAR* CacheKey, TArray<uint8>& OutData) override
 	{
+		const static FName NAME_GetCachedData = FName(TEXT("GetCachedData"));
+		const static FName NAME_HierarchicalDDC = FName(TEXT("HierarchicalDDC"));
+		const static FName NAME_DataSize = FName(TEXT("DataSize"));
+
+		FDDCScopeStatHelper Stat(CacheKey, NAME_GetCachedData);
+		Stat.AddTag(NAME_HierarchicalDDC, FString());
+
 		for (int32 CacheIndex = 0; CacheIndex < InnerBackends.Num(); CacheIndex++)
 		{
 			if (InnerBackends[CacheIndex]->CachedDataProbablyExists(CacheKey) && InnerBackends[CacheIndex]->GetCachedData(CacheKey, OutData))
@@ -119,6 +129,7 @@ public:
 						}
 					}
 				}
+				Stat.AddTag(NAME_DataSize, FString::Printf(TEXT("%d bytes"), OutData.Num()));
 				return true;
 			}
 		}
diff --git a/Engine/Source/Editor/CookingStats/CookingStats.Build.cs b/Engine/Source/Editor/CookingStats/CookingStats.Build.cs
new file mode 100644
--- /dev/null
+++ b/Engine/Source/Editor/CookingStats/CookingStats.Build.cs
@@ -0,0 +1,22 @@
+// Copyright 1998-2015 Epic Games, Inc. All Rights Reserved.
+
+using UnrealBuildTool;
+
+public class CookingStats : ModuleRules
+{
+	public CookingStats(TargetInfo Target)
+	{
+		PrivateDependencyModuleNames.AddRange(
+			new string[] {
+				"Core",
+				"CoreUObject"
+			}
+			);
+
+		if (UEBuildConfiguration.bBuildEditor == true)
+		{
+			PrivateIncludePathModuleNames.AddRange(new string[] { "DirectoryWatcher" });
+			DynamicallyLoadedModuleNames.AddRange(new string[] { "DirectoryWatcher" });
+		}
+	}
+}
diff --git a/Engine/Source/Editor/CookingStats/Private/CookingStats.cpp b/Engine/Source/Editor/CookingStats/Private/CookingStats.cpp
new file mode 100644
--- /dev/null
+++ b/Engine/Source/Editor/CookingStats/Private/CookingStats.cpp
@@ -0,0 +1,85 @@
+// Copyright 1998-2015 Epic Games, Inc. All Rights Reserved.
+
+
+#include "CookingStatsPCH.h"
+#include "ModuleManager.h"
+
+DEFINE_LOG_CATEGORY(LogCookingStats);
+
+FCookingStats::FCookingStats()
+{
+	RunGuid = FName(*FGuid::NewGuid().ToString());
+}
+
+FCookingStats::~FCookingStats()
+{
+}
+
+void FCookingStats::AddRunTag(const FName& Tag, const FString& Value)
+{
+	AddTagValue(RunGuid, Tag, Value);
+}
+
+void FCookingStats::AddTag(const FName& Key, const FName& Tag)
+{
+	AddTagValue(Key, Tag, FString());
+}
+
+void FCookingStats::AddTagValue(const FName& Key, const FName& Tag, const FString& TagValue)
+{
+	FScopeLock ScopeLock(&SyncObject);
+	auto Value = KeyTags.Find(Key);
+
+	if (Value == nullptr)
+	{
+		Value = &KeyTags.Add(Key);
+	}
+
+	FTag FullTag;
+	FullTag.Name = Tag;
+	FullTag.Value = TagValue;
+	
+	
+	Value->Add(FullTag);
+
+}
+
+bool FCookingStats::SaveStatsAsCSV(const FString& Filename)
+{
+	FString Output;
+	FScopeLock ScopeLock(&SyncObject);
+
+	for (const auto& Stat : KeyTags)
+	{
+		Output += Stat.Key.ToString();
+
+		for (const auto& Tag : Stat.Value)
+		{
+			Output += TEXT(", ");
+			Output += Tag.Name.ToString();
+			if (Tag.Value.IsEmpty() == false)
+			{
+				Output += TEXT("=");
+				Output += Tag.Value;
+			}
+			
+		}
+
+		for (const auto& Tag : GlobalTags)
+		{
+			Output += TEXT(", ");
+			Output += Tag.Name.ToString();
+			if (Tag.Value.IsEmpty() == false)
+			{
+				Output += TEXT("=");
+				Output += Tag.Value;
+			}
+		}
+
+		Output += FString::Printf(TEXT(", RunGuid=%s\r\n"), *RunGuid.ToString());
+	}
+
+	FFileHelper::SaveStringToFile(Output, *Filename);
+	return true;
+}
+
diff --git a/Engine/Source/Editor/CookingStats/Private/CookingStats.h b/Engine/Source/Editor/CookingStats/Private/CookingStats.h
new file mode 100644
--- /dev/null
+++ b/Engine/Source/Editor/CookingStats/Private/CookingStats.h
@@ -0,0 +1,62 @@
+// Copyright 1998-2015 Epic Games, Inc. All Rights Reserved.
+
+
+class FCookingStats : public ICookingStats
+{
+public:
+
+	FCookingStats();
+	virtual ~FCookingStats();
+
+
+	/**
+	* AddRunTag
+	* Add a tag and associate it with this run of the editor 
+	* this is for storing global information from this run
+	*
+	* @param Tag tag which we want to keep
+	* @param Value for this tag FString() if we just want to store the tag
+	*/
+	virtual void AddRunTag(const FName& Tag, const FString& Value) override;
+
+	/**
+	* AddTag
+	* add a tag to a key, the tags will be saved out with the key
+	* tags can be timing information or anything
+	*
+	* @param Key key to add the tag to
+	* @param Tag tag to add to the key
+	*/
+	virtual void AddTag(const FName& Key, const FName& Tag) override;
+
+	/**
+	* AddTag
+	* add a tag to a key, the tags will be saved out with the key
+	* tags can be timing information or anything
+	*
+	* @param Key key to add the tag to
+	* @param Tag tag to add to the key
+	*/
+	virtual void AddTagValue(const FName& Key, const FName& Tag, const FString& Value) override;
+
+	/**
+	* SaveStatsAsCSV
+	* Save stats in a CSV file
+	*
+	* @param Filename file name to save comma delimited stats to
+	* @return true if succeeded
+	*/
+	virtual bool SaveStatsAsCSV(const FString& Filename) override;
+
+private:
+	FName RunGuid; // unique guid for this run of the editor
+
+	struct FTag
+	{
+		FName Name;
+		FString Value;
+	};
+	TMap< FName, TArray<FTag> > KeyTags;
+	TArray<FTag> GlobalTags; // tags to add to every key in this runthrough
+	FCriticalSection SyncObject;
+};
\ No newline at end of file
diff --git a/Engine/Source/Editor/CookingStats/Private/CookingStatsModule.cpp b/Engine/Source/Editor/CookingStats/Private/CookingStatsModule.cpp
new file mode 100644
--- /dev/null
+++ b/Engine/Source/Editor/CookingStats/Private/CookingStatsModule.cpp
@@ -0,0 +1,36 @@
+// Copyright 1998-2015 Epic Games, Inc. All Rights Reserved.
+
+
+#include "CookingStatsPCH.h"
+#include "CookingStatsModule.h"
+
+#include "Developer/MessageLog/Public/MessageLogModule.h"
+
+#define LOCTEXT_NAMESPACE "CookingStats"
+
+
+IMPLEMENT_MODULE(FCookingStatsModule, CookingStats);
+
+void FCookingStatsModule::StartupModule()
+{
+	CookingStats = new FCookingStats();
+
+}
+
+
+void FCookingStatsModule::ShutdownModule()
+{
+	if (CookingStats)
+	{
+		delete CookingStats;
+		CookingStats = NULL;
+	}
+}
+
+ICookingStats& FCookingStatsModule::Get() const
+{
+	check(CookingStats);
+	return *CookingStats;
+}
+
+#undef LOCTEXT_NAMESPACE
\ No newline at end of file
diff --git a/Engine/Source/Editor/CookingStats/Private/CookingStatsPCH.h b/Engine/Source/Editor/CookingStats/Private/CookingStatsPCH.h
new file mode 100644
--- /dev/null
+++ b/Engine/Source/Editor/CookingStats/Private/CookingStatsPCH.h
@@ -0,0 +1,12 @@
+#ifndef _COOKING_STATS_PRIVATE_
+#define _COOKING_STATS_PRIVATE_
+
+#include "Core.h"
+#include "CoreUObject.h"
+
+DECLARE_LOG_CATEGORY_EXTERN(LogCookingStats, Log, All);
+
+#include "ICookingStats.h"
+#include "CookingStats.h"
+
+#endif 
diff --git a/Engine/Source/Editor/CookingStats/Public/CookingStatsModule.h b/Engine/Source/Editor/CookingStats/Public/CookingStatsModule.h
new file mode 100644
--- /dev/null
+++ b/Engine/Source/Editor/CookingStats/Public/CookingStatsModule.h
@@ -0,0 +1,37 @@
+// Copyright 1998-2015 Epic Games, Inc. All Rights Reserved.
+
+
+#pragma once
+
+#include "ModuleInterface.h"
+#include "ICookingStats.h"
+#include "CookingStatsInterface.h"
+
+namespace CookingStatsConstants
+{
+	const FName ModuleName("CookingStats");
+}
+
+/**
+* Asset registry module
+*/
+class FCookingStatsModule : public ICookingStatsInterface
+{
+
+public:
+
+	/** Called right after the module DLL has been loaded and the module object has been created */
+	virtual void StartupModule() override;
+
+	/** Called before the module is unloaded, right before the module object is destroyed. */
+	virtual void ShutdownModule() override;
+
+	/** Gets the asset registry singleton */
+	virtual ICookingStats& Get() const;
+
+
+private:
+	ICookingStats* CookingStats;
+};
+
+
diff --git a/Engine/Source/Editor/CookingStats/Public/ICookingStats.h b/Engine/Source/Editor/CookingStats/Public/ICookingStats.h
new file mode 100644
--- /dev/null
+++ b/Engine/Source/Editor/CookingStats/Public/ICookingStats.h
@@ -0,0 +1,52 @@
+// Copyright 1998-2015 Epic Games, Inc. All Rights Reserved.
+#pragma once
+
+class ICookingStats
+{
+public:
+
+	/** Virtual destructor*/
+	virtual ~ICookingStats() {}
+
+
+	/**
+	* AddRunTag
+	* Add a tag and associate it with this run of the editor
+	* this is for storing global information from this run
+	*
+	* @param Tag tag which we want to keep
+	* @param Value for this tag FString() if we just want to store the tag
+	*/
+	virtual void AddRunTag(const FName& Tag, const FString& Value) =0;
+
+	/**
+	 * AddTag
+	 * add a tag to a key, the tags will be saved out with the key
+	 * tags can be timing information or anything
+	 *
+	 * @param Key key to add the tag to
+	 * @param Tag tag to add to the key
+	 */
+	virtual void AddTag(const FName& Key, const FName& Tag) = 0;
+
+	/**
+	* AddTag
+	* add a tag to a key, the tags will be saved out with the key
+	* tags can be timing information or anything
+	*
+	* @param Key key to add the tag to
+	* @param Tag tag to add to the key
+	*/
+	virtual void AddTagValue(const FName& Key, const FName& Tag, const FString& Value) = 0;
+
+	/**
+	 * SaveStatsAsCSV
+	 * Save stats in a CSV file
+	 * 
+	 * @param Filename file name to save comma delimited stats to 
+	 * @return true if succeeded
+	 */
+	virtual bool SaveStatsAsCSV(const FString& Filename) = 0;
+};
+
+
diff --git a/Engine/Source/Editor/UnrealEd/Classes/CookOnTheSide/CookOnTheFlyServer.h b/Engine/Source/Editor/UnrealEd/Classes/CookOnTheSide/CookOnTheFlyServer.h
--- a/Engine/Source/Editor/UnrealEd/Classes/CookOnTheSide/CookOnTheFlyServer.h
+++ b/Engine/Source/Editor/UnrealEd/Classes/CookOnTheSide/CookOnTheFlyServer.h
@@ -607,9 +607,10 @@ private:
 	FString OutputDirectoryOverride;
 	//////////////////////////////////////////////////////////////////////////
 	// Cook by the book options
+public:
 	struct FChildCooker
 	{
-		FChildCooker() : ReadPipe(nullptr), ReturnCode(-1), bFinished(false) { }
+		FChildCooker() : ReadPipe(nullptr), ReturnCode(-1), bFinished(false),  Thread(nullptr) { }
 
 		FProcHandle ProcessHandle;
 		FString ResponseFileName;
@@ -617,7 +618,9 @@ private:
 		void* ReadPipe;
 		int32 ReturnCode;
 		bool bFinished;
+		FRunnableThread* Thread;
 	};
+private:
 	struct FCookByTheBookOptions
 	{
 	public:
@@ -865,6 +868,15 @@ public:
 	void ClearPlatformCookedData( const FName& PlatformName );
 
 
+	/**
+	 * Recompile any global shader changes 
+	 * if any are detected then clear the cooked platform data so that they can be rebuilt
+	 * 
+	 * @return return true if shaders were recompiled
+	 */
+	bool RecompileChangedShaders(const TArray<FName>& TargetPlatforms);
+
+
 	/**
 	 * Force stop whatever pending cook requests are going on and clear all the cooked data
 	 * Note cook on the side / cook on the fly clients may not be able to recover from this if they are waiting on a cook request to complete
@@ -928,6 +940,25 @@ public:
 	/** Returns the configured amount of memory allowed before forcing a GC */
 	uint64 GetMaxMemoryAllowance() const;
 
+	/**
+	* RequestPackage to be cooked
+	*
+	* @param StandardPackageFName name of the package in standard format as returned by FPaths::MakeStandardFilename
+	* @param TargetPlatforms name of the targetplatforms we want this package cooked for
+	* @param bForceFrontOfQueue should we put this package in the front of the cook queue (next to be processed) or at the end
+	*/
+	bool RequestPackage(const FName& StandardPackageFName, const TArray<FName>& TargetPlatforms, const bool bForceFrontOfQueue);
+
+	/**
+	* RequestPackage to be cooked
+	* This function can only be called while the cooker is in cook by the book mode
+	*
+	* @param StandardPackageFName name of the package in standard format as returned by FPaths::MakeStandardFilename
+	* @param bForceFrontOfQueue should we put this package in the front of the cook queue (next to be processed) or at the end
+	*/
+	bool RequestPackage(const FName& StandardPackageFName, const bool bForceFrontOfQueue);
+
+
 	/**
 	 * Callbacks from editor 
 	 */
@@ -988,7 +1019,13 @@ private:
 	 * 
 	 * @return return true if all child cookers are finished
 	 */
-	bool TickChildCookers( const bool bCheckIfFinished );
+	bool TickChildCookers();
+
+	/**
+	 * CleanUPChildCookers
+	 * can only be called after TickChildCookers returns true
+	 */
+	void CleanUpChildCookers();
 
 	/**
 	 * Get all the packages which are listed in asset registry passed in.  
@@ -1151,8 +1188,6 @@ private:
 	FString ConvertToFullSandboxPath( const FString &FileName, bool bForWrite = false ) const;
 	FString ConvertToFullSandboxPath( const FString &FileName, bool bForWrite, const FString& PlatformName ) const;
 
-
-
 	/**
 	 * GetSandboxAssetRegistryFilename
 	 * 
diff --git a/Engine/Source/Editor/UnrealEd/Private/CookOnTheFlyServer.cpp b/Engine/Source/Editor/UnrealEd/Private/CookOnTheFlyServer.cpp
--- a/Engine/Source/Editor/UnrealEd/Private/CookOnTheFlyServer.cpp
+++ b/Engine/Source/Editor/UnrealEd/Private/CookOnTheFlyServer.cpp
@@ -24,6 +24,9 @@
 #include "PhysicsPublic.h"
 #include "AutoSaveUtils.h"
 
+// cooking stats
+#include "CookingStatsModule.h"
+
 // cook by the book requirements
 #include "Commandlets/ChunkManifestGenerator.h"
 #include "Engine/WorldComposition.h"
@@ -430,6 +433,18 @@ const FString& GetAssetRegistryPath()
 	return AssetRegistryPath;
 }
 
+const FString GetStatsFilename(const FString& ResponseFilename)
+{
+	if (ResponseFilename.Len())
+	{
+		return ResponseFilename + TEXT("Stats.csv");
+	}
+	else
+	{
+		static FString StatsPath = FPaths::GameSavedDir() / FString::Printf(TEXT("Stats/%sStats.csv"), *FDateTime::Now().ToString());
+		return StatsPath;
+	}
+}
 
 FString GetChildCookerResultFilename(const FString& ResponseFilename)
 {
@@ -591,6 +606,7 @@ UCookOnTheFlyServer::UCookOnTheFlyServer(const FObjectInitializer& ObjectInitial
 
 UCookOnTheFlyServer::~UCookOnTheFlyServer()
 {
+	check(TickChildCookers() == true);
 	if ( CookByTheBookOptions )
 	{		
 		delete CookByTheBookOptions;
@@ -1137,83 +1153,70 @@ FString UCookOnTheFlyServer::GetDLCContentPath()
 COREUOBJECT_API extern bool GOutputCookingWarnings;
 
 
+bool UCookOnTheFlyServer::RequestPackage(const FName& StandardPackageFName, const TArray<FName>& TargetPlatforms, const bool bForceFrontOfQueue)
+{
+	FFilePlatformRequest FileRequest(StandardPackageFName, TargetPlatforms);
+	CookRequests.EnqueueUnique(FileRequest, bForceFrontOfQueue);
+	return true;
+}
+
 
-bool UCookOnTheFlyServer::TickChildCookers(const bool bCheckIfFinished)
+bool UCookOnTheFlyServer::RequestPackage(const FName& StandardPackageFName, const bool bForceFrontOfQueue)
+{
+	check(IsCookByTheBookMode());
+	// need target platforms if we are not in cook by the book mode
+	FFilePlatformRequest FileRequest(StandardPackageFName, CookByTheBookOptions->TargetPlatformNames);
+	CookRequests.EnqueueUnique(FileRequest, bForceFrontOfQueue);
+	return true;
+}
+
+// should only call this after tickchildcookers returns false
+void UCookOnTheFlyServer::CleanUpChildCookers()
 {
-	SCOPE_TIMER(TickChildCookers);
-	bool bChildCookersFinished = true;
 	if (IsCookByTheBookMode())
 	{
+		const FString MasterCookerStatsFilename = GetStatsFilename(FString());
+
+		FString AllStats;
+		ensure(FFileHelper::LoadFileToString(AllStats, *MasterCookerStatsFilename));
+		
 		for (auto& ChildCooker : CookByTheBookOptions->ChildCookers)
 		{
-			if (ChildCooker.bFinished == false)
-			{
-				bool bIsCookerFinished = false;
-				int32 ReturnCode = 0;
-				if (bCheckIfFinished == true)
-				{
-					SCOPE_TIMER(CheckIfChildCookerIsFinished);
-					if (FPlatformProcess::GetProcReturnCode(ChildCooker.ProcessHandle, &ReturnCode))
-					{
-						UE_LOG(LogCook, Display, TEXT("Child cooker %s returned error code %d"), *ChildCooker.BaseResponseFileName, ReturnCode);
-						ChildCooker.ReturnCode = ReturnCode;
-						ChildCooker.bFinished = true;
-
-						if (ReturnCode != 0)
-						{
-							UE_LOG(LogCook, Error, TEXT("Child cooker %s returned error code %d"), *ChildCooker.BaseResponseFileName, ReturnCode);
-						}
+			check(ChildCooker.bFinished == true);
+			const FString ChildCookerStatsFilename = GetStatsFilename(ChildCooker.ResponseFileName);
 
+			FString ChildStats;
+			ensure(FFileHelper::LoadFileToString(ChildStats, *ChildCookerStatsFilename));
 
-						// if the child cooker completed successfully then it would have output a list of files which the main cooker needs to process
-						const FString AdditionalPackagesFileName = GetChildCookerResultFilename(ChildCooker.ResponseFileName);
+			AllStats += ChildStats;
 
-						FString AdditionalPackages;
-						if (FFileHelper::LoadFileToString(AdditionalPackages, *AdditionalPackagesFileName) == false)
-						{
-							UE_LOG(LogCook, Fatal, TEXT("ChildCooker failed to write out additional packages file to location %s"), *AdditionalPackagesFileName);
-						}
-						TArray<FString> AdditionalPackageList;
-						AdditionalPackages.ParseIntoArrayLines(AdditionalPackageList);
+			if (ChildCooker.Thread != nullptr)
+			{
+				ChildCooker.Thread->WaitForCompletion();
+				delete ChildCooker.Thread;
+				ChildCooker.Thread = nullptr;
+			}
+		}
 
-						for (const auto& AdditionalPackage : AdditionalPackageList)
-						{
-							FName Filename = FName(*AdditionalPackage);
-							FFilePlatformRequest CheckCooked = FFilePlatformRequest(Filename, CookByTheBookOptions->TargetPlatformNames);
-							if (CookedPackages.Exists(CheckCooked) == false)
-							{
-								CookRequests.EnqueueUnique(CheckCooked);
-							}
-						}
-					}
-					else
-					{
-						bChildCookersFinished = false;
-					}
-				}
-				else
-				{
-					bChildCookersFinished = false;
-				}
+		ensure(FFileHelper::SaveStringToFile(AllStats, *MasterCookerStatsFilename));
+	}
+}
 
-				SCOPE_TIMER(ProcessLogOutput);
-				// process the log output from the child cooker even if we just finished to ensure that we get the end of the log
-				FString PipeContents = FPlatformProcess::ReadPipe(ChildCooker.ReadPipe);
-				if (PipeContents.Len())
-				{
-					TArray<FString> PipeLines;
-					PipeContents.ParseIntoArrayLines(PipeLines);
 
-					for (const auto& Line : PipeLines)
-					{
-						UE_LOG(LogCook, Display, TEXT("Cooker output %s: %s"), *ChildCooker.BaseResponseFileName, *Line);
-					}
-				}
+bool UCookOnTheFlyServer::TickChildCookers()
+{
+	if (IsCookByTheBookMode())
+	{
+		for (auto& ChildCooker : CookByTheBookOptions->ChildCookers)
+		{
+			if (ChildCooker.bFinished == false)
+			{
+				return false;
 			}
-
 		}
 	}
-	return bChildCookersFinished;
+	return true;
+
 }
 
 
@@ -1283,7 +1286,7 @@ uint32 UCookOnTheFlyServer::TickCookOnTheSide( const float TimeSlice, uint32 &Co
 		{
 			SCOPE_TIMER(WaitingForChildCookers);
 			// we have nothing left to cook so we are now waiting for child cookers to finish
-			if (TickChildCookers(true) == false)
+			if (TickChildCookers() == false)
 			{
 				Result |= COSR_WaitingOnChildCookers;
 			}
@@ -1291,7 +1294,7 @@ uint32 UCookOnTheFlyServer::TickCookOnTheSide( const float TimeSlice, uint32 &Co
 		}
 		else
 		{
-			if (TickChildCookers(false) == false)
+			if (TickChildCookers() == false)
 			{
 				Result |= COSR_WaitingOnChildCookers;
 			}
@@ -1486,11 +1489,21 @@ uint32 UCookOnTheFlyServer::TickCookOnTheSide( const float TimeSlice, uint32 &Co
 							FString Filename;
 							if ( FPackageName::DoesPackageExist(PackageName, NULL, &Filename) )
 							{
-								check(IsChildCooker() == false);
+
 
 								FString StandardFilename = FPaths::ConvertRelativePathToFull(Filename);
 								FPaths::MakeStandardFilename(StandardFilename);
-								CookRequests.EnqueueUnique(FFilePlatformRequest(MoveTemp(FName(*StandardFilename)), MoveTemp(TargetPlatformNames)));
+								FName StandardPackageFName = FName(*StandardFilename);
+								if (IsChildCooker())
+								{
+									check(IsCookByTheBookMode());
+									// notify the main cooker that it should make sure this package gets cooked
+									CookByTheBookOptions->ChildUnsolicitedPackages.Add(StandardPackageFName);
+								}
+								else
+								{
+									CookRequests.EnqueueUnique(FFilePlatformRequest(MoveTemp(StandardPackageFName), MoveTemp(TargetPlatformNames)));
+								}
 							}
 						}
 
@@ -3428,6 +3441,8 @@ const FString UCookOnTheFlyServer::GetCookedAssetRegistryFilename(const FString&
 	return CookedAssetRegistryFilename;
 }
 
+
+
 void UCookOnTheFlyServer::CookByTheBookFinished()
 {
 	check( IsInGameThread() );
@@ -3438,6 +3453,12 @@ void UCookOnTheFlyServer::CookByTheBookFinished()
 
 	GetDerivedDataCacheRef().WaitForQuiescence(true);
 	
+	static const FName CookingStatsName("CookingStats");
+	FCookingStatsModule& CookingStatsModule = FModuleManager::LoadModuleChecked<FCookingStatsModule>(CookingStatsName);
+	ICookingStats& CookingStats = CookingStatsModule.Get();
+
+	const FString StatsFilename = GetStatsFilename(CookByTheBookOptions->ChildCookFilename);
+	CookingStats.SaveStatsAsCSV(StatsFilename);
 
 	if (IsChildCooker())
 	{
@@ -3448,14 +3469,16 @@ void UCookOnTheFlyServer::CookByTheBookFinished()
 			UncookedPackageList.Append(UncookedPackage.ToString() + TEXT("\n\r"));
 		}
 		FFileHelper::SaveStringToFile(UncookedPackageList, *GetChildCookerResultFilename(CookByTheBookOptions->ChildCookFilename));
+
+		
+		
 	}
 	else
 	{
+		CleanUpChildCookers();
+
 		check(CookByTheBookOptions->ChildUnsolicitedPackages.Num() == 0);
-	}
 
-	if (IsChildCooker() == false)
-	{
 		SCOPE_TIMER(SavingAssetRegistry);
 		// Save modified asset registry with all streaming chunk info generated during cook
 		const FString& SandboxRegistryFilename = GetSandboxAssetRegistryFilename();
@@ -3783,6 +3806,13 @@ void UCookOnTheFlyServer::StartCookByTheBook( const FCookByTheBookStartupOptions
 	InitializeSandbox();
 
 
+	// need to test this out
+	/*if (RecompileChangedShaders(TargetPlatformNames))
+	{
+		// clean everything :(
+		CleanSandbox(false);
+	}*/
+
 	CookByTheBookOptions->bLeakTest = (CookOptions & ECookByTheBookOptions::LeakTest) != ECookByTheBookOptions::None; // this won't work from the editor this needs to be standalone
 	check(!CookByTheBookOptions->bLeakTest || CurrentCookMode == ECookMode::CookByTheBook);
 
@@ -3959,8 +3989,100 @@ void UCookOnTheFlyServer::StartCookByTheBook( const FCookByTheBookStartupOptions
 }
 
 
+bool UCookOnTheFlyServer::RecompileChangedShaders(const TArray<FName>& TargetPlatforms)
+{
+	bool bShadersRecompiled = false;
+	for (const auto& TargetPlatform : TargetPlatforms)
+	{
+		bShadersRecompiled |= RecompileChangedShadersForPlatform(TargetPlatform.ToString());
+	}
+	return bShadersRecompiled;
+}
+
 
 
+
+class FChildCookerRunnable : public FRunnable
+{
+private:
+	UCookOnTheFlyServer* CookServer;
+	UCookOnTheFlyServer::FChildCooker* ChildCooker;
+
+	void ProcessChildLogOutput()
+	{
+		// process the log output from the child cooker even if we just finished to ensure that we get the end of the log
+		FString PipeContents = FPlatformProcess::ReadPipe(ChildCooker->ReadPipe);
+		if (PipeContents.Len())
+		{
+			TArray<FString> PipeLines;
+			PipeContents.ParseIntoArrayLines(PipeLines);
+
+			for (const auto& Line : PipeLines)
+			{
+				UE_LOG(LogCook, Display, TEXT("Cooker output %s: %s"), *ChildCooker->BaseResponseFileName, *Line);
+			}
+		}
+	}
+
+public:
+	FChildCookerRunnable(UCookOnTheFlyServer::FChildCooker* InChildCooker, UCookOnTheFlyServer* InCookServer) : FRunnable(), CookServer(InCookServer), ChildCooker(InChildCooker)
+	{ }
+
+	
+	virtual uint32 Run() override
+	{
+		while (true)
+		{
+			check(ChildCooker != nullptr);
+			check(ChildCooker->bFinished == false); 
+
+			int32 ReturnCode;
+			if (FPlatformProcess::GetProcReturnCode(ChildCooker->ProcessHandle, &ReturnCode))
+			{
+				if (ReturnCode != 0)
+				{
+					UE_LOG(LogCook, Error, TEXT("Child cooker %s returned error code %d"), *ChildCooker->BaseResponseFileName, ReturnCode);
+				}
+				else
+				{
+					UE_LOG(LogCook, Display, TEXT("Child cooker %s returned %d"), *ChildCooker->BaseResponseFileName, ReturnCode);
+				}
+
+
+				// if the child cooker completed successfully then it would have output a list of files which the main cooker needs to process
+				const FString AdditionalPackagesFileName = GetChildCookerResultFilename(ChildCooker->ResponseFileName);
+
+				FString AdditionalPackages;
+				if (FFileHelper::LoadFileToString(AdditionalPackages, *AdditionalPackagesFileName) == false)
+				{
+					UE_LOG(LogCook, Fatal, TEXT("ChildCooker failed to write out additional packages file to location %s"), *AdditionalPackagesFileName);
+				}
+				TArray<FString> AdditionalPackageList;
+				AdditionalPackages.ParseIntoArrayLines(AdditionalPackageList);
+
+				for (const auto& AdditionalPackage : AdditionalPackageList)
+				{
+					UE_LOG(LogCook, Display, TEXT("Child cooker %s requested additional package %s to be cooked"), *ChildCooker->BaseResponseFileName, *AdditionalPackage);
+					FName Filename = FName(*AdditionalPackage);
+					CookServer->RequestPackage(Filename, false);
+				}
+				ProcessChildLogOutput();
+
+				ChildCooker->ReturnCode = ReturnCode;
+				ChildCooker->bFinished = true;
+				ChildCooker = nullptr;
+				return 1;
+			}
+
+			ProcessChildLogOutput();
+
+			FPlatformProcess::Sleep(0.05f);
+		}
+		return 1;
+	}
+
+};
+
 // sue chefs away!!
 void UCookOnTheFlyServer::StartChildCookers(int32 NumCookersToSpawn, const TArray<FName>& TargetPlatformNames)
 {
@@ -4049,6 +4171,9 @@ void UCookOnTheFlyServer::StartChildCookers(int32 NumCookersToSpawn, const TArra
 		TargetPlatformString += TargetPlatformName.ToString();
 	}
 
+	// allocate the memory here, this can't change while we are running the child threads which are handling input because they have a ptr to the childcookers array
+	CookByTheBookOptions->ChildCookers.Empty(NumCookersToSpawn);
+
 	// start the child cookers and give them each some distribution candidates
 	for (int32 CookerCounter = 0; CookerCounter < NumCookersToSpawn; ++CookerCounter )
 	{
@@ -4154,14 +4279,15 @@ void UCookOnTheFlyServer::StartChildCookers(int32 NumCookersToSpawn, const TArra
 
 		ChildCooker.ProcessHandle = FPlatformProcess::CreateProc(*ExecutablePath, *CommandLine, false, true, true, NULL, 0, NULL, WritePipe);
 		ChildCooker.ReadPipe = ReadPipe;
+
+		// start threads to monitor output and finished state
+		ChildCooker.Thread = FRunnableThread::Create(new FChildCookerRunnable(&ChildCooker, this), *FString::Printf(TEXT("ChildCookerInputHandleThreadFor:%s"), *ChildCooker.BaseResponseFileName));
 	}
 
 }
 
 
 
-
-
 /* UCookOnTheFlyServer callbacks
  *****************************************************************************/
 
diff --git a/Engine/Source/Editor/UnrealEd/UnrealEd.Build.cs b/Engine/Source/Editor/UnrealEd/UnrealEd.Build.cs
--- a/Engine/Source/Editor/UnrealEd/UnrealEd.Build.cs
+++ b/Engine/Source/Editor/UnrealEd/UnrealEd.Build.cs
@@ -25,6 +25,7 @@ public class UnrealEd : ModuleRules
 			{
 				"Analytics",
 				"AssetRegistry",
+                "CookingStats",
 				"AssetTools",
                 "BehaviorTreeEditor",
 				"ClassViewer",
diff --git a/Engine/Source/Runtime/CoreUObject/CoreUObject.Build.cs b/Engine/Source/Runtime/CoreUObject/CoreUObject.Build.cs
--- a/Engine/Source/Runtime/CoreUObject/CoreUObject.Build.cs
+++ b/Engine/Source/Runtime/CoreUObject/CoreUObject.Build.cs
@@ -10,10 +10,18 @@ public class CoreUObject : ModuleRules
 
 		PrivateIncludePaths.Add("Runtime/CoreUObject/Private");
 
-		PrivateIncludePathModuleNames.Add("TargetPlatform");
+        PrivateIncludePathModuleNames.AddRange(
+                new string[] 
+			    {
+				    "TargetPlatform",
+				    "CookingStats",
+			    }
+            );
 
 		PublicDependencyModuleNames.Add("Core");
 
 		PrivateDependencyModuleNames.Add("Projects");
+
 	}
+
 }
diff --git a/Engine/Source/Runtime/CoreUObject/Private/UObject/SavePackage.cpp b/Engine/Source/Runtime/CoreUObject/Private/UObject/SavePackage.cpp
--- a/Engine/Source/Runtime/CoreUObject/Private/UObject/SavePackage.cpp
+++ b/Engine/Source/Runtime/CoreUObject/Private/UObject/SavePackage.cpp
@@ -6,7 +6,6 @@
 #include "TargetPlatform.h"
 #include "UObject/UObjectThreadContext.h"
 
-
 DEFINE_LOG_CATEGORY_STATIC(LogSavePackage, Log, All);
 
 static const int32 MAX_MERGED_COMPRESSION_CHUNKSIZE = 1024 * 1024;
@@ -16,26 +15,60 @@ static const FName WorldClassName = FName("World");
 #define VALIDATE_INITIALIZECORECLASSES 0
 #define EXPORT_SORTING_DETAILED_LOGGING 0
 
-#define UE_PROFILE_COOKSAVE 0
+#define UE_PROFILE_COOKSAVE 1
 #if UE_PROFILE_COOKSAVE
 
-#define UE_START_LOG_COOK_TIME(InFilename) double PreviousTime; double StartTime; PreviousTime = StartTime = FPlatformTime::Seconds(); \
-	UE_LOG(LogSavePackage, Log, TEXT("Starting save package for %s"), InFilename);
+#include "CookingStatsModule.h"
+
+#define UE_OUTPUTSTATSTOLOG 0
+
+static const FName CookingStatsName("CookingStats");
+static const FName GSavePackageTransactionId(*FString::Printf(TEXT("SavePackageTransactionId%s"), *FGuid::NewGuid().ToString()));
+static uint32 GSavePackageTransactionNumber = 0;
+void CookStatsStartStat(const FName& Key, const TCHAR *Filename)
+{
+#if UE_OUTPUTSTATSTOLOG
+	UE_LOG(LogSavePackage, Log, TEXT("Starting save package for %s"), Filename);
+#endif
+
+	
+	static FCookingStatsModule& CookingStatsModule = FModuleManager::LoadModuleChecked<FCookingStatsModule>(CookingStatsName);
+	static ICookingStats& CookingStats = CookingStatsModule.Get();
+
+	CookingStats.AddTagValue(Key, TEXT("Filename"), FString(Filename));
 
-#define UE_LOG_COOK_TIME(TimeType) \
-{\
-	double CurrentTime = FPlatformTime::Seconds(); \
-	UE_LOG(LogSavePackage, Log, TEXT("TIMING: %s took %fms"), TimeType, (CurrentTime - PreviousTime) * 1000.0f); \
-	PreviousTime = CurrentTime; \
 }
 
-#define UE_FINISH_LOG_COOK_TIME()\
-	UE_LOG_COOK_TIME(TEXT("Unaccounted time")); \
-{\
-	double CurrentTime = FPlatformTime::Seconds(); \
-	UE_LOG(LogSavePackage, Log, TEXT("Total save time: %fms"), (CurrentTime - StartTime) * 1000.0f); \
+void CookStatsAddStat(const FName& Key, const FName& Tag, const float TimeMilliseconds )
+{
+#if UE_OUTPUTSTATSTOLOG
+	UE_LOG(LogSavePackage, Log, TEXT("TIMING: %s took %fms"), Tag, TimeMilliseconds);
+#endif
+	static FCookingStatsModule& CookingStatsModule = FModuleManager::LoadModuleChecked<FCookingStatsModule>(CookingStatsName);
+	static ICookingStats& CookingStats = CookingStatsModule.Get();
+
+	CookingStats.AddTagValue(Key, Tag, FString::Printf(TEXT("%fms"), TimeMilliseconds));
 }
 
+
+#define UE_START_LOG_COOK_TIME(InFilename) double PreviousTime; double StartTime; PreviousTime = StartTime = FPlatformTime::Seconds(); const FName CookStatsKey = FName(GSavePackageTransactionId,++GSavePackageTransactionNumber); \
+	CookStatsStartStat(CookStatsKey, InFilename);
+
+#define UE_LOG_COOK_TIME(TimeType) \
+	{\
+		double CurrentTime = FPlatformTime::Seconds(); \
+		const FName Tag = FName(TimeType);\
+		CookStatsAddStat(CookStatsKey, Tag, (CurrentTime - PreviousTime) * 1000.0f);\
+		PreviousTime = CurrentTime; \
+	}
+
+#define UE_FINISH_LOG_COOK_TIME()\
+	UE_LOG_COOK_TIME(TEXT("UnaccountedTime")); \
+	{\
+		double CurrentTime = FPlatformTime::Seconds(); \
+		UE_LOG(LogSavePackage, Log, TEXT("Total save time: %fms"), (CurrentTime - StartTime) * 1000.0f); \
+	}
+
 #else
 // don't do anything
 #define UE_START_LOG_COOK_TIME(InFilename)
@@ -2938,7 +2971,6 @@ bool UPackage::SavePackage( UPackage* InOuter, UObject* Base, EObjectFlags TopLe
 			ensureMsgf(false, TEXT("Recursive SavePackage() is not supported"));
 			return false;
 		}
-		UE_LOG_COOK_TIME(TEXT("First"));
 
 
 		
@@ -2974,12 +3006,12 @@ bool UPackage::SavePackage( UPackage* InOuter, UObject* Base, EObjectFlags TopLe
 
 		FlushAsyncLoading();
 
-		UE_LOG_COOK_TIME(TEXT("Flush Async loading"));
+		UE_LOG_COOK_TIME(TEXT("FlushAsyncLoading"));
 		(*GFlushStreamingFunc)();
-		UE_LOG_COOK_TIME(TEXT("Flush streaming func"));
+		UE_LOG_COOK_TIME(TEXT("FlushStreamingFunc"));
 		FIOSystem::Get().BlockTillAllRequestsFinishedAndFlushHandles();
 
-		UE_LOG_COOK_TIME(TEXT("Block till all requests finished and flush handles"));
+		UE_LOG_COOK_TIME(TEXT("BlockTillAllRequestsFinished"));
 
 		check(InOuter);
 		check(Filename);
@@ -3102,7 +3134,7 @@ bool UPackage::SavePackage( UPackage* InOuter, UObject* Base, EObjectFlags TopLe
 		UnMarkAllObjects();
 
 
-		UE_LOG_COOK_TIME(TEXT("ResetLoadersForSave UnMarkAllObjects"));
+		UE_LOG_COOK_TIME(TEXT("ResetLoadersForSaveUnMarkAllObjects"));
 
 		// Size of serialized out package in bytes. This is before compression.
 		int32 PackageSize = INDEX_NONE;
@@ -3317,7 +3349,7 @@ bool UPackage::SavePackage( UPackage* InOuter, UObject* Base, EObjectFlags TopLe
 				TMap<UObject*, UObject*> DuplicateRedirects = UnmarkExportTagFromDuplicates();
 #endif // WITH_EDITOR
 
-				UE_LOG_COOK_TIME(TEXT("Serialize Imports"));
+				UE_LOG_COOK_TIME(TEXT("SerializeImports"));
 				
 				if ( EndSavingIfCancelled( Linker, TempFilename ) ) { return false; }
 				SlowTask.EnterProgressFrame();
@@ -3344,7 +3376,7 @@ bool UPackage::SavePackage( UPackage* InOuter, UObject* Base, EObjectFlags TopLe
 					}
 				}
 
-				UE_LOG_COOK_TIME(TEXT("Gather Localizable Text Data"));
+				UE_LOG_COOK_TIME(TEXT("GatherLocalizableTextData"));
 				
 				if ( EndSavingIfCancelled( Linker, TempFilename ) ) { return false; }
 				SlowTask.EnterProgressFrame();
@@ -3421,7 +3453,7 @@ bool UPackage::SavePackage( UPackage* InOuter, UObject* Base, EObjectFlags TopLe
 					}
 				}
 
-				UE_LOG_COOK_TIME(TEXT("Mark Names"));
+				UE_LOG_COOK_TIME(TEXT("MarkNames"));
 
 				if ( EndSavingIfCancelled( Linker, TempFilename ) ) { return false; }
 				SlowTask.EnterProgressFrame();
@@ -3609,7 +3641,7 @@ bool UPackage::SavePackage( UPackage* InOuter, UObject* Base, EObjectFlags TopLe
 				}
 #endif
 
-				UE_LOG_COOK_TIME(TEXT("Serialize Summary"));
+				UE_LOG_COOK_TIME(TEXT("SerializeSummary"));
 
 				if ( EndSavingIfCancelled( Linker, TempFilename ) ) { return false; }
 				SlowTask.EnterProgressFrame();
@@ -3629,7 +3661,7 @@ bool UPackage::SavePackage( UPackage* InOuter, UObject* Base, EObjectFlags TopLe
 					Linker->NameIndices.Add(Linker->NameMap[i], i);
 				}
 
-				UE_LOG_COOK_TIME(TEXT("Serialize Names"));
+				UE_LOG_COOK_TIME(TEXT("SerializeNames"));
 
 				if ( EndSavingIfCancelled( Linker, TempFilename ) ) { return false; }
 				SlowTask.EnterProgressFrame();
@@ -3648,7 +3680,7 @@ bool UPackage::SavePackage( UPackage* InOuter, UObject* Base, EObjectFlags TopLe
 					}
 				}
 
-				UE_LOG_COOK_TIME(TEXT("Serialize Gatherable Text Data"));
+				UE_LOG_COOK_TIME(TEXT("SerializeGatherableTextData"));
 
 				if ( EndSavingIfCancelled( Linker, TempFilename ) ) { return false; }
 				SlowTask.EnterProgressFrame();
@@ -3710,7 +3742,7 @@ bool UPackage::SavePackage( UPackage* InOuter, UObject* Base, EObjectFlags TopLe
 #endif
 
 
-				UE_LOG_COOK_TIME(TEXT("Build Export Map"));
+				UE_LOG_COOK_TIME(TEXT("BuildExportMap"));
 
 				if ( EndSavingIfCancelled( Linker, TempFilename ) ) { return false; }
 				SlowTask.EnterProgressFrame();
@@ -3719,19 +3751,19 @@ bool UPackage::SavePackage( UPackage* InOuter, UObject* Base, EObjectFlags TopLe
 				FObjectExportSortHelper ExportSortHelper;
 				ExportSortHelper.SortExports( Linker, Conform );
 				
-				UE_LOG_COOK_TIME(TEXT("Sort Exports nonseekfree"));
+				UE_LOG_COOK_TIME(TEXT("SortExportsNonSeekfree"));
 
 				if ( EndSavingIfCancelled( Linker, TempFilename ) ) { return false; }
 				SlowTask.EnterProgressFrame();
 
 				// Sort exports for seek-free loading.
 				FObjectExportSeekFreeSorter SeekFreeSorter;
-				UE_LOG_COOK_TIME(TEXT("Sort Exports seekfree constructor"));
+				UE_LOG_COOK_TIME(TEXT("SortExportsSeekfreeConstructor"));
 				SeekFreeSorter.SortExports( Linker, Conform );
-				UE_LOG_COOK_TIME(TEXT("Sort Exports seekfree inner"));
+				UE_LOG_COOK_TIME(TEXT("SortExportsSeekfreeInner"));
 				Linker->Summary.ExportCount = Linker->ExportMap.Num();
 				
-				UE_LOG_COOK_TIME(TEXT("Sort Exports seekfree"));
+				UE_LOG_COOK_TIME(TEXT("SortExportsSeekfree"));
 
 				if ( EndSavingIfCancelled( Linker, TempFilename ) ) { return false; }
 				SlowTask.EnterProgressFrame();
@@ -3873,7 +3905,7 @@ bool UPackage::SavePackage( UPackage* InOuter, UObject* Base, EObjectFlags TopLe
 				}
 
 
-				UE_LOG_COOK_TIME(TEXT("Build Dependency Map"));
+				UE_LOG_COOK_TIME(TEXT("BuildDependencyMap"));
 
 				if ( EndSavingIfCancelled( Linker, TempFilename ) ) { return false; }
 				SlowTask.EnterProgressFrame();
@@ -3912,7 +3944,7 @@ bool UPackage::SavePackage( UPackage* InOuter, UObject* Base, EObjectFlags TopLe
 
 
 
-				UE_LOG_COOK_TIME(TEXT("Serialize Import Map"));
+				UE_LOG_COOK_TIME(TEXT("SerializeImportMap"));
 
 				if ( EndSavingIfCancelled( Linker, TempFilename ) ) { return false; }
 				SlowTask.EnterProgressFrame();
@@ -3927,7 +3959,7 @@ bool UPackage::SavePackage( UPackage* InOuter, UObject* Base, EObjectFlags TopLe
 				int32 OffsetAfterExportMap = Linker->Tell();
 
 
-				UE_LOG_COOK_TIME(TEXT("Serialize Export Map"));
+				UE_LOG_COOK_TIME(TEXT("SerializeExportMap"));
 
 				if ( EndSavingIfCancelled( Linker, TempFilename ) ) { return false; }
 				SlowTask.EnterProgressFrame();
@@ -3942,7 +3974,7 @@ bool UPackage::SavePackage( UPackage* InOuter, UObject* Base, EObjectFlags TopLe
 				}
 
 
-				UE_LOG_COOK_TIME(TEXT("Serialize Dependency Map"));
+				UE_LOG_COOK_TIME(TEXT("SerializeDependencyMap"));
 
 				if (EndSavingIfCancelled(Linker, TempFilename)) { return false; }
 				SlowTask.EnterProgressFrame();
@@ -3955,7 +3987,7 @@ bool UPackage::SavePackage( UPackage* InOuter, UObject* Base, EObjectFlags TopLe
 					*Linker << StringAssetReferencesMap[i];
 				}
 
-				UE_LOG_COOK_TIME(TEXT("Serialize StringAssetReference Map"));
+				UE_LOG_COOK_TIME(TEXT("SerializeStringAssetReferenceMap"));
 	
 				// Save thumbnails
 				UPackage::SaveThumbnails( InOuter, Linker );
@@ -3981,7 +4013,7 @@ bool UPackage::SavePackage( UPackage* InOuter, UObject* Base, EObjectFlags TopLe
 					Linker->StartScriptSHAGeneration();
 				}
 
-				UE_LOG_COOK_TIME(TEXT("SaveThumbNails Save AssetRegistryData SaveWorldLevelInfo"));
+				UE_LOG_COOK_TIME(TEXT("SaveThumbNailsAssetRegistryDataWorldLevelInfo"));
 				
 				{
 					FScopedSlowTask ExportScope(Linker->ExportMap.Num());
@@ -4066,7 +4098,7 @@ bool UPackage::SavePackage( UPackage* InOuter, UObject* Base, EObjectFlags TopLe
 					}
 				}
 
-				UE_LOG_COOK_TIME(TEXT("Serialize Exports"));
+				UE_LOG_COOK_TIME(TEXT("SerializeExports"));
 
 				// if we want to generate the SHA key, get it out now that the package has finished saving
 				if (ScriptSHABytes && Linker->ContainsCode())
@@ -4131,7 +4163,7 @@ bool UPackage::SavePackage( UPackage* InOuter, UObject* Base, EObjectFlags TopLe
 				uint32 Tag = PACKAGE_FILE_TAG;
 				*Linker << Tag;
 
-				UE_LOG_COOK_TIME(TEXT("Serialize BulkData"));
+				UE_LOG_COOK_TIME(TEXT("SerializeBulkData"));
 
 				// We capture the package size before the first seek to work with archives that don't report the file
 				// size correctly while the file is still being written to.
@@ -4176,7 +4208,7 @@ bool UPackage::SavePackage( UPackage* InOuter, UObject* Base, EObjectFlags TopLe
 				}
 				
 
-				UE_LOG_COOK_TIME(TEXT("Serialize Import Map"));
+				UE_LOG_COOK_TIME(TEXT("SerializeImportMap"));
 
 				check( Linker->Tell() == OffsetAfterImportMap );
 
@@ -4188,7 +4220,7 @@ bool UPackage::SavePackage( UPackage* InOuter, UObject* Base, EObjectFlags TopLe
 				}
 				check( Linker->Tell() == OffsetAfterExportMap );
 
-				UE_LOG_COOK_TIME(TEXT("Serialize Export Map"));
+				UE_LOG_COOK_TIME(TEXT("SerializeExportMap"));
 
 				if ( EndSavingIfCancelled( Linker, TempFilename ) ) { return false; }
 				SlowTask.EnterProgressFrame();
@@ -4250,7 +4282,7 @@ bool UPackage::SavePackage( UPackage* InOuter, UObject* Base, EObjectFlags TopLe
 							AsyncWriteCompressedFile(*(FBufferArchive*)(Linker->Saver), *NewPath, FinalTimeStamp, Linker->ForceByteSwapping(), Linker->Summary.TotalHeaderSize, ExportSizes);
 							Linker->Detach();
 						}
-						UE_LOG_COOK_TIME(TEXT("Async write"));
+						UE_LOG_COOK_TIME(TEXT("AsyncWrite"));
 					}
 					else if( bCompressFromMemory )
 					{
@@ -4262,7 +4294,7 @@ bool UPackage::SavePackage( UPackage* InOuter, UObject* Base, EObjectFlags TopLe
 						{
 							Linker->Detach();
 						}
-						UE_LOG_COOK_TIME(TEXT("Compress from memory"));
+						UE_LOG_COOK_TIME(TEXT("CompressFromMemory"));
 					}
 					// Compress the temporarily file to destination.
 					else if( InOuter->PackageFlags & PKG_StoreCompressed )
@@ -4271,7 +4303,7 @@ bool UPackage::SavePackage( UPackage* InOuter, UObject* Base, EObjectFlags TopLe
 						FFileCompressionHelper CompressionHelper;
 						Success = CompressionHelper.CompressFile( *TempFilename, *NewPath, Linker );
 
-						UE_LOG_COOK_TIME(TEXT("Compress from temp file"));
+						UE_LOG_COOK_TIME(TEXT("CompressFromTempFile"));
 					}
 					// Fully compress package in one "block".
 					else if( InOuter->PackageFlags & PKG_StoreFullyCompressed )
@@ -4280,7 +4312,7 @@ bool UPackage::SavePackage( UPackage* InOuter, UObject* Base, EObjectFlags TopLe
 						FFileCompressionHelper CompressionHelper;
 						Success = CompressionHelper.FullyCompressFile( *TempFilename, *NewPath, Linker->ForceByteSwapping() );
 
-						UE_LOG_COOK_TIME(TEXT("Fully compress from temp file"));
+						UE_LOG_COOK_TIME(TEXT("FullyCompressFromTempFile"));
 					}
 					else if (bSaveAsync)
 					{
@@ -4292,7 +4324,7 @@ bool UPackage::SavePackage( UPackage* InOuter, UObject* Base, EObjectFlags TopLe
 							AsyncWriteFile(*(FBufferArchive*)(Linker->Saver), *NewPath, FinalTimeStamp);
 							Linker->Detach();
 						}
-						UE_LOG_COOK_TIME(TEXT("Async write"));
+						UE_LOG_COOK_TIME(TEXT("AsyncWrite"));
 					}
 					// Move the temporary file.
 					else
@@ -4304,7 +4336,7 @@ bool UPackage::SavePackage( UPackage* InOuter, UObject* Base, EObjectFlags TopLe
 							IFileManager::Get().SetTimeStamp(*NewPath, FinalTimeStamp);
 						}
 
-						UE_LOG_COOK_TIME(TEXT("Move file"));
+						UE_LOG_COOK_TIME(TEXT("MoveFile"));
 					}
 
 					// Make sure to clean up any stale .uncompressed_size files.
diff --git a/Engine/Source/Runtime/Engine/Engine.Build.cs b/Engine/Source/Runtime/Engine/Engine.Build.cs
--- a/Engine/Source/Runtime/Engine/Engine.Build.cs
+++ b/Engine/Source/Runtime/Engine/Engine.Build.cs
@@ -51,6 +51,7 @@ public class Engine : ModuleRules
 				"RHI",
 				"ShaderCore",
 				"AssetRegistry", // Here until FAssetData is moved to engine
+                "CookingStats",
 				"EngineMessages",
 				"EngineSettings",
 				"SynthBenchmark",
diff --git a/Engine/Source/Runtime/Engine/Private/GlobalShader.cpp b/Engine/Source/Runtime/Engine/Private/GlobalShader.cpp
--- a/Engine/Source/Runtime/Engine/Private/GlobalShader.cpp
+++ b/Engine/Source/Runtime/Engine/Private/GlobalShader.cpp
@@ -508,6 +508,64 @@ void RecompileGlobalShaders()
 	}
 }
 
+bool RecompileChangedShadersForPlatform( const FString& PlatformName )
+{
+	// figure out what shader platforms to recompile
+	ITargetPlatformManagerModule* TPM = GetTargetPlatformManager();
+	ITargetPlatform* TargetPlatform = TPM->FindTargetPlatform(PlatformName);
+	if (TargetPlatform == NULL)
+	{
+		UE_LOG(LogShaders, Display, TEXT("Failed to find target platform module for %s"), *PlatformName);
+		return false;
+	}
+
+	TArray<FName> DesiredShaderFormats;
+	TargetPlatform->GetAllTargetedShaderFormats(DesiredShaderFormats);
+
+
+
+	// figure out which shaders are out of date
+	TArray<FShaderType*> OutdatedShaderTypes;
+	TArray<const FVertexFactoryType*> OutdatedFactoryTypes;
+
+	// Pick up new changes to shader files
+	FlushShaderFileCache();
+
+	FShaderType::GetOutdatedTypes(OutdatedShaderTypes, OutdatedFactoryTypes);
+	UE_LOG(LogShaders, Display, TEXT("We found %d out of date shader types, and %d out of date VF types!"), OutdatedShaderTypes.Num(), OutdatedFactoryTypes.Num());
+
+	for (int32 FormatIndex = 0; FormatIndex < DesiredShaderFormats.Num(); FormatIndex++)
+	{
+		// get the shader platform enum
+		const EShaderPlatform ShaderPlatform = ShaderFormatToLegacyShaderPlatform(DesiredShaderFormats[FormatIndex]);
+
+		// Only compile for the desired platform if requested
+		// Kick off global shader recompiles
+		BeginRecompileGlobalShaders(OutdatedShaderTypes, ShaderPlatform);
+		
+		// Block on global shaders
+		FinishRecompileGlobalShaders();
+#if WITH_EDITOR
+		// we only want to actually compile mesh shaders if we have out of date ones
+		if (OutdatedShaderTypes.Num() || OutdatedFactoryTypes.Num())
+		{
+			for (TObjectIterator<UMaterialInterface> It; It; ++It)
+			{
+				(*It)->ClearCachedCookedPlatformData(TargetPlatform);
+			}
+		}
+#endif
+	}
+
+	if (OutdatedFactoryTypes.Num() || OutdatedShaderTypes.Num())
+	{
+		return true;
+	}
+	return false;
+}
+
+
+
 void RecompileShadersForRemote( 
 	const FString& PlatformName,
 	EShaderPlatform ShaderPlatformToCompile,
diff --git a/Engine/Source/Runtime/Engine/Public/GlobalShader.h b/Engine/Source/Runtime/Engine/Public/GlobalShader.h
--- a/Engine/Source/Runtime/Engine/Public/GlobalShader.h
+++ b/Engine/Source/Runtime/Engine/Public/GlobalShader.h
@@ -224,6 +224,15 @@ extern ENGINE_API FString GetGlobalShaderMapDDCKey();
 extern ENGINE_API FString GetMaterialShaderMapDDCKey();
 
 
+/**
+ * Recompiles global shaders and material shaders
+ * rebuilds global shaders and also 
+ * clears the cooked platform data for all materials if there is a global shader change detected
+ * can be slow
+ */
+extern ENGINE_API bool RecompileChangedShadersForPlatform(const FString& PlatformName);
+
+
 /** 
  * Recompiles global shaders
  *
