Fix shadow variable warning.

diff --git a/Engine/Source/Runtime/CoreUObject/Private/UObject/SavePackage.cpp b/Engine/Source/Runtime/CoreUObject/Private/UObject/SavePackage.cpp
--- a/Engine/Source/Runtime/CoreUObject/Private/UObject/SavePackage.cpp
+++ b/Engine/Source/Runtime/CoreUObject/Private/UObject/SavePackage.cpp
@@ -57,8 +57,8 @@ void CookStatsAddStat(const FName& Key, const FName& Tag, const float TimeMillis
 #define UE_LOG_COOK_TIME(TimeType) \
 	{\
 		double CurrentTime = FPlatformTime::Seconds(); \
-		const FName Tag = FName(TimeType);\
-		CookStatsAddStat(CookStatsKey, Tag, (CurrentTime - PreviousTime) * 1000.0f);\
+		const FName TagName = FName(TimeType);\
+		CookStatsAddStat(CookStatsKey, TagName, (CurrentTime - PreviousTime) * 1000.0f);\
 		PreviousTime = CurrentTime; \
 	}
 
