Support for GCRT_AddTMapReferencedObjects in FGCReferenceTokenStream::ReplaceOrAddAddReferencedObjectsCall.

diff --git a/Engine/Source/Runtime/CoreUObject/Private/UObject/GarbageCollection.cpp b/Engine/Source/Runtime/CoreUObject/Private/UObject/GarbageCollection.cpp
--- a/Engine/Source/Runtime/CoreUObject/Private/UObject/GarbageCollection.cpp
+++ b/Engine/Source/Runtime/CoreUObject/Private/UObject/GarbageCollection.cpp
@@ -1925,7 +1925,7 @@ void FGCReferenceTokenStream::ReplaceOrAddAddReferencedObjectsCall(void (*AddRef
 	for (int32 TokenStreamIndex = 0; TokenStreamIndex < Tokens.Num(); ++TokenStreamIndex)
 	{
 		uint32 TokenIndex = (uint32)TokenStreamIndex;
-		const uint32 TokenType = AccessReferenceInfo(TokenIndex).Type;
+		const EGCReferenceType TokenType = (EGCReferenceType)AccessReferenceInfo(TokenIndex).Type;
 		// Read token type and skip additional data if present.
 		switch (TokenType)
 		{
@@ -1959,6 +1959,12 @@ void FGCReferenceTokenStream::ReplaceOrAddAddReferencedObjectsCall(void (*AddRef
 				StorePointer(&Tokens[++TokenIndex], (const void*)AddReferencedObjectsPtr);
 				return;
 			}
+		case GCRT_AddTMapReferencedObjects:
+			{
+				// Skip pointer
+				TokenIndex += GNumTokensPerPointer;
+			}
+			break;
 		case GCRT_None:
 		case GCRT_Object:
 		case GCRT_PersistentObject:
@@ -1966,7 +1972,7 @@ void FGCReferenceTokenStream::ReplaceOrAddAddReferencedObjectsCall(void (*AddRef
 		case GCRT_EndOfStream:
 			break;
 		default:
-			UE_LOG(LogGarbage, Fatal, TEXT("Unknown token type (%u) when trying to add ARO token."), TokenType);
+			UE_LOG(LogGarbage, Fatal, TEXT("Unknown token type (%u) when trying to add ARO token."), (uint32)TokenType);
 			break;
 		};
 		TokenStreamIndex = (int32)TokenIndex;
