Adding stats to detect BulkData async buffer leaks

diff --git a/Engine/Source/Runtime/CoreUObject/Private/Serialization/BulkData.cpp b/Engine/Source/Runtime/CoreUObject/Private/Serialization/BulkData.cpp
--- a/Engine/Source/Runtime/CoreUObject/Private/Serialization/BulkData.cpp
+++ b/Engine/Source/Runtime/CoreUObject/Private/Serialization/BulkData.cpp
@@ -11,6 +11,8 @@
 /** Whether to track information of how bulk data is being used */
 #define TRACK_BULKDATA_USE 0
 
+DECLARE_STATS_GROUP(TEXT("Bulk Data"), STATGROUP_BulkData, STATCAT_Advanced);
+DECLARE_DWORD_ACCUMULATOR_STAT(TEXT("Leaked async buffers"), STAT_BulkDataAsyncLeaks, STATGROUP_BulkData);
 
 #if TRACK_BULKDATA_USE
 
@@ -133,6 +135,7 @@ FUntypedBulkData::~FUntypedBulkData()
 	if (BulkData != BulkDataAsync)
 	{
 		FMemory::Free(BulkDataAsync);
+		DEC_DWORD_STAT_BY(STAT_BulkDataAsyncLeaks, 1);
 	}
 	BulkData = nullptr;
 	BulkDataAsync = nullptr;
@@ -415,6 +418,7 @@ void FUntypedBulkData::GetCopy( void** Dest, bool bDiscardInternalCopy )
 		{
 			WaitForAsyncLoading();
 			BulkData = BulkDataAsync;
+			DEC_DWORD_STAT_BY(STAT_BulkDataAsyncLeaks, 1);
 			ResetAsyncData();
 		}
 		// The data is already loaded so we can simply use a mempcy.
@@ -671,6 +675,7 @@ void FUntypedBulkData::StartSerializingBulkData(FArchive& Ar, UObject* Owner, in
 		{
 			BulkDataAsync = FMemory::Realloc(BulkDataAsync, GetBulkDataSize(), BulkDataAlignment);
 		}
+		INC_DWORD_STAT_BY(STAT_BulkDataAsyncLeaks, 1);
 
 		FArchive* FileReaderAr = IFileManager::Get().CreateFileReader(*Filename, FILEREAD_Silent);
 		checkf(FileReaderAr != NULL, TEXT("Attempted to load bulk data from an invalid filename '%s'."), *Filename);
@@ -1231,6 +1236,7 @@ void FUntypedBulkData::MakeSureBulkDataIsLoaded()
 		{
 			WaitForAsyncLoading();
 			BulkData = BulkDataAsync;
+			DEC_DWORD_STAT_BY(STAT_BulkDataAsyncLeaks, 1);
 			ResetAsyncData();
 		}
 		else
@@ -1267,6 +1273,7 @@ bool FUntypedBulkData::FlushAsyncLoading(void* Dest)
 	{
 		WaitForAsyncLoading();
 		FMemory::Memcpy(Dest, BulkDataAsync, GetBulkDataSize());
+		DEC_DWORD_STAT_BY(STAT_BulkDataAsyncLeaks, 1);
 	}
 	return bIsLoadingAsync;
 }
