Editor-only data is now gathered for localization by opting-in through a new setting.

diff --git a/Engine/Source/Editor/UnrealEd/Classes/Commandlets/GatherTextFromAssetsCommandlet.h b/Engine/Source/Editor/UnrealEd/Classes/Commandlets/GatherTextFromAssetsCommandlet.h
--- a/Engine/Source/Editor/UnrealEd/Classes/Commandlets/GatherTextFromAssetsCommandlet.h
+++ b/Engine/Source/Editor/UnrealEd/Classes/Commandlets/GatherTextFromAssetsCommandlet.h
@@ -63,5 +63,5 @@ private:
 	FConflictTracker ConflictTracker;
 
 	bool bFixBroken;
-	bool ShouldGatherBlueprintPinNames;
+	bool ShouldGatherFromEditorOnlyData;
 };
diff --git a/Engine/Source/Editor/UnrealEd/Classes/Commandlets/GatherTextFromMetadataCommandlet.h b/Engine/Source/Editor/UnrealEd/Classes/Commandlets/GatherTextFromMetadataCommandlet.h
--- a/Engine/Source/Editor/UnrealEd/Classes/Commandlets/GatherTextFromMetadataCommandlet.h
+++ b/Engine/Source/Editor/UnrealEd/Classes/Commandlets/GatherTextFromMetadataCommandlet.h
@@ -28,4 +28,7 @@ private:
 private:
 	void GatherTextFromUObjects(const TArray<FString>& IncludePaths, const TArray<FString>& ExcludePaths, const FGatherParameters& Arguments);
 	void GatherTextFromUObject(UField* const Field, const FGatherParameters& Arguments);
+
+private:
+	bool ShouldGatherFromEditorOnlyData;
 };
\ No newline at end of file
diff --git a/Engine/Source/Editor/UnrealEd/Private/Commandlets/GatherTextFromAssetsCommandlet.cpp b/Engine/Source/Editor/UnrealEd/Private/Commandlets/GatherTextFromAssetsCommandlet.cpp
--- a/Engine/Source/Editor/UnrealEd/Private/Commandlets/GatherTextFromAssetsCommandlet.cpp
+++ b/Engine/Source/Editor/UnrealEd/Private/Commandlets/GatherTextFromAssetsCommandlet.cpp
@@ -1,4 +1,4 @@
-// Copyright 1998-2015 Epic Games, Inc. All Rights Reserved.
+﻿// Copyright 1998-2015 Epic Games, Inc. All Rights Reserved.
 
 #include "UnrealEd.h"
 #include "PackageTools.h"
@@ -77,8 +77,11 @@ void UGatherTextFromAssetsCommandlet::ProcessGatherableTextDataArray(const FStri
 			// Gather if it requires gathering and isn't in a bad state.
 			if ((NewEntry.Status & EAssetTextGatherStatus::BadBitMask) == false)
 			{
+				if (!TextSourceSiteContext.IsEditorOnly || ShouldGatherFromEditorOnlyData)
+				{
 				ManifestInfo->AddEntry(TextSourceSiteContext.SiteDescription, GatherableTextData.NamespaceName, NewEntry.SourceString.Get() ? *(NewEntry.SourceString) : TEXT(""), Context );
 			}
+			}
 
 			// Add to conflict tracker.
 			FConflictTracker::FKeyTable& KeyTable = ConflictTracker.Namespaces.FindOrAdd(GatherableTextData.NamespaceName);
@@ -93,14 +96,13 @@ void UGatherTextFromAssetsCommandlet::ProcessPackages( const TArray< UPackage* >
 	for(UPackage* const Package : PackagesToProcess)
 	{
 		TArray<UObject*> ObjectsInPackage;
-		GetObjectsWithOuter(Package, ObjectsInPackage);
+		GetObjectsWithOuter(Package, ObjectsInPackage, true, RF_Transient | RF_PendingKill);
 		for (UObject* const Object : ObjectsInPackage)
 		{
 			TArray<FGatherableTextData> GatherableTextDataArray;
 
-			if( !Object->HasAnyFlags( RF_Transient | RF_PendingKill ) )
-			{
-				GatherLocalizationDataFromPropertiesOfDataStructure(Object->GetClass(), Object, GatherableTextDataArray);
+			FPropertyLocalizationDataGatherer PropertyLocalizationDataGatherer(GatherableTextDataArray);
+			PropertyLocalizationDataGatherer.GatherLocalizationDataFromPropertiesOfDataStructure(Object->GetClass(), Object);
 
 				for(UClass* Class = Object->GetClass(); Class != nullptr; Class = Class->GetSuperClass())
 				{
@@ -114,7 +116,6 @@ void UGatherTextFromAssetsCommandlet::ProcessPackages( const TArray< UPackage* >
 				ProcessGatherableTextDataArray(Package->FileName.ToString(), GatherableTextDataArray);
 			}
 		}
-	}
 }
 
 int32 UGatherTextFromAssetsCommandlet::Main(const FString& Params)
@@ -256,10 +257,10 @@ int32 UGatherTextFromAssetsCommandlet::Main(const FString& Params)
 		bFixBroken = false;
 	}
 
-	// Get whether we should gather pin name properties that we find. Typically only useful for the localization of UE4 itself.
-	if (!GetBoolFromConfig(*SectionName, TEXT("ShouldGatherBlueprintPinNames"), ShouldGatherBlueprintPinNames, GatherTextConfigPath))
+	// Get whether we should gather editor-only data. Typically only useful for the localization of UE4 itself.
+	if (!GetBoolFromConfig(*SectionName, TEXT("ShouldGatherFromEditorOnlyData"), ShouldGatherFromEditorOnlyData, GatherTextConfigPath))
 	{
-		ShouldGatherBlueprintPinNames = false;
+		ShouldGatherFromEditorOnlyData = false;
 	}
 
 	// Add any manifest dependencies if they were provided
diff --git a/Engine/Source/Editor/UnrealEd/Private/Commandlets/GatherTextFromMetadataCommandlet.cpp b/Engine/Source/Editor/UnrealEd/Private/Commandlets/GatherTextFromMetadataCommandlet.cpp
--- a/Engine/Source/Editor/UnrealEd/Private/Commandlets/GatherTextFromMetadataCommandlet.cpp
+++ b/Engine/Source/Editor/UnrealEd/Private/Commandlets/GatherTextFromMetadataCommandlet.cpp
@@ -94,6 +94,12 @@ int32 UGatherTextFromMetaDataCommandlet::Main( const FString& Params )
 		}
 	}
 
+	// Get whether we should gather editor-only data. Typically only useful for the localization of UE4 itself.
+	if (!GetBoolFromConfig(*SectionName, TEXT("ShouldGatherFromEditorOnlyData"), ShouldGatherFromEditorOnlyData, GatherTextConfigPath))
+	{
+		ShouldGatherFromEditorOnlyData = false;
+	}
+
 	FGatherParameters Arguments;
 	GetStringArrayFromConfig(*SectionName, TEXT("InputKeys"), Arguments.InputKeys, GatherTextConfigPath);
 	GetStringArrayFromConfig(*SectionName, TEXT("OutputNamespaces"), Arguments.OutputNamespaces, GatherTextConfigPath);
@@ -125,6 +131,13 @@ void UGatherTextFromMetaDataCommandlet::GatherTextFromUObjects(const TArray<FStr
 {
 	for(TObjectIterator<UField> It; It; ++It)
 	{
+		// Skip editor-only properties if we're not gathering for editor-only data.
+		UProperty* Property = Cast<UProperty>(*It);
+		if (Property && !ShouldGatherFromEditorOnlyData && Property->HasAnyPropertyFlags(CPF_EditorOnly))
+		{
+			continue;
+		}
+
 		FString SourceFilePath;
 		FSourceCodeNavigation::FindClassHeaderPath(*It, SourceFilePath);
 		SourceFilePath = FPaths::ConvertRelativePathToFull(SourceFilePath);
diff --git a/Engine/Source/Runtime/CoreUObject/Private/Serialization/PropertyLocalizationDataGathering.cpp b/Engine/Source/Runtime/CoreUObject/Private/Serialization/PropertyLocalizationDataGathering.cpp
--- a/Engine/Source/Runtime/CoreUObject/Private/Serialization/PropertyLocalizationDataGathering.cpp
+++ b/Engine/Source/Runtime/CoreUObject/Private/Serialization/PropertyLocalizationDataGathering.cpp
@@ -1,10 +1,10 @@
-// Copyright 1998-2015 Epic Games, Inc. All Rights Reserved.
+﻿// Copyright 1998-2015 Epic Games, Inc. All Rights Reserved.
 
 #include "CoreUObjectPrivate.h"
 #include "Serialization/PropertyLocalizationDataGathering.h"
 #include "UTextProperty.h"
 
-void GatherLocalizationDataFromPropertiesOfDataStructure(UStruct* const Structure, void* const Data, TArray<FGatherableTextData>& GatherableTextDataArray)
+void FPropertyLocalizationDataGatherer::GatherLocalizationDataFromPropertiesOfDataStructure(UStruct* const Structure, void* const Data, const bool ForceIsEditorOnly)
 {
 	FString Path;
 	{
@@ -19,11 +19,11 @@ void GatherLocalizationDataFromPropertiesOfDataStructure(UStruct* const Structur
 	// Iterate over all fields of the object's class.
 	for (TFieldIterator<UProperty> PropIt(Structure, EFieldIteratorFlags::IncludeSuper, EFieldIteratorFlags::ExcludeDeprecated, EFieldIteratorFlags::IncludeInterfaces); PropIt; ++PropIt)
 	{
-		GatherLocalizationDataFromChildTextProperies(Path, *PropIt, PropIt->ContainerPtrToValuePtr<void>(Data), GatherableTextDataArray, false);
+		GatherLocalizationDataFromChildTextProperies(Path, *PropIt, PropIt->ContainerPtrToValuePtr<void>(Data));
 	}
 }
 
-void GatherLocalizationDataFromChildTextProperies(const FString& PathToParent, UProperty* const Property, void* const ValueAddress, TArray<FGatherableTextData>& GatherableTextDataArray, const bool ForceIsEditorOnly)
+void FPropertyLocalizationDataGatherer::GatherLocalizationDataFromChildTextProperies(const FString& PathToParent, UProperty* const Property, void* const ValueAddress, const bool ForceIsEditorOnly)
 {
 	UTextProperty* const TextProperty = Cast<UTextProperty>(Property);
 	UArrayProperty* const ArrayProperty = Cast<UArrayProperty>(Property);
@@ -32,7 +32,7 @@ void GatherLocalizationDataFromChildTextProperies(const FString& PathToParent, U
 	// Property is a text property.
 	if (TextProperty)
 	{
-		GatherLocalizationDataFromTextProperty(PathToParent, TextProperty, ValueAddress, GatherableTextDataArray, ForceIsEditorOnly);
+		GatherLocalizationDataFromTextProperty(PathToParent, TextProperty, ValueAddress, ForceIsEditorOnly);
 	}
 	else
 	{
@@ -51,7 +51,7 @@ void GatherLocalizationDataFromChildTextProperies(const FString& PathToParent, U
 				const int32 ElementCount = ScriptArrayHelper.Num();
 				for(int32 j = 0; j < ElementCount; ++j)
 				{
-					GatherLocalizationDataFromChildTextProperies(PathToElement + FString::Printf(TEXT("(%d)"), j), ArrayProperty->Inner, ScriptArrayHelper.GetRawPtr(j), GatherableTextDataArray, ForceIsEditorOnly || (Property->PropertyFlags & CPF_EditorOnly) != 0);
+					GatherLocalizationDataFromChildTextProperies(PathToElement + FString::Printf(TEXT("(%d)"), j), ArrayProperty->Inner, ScriptArrayHelper.GetRawPtr(j), ForceIsEditorOnly || Property->HasAnyPropertyFlags(CPF_EditorOnly));
 				}
 			}
 			// Property is a struct property.
@@ -62,14 +62,14 @@ void GatherLocalizationDataFromChildTextProperies(const FString& PathToParent, U
 				{
 					const FString Path = PathToElement;
 
-					GatherLocalizationDataFromChildTextProperies(PathToElement, *PropIt, PropIt->ContainerPtrToValuePtr<void>(ElementValueAddress), GatherableTextDataArray, ForceIsEditorOnly || (Property->PropertyFlags & CPF_EditorOnly) != 0);
+					GatherLocalizationDataFromChildTextProperies(PathToElement, *PropIt, PropIt->ContainerPtrToValuePtr<void>(ElementValueAddress), ForceIsEditorOnly || Property->HasAnyPropertyFlags(CPF_EditorOnly));
 				}
 			}
 		}
 	}
 }
 
-void GatherLocalizationDataFromTextProperty(const FString& PathToParent, UTextProperty* const TextProperty, void* const ValueAddress, TArray<FGatherableTextData>& GatherableTextDataArray, const bool ForceIsEditorOnly)
+void FPropertyLocalizationDataGatherer::GatherLocalizationDataFromTextProperty(const FString& PathToParent, UTextProperty* const TextProperty, void* const ValueAddress, bool ForceIsEditorOnly)
 {
 	const bool IsFixedSizeArray = TextProperty->ArrayDim > 1;
 	for(int32 i = 0; i < TextProperty->ArrayDim; ++i)
@@ -137,7 +137,7 @@ void GatherLocalizationDataFromTextProperty(const FString& PathToParent, UTextPr
 			FTextSourceSiteContext& SourceSiteContext = *(new(GatherableTextData->SourceSiteContexts) FTextSourceSiteContext);
 			SourceSiteContext.KeyName = Key;
 			SourceSiteContext.SiteDescription = PathToElement;
-			SourceSiteContext.IsEditorOnly = ForceIsEditorOnly || (TextProperty->PropertyFlags & CPF_EditorOnly) != 0;
+			SourceSiteContext.IsEditorOnly = ForceIsEditorOnly || TextProperty->HasAnyPropertyFlags(CPF_EditorOnly);
 			SourceSiteContext.IsOptional = false;
 		}
 	}
diff --git a/Engine/Source/Runtime/CoreUObject/Private/UObject/SavePackage.cpp b/Engine/Source/Runtime/CoreUObject/Private/UObject/SavePackage.cpp
--- a/Engine/Source/Runtime/CoreUObject/Private/UObject/SavePackage.cpp
+++ b/Engine/Source/Runtime/CoreUObject/Private/UObject/SavePackage.cpp
@@ -1,4 +1,4 @@
-// Copyright 1998-2015 Epic Games, Inc. All Rights Reserved.
+﻿// Copyright 1998-2015 Epic Games, Inc. All Rights Reserved.
 
 #include "CoreUObjectPrivate.h"
 #include "UObject/UTextProperty.h"
@@ -3356,21 +3356,19 @@ bool UPackage::SavePackage( UPackage* InOuter, UObject* Base, EObjectFlags TopLe
 
 				if ( !(Linker->Summary.PackageFlags & PKG_FilterEditorOnly) )
 				{
-					TArray<UObject*> TagExpObjects;
-					GetObjectsWithAnyMarks(TagExpObjects, OBJECTMARK_TagExp);
-					for (UObject* const Object : TagExpObjects)
+					TArray<UObject*> ObjectsInPackage;
+					GetObjectsWithOuter(InOuter, ObjectsInPackage, true, RF_Transient | RF_PendingKill);
+					for (UObject* const Object : ObjectsInPackage)
 					{
-						if( !Object->HasAnyFlags( RF_Transient | RF_PendingKill ) )
-						{
-							GatherLocalizationDataFromPropertiesOfDataStructure(Object->GetClass(), Object, Linker->GatherableTextDataMap);
+						FPropertyLocalizationDataGatherer PropertyLocalizationDataGatherer(Linker->GatherableTextDataMap);
+						PropertyLocalizationDataGatherer.GatherLocalizationDataFromPropertiesOfDataStructure(Object->GetClass(), Object);
 
-							for(UClass* Class = Object->GetClass(); Class != nullptr; Class = Class->GetSuperClass())
+						for(UClass* Class = Object->GetClass(); Class != nullptr; Class = Class->GetSuperClass())
+						{
+							FLocalizationDataGatheringCallback* const CustomCallback = GetTypeSpecificLocalizationDataGatheringCallbacks().Find(Class);
+							if (CustomCallback)
 							{
-								FLocalizationDataGatheringCallback* const CustomCallback = GetTypeSpecificLocalizationDataGatheringCallbacks().Find(Class);
-								if (CustomCallback)
-								{
-									(*CustomCallback)(Object, Linker->GatherableTextDataMap);
-								}
+								(*CustomCallback)(Object, Linker->GatherableTextDataMap);
 							}
 						}
 					}
diff --git a/Engine/Source/Runtime/CoreUObject/Public/Serialization/PropertyLocalizationDataGathering.h b/Engine/Source/Runtime/CoreUObject/Public/Serialization/PropertyLocalizationDataGathering.h
--- a/Engine/Source/Runtime/CoreUObject/Public/Serialization/PropertyLocalizationDataGathering.h
+++ b/Engine/Source/Runtime/CoreUObject/Public/Serialization/PropertyLocalizationDataGathering.h
@@ -4,6 +4,19 @@
 
 #include "GatherableTextData.h"
 
-COREUOBJECT_API extern void GatherLocalizationDataFromPropertiesOfDataStructure(UStruct* const Structure, void* const Data, TArray<FGatherableTextData>& GatherableTextDataArray);
-COREUOBJECT_API extern void GatherLocalizationDataFromChildTextProperies(const FString& PathToParent, UProperty* const Property, void* const ValueAddress, TArray<FGatherableTextData>& GatherableTextDataArray, const bool ForceIsEditorOnly);
-COREUOBJECT_API extern void GatherLocalizationDataFromTextProperty(const FString& PathToParent, UTextProperty* const TextProperty, void* const ValueAddress, TArray<FGatherableTextData>& GatherableTextDataArray, const bool ForceIsEditorOnly);
\ No newline at end of file
+class FPropertyLocalizationDataGatherer
+{
+public:
+	COREUOBJECT_API FPropertyLocalizationDataGatherer(TArray<FGatherableTextData>& InGatherableTextDataArray)
+		: GatherableTextDataArray(InGatherableTextDataArray)
+	{
+	}
+
+	COREUOBJECT_API void GatherLocalizationDataFromPropertiesOfDataStructure(UStruct* const Structure, void* const Data, const bool ForceIsEditorOnly = false);
+	COREUOBJECT_API void GatherLocalizationDataFromChildTextProperies(const FString& PathToParent, UProperty* const Property, void* const ValueAddress, const bool ForceIsEditorOnly = false);
+	COREUOBJECT_API void GatherLocalizationDataFromTextProperty(const FString& PathToParent, UTextProperty* const TextProperty, void* const ValueAddress, const bool ForceIsEditorOnly = false);
+
+private:
+	TArray<FGatherableTextData>& GatherableTextDataArray;
+	TArray<UObject*> ProcessedObjects;
+};
\ No newline at end of file
diff --git a/Engine/Source/Runtime/Engine/Private/DataTable.cpp b/Engine/Source/Runtime/Engine/Private/DataTable.cpp
--- a/Engine/Source/Runtime/Engine/Private/DataTable.cpp
+++ b/Engine/Source/Runtime/Engine/Private/DataTable.cpp
@@ -22,7 +22,8 @@ namespace
 			const FString PathToRow = PathToObject + TEXT(".") + Pair.Key.ToString();
 			for (TFieldIterator<UProperty> PropIt(DataTable->RowStruct, EFieldIteratorFlags::IncludeSuper, EFieldIteratorFlags::ExcludeDeprecated, EFieldIteratorFlags::IncludeInterfaces); PropIt; ++PropIt)
 			{
-				GatherLocalizationDataFromChildTextProperies(PathToRow, *PropIt, PropIt->ContainerPtrToValuePtr<void>(Pair.Value), GatherableTextDataArray, false);
+				FPropertyLocalizationDataGatherer PropertyLocalizationDataGatherer(GatherableTextDataArray);
+				PropertyLocalizationDataGatherer.GatherLocalizationDataFromChildTextProperies(PathToRow, *PropIt, PropIt->ContainerPtrToValuePtr<void>(Pair.Value));
 			}
 		}
 	}
diff --git a/Engine/Source/Runtime/Engine/Private/UserDefinedStruct.cpp b/Engine/Source/Runtime/Engine/Private/UserDefinedStruct.cpp
--- a/Engine/Source/Runtime/Engine/Private/UserDefinedStruct.cpp
+++ b/Engine/Source/Runtime/Engine/Private/UserDefinedStruct.cpp
@@ -22,7 +22,8 @@ namespace
 		// Iterate over all fields of the object's class.
 		for (TFieldIterator<UProperty> PropIt(StructData.GetStruct(), EFieldIteratorFlags::IncludeSuper, EFieldIteratorFlags::ExcludeDeprecated, EFieldIteratorFlags::IncludeInterfaces); PropIt; ++PropIt)
 		{
-			GatherLocalizationDataFromChildTextProperies(PathToObject, *PropIt, PropIt->ContainerPtrToValuePtr<void>(StructData.GetStructMemory()), GatherableTextDataArray, false);
+			FPropertyLocalizationDataGatherer PropertyLocalizationDataGatherer(GatherableTextDataArray);
+			PropertyLocalizationDataGatherer.GatherLocalizationDataFromChildTextProperies(PathToObject, *PropIt, PropIt->ContainerPtrToValuePtr<void>(StructData.GetStructMemory()));
 		}
 	}
 }
