Reducing log verbosity when inserting package dependencies while async loading

diff --git a/Engine/Source/Runtime/CoreUObject/Private/Serialization/AsyncLoading.cpp b/Engine/Source/Runtime/CoreUObject/Private/Serialization/AsyncLoading.cpp
--- a/Engine/Source/Runtime/CoreUObject/Private/Serialization/AsyncLoading.cpp
+++ b/Engine/Source/Runtime/CoreUObject/Private/Serialization/AsyncLoading.cpp
@@ -509,7 +509,7 @@ int32 FAsyncLoadingThread::CreateAsyncPackagesFromQueue()
 				ProcessAsyncPackageRequest(PackageRequest, nullptr, DependencyTracker);
 			}
 		}
-		UE_LOG(LogStreaming, Display, TEXT("Async package requests inserted in %fms"), Timer * 1000.0);
+		UE_LOG(LogStreaming, Verbose, TEXT("Async package requests inserted in %fms"), Timer * 1000.0);
 	}
 
 	NumCreated = QueueCopy.Num();
