A few more stats for async loading

diff --git a/Engine/Source/Runtime/CoreUObject/Private/Serialization/AsyncLoading.cpp b/Engine/Source/Runtime/CoreUObject/Private/Serialization/AsyncLoading.cpp
--- a/Engine/Source/Runtime/CoreUObject/Private/Serialization/AsyncLoading.cpp
+++ b/Engine/Source/Runtime/CoreUObject/Private/Serialization/AsyncLoading.cpp
@@ -30,6 +30,8 @@ DECLARE_CYCLE_STAT(TEXT("LoadImports AsyncPackage"),STAT_FAsyncPackage_LoadImpor
 DECLARE_CYCLE_STAT(TEXT("CreateImports AsyncPackage"),STAT_FAsyncPackage_CreateImports,STATGROUP_AsyncLoad);
 DECLARE_CYCLE_STAT(TEXT("FinishTextureAllocations AsyncPackage"),STAT_FAsyncPackage_FinishTextureAllocations,STATGROUP_AsyncLoad);
 DECLARE_CYCLE_STAT(TEXT("CreateExports AsyncPackage"),STAT_FAsyncPackage_CreateExports,STATGROUP_AsyncLoad);
+DECLARE_CYCLE_STAT(TEXT("FreeReferencedImports AsyncPackage"), STAT_FAsyncPackage_FreeReferencedImports, STATGROUP_AsyncLoad);
+DECLARE_CYCLE_STAT(TEXT("Precache ArchiveAsync"), STAT_FArchiveAsync_Precache, STATGROUP_AsyncLoad);
 DECLARE_CYCLE_STAT(TEXT("PreLoadObjects AsyncPackage"),STAT_FAsyncPackage_PreLoadObjects,STATGROUP_AsyncLoad);
 DECLARE_CYCLE_STAT(TEXT("PostLoadObjects AsyncPackage"),STAT_FAsyncPackage_PostLoadObjects,STATGROUP_AsyncLoad);
 DECLARE_CYCLE_STAT(TEXT("FinishObjects AsyncPackage"),STAT_FAsyncPackage_FinishObjects,STATGROUP_AsyncLoad);
@@ -1043,7 +1045,8 @@ bool FAsyncPackage::IsTimeLimitExceeded()
  */
 bool FAsyncPackage::GiveUpTimeSlice()
 {
-	if (!FPlatformProcess::SupportsMultithreading())
+	static const bool bPlatformIsSingleThreaded = !FPlatformProcess::SupportsMultithreading();
+	if (bPlatformIsSingleThreaded)
 	{
 		FIOSystem::Get().TickSingleThreaded();
 	}
@@ -1633,21 +1636,20 @@ EAsyncPackageState::Type FAsyncPackage::CreateExports()
 			UObject* Object	= Linker->CreateExport( ExportIndex++ );
 			// ... and preload it.
 			if( Object )
-			{
+			{				
 				// This will cause the object to be serialized. We do this here for all objects and
 				// not just UClass and template objects, for which this is required in order to ensure
 				// seek free loading, to be able introduce async file I/O.
 				Linker->Preload( Object );
-
 			}
 
 			LastObjectWorkWasPerformedOn	= Object;
-			LastTypeOfWorkPerformed			= TEXT("creating exports for");
-			
+			LastTypeOfWorkPerformed = TEXT("creating exports for");
+				
 			UpdateLoadPercentage();
 		}
 		// Data isn't ready yet. Give up remainder of time slice if we're not using a time limit.
-		else if( GiveUpTimeSlice() )
+		else if (GiveUpTimeSlice())
 		{
 			INC_FLOAT_STAT_BY(STAT_AsyncIO_AsyncPackagePrecacheWaitTime, (float)FApp::GetDeltaTime());
 			return EAsyncPackageState::TimeOut;
@@ -1665,6 +1667,8 @@ EAsyncPackageState::Type FAsyncPackage::CreateExports()
  */
 void FAsyncPackage::FreeReferencedImports()
 {	
+	SCOPE_CYCLE_COUNTER(STAT_FAsyncPackage_FreeReferencedImports);	
+
 	for (int32 ReferenceIndex = 0; ReferenceIndex < ReferencedImports.Num(); ++ReferenceIndex)
 	{
 		FAsyncPackage& Ref = *ReferencedImports[ReferenceIndex];
@@ -2400,6 +2404,8 @@ void FArchiveAsync::PrecacheCompressedChunk( int64 ChunkIndex, int64 BufferIndex
  */
 bool FArchiveAsync::Precache( int64 RequestOffset, int64 RequestSize )
 {
+	SCOPE_CYCLE_COUNTER(STAT_FArchiveAsync_Precache);
+
 	// Check whether we're currently waiting for a read request to finish.
 	bool bFinishedReadingCurrent	= PrecacheReadStatus[CURRENT].GetValue()==0 ? true : false;
 	bool bFinishedReadingNext		= PrecacheReadStatus[NEXT].GetValue()==0 ? true : false;
