Don't mark packages as dirty from within async loading code path

diff --git a/Engine/Source/Runtime/CoreUObject/Private/UObject/UObjectBaseUtility.cpp b/Engine/Source/Runtime/CoreUObject/Private/UObject/UObjectBaseUtility.cpp
--- a/Engine/Source/Runtime/CoreUObject/Private/UObject/UObjectBaseUtility.cpp
+++ b/Engine/Source/Runtime/CoreUObject/Private/UObject/UObjectBaseUtility.cpp
@@ -141,7 +141,7 @@ bool UObjectBaseUtility::MarkPackageDirty() const
 			// we explicitly disable the ability to dirty a package or map during load.  Commandlets can still
 			// set the dirty state on load.
 			if( IsRunningCommandlet() || 
-				(GIsEditor && !GIsEditorLoadingPackage && !GIsPlayInEditorWorld
+				(GIsEditor && !GIsEditorLoadingPackage && !GIsPlayInEditorWorld && !IsInAsyncLoadingThread()
 #if WITH_HOT_RELOAD
 				&& !GIsHotReload
 #endif // WITH_HOT_RELOAD
