[UE-17388] Making ResolveDeferredExports() fully reentrant, so we don't attempt to recreate/mutate package exports after the Blueprint has been subsequently finalized.

diff --git a/Engine/Source/Runtime/CoreUObject/Private/Blueprint/BlueprintSupport.cpp b/Engine/Source/Runtime/CoreUObject/Private/Blueprint/BlueprintSupport.cpp
--- a/Engine/Source/Runtime/CoreUObject/Private/Blueprint/BlueprintSupport.cpp
+++ b/Engine/Source/Runtime/CoreUObject/Private/Blueprint/BlueprintSupport.cpp
@@ -61,27 +61,6 @@ bool FBlueprintSupport::UseDeferredDependencyLoading()
 #endif
 }
 
-
-bool FBlueprintSupport::IsResolvingDeferredDependenciesDisabled()
-{
-#if USE_DEFERRED_DEPENDENCY_CHECK_VERIFICATION_TESTS
-	static const FBoolConfigValueHelper NoDeferredDependencyResolves(TEXT("Kismet"), TEXT("bForceDisableDeferredDependencyResolving"), GEngineIni);
-	return !UseDeferredDependencyLoading() || NoDeferredDependencyResolves;
-#else 
-	return false;
-#endif
-}
-
-bool FBlueprintSupport::IsDeferredCDOSerializationDisabled()
-{
-#if USE_DEFERRED_DEPENDENCY_CHECK_VERIFICATION_TESTS
-	static const FBoolConfigValueHelper NoDeferredCDOLoading(TEXT("Kismet"), TEXT("bForceDisableDeferredCDOLoading"), GEngineIni);
-	return !UseDeferredDependencyLoading() || NoDeferredCDOLoading;
-#else 
-	return false;
-#endif
-}
-
 bool FBlueprintSupport::IsDeferredExportCreationDisabled()
 {
 #if USE_CIRCULAR_DEPENDENCY_LOAD_DEFERRING
@@ -871,7 +850,7 @@ void FLinkerLoad::ResolveDeferredDependencies(UStruct* LoadStruct)
 	//LoadAllObjects();
 
 #if USE_DEFERRED_DEPENDENCY_CHECK_VERIFICATION_TESTS
-	for (TObjectIterator<ULinkerPlaceholderClass> PlaceholderIt; PlaceholderIt && !FBlueprintSupport::IsResolvingDeferredDependenciesDisabled(); ++PlaceholderIt)
+	for (TObjectIterator<ULinkerPlaceholderClass> PlaceholderIt; PlaceholderIt; ++PlaceholderIt)
 	{
 		ULinkerPlaceholderClass* PlaceholderClass = *PlaceholderIt;
 		if (PlaceholderClass->GetOuter() == LinkerRoot)
@@ -907,13 +886,6 @@ bool FLinkerLoad::HasUnresolvedDependencies() const
 int32 FLinkerLoad::ResolveDependencyPlaceholder(FLinkerPlaceholderBase* PlaceholderIn, UClass* ReferencingClass)
 {
 #if USE_CIRCULAR_DEPENDENCY_LOAD_DEFERRING
-#if USE_DEFERRED_DEPENDENCY_CHECK_VERIFICATION_TESTS
-	if (FBlueprintSupport::IsResolvingDeferredDependenciesDisabled())
-	{
-		return 0;
-	}
-#endif // USE_DEFERRED_DEPENDENCY_CHECK_VERIFICATION_TESTS
-
 	TGuardValue<uint32>  LoadFlagsGuard(LoadFlags, (LoadFlags & ~LOAD_DeferDependencyLoads));
 	TGuardValue<FLinkerPlaceholderBase*> ResolvingClassGuard(ResolvingDeferredPlaceholder, PlaceholderIn);
 
@@ -1113,14 +1085,17 @@ void FLinkerLoad::FinalizeBlueprint(UClass* LoadClass)
 	// we need it ran for any super-classes before we regen
 	ResolveAllImports();
 
-	// the above loop could have caused some recursion... if it ended up 
-	// finalizing a sub-class (of LoadClass), then that would have finalized
-	// this as well; so, before we continue, make sure that this didn't already 
-	// get fully executed in some nested call
+	// replace any export placeholders that were created, and serialize in the 
+	// class's CDO
+	ResolveDeferredExports(LoadClass);
+
+	// the above calls (ResolveAllImports(), ResolveDeferredExports(), etc.) 
+	// could have caused some recursion... if it ended up finalizing a sub-class 
+	// (of LoadClass), then that would have finalized this as well; so, before 
+	// we continue, make sure that this didn't already get fully executed in 
+	// some nested call
 	if (IsBlueprintFinalizationPending())
 	{
-		ResolveDeferredExports(LoadClass);
-
 		FObjectExport& CDOExport = ExportMap[DeferredCDOIndex];
 		// clear this so IsBlueprintFinalizationPending() doesn't report true
 		DeferredCDOIndex = INDEX_NONE;
@@ -1181,6 +1156,10 @@ void FLinkerLoad::FinalizeBlueprint(UClass* LoadClass)
 void FLinkerLoad::ResolveDeferredExports(UClass* LoadClass)
 {
 #if USE_CIRCULAR_DEPENDENCY_LOAD_DEFERRING
+	if (!IsBlueprintFinalizationPending())
+	{
+		return;
+	}
 
 	DEFERRED_DEPENDENCY_CHECK(DeferredCDOIndex != INDEX_NONE);
 	FObjectExport& CDOExport = ExportMap[DeferredCDOIndex];
@@ -1188,37 +1167,29 @@ void FLinkerLoad::ResolveDeferredExports(UClass* LoadClass)
 	UObject* BlueprintCDO = CDOExport.Object;
 	DEFERRED_DEPENDENCY_CHECK(BlueprintCDO != nullptr);
 
+	if (!FBlueprintSupport::IsDeferredExportCreationDisabled())
+	{
 #if USE_DEFERRED_DEPENDENCY_CHECK_VERIFICATION_TESTS
-	static TSet<FLinkerLoad*> ResolvingExportsTracker;
-	// the first half of this function CANNOT be re-entrant (specifically for 
-	// this linker)
-	//
-	// if this ends up being a problem, then we have two blueprints that 
-	// rely on an instances of each other... this is an unexpected problem, and 
-	// we'll have to re-instance on load (an operation we'd rather avoid,  
-	// because of how slow it is)
-	DEFERRED_DEPENDENCY_CHECK(!ResolvingExportsTracker.Contains(this));
-	ResolvingExportsTracker.Add(this);
+		auto IsPlaceholderReferenced = [](ULinkerPlaceholderExportObject* ExportPlaceholder)->bool
+		{
+			FReferencerInformationList UnresolvedReferences;
+			UObject* PlaceholderObj = ExportPlaceholder;
+			return IsReferenced(PlaceholderObj, RF_NoFlags, /*bCheckSubObjects =*/false, &UnresolvedReferences);
+		};
 #endif // USE_DEFERRED_DEPENDENCY_CHECK_VERIFICATION_TESTS
 
-	if (!FBlueprintSupport::IsDeferredExportCreationDisabled())
-	{
 		// we may have circumvented an export creation or two to avoid instantiating
 		// an BP object before its class has been finalized (to avoid costly re-
 		// instancing operations when the class ultimately finalizes)... so here, we
 		// find those skipped exports and properly create them (before we finalize 
 		// our own class)
-		for (int32 ExportIndex = 0; ExportIndex < ExportMap.Num(); ++ExportIndex)
+		for (int32 ExportIndex = 0; ExportIndex < ExportMap.Num() && IsBlueprintFinalizationPending(); ++ExportIndex)
 		{
 			FObjectExport& Export = ExportMap[ExportIndex];
 			if (ULinkerPlaceholderExportObject* PlaceholderExport = Cast<ULinkerPlaceholderExportObject>(Export.Object))
 			{
 				DEFERRED_DEPENDENCY_CHECK(Export.ClassIndex.IsImport());
 
-				// make sure we're not already in ForceRegenerateClass() for
-				// this export (that could cause some bad infinite recursion)
-				DEFERRED_DEPENDENCY_CHECK(!IsExportBeingResolved(ExportIndex));
-
 				UClass* ExportClass = GetExportLoadClass(ExportIndex);
 				// export class could be null... we create these placeholder 
 				// exports for objects that are instances of an external 
@@ -1238,6 +1209,12 @@ void FLinkerLoad::ResolveDeferredExports(UClass* LoadClass)
 					// make sure this export's class is fully regenerated before  
 					// we instantiate it (so we don't have to re-inst on load)
 					ForceRegenerateClass(ExportClass);
+
+					if (PlaceholderExport != Export.Object)
+					{
+						DEFERRED_DEPENDENCY_CHECK( !IsPlaceholderReferenced(PlaceholderExport) );
+						continue;
+					}
 				}
 
 				// replace the placeholder with the proper object instance
@@ -1257,13 +1234,7 @@ void FLinkerLoad::ResolveDeferredExports(UClass* LoadClass)
 				{
 					Preload(ExportObj);
 				}
-
-#if USE_DEFERRED_DEPENDENCY_CHECK_VERIFICATION_TESTS
-				FReferencerInformationList UnresolvedReferences;
-				UObject* PlaceholderObj = PlaceholderExport;
-				bool const bIsReferenced = IsReferenced(PlaceholderObj, RF_NoFlags, /*bCheckSubObjects =*/false, &UnresolvedReferences);
-				DEFERRED_DEPENDENCY_CHECK(!bIsReferenced);
-#endif // USE_DEFERRED_DEPENDENCY_CHECK_VERIFICATION_TESTS
+				DEFERRED_DEPENDENCY_CHECK( !IsPlaceholderReferenced(PlaceholderExport) );
 			}
 		}
 	}
@@ -1273,95 +1244,20 @@ void FLinkerLoad::ResolveDeferredExports(UClass* LoadClass)
 	// between now and when DeferredCDOIndex is cleared (they won't be resolved,
 	// so that is a problem!)
 	FResolvingExportTracker::Get().FlagFullExportResolvePassComplete(this);
-	// we can be re-entrant from here on (in-fact, we expect to with CDO 
-	// serialization)
-	ResolvingExportsTracker.Remove(this);
 
 	for (TObjectIterator<ULinkerPlaceholderExportObject> PlaceholderIt; PlaceholderIt; ++PlaceholderIt)
 	{
 		ULinkerPlaceholderExportObject* PlaceholderObj = *PlaceholderIt;
 		DEFERRED_DEPENDENCY_CHECK((PlaceholderObj->GetOuter() != LinkerRoot) || PlaceholderObj->IsPendingKill());
 	}
-
-	if (!FBlueprintSupport::IsDeferredCDOSerializationDisabled())
 #endif // USE_DEFERRED_DEPENDENCY_CHECK_VERIFICATION_TESTS
-	{
-		// NOTE: this doesn't work... mainly because of how we handle the case  
-		//       where SuperClass->HasAnyClassFlags(CLASS_NewerVersionExists) is 
-		//       true; at a surface level it works, up until GC is ran (the old
-		//       CDO is still tied to the active class, and the CDO's 
-		//       DestroyNonNativeProperties() attempts to call the class's
-		//       DestroyPersistentUberGraphFrame()... this is why 
-		//       ResolveDeferredSubClassObjects() was introduced instead
-// 		UClass* SuperClass = LoadClass->GetSuperClass();
-// 		// if this class's CDO had its initialization deferred (presumably 
-// 		// because its super's CDO hadn't been fully serialized), then we want 
-// 		// to make sure it is ran here before we serialize in the CDO (so it 
-// 		// gets stocked with inherited values first, that the Preload() may or
-// 		// may not override)
-// 		if (!FPlatformProperties::RequiresCookedData() && FDeferredObjInitializerTracker::IsCdoDeferred(LoadClass))
-// 		{
-// 			// if LoadClass's super has been regenerated (presumably since we 
-// 			// deferred its CDO initialization), then we need to throw out the
-// 			// CDO and recreate a new one... we do this because the super's
-// 			// layout could have changed (and property offsets are now different)
-// 			// for the deferred initialization to work, inherited property 
-// 			// offsets have to match (FObjectInitializer uses one property to 
-// 			// copy values between the CDOs)
-// 			if (SuperClass->HasAnyClassFlags(CLASS_NewerVersionExists))
-// 			{
-// 				// reset the class's super and relink (so that the class's 
-// 				// property chain reflects the super's regenerated layout)
-// 				UClass* RegenedSuperClass = SuperClass->GetAuthoritativeClass();
-// 				LoadClass->SetSuperStruct(RegenedSuperClass);
-// 				LoadClass->StaticLink(/*bRelinkExistingProperties =*/true);
-// 				SuperClass = RegenedSuperClass;
-// 
-// 				UObject* OldCDO = BlueprintCDO;
-// 				DEFERRED_DEPENDENCY_CHECK(LoadClass->GetDefaultObject(/*bCreateIfNeeded =*/false) == BlueprintCDO);
-// 
-// 				// we have to move the OldCDO out of its package, otherwise 
-// 				// when we attempt to create a new one StaticAllocateObject() 
-// 				// will just return us this stale one
-// 				UPackage* TransientPackage = GetTransientPackage();
-// 				const FName TransientName = MakeUniqueObjectName(TransientPackage, LoadClass, *FString::Printf(TEXT("CYCLIC_%s"), *OldCDO->GetName()));
-// 				const ERenameFlags RenameFlags = (REN_DontCreateRedirectors | REN_ForceNoResetLoaders | REN_NonTransactional | REN_DoNotDirty);
-// 				OldCDO->Rename(*TransientName.ToString(), TransientPackage, RenameFlags);
-// 
-// 				// have to reset the LOAD_DeferDependencyLoads flag so that 
-// 				// CreateExport() doesn't try to regenerate the Blueprint on the spot
-// 				const uint32 DeferDependencyFlag = (LoadFlags & LOAD_DeferDependencyLoads);
-// 				LoadFlags |= LOAD_DeferDependencyLoads;
-// 				
-// 				// recreate the CDO from scratch (going through CreateExport() instead
-// 				// of GetDefaultObject() so that it gets properly marked with RF_NeedLoad, etc.)
-// 				LoadClass->ClassDefaultObject = nullptr;
-// 				CDOExport.Object = nullptr;
-// 				UObject* NewCDO = CreateExport(DeferredCDOIndex);
-// 				DEFERRED_DEPENDENCY_CHECK(NewCDO != OldCDO);
-// 				LoadClass->ClassDefaultObject = NewCDO;
-// 
-// 				// restore (clear) the LOAD_DeferDependencyLoads flag
-// 				LoadFlags = (LoadFlags & ~LOAD_DeferDependencyLoads) | DeferDependencyFlag;
-// 
-// 				BlueprintCDO = NewCDO;
-// 				// NOTE: since we're recreating the CDO, CreateExport() will take care of initializting
-// 				//       it for us (the cached FObjectInitializer is no longer needed), but we do still
-// 				//       need to load (and re-parent) deferred sub-objects
-// 				FDeferredObjInitializerTracker::ResolveDeferredSubObjects(NewCDO);
-// 				FDeferredObjInitializerTracker::Remove(LoadClass);
-// 			}
-// 			else
-// 			{
-// 				FDeferredObjInitializerTracker::ResolveDeferredInitialization(LoadClass);
-// 			}
-// 		}
-// 		else
-// 		{
-//  			DEFERRED_DEPENDENCY_CHECK(!SuperClass->HasAnyClassFlags(CLASS_NewerVersionExists));
-//  			FDeferredObjInitializerTracker::ResolveDeferredInitialization(LoadClass);
-// 		}
 
+	// the ExportMap loop above could have recursed back into "finalization" for  
+	// this asset, subsequently resolving all exports before this function could 
+	// finish... that means there's no work left for this to do (and trying to
+	// redo the work would cause a crash), so we guard here against that
+	if (IsBlueprintFinalizationPending())
+	{
 		// have to prematurely set the CDO's linker so we can force a Preload()/
 		// Serialization of the CDO before we regenerate the Blueprint class
 		{
diff --git a/Engine/Source/Runtime/CoreUObject/Private/UObject/LinkerLoad.cpp b/Engine/Source/Runtime/CoreUObject/Private/UObject/LinkerLoad.cpp
--- a/Engine/Source/Runtime/CoreUObject/Private/UObject/LinkerLoad.cpp
+++ b/Engine/Source/Runtime/CoreUObject/Private/UObject/LinkerLoad.cpp
@@ -3096,25 +3096,21 @@ void FLinkerLoad::Preload( UObject* Object )
 						if ((LoadFlags & LOAD_DeferDependencyLoads) != 0)
 						{
 #if USE_DEFERRED_DEPENDENCY_CHECK_VERIFICATION_TESTS
-							if (!FBlueprintSupport::IsDeferredCDOSerializationDisabled())
-							{
-								check((DeferredCDOIndex == INDEX_NONE) || (DeferredCDOIndex == ExportIndex));
-#else // USE_DEFERRED_DEPENDENCY_CHECK_VERIFICATION_TESTS
-							{
+							check((DeferredCDOIndex == INDEX_NONE) || (DeferredCDOIndex == ExportIndex));
 #endif // USE_DEFERRED_DEPENDENCY_CHECK_VERIFICATION_TESTS
-								// since serializing the CDO can introduce circular 
-								// dependencies, we want to stave that off until 
-								// we're ready to handle those 
-								DeferredCDOIndex = ExportIndex;
-								// don't need to actually "consume" the data through
-								// serialization though (since we seek back to 
-								// SavedPos later on)
-
-								// reset the flag and return (don't worry, we make
-								// sure to force load this later)
-								Object->SetFlags(RF_NeedLoad);
-								return;
-							}
+							
+							// since serializing the CDO can introduce circular 
+							// dependencies, we want to stave that off until 
+							// we're ready to handle those 
+							DeferredCDOIndex = ExportIndex;
+							// don't need to actually "consume" the data through
+							// serialization though (since we seek back to 
+							// SavedPos later on)
+
+							// reset the flag and return (don't worry, we make
+							// sure to force load this later)
+							Object->SetFlags(RF_NeedLoad);
+							return;
 						}
 #endif // USE_CIRCULAR_DEPENDENCY_LOAD_DEFERRING
 
@@ -3158,9 +3154,9 @@ void FLinkerLoad::Preload( UObject* Object )
 							// however, sometimes DeferredExportIndex doesn't get set at all (we have to utilize FindCDOExportIndex() to set
 							// it), and that happens when the class's ClassGeneratedBy is serialized in null... this will happen for cooked 
 							// builds (because Blueprints are editor-only objects)
-							check((DeferredCDOIndex != INDEX_NONE) || FPlatformProperties::RequiresCookedData() || FBlueprintSupport::IsDeferredCDOSerializationDisabled());
+							check((DeferredCDOIndex != INDEX_NONE) || FPlatformProperties::RequiresCookedData());
 
-							if ((DeferredCDOIndex == INDEX_NONE) && !FBlueprintSupport::IsDeferredCDOSerializationDisabled())
+							if (DeferredCDOIndex == INDEX_NONE)
 							{
 								DeferredCDOIndex = FindCDOExportIndex(ObjectAsClass);
 								check(DeferredCDOIndex != INDEX_NONE);
@@ -3692,11 +3688,7 @@ UObject* FLinkerLoad::CreateExport( int32 Index )
 			bool const bIsBlueprintCDO = ((Export.ObjectFlags & RF_ClassDefaultObject) != 0) && (LoadClass->ClassGeneratedBy != nullptr);
 
 #if USE_CIRCULAR_DEPENDENCY_LOAD_DEFERRING
-			bool bDeferCDOSerialization = bIsBlueprintCDO && ((LoadFlags & LOAD_DeferDependencyLoads) != 0);
-#if USE_DEFERRED_DEPENDENCY_CHECK_VERIFICATION_TESTS
-			bDeferCDOSerialization &= !FBlueprintSupport::IsDeferredCDOSerializationDisabled();
-#endif // USE_DEFERRED_DEPENDENCY_CHECK_VERIFICATION_TESTS	
-
+			const bool bDeferCDOSerialization = bIsBlueprintCDO && ((LoadFlags & LOAD_DeferDependencyLoads) != 0);
 			if (bDeferCDOSerialization)			
 			{
 				// if LOAD_DeferDependencyLoads is set, then we're already
diff --git a/Engine/Source/Runtime/CoreUObject/Private/UObject/LinkerPlaceholderBase.inl b/Engine/Source/Runtime/CoreUObject/Private/UObject/LinkerPlaceholderBase.inl
--- a/Engine/Source/Runtime/CoreUObject/Private/UObject/LinkerPlaceholderBase.inl
+++ b/Engine/Source/Runtime/CoreUObject/Private/UObject/LinkerPlaceholderBase.inl
@@ -92,7 +92,7 @@ void TLinkerImportPlaceholder<PlaceholderType>::AddReferencingProperty(UProperty
 			}
 		}
 		check(ThisAsObject->GetOutermost() == PropertyLinker->LinkerRoot);
-		check((PropertyLinker->LoadFlags & LOAD_DeferDependencyLoads) || FBlueprintSupport::IsResolvingDeferredDependenciesDisabled());
+		check(PropertyLinker->LoadFlags & LOAD_DeferDependencyLoads);
 	}
 	// if this check hits, then we're adding dependencies after we've 
 	// already resolved the placeholder (it won't be resolved again)
diff --git a/Engine/Source/Runtime/CoreUObject/Public/Blueprint/BlueprintSupport.h b/Engine/Source/Runtime/CoreUObject/Public/Blueprint/BlueprintSupport.h
--- a/Engine/Source/Runtime/CoreUObject/Public/Blueprint/BlueprintSupport.h
+++ b/Engine/Source/Runtime/CoreUObject/Public/Blueprint/BlueprintSupport.h
@@ -23,8 +23,6 @@ struct FBlueprintSupport
 	 * aspects of the deferred loading (mostly for testing purposes). 
 	 */
 	static bool UseDeferredDependencyLoading();
-	static bool IsResolvingDeferredDependenciesDisabled();
-	static bool IsDeferredCDOSerializationDisabled();
 	static bool IsDeferredExportCreationDisabled();
 	static bool IsDeferredCDOInitializationDisabled();
 };
