Merging change #2624225

Fixing Token Stream not properly handling AddReferencedObjects pointers when added just before the end of array tokens. Introduced new token type that marks the end of a pointer data inside of a token stream to correctly handle return count for the array.

diff --git a/Engine/Source/Runtime/CoreUObject/Private/UObject/GarbageCollection.cpp b/Engine/Source/Runtime/CoreUObject/Private/UObject/GarbageCollection.cpp
--- a/Engine/Source/Runtime/CoreUObject/Private/UObject/GarbageCollection.cpp
+++ b/Engine/Source/Runtime/CoreUObject/Private/UObject/GarbageCollection.cpp
@@ -863,6 +863,10 @@ public:
 						FSimpleObjectReferenceCollectorArchive CollectorArchive(CurrentObject, ReferenceCollector);
 						MapProperty->SerializeItem(CollectorArchive, Map, nullptr);
 					}
+					else if (REFERENCE_INFO.Type == GCRT_EndOfPointer)
+					{
+						TokenReturnCount = REFERENCE_INFO.ReturnCount;
+					}
 					else if( REFERENCE_INFO.Type == GCRT_EndOfStream )
 					{
 						// Break out of loop.
@@ -1969,7 +1973,8 @@ void FGCReferenceTokenStream::ReplaceOrAddAddReferencedObjectsCall(void (*AddRef
 		case GCRT_Object:
 		case GCRT_PersistentObject:
 		case GCRT_ArrayObject:
-		case GCRT_EndOfStream:
+		case GCRT_EndOfPointer:
+		case GCRT_EndOfStream:		
 			break;
 		default:
 			UE_LOG(LogGarbage, Fatal, TEXT("Unknown token type (%u) when trying to add ARO token."), (uint32)TokenType);
@@ -2007,7 +2012,7 @@ uint32 FGCReferenceTokenStream::EmitSkipIndexPlaceholder()
  */
 void FGCReferenceTokenStream::UpdateSkipIndexPlaceholder( uint32 SkipIndexIndex, uint32 SkipIndex )
 {
-	check( SkipIndex > 0 && SkipIndex <= (uint32)Tokens.Num() );			
+	check( SkipIndex > 0 && SkipIndex <= (uint32)Tokens.Num() );
 	const FGCReferenceInfo& ReferenceInfo = Tokens[SkipIndex-1];
 	check( ReferenceInfo.Type != GCRT_None );
 	check( Tokens[SkipIndexIndex] == E_GCSkipIndexPlaceholder );
@@ -2035,6 +2040,9 @@ void FGCReferenceTokenStream::EmitPointer( void const* Ptr )
 	const int32 StoreIndex = Tokens.Num();
 	Tokens.AddUninitialized(GNumTokensPerPointer);
 	StorePointer(&Tokens[StoreIndex], Ptr);
+	// Now inser the end of pointer marker, this will mostly be used for storing ReturnCount value
+	// if the pointer was stored at the end of struct array stream.
+	EmitReferenceInfo(FGCReferenceInfo(GCRT_EndOfPointer, 0));
 }
 
 /**
@@ -2055,7 +2063,7 @@ void FGCReferenceTokenStream::EmitStride( uint32 Stride )
 uint32 FGCReferenceTokenStream::EmitReturn()
 {
 	FGCReferenceInfo ReferenceInfo = Tokens.Last();
-	check( ReferenceInfo.Type != GCRT_None );
+	check(ReferenceInfo.Type != GCRT_None);
 	ReferenceInfo.ReturnCount++;
 	Tokens.Last() = ReferenceInfo;
 	return Tokens.Num();
diff --git a/Engine/Source/Runtime/CoreUObject/Public/UObject/GarbageCollection.h b/Engine/Source/Runtime/CoreUObject/Public/UObject/GarbageCollection.h
--- a/Engine/Source/Runtime/CoreUObject/Public/UObject/GarbageCollection.h
+++ b/Engine/Source/Runtime/CoreUObject/Public/UObject/GarbageCollection.h
@@ -26,6 +26,7 @@ enum EGCReferenceType
 	GCRT_AddStructReferencedObjects,
 	GCRT_AddReferencedObjects,
 	GCRT_AddTMapReferencedObjects,
+	GCRT_EndOfPointer,
 	GCRT_EndOfStream,
 };
 
