Removed FDeferredMessageLog. Added concept of 'unseen messages' to IMessageLog.  Opened assets now automatically open the message log if they encounter load errors.

diff --git a/Engine/Source/Developer/MessageLog/Private/Presentation/MessageLogListingViewModel.cpp b/Engine/Source/Developer/MessageLog/Private/Presentation/MessageLogListingViewModel.cpp
--- a/Engine/Source/Developer/MessageLog/Private/Presentation/MessageLogListingViewModel.cpp
+++ b/Engine/Source/Developer/MessageLog/Private/Presentation/MessageLogListingViewModel.cpp
@@ -91,6 +91,7 @@ void FMessageLogListingViewModel::OpenMessageLog( )
 {
 	FMessageLogModule& MessageLogModule = FModuleManager::GetModuleChecked<FMessageLogModule>("MessageLog");
 	MessageLogModule.OpenMessageLog(GetName());
+	bHasUnseenMessages = false;
 }
 
 int32 FMessageLogListingViewModel::NumMessagesPresent( uint32 PageIndex, EMessageSeverity::Type InSeverity ) const
@@ -138,6 +139,7 @@ void FMessageLogListingViewModel::AddMessage( const TSharedRef< class FTokenized
 	}
 
 	MessageLogListingModel->AddMessage(NewMessage);
+	bHasUnseenMessages = true;
 }
 
 void FMessageLogListingViewModel::AddMessages( const TArray< TSharedRef< class FTokenizedMessage > >& NewMessages )
@@ -153,12 +155,14 @@ void FMessageLogListingViewModel::AddMessages( const TArray< TSharedRef< class F
 	else
 	{
 		MessageLogListingModel->AddMessages(NewMessages);
+		bHasUnseenMessages = true;
 	}
 }
 
 void FMessageLogListingViewModel::ClearMessages()
 {
 	MessageLogListingModel->ClearMessages();
+	bHasUnseenMessages = false;
 }
 
 TSharedPtr<FTokenizedMessage> FMessageLogListingViewModel::GetMessageFromData(const struct FTokenizedMiscData& MessageData) const
@@ -330,6 +334,11 @@ int32 FMessageLogListingViewModel::NumMessages( EMessageSeverity::Type SeverityF
 	return NumMessagesPresent(0, SeverityFilter);
 }
 
+bool FMessageLogListingViewModel::HasUnseenMessages() const
+{
+	return bHasUnseenMessages;
+}
+
 void FMessageLogListingViewModel::Open()
 {
 	OpenMessageLog();
diff --git a/Engine/Source/Developer/MessageLog/Private/Presentation/MessageLogListingViewModel.h b/Engine/Source/Developer/MessageLog/Private/Presentation/MessageLogListingViewModel.h
--- a/Engine/Source/Developer/MessageLog/Private/Presentation/MessageLogListingViewModel.h
+++ b/Engine/Source/Developer/MessageLog/Private/Presentation/MessageLogListingViewModel.h
@@ -59,6 +59,7 @@ public:
 	virtual void NotifyIfAnyMessages( const FText& Message, EMessageSeverity::Type SeverityFilter = EMessageSeverity::Info, bool bForce = false ) override;
 	virtual void Open() override;
 	virtual int32 NumMessages( EMessageSeverity::Type SeverityFilter ) override;
+	virtual bool HasUnseenMessages() const override;
 
 	DECLARE_DERIVED_EVENT(FMessageLogListingViewModel, IMessageLogListing::IMessageTokenClickedEvent, IMessageTokenClickedEvent)
 	virtual IMessageLogListing::IMessageTokenClickedEvent& OnMessageTokenClicked() override { return TokenClickedEvent; }
@@ -154,6 +155,7 @@ private:
 		, bShowPages( InitializationOptions.bShowPages )
 		, bAllowClear( InitializationOptions.bAllowClear )
 		, bDiscardDuplicates( InitializationOptions.bDiscardDuplicates )
+		, bHasUnseenMessages( false )
 		, MaxPageCount( InitializationOptions.MaxPageCount )
 		, CurrentPageIndex( 0 )
 		, bIsRefreshing( false )
@@ -187,6 +189,9 @@ private:
 	/** Whether to check for duplicate messages & discard them */
 	bool bDiscardDuplicates;
 
+	/** Whether there are unseen messages in this listing */
+	bool bHasUnseenMessages;
+
 	/** The limit on the number of displayed pages for this listing */
 	uint32 MaxPageCount;
 
diff --git a/Engine/Source/Editor/ContentBrowser/Private/SContentBrowser.cpp b/Engine/Source/Editor/ContentBrowser/Private/SContentBrowser.cpp
--- a/Engine/Source/Editor/ContentBrowser/Private/SContentBrowser.cpp
+++ b/Engine/Source/Editor/ContentBrowser/Private/SContentBrowser.cpp
@@ -1769,6 +1769,8 @@ void SContentBrowser::OnAssetsActivated(const TArray<FAssetData>& ActivatedAsset
 	{
 		ContentBrowserUtils::OpenEditorForAsset(ObjectsWithoutTypeActions);
 	}
+
+	FEditorDelegates::DisplayLoadErrors.Broadcast();
 }
 
 TSharedPtr<SWidget> SContentBrowser::OnGetAssetContextMenu(const TArray<FAssetData>& SelectedAssets)
diff --git a/Engine/Source/Editor/UnrealEd/Private/EditorServer.cpp b/Engine/Source/Editor/UnrealEd/Private/EditorServer.cpp
--- a/Engine/Source/Editor/UnrealEd/Private/EditorServer.cpp
+++ b/Engine/Source/Editor/UnrealEd/Private/EditorServer.cpp
@@ -2520,7 +2520,11 @@ bool UEditorEngine::Map_Load(const TCHAR* Str, FOutputDevice& Ar)
 
 					InitializingFeedback.EnterProgressFrame();
 
-					FEditorDelegates::DisplayLoadErrors.Broadcast();
+					if (!FEditorFileUtils::IsLoadingStartupMap())
+					{
+						// Don't attempt to display load errors if loading the startup map; this will be done in the first editor tick.
+						FEditorDelegates::DisplayLoadErrors.Broadcast();
+					}
 
 					if( FModuleManager::Get().IsModuleLoaded("LevelEditor") )
 					{
diff --git a/Engine/Source/Editor/UnrealEd/Private/UnrealEdMisc.cpp b/Engine/Source/Editor/UnrealEd/Private/UnrealEdMisc.cpp
--- a/Engine/Source/Editor/UnrealEd/Private/UnrealEdMisc.cpp
+++ b/Engine/Source/Editor/UnrealEd/Private/UnrealEdMisc.cpp
@@ -869,7 +869,11 @@ void FUnrealEdMisc::CB_DisplayLoadErrors()
 		const bool bIsImmersive = FParse::Param( FCommandLine::Get(), TEXT( "immersive" ) );
 		if( !bIsImmersive && !GIsAutomationTesting)
 		{
-			FMessageLog("LoadErrors").Open();
+			FMessageLog MessageLog("LoadErrors");
+			if (MessageLog.HasUnseenMessages())
+			{
+				MessageLog.Open();
+			}
 		}
 	}
 }
diff --git a/Engine/Source/Runtime/Core/Private/Logging/MessageLog.cpp b/Engine/Source/Runtime/Core/Private/Logging/MessageLog.cpp
--- a/Engine/Source/Runtime/Core/Private/Logging/MessageLog.cpp
+++ b/Engine/Source/Runtime/Core/Private/Logging/MessageLog.cpp
@@ -32,6 +32,11 @@ public:
 		}
 	}
 
+	virtual bool HasUnseenMessages() const override
+	{
+		return true;
+	}
+
 	virtual void NewPage( const FText& Title ) override
 	{
 		FFormatNamedArguments Arguments;
@@ -155,14 +160,16 @@ int32 FMessageLog::NumMessages( EMessageSeverity::Type InSeverityFilter )
 	return MessageLog->NumMessages(InSeverityFilter);
 }
 
+bool FMessageLog::HasUnseenMessages()
+{
+	Flush();
+	return MessageLog->HasUnseenMessages();
+}
+
 void FMessageLog::Open( EMessageSeverity::Type InSeverityFilter, bool bOpenEvenIfEmpty )
 {
 	Flush();
-	if(bOpenEvenIfEmpty)
-	{
-		MessageLog->Open();
-	}
-	else if(MessageLog->NumMessages(InSeverityFilter) > 0)
+	if(bOpenEvenIfEmpty || MessageLog->NumMessages(InSeverityFilter) > 0)
 	{
 		MessageLog->Open();
 	}
diff --git a/Engine/Source/Runtime/Core/Public/Logging/IMessageLog.h b/Engine/Source/Runtime/Core/Public/Logging/IMessageLog.h
--- a/Engine/Source/Runtime/Core/Public/Logging/IMessageLog.h
+++ b/Engine/Source/Runtime/Core/Public/Logging/IMessageLog.h
@@ -63,4 +63,9 @@ public:
 	 * @return The number of messages that pass our filter
 	 */
 	virtual int32 NumMessages( EMessageSeverity::Type SeverityFilter ) = 0;
+
+	/**
+	 * Returns whether there are any messages in the log which have not yet been seen.
+	 */
+	virtual bool HasUnseenMessages() const = 0;
 };
\ No newline at end of file
diff --git a/Engine/Source/Runtime/Core/Public/Logging/MessageLog.h b/Engine/Source/Runtime/Core/Public/Logging/MessageLog.h
--- a/Engine/Source/Runtime/Core/Public/Logging/MessageLog.h
+++ b/Engine/Source/Runtime/Core/Public/Logging/MessageLog.h
@@ -49,12 +49,19 @@ public:
 	/**
 	 * Check whether there are any messages present for this log.
 	 * If the log is paged, this will check if there are any messages on the current page.
-	 * This call will cause a flush so that the logs state is properly reflected.
+	 * This call will cause a flush so that the log's state is properly reflected.
 	 * @param	InSeverityFilter	Only messages of higher severity than this filter will be considered when checking.
 	 * @return The number of messages that pass our filter
 	 */
 	CORE_API int32 NumMessages( EMessageSeverity::Type InSeverityFilter = EMessageSeverity::Info );
 
+	/**
+	 * Check whether there are any unseen messages present for this log.
+	 * This call will cause a flush so that the log's state is properly reflected.
+	 * @return Whether there are unseen message present for this log.
+	 */
+	CORE_API bool HasUnseenMessages();
+
 	/**
 	 * Opens the log for display to the user given certain conditions.
 	 * This call will cause a flush so that the logs state is properly reflected.
diff --git a/Engine/Source/Runtime/CoreUObject/Private/Serialization/AsyncLoading.cpp b/Engine/Source/Runtime/CoreUObject/Private/Serialization/AsyncLoading.cpp
--- a/Engine/Source/Runtime/CoreUObject/Private/Serialization/AsyncLoading.cpp
+++ b/Engine/Source/Runtime/CoreUObject/Private/Serialization/AsyncLoading.cpp
@@ -6,7 +6,6 @@
 
 #include "CoreUObjectPrivate.h"
 #include "Serialization/AsyncLoading.h"
-#include "Serialization/DeferredMessageLog.h"
 #include "Serialization/AsyncPackage.h"
 #include "UObject/UObjectThreadContext.h"
 #include "UObject/LinkerManager.h"
@@ -721,7 +720,6 @@ EAsyncPackageState::Type FAsyncLoadingThread::TickAsyncLoading(bool bUseTimeLimi
 #endif
 			if (AsyncPackagesCounter.GetValue() == 0 && LoadedPackagesToProcess.Num() == 0)
 			{
-				FDeferredMessageLog::Flush();
 				FAsyncObjectsReferencer::Get().EmptyReferencedObjects();
 			}
 		}
diff --git a/Engine/Source/Runtime/CoreUObject/Private/Serialization/DeferredMessageLog.cpp b/Engine/Source/Runtime/CoreUObject/Private/Serialization/DeferredMessageLog.cpp
deleted file mode 100644
diff --git a/Engine/Source/Runtime/CoreUObject/Private/Serialization/DeferredMessageLog.h b/Engine/Source/Runtime/CoreUObject/Private/Serialization/DeferredMessageLog.h
deleted file mode 100644
diff --git a/Engine/Source/Runtime/CoreUObject/Private/UObject/LinkerLoad.cpp b/Engine/Source/Runtime/CoreUObject/Private/UObject/LinkerLoad.cpp
--- a/Engine/Source/Runtime/CoreUObject/Private/UObject/LinkerLoad.cpp
+++ b/Engine/Source/Runtime/CoreUObject/Private/UObject/LinkerLoad.cpp
@@ -9,7 +9,6 @@
 #include "LinkerPlaceholderExportObject.h"
 #include "LinkerPlaceholderFunction.h"
 #include "LinkerManager.h"
-#include "Serialization/DeferredMessageLog.h"
 #include "UObject/UObjectThreadContext.h"
 #include "GatherableTextData.h"
 #include "Serialization/AsyncLoading.h"
@@ -85,7 +84,7 @@ void FLinkerLoad::CreateActiveRedirectsMap(const FString& GEngineIniName)
 		FConfigSection* PackageRedirects = GConfig->GetSectionPrivate( TEXT("/Script/Engine.Engine"), false, true, GEngineIniName );
 		if (PackageRedirects)
 		{
-			FDeferredMessageLog RedirectErrors(NAME_LoadErrors);
+			FMessageLog RedirectErrors(NAME_LoadErrors);
 
 			static FName ActiveClassRedirectsKey(TEXT("ActiveClassRedirects"));
 			for( FConfigSection::TIterator It(*PackageRedirects); It; ++It )
@@ -2225,7 +2224,7 @@ FLinkerLoad::EVerifyResult FLinkerLoad::VerifyImport(int32 ImportIndex)
 				UClass* FindClass = ClassPackage ? FindObject<UClass>( ClassPackage, *OriginalImport.ClassName.ToString() ) : NULL;
 				if( GIsEditor && !IsRunningCommandlet())
 				{
-					FDeferredMessageLog LoadErrors(NAME_LoadErrors);
+					FMessageLog LoadErrors(NAME_LoadErrors);
 					// put something into the load warnings dialog, with any extra information from above (in WarningAppend)
 					TSharedRef<FTokenizedMessage> TokenizedMessage = LoadErrors.Error(FText());
 					TokenizedMessage->AddToken(FAssetNameToken::Create(LinkerRoot->GetName()));
@@ -3311,7 +3310,7 @@ bool FLinkerLoad::WillTextureBeLoaded( UClass* Class, int32 ExportIndex )
 UObject* FLinkerLoad::CreateExport( int32 Index )
 {
 	FScopedCreateExportCounter ScopedCounter( this, Index );
-	FDeferredMessageLog LoadErrors(NAME_LoadErrors);
+	FMessageLog LoadErrors(NAME_LoadErrors);
 
 	// Map the object into our table.
 	FObjectExport& Export = ExportMap[ Index ];
