Make LoadObject to also respect the legacy struct's path.

It turnes out that LoadObject had not a legacy path that FindObject did. Moved legacy logic to StaticFindObjectFast that is shared by both functions

UDN: https://udn.unrealengine.com/questions/252115/failed-to-find-object-when-cooking.html

diff --git a/Engine/Source/Runtime/CoreUObject/Private/UObject/UObjectGlobals.cpp b/Engine/Source/Runtime/CoreUObject/Private/UObject/UObjectGlobals.cpp
--- a/Engine/Source/Runtime/CoreUObject/Private/UObject/UObjectGlobals.cpp
+++ b/Engine/Source/Runtime/CoreUObject/Private/UObject/UObjectGlobals.cpp
@@ -84,53 +84,28 @@ static FORCEINLINE bool IsGarbageCollectingOnGameThread()
 	return IsInGameThread() && IsGarbageCollecting();
 }
 
-/**
- * Fast version of StaticFindObject that relies on the passed in FName being the object name
- * without any group/ package qualifiers.
- *
- * @param	ObjectClass		The to be found object's class
- * @param	ObjectPackage	The to be found object's outer
- * @param	ObjectName		The to be found object's class
- * @param	ExactClass		Whether to require an exact match with the passed in class
- * @param	AnyPackage		Whether to look in any package
- * @param	ExclusiveFlags	Ignores objects that contain any of the specified exclusive flags
- * @return	Returns a pointer to the found object or NULL if none could be found
- */
-UObject* StaticFindObjectFast( UClass* ObjectClass, UObject* ObjectPackage, FName ObjectName, bool ExactClass, bool AnyPackage, EObjectFlags ExclusiveFlags )
-{
-	if (GIsSavingPackage || IsGarbageCollectingOnGameThread())
-	{
-		UE_LOG(LogUObjectGlobals, Fatal,TEXT("Illegal call to StaticFindObjectFast() while serializing object data or garbage collecting!"));
-	}
-
-	// We don't want to return any objects that are currently being background loaded unless we're using FindObject during async loading.
-	ExclusiveFlags |= IsInAsyncLoadingThread() ? RF_NoFlags : RF_AsyncLoading;
-	return StaticFindObjectFastInternal( ObjectClass, ObjectPackage, ObjectName, ExactClass, AnyPackage, ExclusiveFlags );
-}
-
 // Anonymous namespace to not pollute global.
 namespace
 {
 	/**
-	 * Legacy static find object helper, that helps to find reflected types, that
-	 * are no longer a subobjects of UCLASS defined in the same header.
-	 *
-	 * If the class looked for is of one of the relocated types (or theirs subclass)
-	 * then it performs another search in containing package.
-	 *
-	 * If the class match wasn't exact (i.e. either nullptr or subclass of allowed
-	 * ones) and we've found an object we're revalidating it to make sure the
-	 * legacy search was valid.
-	 *
-	 * @param ObjectClass Class of the object to find.
-	 * @param ObjectPackage Package of the object to find.
-	 * @param ObjectName Name of the object to find.
-	 * @param ExactClass If the class match has to be exact. I.e. ObjectClass == FoundObjects.GetClass()
-	 * @param OrigInName Original name provided to StaticFindObject, before resolving. Used to determine if we're looking for delegate signature.
-	 *
-	 * @returns Found object.
-	 */
-	UObject* StaticFindObjectWithChangedLegacyPath(UClass* ObjectClass, UObject* ObjectPackage, FName ObjectName, bool ExactClass, const TCHAR* OrigInName)
+	* Legacy static find object helper, that helps to find reflected types, that
+	* are no longer a subobjects of UCLASS defined in the same header.
+	*
+	* If the class looked for is of one of the relocated types (or theirs subclass)
+	* then it performs another search in containing package.
+	*
+	* If the class match wasn't exact (i.e. either nullptr or subclass of allowed
+	* ones) and we've found an object we're revalidating it to make sure the
+	* legacy search was valid.
+	*
+	* @param ObjectClass Class of the object to find.
+	* @param ObjectPackage Package of the object to find.
+	* @param ObjectName Name of the object to find.
+	* @param ExactClass If the class match has to be exact. I.e. ObjectClass == FoundObjects.GetClass()
+	*
+	* @returns Found object.
+	*/
+	UObject* StaticFindObjectWithChangedLegacyPath(UClass* ObjectClass, UObject* ObjectPackage, FName ObjectName, bool ExactClass)
 	{
 		UObject* MatchingObject = nullptr;
 
@@ -141,7 +116,7 @@ namespace
 			ObjectPackage != nullptr &&
 			ObjectPackage->IsA<UClass>()) // Only if outer is a class.
 		{
-			bool bHasDelegateSignaturePostfix = FString(OrigInName).EndsWith(HEADER_GENERATED_DELEGATE_SIGNATURE_SUFFIX);
+			bool bHasDelegateSignaturePostfix = ObjectName.ToString().EndsWith(HEADER_GENERATED_DELEGATE_SIGNATURE_SUFFIX);
 
 			bool bExactPathChangedClass = ObjectClass == UEnum::StaticClass() // Enums
 				|| ObjectClass == UScriptStruct::StaticClass() || ObjectClass == UStruct::StaticClass() // Structs
@@ -178,6 +153,38 @@ namespace
 	}
 }
 
+/**
+ * Fast version of StaticFindObject that relies on the passed in FName being the object name
+ * without any group/ package qualifiers.
+ *
+ * @param	ObjectClass		The to be found object's class
+ * @param	ObjectPackage	The to be found object's outer
+ * @param	ObjectName		The to be found object's class
+ * @param	ExactClass		Whether to require an exact match with the passed in class
+ * @param	AnyPackage		Whether to look in any package
+ * @param	ExclusiveFlags	Ignores objects that contain any of the specified exclusive flags
+ * @return	Returns a pointer to the found object or NULL if none could be found
+ */
+UObject* StaticFindObjectFast( UClass* ObjectClass, UObject* ObjectPackage, FName ObjectName, bool ExactClass, bool AnyPackage, EObjectFlags ExclusiveFlags )
+{
+	if (GIsSavingPackage || IsGarbageCollectingOnGameThread())
+	{
+		UE_LOG(LogUObjectGlobals, Fatal,TEXT("Illegal call to StaticFindObjectFast() while serializing object data or garbage collecting!"));
+	}
+
+	// We don't want to return any objects that are currently being background loaded unless we're using FindObject during async loading.
+	ExclusiveFlags |= IsInAsyncLoadingThread() ? RF_NoFlags : RF_AsyncLoading;
+	
+	UObject* FoundObject = StaticFindObjectFastInternal( ObjectClass, ObjectPackage, ObjectName, ExactClass, AnyPackage, ExclusiveFlags );
+
+	if (!FoundObject)
+	{
+		FoundObject = StaticFindObjectWithChangedLegacyPath(ObjectClass, ObjectPackage, ObjectName, ExactClass);
+	}
+
+	return FoundObject;
+}
+
 //
 // Find an optional object.
 //
@@ -234,14 +241,7 @@ UObject* StaticFindObject( UClass* ObjectClass, UObject* InObjectPackage, const
 	}
 
 	FName ObjectName(*InName, FNAME_Add, true);
-	MatchingObject = StaticFindObjectFast( ObjectClass, ObjectPackage, ObjectName, ExactClass, bAnyPackage );
-
-	if (!MatchingObject)
-	{
-		return StaticFindObjectWithChangedLegacyPath(ObjectClass, ObjectPackage, ObjectName, ExactClass, OrigInName);
-	}
-
-	return MatchingObject;
+	return StaticFindObjectFast(ObjectClass, ObjectPackage, ObjectName, ExactClass, bAnyPackage);
 }
 
 //
