UE-16118: 4.8 Preview 3: Hot reload error after adding new enums to a UObject file.

diff --git a/Engine/Source/Runtime/CoreUObject/Private/UObject/UObjectBase.cpp b/Engine/Source/Runtime/CoreUObject/Private/UObject/UObjectBase.cpp
--- a/Engine/Source/Runtime/CoreUObject/Private/UObject/UObjectBase.cpp
+++ b/Engine/Source/Runtime/CoreUObject/Private/UObject/UObjectBase.cpp
@@ -498,7 +498,7 @@ class UEnum *GetStaticEnum(class UEnum *(*InRegister)(), UObject* EnumOuter, con
 #if WITH_HOT_RELOAD
 	if (GIsHotReload)
 	{
-		UEnum* ReturnEnum = FindObjectChecked<UEnum>(EnumOuter, EnumName);
+		UEnum* ReturnEnum = FindObject<UEnum>(EnumOuter, EnumName);
 		if (ReturnEnum)
 		{
 			UE_LOG(LogClass, Log, TEXT( "%s HotReload."), EnumName);
@@ -799,7 +799,7 @@ static void UObjectLoadAllCompiledInStructs()
 
 bool AnyNewlyLoadedUObjects()
 {
-	return GFirstPendingRegistrant != NULL || GetDeferredCompiledInRegistration().Num() || GetDeferredCompiledInStructRegistration().Num();
+	return GFirstPendingRegistrant != NULL || GetDeferredCompiledInRegistration().Num() || GetDeferredCompiledInStructRegistration().Num() || GetDeferredCompiledInEnumRegistration().Num();
 }
 
 
