Back out changelist 2625724

diff --git a/Engine/Source/Developer/MessageLog/Private/Presentation/MessageLogListingViewModel.cpp b/Engine/Source/Developer/MessageLog/Private/Presentation/MessageLogListingViewModel.cpp
--- a/Engine/Source/Developer/MessageLog/Private/Presentation/MessageLogListingViewModel.cpp
+++ b/Engine/Source/Developer/MessageLog/Private/Presentation/MessageLogListingViewModel.cpp
@@ -91,7 +91,6 @@ void FMessageLogListingViewModel::OpenMessageLog( )
 {
 	FMessageLogModule& MessageLogModule = FModuleManager::GetModuleChecked<FMessageLogModule>("MessageLog");
 	MessageLogModule.OpenMessageLog(GetName());
-	bHasUnseenMessages = false;
 }
 
 int32 FMessageLogListingViewModel::NumMessagesPresent( uint32 PageIndex, EMessageSeverity::Type InSeverity ) const
@@ -139,7 +138,6 @@ void FMessageLogListingViewModel::AddMessage( const TSharedRef< class FTokenized
 	}
 
 	MessageLogListingModel->AddMessage(NewMessage);
-	bHasUnseenMessages = true;
 }
 
 void FMessageLogListingViewModel::AddMessages( const TArray< TSharedRef< class FTokenizedMessage > >& NewMessages )
@@ -155,14 +153,12 @@ void FMessageLogListingViewModel::AddMessages( const TArray< TSharedRef< class F
 	else
 	{
 		MessageLogListingModel->AddMessages(NewMessages);
-		bHasUnseenMessages = true;
 	}
 }
 
 void FMessageLogListingViewModel::ClearMessages()
 {
 	MessageLogListingModel->ClearMessages();
-	bHasUnseenMessages = false;
 }
 
 TSharedPtr<FTokenizedMessage> FMessageLogListingViewModel::GetMessageFromData(const struct FTokenizedMiscData& MessageData) const
@@ -334,11 +330,6 @@ int32 FMessageLogListingViewModel::NumMessages( EMessageSeverity::Type SeverityF
 	return NumMessagesPresent(0, SeverityFilter);
 }
 
-bool FMessageLogListingViewModel::HasUnseenMessages() const
-{
-	return bHasUnseenMessages;
-}
-
 void FMessageLogListingViewModel::Open()
 {
 	OpenMessageLog();
diff --git a/Engine/Source/Developer/MessageLog/Private/Presentation/MessageLogListingViewModel.h b/Engine/Source/Developer/MessageLog/Private/Presentation/MessageLogListingViewModel.h
--- a/Engine/Source/Developer/MessageLog/Private/Presentation/MessageLogListingViewModel.h
+++ b/Engine/Source/Developer/MessageLog/Private/Presentation/MessageLogListingViewModel.h
@@ -59,7 +59,6 @@ public:
 	virtual void NotifyIfAnyMessages( const FText& Message, EMessageSeverity::Type SeverityFilter = EMessageSeverity::Info, bool bForce = false ) override;
 	virtual void Open() override;
 	virtual int32 NumMessages( EMessageSeverity::Type SeverityFilter ) override;
-	virtual bool HasUnseenMessages() const override;
 
 	DECLARE_DERIVED_EVENT(FMessageLogListingViewModel, IMessageLogListing::IMessageTokenClickedEvent, IMessageTokenClickedEvent)
 	virtual IMessageLogListing::IMessageTokenClickedEvent& OnMessageTokenClicked() override { return TokenClickedEvent; }
@@ -155,7 +154,6 @@ private:
 		, bShowPages( InitializationOptions.bShowPages )
 		, bAllowClear( InitializationOptions.bAllowClear )
 		, bDiscardDuplicates( InitializationOptions.bDiscardDuplicates )
-		, bHasUnseenMessages( false )
 		, MaxPageCount( InitializationOptions.MaxPageCount )
 		, CurrentPageIndex( 0 )
 		, bIsRefreshing( false )
@@ -189,9 +187,6 @@ private:
 	/** Whether to check for duplicate messages & discard them */
 	bool bDiscardDuplicates;
 
-	/** Whether there are unseen messages in this listing */
-	bool bHasUnseenMessages;
-
 	/** The limit on the number of displayed pages for this listing */
 	uint32 MaxPageCount;
 
diff --git a/Engine/Source/Editor/ContentBrowser/Private/SContentBrowser.cpp b/Engine/Source/Editor/ContentBrowser/Private/SContentBrowser.cpp
--- a/Engine/Source/Editor/ContentBrowser/Private/SContentBrowser.cpp
+++ b/Engine/Source/Editor/ContentBrowser/Private/SContentBrowser.cpp
@@ -1769,8 +1769,6 @@ void SContentBrowser::OnAssetsActivated(const TArray<FAssetData>& ActivatedAsset
 	{
 		ContentBrowserUtils::OpenEditorForAsset(ObjectsWithoutTypeActions);
 	}
-
-	FEditorDelegates::DisplayLoadErrors.Broadcast();
 }
 
 TSharedPtr<SWidget> SContentBrowser::OnGetAssetContextMenu(const TArray<FAssetData>& SelectedAssets)
diff --git a/Engine/Source/Editor/UnrealEd/Private/EditorServer.cpp b/Engine/Source/Editor/UnrealEd/Private/EditorServer.cpp
--- a/Engine/Source/Editor/UnrealEd/Private/EditorServer.cpp
+++ b/Engine/Source/Editor/UnrealEd/Private/EditorServer.cpp
@@ -2520,11 +2520,7 @@ bool UEditorEngine::Map_Load(const TCHAR* Str, FOutputDevice& Ar)
 
 					InitializingFeedback.EnterProgressFrame();
 
-					if (!FEditorFileUtils::IsLoadingStartupMap())
-					{
-						// Don't attempt to display load errors if loading the startup map; this will be done in the first editor tick.
-						FEditorDelegates::DisplayLoadErrors.Broadcast();
-					}
+					FEditorDelegates::DisplayLoadErrors.Broadcast();
 
 					if( FModuleManager::Get().IsModuleLoaded("LevelEditor") )
 					{
diff --git a/Engine/Source/Editor/UnrealEd/Private/UnrealEdMisc.cpp b/Engine/Source/Editor/UnrealEd/Private/UnrealEdMisc.cpp
--- a/Engine/Source/Editor/UnrealEd/Private/UnrealEdMisc.cpp
+++ b/Engine/Source/Editor/UnrealEd/Private/UnrealEdMisc.cpp
@@ -869,11 +869,7 @@ void FUnrealEdMisc::CB_DisplayLoadErrors()
 		const bool bIsImmersive = FParse::Param( FCommandLine::Get(), TEXT( "immersive" ) );
 		if( !bIsImmersive && !GIsAutomationTesting)
 		{
-			FMessageLog MessageLog("LoadErrors");
-			if (MessageLog.HasUnseenMessages())
-			{
-				MessageLog.Open();
-			}
+			FMessageLog("LoadErrors").Open();
 		}
 	}
 }
diff --git a/Engine/Source/Runtime/Core/Private/Logging/MessageLog.cpp b/Engine/Source/Runtime/Core/Private/Logging/MessageLog.cpp
--- a/Engine/Source/Runtime/Core/Private/Logging/MessageLog.cpp
+++ b/Engine/Source/Runtime/Core/Private/Logging/MessageLog.cpp
@@ -32,11 +32,6 @@ public:
 		}
 	}
 
-	virtual bool HasUnseenMessages() const override
-	{
-		return true;
-	}
-
 	virtual void NewPage( const FText& Title ) override
 	{
 		FFormatNamedArguments Arguments;
@@ -160,16 +155,14 @@ int32 FMessageLog::NumMessages( EMessageSeverity::Type InSeverityFilter )
 	return MessageLog->NumMessages(InSeverityFilter);
 }
 
-bool FMessageLog::HasUnseenMessages()
-{
-	Flush();
-	return MessageLog->HasUnseenMessages();
-}
-
 void FMessageLog::Open( EMessageSeverity::Type InSeverityFilter, bool bOpenEvenIfEmpty )
 {
 	Flush();
-	if(bOpenEvenIfEmpty || MessageLog->NumMessages(InSeverityFilter) > 0)
+	if(bOpenEvenIfEmpty)
+	{
+		MessageLog->Open();
+	}
+	else if(MessageLog->NumMessages(InSeverityFilter) > 0)
 	{
 		MessageLog->Open();
 	}
diff --git a/Engine/Source/Runtime/Core/Public/Logging/IMessageLog.h b/Engine/Source/Runtime/Core/Public/Logging/IMessageLog.h
--- a/Engine/Source/Runtime/Core/Public/Logging/IMessageLog.h
+++ b/Engine/Source/Runtime/Core/Public/Logging/IMessageLog.h
@@ -63,9 +63,4 @@ public:
 	 * @return The number of messages that pass our filter
 	 */
 	virtual int32 NumMessages( EMessageSeverity::Type SeverityFilter ) = 0;
-
-	/**
-	 * Returns whether there are any messages in the log which have not yet been seen.
-	 */
-	virtual bool HasUnseenMessages() const = 0;
 };
\ No newline at end of file
diff --git a/Engine/Source/Runtime/Core/Public/Logging/MessageLog.h b/Engine/Source/Runtime/Core/Public/Logging/MessageLog.h
--- a/Engine/Source/Runtime/Core/Public/Logging/MessageLog.h
+++ b/Engine/Source/Runtime/Core/Public/Logging/MessageLog.h
@@ -49,19 +49,12 @@ public:
 	/**
 	 * Check whether there are any messages present for this log.
 	 * If the log is paged, this will check if there are any messages on the current page.
-	 * This call will cause a flush so that the log's state is properly reflected.
+	 * This call will cause a flush so that the logs state is properly reflected.
 	 * @param	InSeverityFilter	Only messages of higher severity than this filter will be considered when checking.
 	 * @return The number of messages that pass our filter
 	 */
 	CORE_API int32 NumMessages( EMessageSeverity::Type InSeverityFilter = EMessageSeverity::Info );
 
-	/**
-	 * Check whether there are any unseen messages present for this log.
-	 * This call will cause a flush so that the log's state is properly reflected.
-	 * @return Whether there are unseen message present for this log.
-	 */
-	CORE_API bool HasUnseenMessages();
-
 	/**
 	 * Opens the log for display to the user given certain conditions.
 	 * This call will cause a flush so that the logs state is properly reflected.
diff --git a/Engine/Source/Runtime/CoreUObject/Private/Serialization/AsyncLoading.cpp b/Engine/Source/Runtime/CoreUObject/Private/Serialization/AsyncLoading.cpp
--- a/Engine/Source/Runtime/CoreUObject/Private/Serialization/AsyncLoading.cpp
+++ b/Engine/Source/Runtime/CoreUObject/Private/Serialization/AsyncLoading.cpp
@@ -6,6 +6,7 @@
 
 #include "CoreUObjectPrivate.h"
 #include "Serialization/AsyncLoading.h"
+#include "Serialization/DeferredMessageLog.h"
 #include "Serialization/AsyncPackage.h"
 #include "UObject/UObjectThreadContext.h"
 #include "UObject/LinkerManager.h"
@@ -720,6 +721,7 @@ EAsyncPackageState::Type FAsyncLoadingThread::TickAsyncLoading(bool bUseTimeLimi
 #endif
 			if (AsyncPackagesCounter.GetValue() == 0 && LoadedPackagesToProcess.Num() == 0)
 			{
+				FDeferredMessageLog::Flush();
 				FAsyncObjectsReferencer::Get().EmptyReferencedObjects();
 			}
 		}
diff --git a/Engine/Source/Runtime/CoreUObject/Private/Serialization/DeferredMessageLog.cpp b/Engine/Source/Runtime/CoreUObject/Private/Serialization/DeferredMessageLog.cpp
new file mode 100644
--- /dev/null
+++ b/Engine/Source/Runtime/CoreUObject/Private/Serialization/DeferredMessageLog.cpp
@@ -0,0 +1,63 @@
+// Copyright 1998-2014 Epic Games, Inc. All Rights Reserved.
+
+/*=============================================================================
+	UnAsyncLoading.cpp: Unreal async messsage log.
+=============================================================================*/
+
+#include "CoreUObjectPrivate.h"
+#include "Serialization/DeferredMessageLog.h"
+
+
+FDeferredMessageLog::FDeferredMessageLog(const FName& InLogCategory)
+: LogCategory(InLogCategory)
+{
+	TArray<TSharedRef<FTokenizedMessage>>** ExistingCategoryMessages = Messages.Find(LogCategory);
+	if (!ExistingCategoryMessages)
+	{
+		TArray<TSharedRef<FTokenizedMessage>>* CategoryMessages = new TArray<TSharedRef<FTokenizedMessage>>();
+		Messages.Add(LogCategory, CategoryMessages);
+	}
+}
+
+void FDeferredMessageLog::AddMessage(TSharedRef<FTokenizedMessage>& Message)
+{
+	TArray<TSharedRef<FTokenizedMessage>>* CategoryMessages = Messages.FindRef(LogCategory);
+	check(CategoryMessages);
+	CategoryMessages->Add(Message);
+}
+
+TSharedRef<FTokenizedMessage> FDeferredMessageLog::Info(const FText& InMessage)
+{
+	TSharedRef<FTokenizedMessage> Message = FTokenizedMessage::Create(EMessageSeverity::Info, InMessage);
+	AddMessage(Message);
+	return Message;
+}
+
+TSharedRef<FTokenizedMessage> FDeferredMessageLog::Warning(const FText& InMessage)
+{
+	TSharedRef<FTokenizedMessage> Message = FTokenizedMessage::Create(EMessageSeverity::Warning, InMessage);
+	AddMessage(Message);
+	return Message;
+}
+
+TSharedRef<FTokenizedMessage> FDeferredMessageLog::Error(const FText& InMessage)
+{
+	TSharedRef<FTokenizedMessage> Message = FTokenizedMessage::Create(EMessageSeverity::Error, InMessage);
+	AddMessage(Message);
+	return Message;
+}
+
+void FDeferredMessageLog::Flush()
+{
+	for (auto& CategoryMessages : Messages)
+	{
+		if (CategoryMessages.Value->Num())
+		{
+			FMessageLog LoaderLog(CategoryMessages.Key);
+			LoaderLog.AddMessages(*CategoryMessages.Value);
+			CategoryMessages.Value->Empty(CategoryMessages.Value->Num());
+		}
+	}
+}
+
+TMap<FName, TArray<TSharedRef<FTokenizedMessage>>*> FDeferredMessageLog::Messages;
diff --git a/Engine/Source/Runtime/CoreUObject/Private/Serialization/DeferredMessageLog.h b/Engine/Source/Runtime/CoreUObject/Private/Serialization/DeferredMessageLog.h
new file mode 100644
--- /dev/null
+++ b/Engine/Source/Runtime/CoreUObject/Private/Serialization/DeferredMessageLog.h
@@ -0,0 +1,27 @@
+// Copyright 1998-2014 Epic Games, Inc. All Rights Reserved.
+
+/*=============================================================================
+	DeferredMessgaeLog.h: Unreal async loading log.
+=============================================================================*/
+
+#pragma once
+
+#include "MessageLog.h"
+
+class FDeferredMessageLog
+{
+	FName LogCategory;
+
+	static TMap<FName, TArray<TSharedRef<FTokenizedMessage>>*> Messages;
+
+	void AddMessage(TSharedRef<FTokenizedMessage>& Message);
+
+public:
+	FDeferredMessageLog(const FName& InLogCategory);
+	
+	TSharedRef<FTokenizedMessage> Info(const FText& Message);
+	TSharedRef<FTokenizedMessage> Warning(const FText& Message);
+	TSharedRef<FTokenizedMessage> Error(const FText& Message);
+
+	static void Flush();
+};
diff --git a/Engine/Source/Runtime/CoreUObject/Private/UObject/LinkerLoad.cpp b/Engine/Source/Runtime/CoreUObject/Private/UObject/LinkerLoad.cpp
--- a/Engine/Source/Runtime/CoreUObject/Private/UObject/LinkerLoad.cpp
+++ b/Engine/Source/Runtime/CoreUObject/Private/UObject/LinkerLoad.cpp
@@ -9,6 +9,7 @@
 #include "LinkerPlaceholderExportObject.h"
 #include "LinkerPlaceholderFunction.h"
 #include "LinkerManager.h"
+#include "Serialization/DeferredMessageLog.h"
 #include "UObject/UObjectThreadContext.h"
 #include "GatherableTextData.h"
 #include "Serialization/AsyncLoading.h"
@@ -84,7 +85,7 @@ void FLinkerLoad::CreateActiveRedirectsMap(const FString& GEngineIniName)
 		FConfigSection* PackageRedirects = GConfig->GetSectionPrivate( TEXT("/Script/Engine.Engine"), false, true, GEngineIniName );
 		if (PackageRedirects)
 		{
-			FMessageLog RedirectErrors(NAME_LoadErrors);
+			FDeferredMessageLog RedirectErrors(NAME_LoadErrors);
 
 			static FName ActiveClassRedirectsKey(TEXT("ActiveClassRedirects"));
 			for( FConfigSection::TIterator It(*PackageRedirects); It; ++It )
@@ -2224,7 +2225,7 @@ FLinkerLoad::EVerifyResult FLinkerLoad::VerifyImport(int32 ImportIndex)
 				UClass* FindClass = ClassPackage ? FindObject<UClass>( ClassPackage, *OriginalImport.ClassName.ToString() ) : NULL;
 				if( GIsEditor && !IsRunningCommandlet())
 				{
-					FMessageLog LoadErrors(NAME_LoadErrors);
+					FDeferredMessageLog LoadErrors(NAME_LoadErrors);
 					// put something into the load warnings dialog, with any extra information from above (in WarningAppend)
 					TSharedRef<FTokenizedMessage> TokenizedMessage = LoadErrors.Error(FText());
 					TokenizedMessage->AddToken(FAssetNameToken::Create(LinkerRoot->GetName()));
@@ -3310,7 +3311,7 @@ bool FLinkerLoad::WillTextureBeLoaded( UClass* Class, int32 ExportIndex )
 UObject* FLinkerLoad::CreateExport( int32 Index )
 {
 	FScopedCreateExportCounter ScopedCounter( this, Index );
-	FMessageLog LoadErrors(NAME_LoadErrors);
+	FDeferredMessageLog LoadErrors(NAME_LoadErrors);
 
 	// Map the object into our table.
 	FObjectExport& Export = ExportMap[ Index ];
