Making FDefferedMessageLog slighly more thread safe

diff --git a/Engine/Source/Runtime/CoreUObject/Private/Serialization/DeferredMessageLog.cpp b/Engine/Source/Runtime/CoreUObject/Private/Serialization/DeferredMessageLog.cpp
--- a/Engine/Source/Runtime/CoreUObject/Private/Serialization/DeferredMessageLog.cpp
+++ b/Engine/Source/Runtime/CoreUObject/Private/Serialization/DeferredMessageLog.cpp
@@ -7,10 +7,13 @@
 #include "CoreUObjectPrivate.h"
 #include "Serialization/DeferredMessageLog.h"
 
+TMap<FName, TArray<TSharedRef<FTokenizedMessage>>*> FDeferredMessageLog::Messages;
+FCriticalSection FDeferredMessageLog::MessagesCritical;
 
 FDeferredMessageLog::FDeferredMessageLog(const FName& InLogCategory)
 : LogCategory(InLogCategory)
 {
+	FScopeLock MessagesLock(&MessagesCritical);
 	TArray<TSharedRef<FTokenizedMessage>>** ExistingCategoryMessages = Messages.Find(LogCategory);
 	if (!ExistingCategoryMessages)
 	{
@@ -21,6 +24,7 @@ FDeferredMessageLog::FDeferredMessageLog(const FName& InLogCategory)
 
 void FDeferredMessageLog::AddMessage(TSharedRef<FTokenizedMessage>& Message)
 {
+	FScopeLock MessagesLock(&MessagesCritical);
 	TArray<TSharedRef<FTokenizedMessage>>* CategoryMessages = Messages.FindRef(LogCategory);
 	check(CategoryMessages);
 	CategoryMessages->Add(Message);
@@ -49,6 +53,7 @@ TSharedRef<FTokenizedMessage> FDeferredMessageLog::Error(const FText& InMessage)
 
 void FDeferredMessageLog::Flush()
 {
+	FScopeLock MessagesLock(&MessagesCritical);
 	for (auto& CategoryMessages : Messages)
 	{
 		if (CategoryMessages.Value->Num())
@@ -60,4 +65,11 @@ void FDeferredMessageLog::Flush()
 	}
 }
 
-TMap<FName, TArray<TSharedRef<FTokenizedMessage>>*> FDeferredMessageLog::Messages;
+void FDeferredMessageLog::Cleanup()
+{
+	for (auto& CategoryMessages : Messages)
+	{
+		delete CategoryMessages.Value;
+	}
+	Messages.Empty();
+}
diff --git a/Engine/Source/Runtime/CoreUObject/Private/Serialization/DeferredMessageLog.h b/Engine/Source/Runtime/CoreUObject/Private/Serialization/DeferredMessageLog.h
--- a/Engine/Source/Runtime/CoreUObject/Private/Serialization/DeferredMessageLog.h
+++ b/Engine/Source/Runtime/CoreUObject/Private/Serialization/DeferredMessageLog.h
@@ -8,11 +8,18 @@
 
 #include "MessageLog.h"
 
+/**
+ * Thread safe proxy for the FMessageLog while performing async loading.
+ * Makes sure the messages does not get added to the log until async loading is 
+ * finished to prevent modules from being loaded outside of game thread.
+ * Also makes sure the messages are added to the message queue in a thread-safe way.
+*/
 class FDeferredMessageLog
 {
 	FName LogCategory;
 
 	static TMap<FName, TArray<TSharedRef<FTokenizedMessage>>*> Messages;
+	static FCriticalSection MessagesCritical;
 
 	void AddMessage(TSharedRef<FTokenizedMessage>& Message);
 
@@ -24,4 +31,5 @@ public:
 	TSharedRef<FTokenizedMessage> Error(const FText& Message);
 
 	static void Flush();
+	static void Cleanup();
 };
diff --git a/Engine/Source/Runtime/CoreUObject/Private/UObject/Obj.cpp b/Engine/Source/Runtime/CoreUObject/Private/UObject/Obj.cpp
--- a/Engine/Source/Runtime/CoreUObject/Private/UObject/Obj.cpp
+++ b/Engine/Source/Runtime/CoreUObject/Private/UObject/Obj.cpp
@@ -15,6 +15,7 @@
 #include "HotReloadInterface.h"
 #include "UObject/TlsObjectInitializers.h"
 #include "UObject/UObjectThreadContext.h"
+#include "Serialization/DeferredMessageLog.h"
 
 DEFINE_LOG_CATEGORY(LogObj);
 
@@ -3730,6 +3731,7 @@ void StaticExit()
 	UObjectBaseShutdown();
 	// Empty arrays to prevent falsely-reported memory leaks.
 	FUObjectThreadContext::Get().ObjLoaded.Empty();
+	FDeferredMessageLog::Cleanup();
 
 	UE_LOG(LogExit, Log, TEXT("Object subsystem successfully closed.") );
 }
