Adding missing Unhash Object stat

diff --git a/Engine/Source/Runtime/CoreUObject/Private/UObject/UObjectHash.cpp b/Engine/Source/Runtime/CoreUObject/Private/UObject/UObjectHash.cpp
--- a/Engine/Source/Runtime/CoreUObject/Private/UObject/UObjectHash.cpp
+++ b/Engine/Source/Runtime/CoreUObject/Private/UObject/UObjectHash.cpp
@@ -884,6 +884,8 @@ void HashObject(UObjectBase* Object)
  */
 void UnhashObject(UObjectBase* Object)
 {
+	SCOPE_CYCLE_COUNTER(STAT_Hash_UnhashObject);
+
 	FName Name = Object->GetFName();
 	if (Name != NAME_None)
 	{
