[UE-17464] Making sure we defer serialization for CDO that have had their initialization deferred (need to initialize CDOs with their parent's values before we load overridden values, not vise versa).

diff --git a/Engine/Source/Runtime/CoreUObject/Private/UObject/LinkerLoad.cpp b/Engine/Source/Runtime/CoreUObject/Private/UObject/LinkerLoad.cpp
--- a/Engine/Source/Runtime/CoreUObject/Private/UObject/LinkerLoad.cpp
+++ b/Engine/Source/Runtime/CoreUObject/Private/UObject/LinkerLoad.cpp
@@ -3022,14 +3022,30 @@ void FLinkerLoad::Preload( UObject* Object )
 		if (Object->GetLinker() == this)
 		{
 #if USE_CIRCULAR_DEPENDENCY_LOAD_DEFERRING
+			// Because of delta serialization, we require that a parent's CDO be 
+			// fully serialized before its children's CDOs are created. However, 
+			// due to cyclic parent/child dependencies, we have some cases where 
+			// the linker breaks that expected behavior. In those cases, we 
+			// defer the child's initialization (i.e. defer copying of parent  
+			// property values, etc.), and wait until we can guarantee that the 
+			// parent CDO has been fully loaded.
+			//
+			// In a normal scenario, the order of property initialization is:
+			// Creation (zeroed) -> Initialization (copied super's values) -> Serialization (overridden values loaded)
+			// When the initialization has been deferred we have to make sure to
+			// defer serialization here as well (don't worry, it will be invoked 
+			// again from FinalizeBlueprint()->ResolveDeferredExports())
+			if (Object->HasAnyFlags(RF_ClassDefaultObject) && FDeferredObjInitializerTracker::IsCdoDeferred(Object->GetClass()))
+			{
+				return;
+			}
 			// if this is an inherited sub-object on a CDO, and that CDO has had
-			// its initialization deferred, then we shouldn't serialize in data 
-			// for this quite yet... not until the CDO owner has had a chance to
-			// initialize itself (because, as part of CDO initialization, 
-			// inherited sub-objects get filled in with values inherited from 
-			// the parent class... it is expected that this happens prior to 
-			// sub-object serialization)
-			if (Object->HasAnyFlags(RF_DefaultSubObject) && FDeferredObjInitializerTracker::DeferSubObjectPreload(Object))
+			// its initialization deferred (for reasons explained above), then 
+			// we shouldn't serialize in data for this quite yet... not until 
+			// its owner has had a chance to initialize itself (because, as part
+			// of CDO initialization, inherited sub-objects get filled in with 
+			// values inherited from the super)
+			else if (Object->HasAnyFlags(RF_DefaultSubObject) && FDeferredObjInitializerTracker::DeferSubObjectPreload(Object))
 			{
 				// don't worry, FDeferredObjInitializerTracker::DeferSubObjectPreload() 
 				// should have cached this object, and it will run Preload() on 
diff --git a/Engine/Source/Runtime/CoreUObject/Private/UObject/UObjectGlobals.cpp b/Engine/Source/Runtime/CoreUObject/Private/UObject/UObjectGlobals.cpp
--- a/Engine/Source/Runtime/CoreUObject/Private/UObject/UObjectGlobals.cpp
+++ b/Engine/Source/Runtime/CoreUObject/Private/UObject/UObjectGlobals.cpp
@@ -2198,13 +2198,13 @@ FObjectInitializer::~FObjectInitializer()
 	if (bIsCDO && (ObjectArchetype != nullptr) && !FBlueprintSupport::IsDeferredCDOInitializationDisabled())
 	{
 		UClass* ArchetypeClass = ObjectArchetype->GetClass();
+		const bool bSuperCDONeedsLoad = ObjectArchetype->HasAnyFlags(RF_NeedLoad) ||
+			(ArchetypeClass->GetLinker() && ArchetypeClass->GetLinker()->IsBlueprintFinalizationPending()) ||
+			FDeferredObjInitializerTracker::IsCdoDeferred(ArchetypeClass);
+
 		// if this is a blueprint CDO that derives from another blueprint, and 
 		// that parent (archetype) CDO isn't fully serialized
-		if ( !Class->HasAnyFlags(RF_Native) && !ArchetypeClass->HasAnyFlags(RF_Native) && 
-			(	ObjectArchetype->HasAnyFlags(RF_NeedLoad) ||
-				(ArchetypeClass->GetLinker() && ArchetypeClass->GetLinker()->IsBlueprintFinalizationPending()) ||
-				FDeferredObjInitializerTracker::IsCdoDeferred(ArchetypeClass))
-			)
+		if (!Class->HasAnyFlags(RF_Native) && !ArchetypeClass->HasAnyFlags(RF_Native) && bSuperCDONeedsLoad)
 		{
 			FLinkerLoad* ClassLinker = Class->GetLinker();
 			if ((ClassLinker != nullptr) && (ClassLinker->LoadFlags & LOAD_DeferDependencyLoads) != 0x00)
