#UE4 Tightening up accounting for exclusive load package time tracking. There was a case where the value returned from IsAsset in PostLoad was changing for some objects. Now only objects where IsAsset returns true before PostLoad is called are tracked. Also moved asset type tracking to be evaluated at report time instead of tracking time.

diff --git a/Engine/Source/Runtime/CoreUObject/Private/Misc/ExclusiveLoadPackageTimeTracker.cpp b/Engine/Source/Runtime/CoreUObject/Private/Misc/ExclusiveLoadPackageTimeTracker.cpp
--- a/Engine/Source/Runtime/CoreUObject/Private/Misc/ExclusiveLoadPackageTimeTracker.cpp
+++ b/Engine/Source/Runtime/CoreUObject/Private/Misc/ExclusiveLoadPackageTimeTracker.cpp
@@ -68,31 +68,18 @@ void FExclusiveLoadPackageTimeTracker::InternalPopLoadPackage(UPackage* LoadedPa
 			}
 		}
 
-		bool bIncrementCount = false;
 		double InclusiveTime = CurrentTime - Time.OriginalStartTime;
 		FLoadTime* LoadTimePtr = LoadTimes.Find(Time.TimeName);
 		if ( LoadTimePtr )
 		{
+			// Use the most recently discovered classname. This may be "Unknown" during async loading but future pops will correct it.
+			LoadTimePtr->AssetClass = ClassName;
 			LoadTimePtr->ExclusiveTime += Time.ExclusiveTime;
 			LoadTimePtr->InclusiveTime += InclusiveTime;
 		}
 		else
 		{
-			LoadTimes.Add(Time.TimeName, FLoadTime(Time.TimeName, Time.ExclusiveTime, InclusiveTime));
-			bIncrementCount = true;
-		}
-
-		if (IsPackageLoadTime(Time))
-		{
-			FTimeCount& TypeTimeCount = AssetTypeLoadTimes.FindOrAdd(ClassName);
-			TypeTimeCount.AssetClass = ClassName;
-			TypeTimeCount.Time.ExclusiveTime += Time.ExclusiveTime;
-			TypeTimeCount.Time.InclusiveTime += InclusiveTime;
-
-			if (bIncrementCount)
-			{
-				TypeTimeCount.Count++;
-			}
+			LoadTimes.Add(Time.TimeName, FLoadTime(Time.TimeName, ClassName, Time.ExclusiveTime, InclusiveTime));
 		}
 
 		TimeStack.RemoveAt(TimeStack.Num() - 1, 1, false);
@@ -117,20 +104,30 @@ void FExclusiveLoadPackageTimeTracker::InternalDumpReport() const
 	FName LongestLoadName;
 	double TotalLoadTime = 0;
 	const double SlowAssetThreshold = 0.10f;
+	TMap<FName, FTimeCount> AssetTypeLoadTimes;
 	for (const auto& LoadTimeIt : LoadTimes)
 	{
-		if (LoadTimeIt.Value.ExclusiveTime > LongestLoadTime && IsPackageLoadTime(LoadTimeIt.Value))
+		const FLoadTime& Time = LoadTimeIt.Value;
+		if (IsPackageLoadTime(Time))
 		{
-			LongestLoadName = LoadTimeIt.Key;
-			LongestLoadTime = LoadTimeIt.Value.ExclusiveTime;
-		}
+			if (Time.ExclusiveTime > LongestLoadTime)
+			{
+				LongestLoadName = LoadTimeIt.Key;
+				LongestLoadTime = Time.ExclusiveTime;
+			}
 
-		if (LoadTimeIt.Value.ExclusiveTime > SlowAssetThreshold && IsPackageLoadTime(LoadTimeIt.Value))
-		{
-			SlowAssetTime += LoadTimeIt.Value.ExclusiveTime;
+			if (Time.ExclusiveTime > SlowAssetThreshold)
+			{
+				SlowAssetTime += Time.ExclusiveTime;
+			}
+
+			FTimeCount& TypeTimeCount = AssetTypeLoadTimes.FindOrAdd(Time.AssetClass);
+			TypeTimeCount.Time.AssetClass = Time.AssetClass;
+			TypeTimeCount.Time.ExclusiveTime += Time.ExclusiveTime;
+			TypeTimeCount.Count++;
 		}
 
-		TotalLoadTime += LoadTimeIt.Value.ExclusiveTime;
+		TotalLoadTime += Time.ExclusiveTime;
 	}
 
 	const FLoadTime* EndLoadTime = LoadTimes.Find(EndLoadName);
@@ -148,7 +145,7 @@ void FExclusiveLoadPackageTimeTracker::InternalDumpReport() const
 	SortedAssetTypeLoadTimes.Sort([](const FTimeCount& A, const FTimeCount& B){ return A.Time.ExclusiveTime >= B.Time.ExclusiveTime; });
 	for (const FTimeCount& TimeCount : SortedAssetTypeLoadTimes)
 	{
-		UE_LOG(LogLoad, Log, TEXT("    %.3f: %s (%d packages, %.1fms per package)"), TimeCount.Time.ExclusiveTime, *TimeCount.AssetClass.ToString(), TimeCount.Count, TimeCount.Time.ExclusiveTime / TimeCount.Count * 1000);
+		UE_LOG(LogLoad, Log, TEXT("    %.3f: %s (%d packages, %.1fms per package)"), TimeCount.Time.ExclusiveTime, *TimeCount.Time.AssetClass.ToString(), TimeCount.Count, TimeCount.Time.ExclusiveTime / TimeCount.Count * 1000);
 	}
 
 	// Assets faster than this will be excluded
diff --git a/Engine/Source/Runtime/CoreUObject/Private/UObject/Obj.cpp b/Engine/Source/Runtime/CoreUObject/Private/UObject/Obj.cpp
--- a/Engine/Source/Runtime/CoreUObject/Private/UObject/Obj.cpp
+++ b/Engine/Source/Runtime/CoreUObject/Private/UObject/Obj.cpp
@@ -720,9 +720,10 @@ void UObject::ConditionalPostLoad()
 
 		ConditionalPostLoadSubobjects();
 
-		FExclusiveLoadPackageTimeTracker::PushPostLoad(this);
-		PostLoad();
-		FExclusiveLoadPackageTimeTracker::PopPostLoad(this);
+		{
+			FExclusiveLoadPackageTimeTracker::FScopedPostLoadTracker Tracker(this);
+			PostLoad();
+		}
 
 #if !(UE_BUILD_SHIPPING || UE_BUILD_TEST)
 		if (ThreadContext.DebugPostLoad.Contains(this))
diff --git a/Engine/Source/Runtime/CoreUObject/Private/UObject/UObjectGlobals.cpp b/Engine/Source/Runtime/CoreUObject/Private/UObject/UObjectGlobals.cpp
--- a/Engine/Source/Runtime/CoreUObject/Private/UObject/UObjectGlobals.cpp
+++ b/Engine/Source/Runtime/CoreUObject/Private/UObject/UObjectGlobals.cpp
@@ -958,11 +958,8 @@ UPackage* LoadPackageInternal(UPackage* InOuter, const TCHAR* InLongPackageName,
 			return Result;
 		}
 
-		if ( Result )
-		{
-			// We have begun loading a package that we know the name of. Let the package time tracker know.
-			FExclusiveLoadPackageTimeTracker::PushLoadPackage(Result->GetFName());
-		}
+		// The time tracker keeps track of time spent in LoadPackage.
+		FExclusiveLoadPackageTimeTracker::FScopedPackageTracker Tracker(Result);
 
 #if !WITH_EDITOR
 		static auto CVarPreloadDependencies = IConsoleManager::Get().FindConsoleVariable(TEXT("s.PreloadPackageDependencies"));
@@ -1047,12 +1044,6 @@ UPackage* LoadPackageInternal(UPackage* InOuter, const TCHAR* InLongPackageName,
 
 		EndLoadAndCopyLocalizationGatherFlag();
 
-		if ( Result )
-		{
-			// We are done loading the package. Let the package time tracker know.
-			FExclusiveLoadPackageTimeTracker::PopLoadPackage(Result);
-		}
-
 #if WITH_EDITOR
 		GIsEditorLoadingPackage = *IsEditorLoadingPackage;
 #endif
@@ -1229,13 +1220,6 @@ void EndLoad()
 		return;
 	}
 
-	// The package time tracker keeps track of time spent in EndLoad.
-	struct FEndLoadScopedTracker
-	{
-		FEndLoadScopedTracker() { FExclusiveLoadPackageTimeTracker::PushEndLoad(); }
-		~FEndLoadScopedTracker() { FExclusiveLoadPackageTimeTracker::PopEndLoad(); }
-	} EndLoadScopedTracker;
-
 #if WITH_EDITOR
 	FScopedSlowTask SlowTask(0, NSLOCTEXT("Core", "PerformingPostLoad", "Performing post-load..."), ShouldReportProgress());
 
@@ -1244,6 +1228,9 @@ void EndLoad()
 
 	while (--ThreadContext.ObjBeginLoadCount == 0 && (ThreadContext.ObjLoaded.Num() || ThreadContext.ImportCount || ThreadContext.ForcedExportCount))
 	{
+		// The time tracker keeps track of time spent in EndLoad.
+		FExclusiveLoadPackageTimeTracker::FScopedEndLoadTracker Tracker;
+
 		// Make sure we're not recursively calling EndLoad as e.g. loading a config file could cause
 		// BeginLoad/EndLoad to be called.
 		ThreadContext.ObjBeginLoadCount++;
diff --git a/Engine/Source/Runtime/CoreUObject/Public/Misc/ExclusiveLoadPackageTimeTracker.h b/Engine/Source/Runtime/CoreUObject/Public/Misc/ExclusiveLoadPackageTimeTracker.h
--- a/Engine/Source/Runtime/CoreUObject/Public/Misc/ExclusiveLoadPackageTimeTracker.h
+++ b/Engine/Source/Runtime/CoreUObject/Public/Misc/ExclusiveLoadPackageTimeTracker.h
@@ -14,6 +14,57 @@
 class FExclusiveLoadPackageTimeTracker
 {
 public:
+	/** Scoped helper for LoadPackage */
+	struct FScopedPackageTracker
+	{
+		FScopedPackageTracker(UPackage* InPackageToTrack)
+			: PackageToTrack(InPackageToTrack)
+		{
+			FExclusiveLoadPackageTimeTracker::PushLoadPackage(PackageToTrack ? PackageToTrack->GetFName() : NAME_None);
+		}
+		FScopedPackageTracker(FName PackageNameToTrack)
+			: PackageToTrack(nullptr)
+		{
+			FExclusiveLoadPackageTimeTracker::PushLoadPackage(PackageNameToTrack);
+		}
+		~FScopedPackageTracker()
+		{
+			FExclusiveLoadPackageTimeTracker::PopLoadPackage(PackageToTrack);
+		}
+		UPackage* PackageToTrack;
+	};
+
+	/** Scoped helper for PostLoad */
+	struct FScopedPostLoadTracker
+	{
+		FScopedPostLoadTracker(UObject* InPostLoadObject)
+		{
+			if (FExclusiveLoadPackageTimeTracker::PushPostLoad(InPostLoadObject))
+			{
+				PostLoadObject = InPostLoadObject;
+			}
+			else
+			{
+				PostLoadObject = nullptr;
+			}
+		}
+		~FScopedPostLoadTracker()
+		{
+			if (PostLoadObject)
+			{
+				FExclusiveLoadPackageTimeTracker::PopPostLoad(PostLoadObject);
+			}
+		}
+		UObject* PostLoadObject;
+	};
+
+	/** Scoped helper for EndLoad */
+	struct FScopedEndLoadTracker
+	{
+		FScopedEndLoadTracker() { FExclusiveLoadPackageTimeTracker::PushEndLoad(); }
+		~FScopedEndLoadTracker() { FExclusiveLoadPackageTimeTracker::PopEndLoad(); }
+	};
+
 	/** Starts a time for the specified package name. */
 	FORCEINLINE static void PushLoadPackage(FName PackageName)
 	{
@@ -31,24 +82,23 @@ public:
 	}
 
 	/** Starts a time for the specified package name. */
-	FORCEINLINE static void PushPostLoad(UObject* PostLoadObject)
+	FORCEINLINE static bool PushPostLoad(UObject* PostLoadObject)
 	{
 #if WITH_LOADPACKAGE_TIME_TRACKER
 		if (PostLoadObject && PostLoadObject->IsAsset())
 		{
 			Get().InternalPushLoadPackage(PostLoadObject->GetOutermost()->GetFName());
+			return true;
 		}
 #endif
+		return false;
 	}
 
 	/** Records a time and stats for the loaded package. Optionally provide the loaded asset if it is known, otherwise the asset in the package is detected. */
 	FORCEINLINE static void PopPostLoad(UObject* PostLoadObject)
 	{
 #if WITH_LOADPACKAGE_TIME_TRACKER
-		if (PostLoadObject && PostLoadObject->IsAsset())
-		{
-			Get().InternalPopLoadPackage(nullptr, PostLoadObject);
-		}
+		Get().InternalPopLoadPackage(nullptr, PostLoadObject);
 #endif
 	}
 
@@ -105,6 +155,7 @@ private:
 	struct FLoadTime
 	{
 		FName TimeName;
+		FName AssetClass;
 		double ExclusiveTime;
 		double InclusiveTime;
 		double LastStartTime;
@@ -118,7 +169,7 @@ private:
 			: TimeName(InName), ExclusiveTime(0), InclusiveTime(0), LastStartTime(InStartTime), OriginalStartTime(InStartTime)
 		{}
 
-		FLoadTime(FName InName, double InExclusive, double InInclusive)
+		FLoadTime(FName InName, FName InAssetClass, double InExclusive, double InInclusive)
 			: TimeName(InName), ExclusiveTime(InExclusive), InclusiveTime(InInclusive), LastStartTime(0), OriginalStartTime(0)
 		{}
 	};
@@ -126,7 +177,6 @@ private:
 	/** Struct to assemble times for assets loaded by class */
 	struct FTimeCount
 	{
-		FName AssetClass;
 		FLoadTime Time;
 		int32 Count;
 
@@ -165,7 +215,6 @@ private:
 
 	/** A couple maps used to make the report and to report calculated times */
 	TMap<FName, FLoadTime> LoadTimes;
-	TMap<FName, FTimeCount> AssetTypeLoadTimes;
 
 	/** The amount of time spent in this singleton keeping track of load times. This overhead is included in the report. */
 	double TrackerOverhead;
