#UE4 Minor data structure cleanup after CL#2629314

diff --git a/Engine/Source/Runtime/CoreUObject/Private/Misc/ExclusiveLoadPackageTimeTracker.cpp b/Engine/Source/Runtime/CoreUObject/Private/Misc/ExclusiveLoadPackageTimeTracker.cpp
--- a/Engine/Source/Runtime/CoreUObject/Private/Misc/ExclusiveLoadPackageTimeTracker.cpp
+++ b/Engine/Source/Runtime/CoreUObject/Private/Misc/ExclusiveLoadPackageTimeTracker.cpp
@@ -99,6 +99,16 @@ void FExclusiveLoadPackageTimeTracker::InternalDumpReport() const
 {
 	FScopeLock Lock(&TimesCritical);
 
+	/** Struct to assemble times for assets loaded by class */
+	struct FTimeCount
+	{
+		FName AssetClass;
+		double ExclusiveTime;
+		int32 Count;
+
+		FTimeCount() : ExclusiveTime(0), Count(0) {}
+	};
+
 	double LongestLoadTime = 0;
 	double SlowAssetTime = 0;
 	FName LongestLoadName;
@@ -122,8 +132,8 @@ void FExclusiveLoadPackageTimeTracker::InternalDumpReport() const
 			}
 
 			FTimeCount& TypeTimeCount = AssetTypeLoadTimes.FindOrAdd(Time.AssetClass);
-			TypeTimeCount.Time.AssetClass = Time.AssetClass;
-			TypeTimeCount.Time.ExclusiveTime += Time.ExclusiveTime;
+			TypeTimeCount.AssetClass = Time.AssetClass;
+			TypeTimeCount.ExclusiveTime += Time.ExclusiveTime;
 			TypeTimeCount.Count++;
 		}
 
@@ -142,10 +152,10 @@ void FExclusiveLoadPackageTimeTracker::InternalDumpReport() const
 	UE_LOG(LogLoad, Log, TEXT("Dumping asset type load times sorted by exclusive time:"));
 	TArray<FTimeCount> SortedAssetTypeLoadTimes;
 	AssetTypeLoadTimes.GenerateValueArray(SortedAssetTypeLoadTimes);
-	SortedAssetTypeLoadTimes.Sort([](const FTimeCount& A, const FTimeCount& B){ return A.Time.ExclusiveTime >= B.Time.ExclusiveTime; });
+	SortedAssetTypeLoadTimes.Sort([](const FTimeCount& A, const FTimeCount& B){ return A.ExclusiveTime >= B.ExclusiveTime; });
 	for (const FTimeCount& TimeCount : SortedAssetTypeLoadTimes)
 	{
-		UE_LOG(LogLoad, Log, TEXT("    %.3f: %s (%d packages, %.1fms per package)"), TimeCount.Time.ExclusiveTime, *TimeCount.Time.AssetClass.ToString(), TimeCount.Count, TimeCount.Time.ExclusiveTime / TimeCount.Count * 1000);
+		UE_LOG(LogLoad, Log, TEXT("    %.3f: %s (%d packages, %.1fms per package)"), TimeCount.ExclusiveTime, *TimeCount.AssetClass.ToString(), TimeCount.Count, TimeCount.ExclusiveTime / TimeCount.Count * 1000);
 	}
 
 	// Assets faster than this will be excluded
diff --git a/Engine/Source/Runtime/CoreUObject/Public/Misc/ExclusiveLoadPackageTimeTracker.h b/Engine/Source/Runtime/CoreUObject/Public/Misc/ExclusiveLoadPackageTimeTracker.h
--- a/Engine/Source/Runtime/CoreUObject/Public/Misc/ExclusiveLoadPackageTimeTracker.h
+++ b/Engine/Source/Runtime/CoreUObject/Public/Misc/ExclusiveLoadPackageTimeTracker.h
@@ -174,15 +174,6 @@ private:
 		{}
 	};
 
-	/** Struct to assemble times for assets loaded by class */
-	struct FTimeCount
-	{
-		FLoadTime Time;
-		int32 Count;
-
-		FTimeCount() : Count(0) {}
-	};
-
 	FExclusiveLoadPackageTimeTracker();
 	COREUOBJECT_API static FExclusiveLoadPackageTimeTracker& Get();
 
