Fixing crash when destroying a linker in level streaming.

UE-19028

diff --git a/Engine/Source/Runtime/CoreUObject/Private/UObject/Linker.cpp b/Engine/Source/Runtime/CoreUObject/Private/UObject/Linker.cpp
--- a/Engine/Source/Runtime/CoreUObject/Private/UObject/Linker.cpp
+++ b/Engine/Source/Runtime/CoreUObject/Private/UObject/Linker.cpp
@@ -326,12 +326,16 @@ void ResetLoaders(UObject* InPkg)
 	FLinkerManager::Get().ResetLoaders(InPkg);
 }
 
-/** Deletes all linkers that have finished loading */
 void DeleteLoaders()
 {
 	FLinkerManager::Get().DeleteLinkers();
 }
 
+void DeleteLoader(FLinkerLoad* Loader)
+{
+	FLinkerManager::Get().RemoveLinker(Loader);
+}
+
 /**
  * Dissociates all linker import and forced export object references. This currently needs to 
  * happen as the referred objects might be destroyed at any time.
diff --git a/Engine/Source/Runtime/CoreUObject/Public/UObject/Linker.h b/Engine/Source/Runtime/CoreUObject/Public/UObject/Linker.h
--- a/Engine/Source/Runtime/CoreUObject/Public/UObject/Linker.h
+++ b/Engine/Source/Runtime/CoreUObject/Public/UObject/Linker.h
@@ -1814,7 +1814,7 @@ public:
 	/**
 	* Detaches linker from bulk data/ exports and removes itself from array of loaders.
 	*/
-	void Detach();
+	COREUOBJECT_API void Detach();
 
 private:
 
@@ -2264,6 +2264,9 @@ COREUOBJECT_API void ResetLoaders( UObject* InOuter );
 /** Deletes all linkers that have finished loading */
 COREUOBJECT_API void DeleteLoaders();
 
+/** Queues linker for deletion */
+COREUOBJECT_API void DeleteLoader(FLinkerLoad* Loader);
+
 /**
  * Dissociates all linker import and forced export object references. This currently needs to 
  * happen as the referred objects might be destroyed at any time.
diff --git a/Engine/Source/Runtime/Engine/Private/LevelStreaming.cpp b/Engine/Source/Runtime/Engine/Private/LevelStreaming.cpp
--- a/Engine/Source/Runtime/Engine/Private/LevelStreaming.cpp
+++ b/Engine/Source/Runtime/Engine/Private/LevelStreaming.cpp
@@ -598,7 +598,8 @@ void ULevelStreaming::AsyncLevelLoadComplete(const FName& InPackageName, UPackag
 					FLinkerLoad* PackageLinker = FLinkerLoad::FindExistingLinkerForPackage(LevelPackage);
 					if (PackageLinker)
 					{
-						delete PackageLinker;
+						PackageLinker->Detach();
+						DeleteLoader(PackageLinker);
 						PackageLinker = nullptr;
 					}
 
