Merging CL #2632216: Fix for crash when doing reference search

diff --git a/Engine/Source/Runtime/CoreUObject/Private/UObject/ReferenceChainSearch.cpp b/Engine/Source/Runtime/CoreUObject/Private/UObject/ReferenceChainSearch.cpp
--- a/Engine/Source/Runtime/CoreUObject/Private/UObject/ReferenceChainSearch.cpp
+++ b/Engine/Source/Runtime/CoreUObject/Private/UObject/ReferenceChainSearch.cpp
@@ -638,6 +638,10 @@ bool FReferenceChainSearch::ProcessObject( UObject* CurrentObject )
 				AddToReferenceList(ReferenceList, Ref);
 			}
 		}
+		else if (REFERENCE_INFO.Type == GCRT_EndOfPointer)
+		{
+			TokenReturnCount = REFERENCE_INFO.ReturnCount;
+		}
 		else if( REFERENCE_INFO.Type == GCRT_EndOfStream )
 		{
 			// Break out of loop.
