Added UPropertyInterface::AddReferencedObjects function.
Added CPPF_BlueprintCppBackend flag.

diff --git a/Engine/Source/Runtime/CoreUObject/Private/UObject/PropertyInterface.cpp b/Engine/Source/Runtime/CoreUObject/Private/UObject/PropertyInterface.cpp
--- a/Engine/Source/Runtime/CoreUObject/Private/UObject/PropertyInterface.cpp
+++ b/Engine/Source/Runtime/CoreUObject/Private/UObject/PropertyInterface.cpp
@@ -54,9 +54,12 @@ FString UInterfaceProperty::GetCPPType( FString* ExtendedTypeText/*=NULL*/, uint
 	if ( ExtendedTypeText != NULL )
 	{
 		UClass* ExportClass = InterfaceClass;
-		while ( ExportClass && !ExportClass->HasAnyClassFlags(CLASS_Native) )
+		if (0 == (CPPF_BlueprintCppBackend & CPPExportFlags))
 		{
-			ExportClass = ExportClass->GetSuperClass();
+			while (ExportClass && !ExportClass->HasAnyClassFlags(CLASS_Native))
+			{
+				ExportClass = ExportClass->GetSuperClass();
+			}
 		}
 		check(ExportClass);
 		check(ExportClass->HasAnyClassFlags(CLASS_Interface));
@@ -232,9 +235,6 @@ void UInterfaceProperty::Serialize( FArchive& Ar )
 		{
 			checkf(false, TEXT("Interface property tried to serialize a missing interface.  Did you remove a native class and not fully recompile?"));
 		}
-
-		int32 a = 0;
-		++a;
  	}
 }
 
@@ -260,6 +260,13 @@ bool UInterfaceProperty::SameType(const UProperty* Other) const
 	return Super::SameType(Other) && (InterfaceClass == ((UInterfaceProperty*)Other)->InterfaceClass);
 }
 
+void UInterfaceProperty::AddReferencedObjects(UObject* InThis, FReferenceCollector& Collector)
+{
+	UInterfaceProperty* This = CastChecked<UInterfaceProperty>(InThis);
+	Collector.AddReferencedObject(This->InterfaceClass, This);
+	Super::AddReferencedObjects(This, Collector);
+}
+
 IMPLEMENT_CORE_INTRINSIC_CLASS(UInterfaceProperty, UProperty,
 	{
 		Class->EmitObjectReference(STRUCT_OFFSET(UInterfaceProperty, InterfaceClass), TEXT("InterfaceClass"));
diff --git a/Engine/Source/Runtime/CoreUObject/Public/UObject/UnrealType.h b/Engine/Source/Runtime/CoreUObject/Public/UObject/UnrealType.h
--- a/Engine/Source/Runtime/CoreUObject/Public/UObject/UnrealType.h
+++ b/Engine/Source/Runtime/CoreUObject/Public/UObject/UnrealType.h
@@ -31,6 +31,8 @@ enum EPropertyExportCPPFlags
 	CPPF_NoConst					=	0x00000010,
 	/** No reference '&' sign */
 	CPPF_NoRef						=	0x00000020,
+	/** Blueprint kompiler generated C++ code */
+	CPPF_BlueprintCppBackend		=	0x00000080,
 };
 
 namespace EExportedDeclaration
@@ -2205,6 +2207,7 @@ public:
 	virtual void Serialize( FArchive& Ar ) override;
 	virtual void EmitReferenceInfo(UClass& OwnerClass, int32 BaseOffset) override;
 	virtual void BeginDestroy() override;
+	static void AddReferencedObjects(UObject* InThis, FReferenceCollector& Collector);
 	// End of UObject interface
 
 	/**
