Checking in Mikolaj Sieluzycki's change to fix up callsites to ask for the value of an enum entry instead of the index. The value and the index are the same unless the enum specified element values in its declaration (i.e. enum MyEnum { FirstValue = 1 })

diff --git a/Engine/Source/Editor/AddContentDialog/Private/ContentSourceProviders/FeaturePack/FeaturePackContentSource.cpp b/Engine/Source/Editor/AddContentDialog/Private/ContentSourceProviders/FeaturePack/FeaturePackContentSource.cpp
--- a/Engine/Source/Editor/AddContentDialog/Private/ContentSourceProviders/FeaturePack/FeaturePackContentSource.cpp
+++ b/Engine/Source/Editor/AddContentDialog/Private/ContentSourceProviders/FeaturePack/FeaturePackContentSource.cpp
@@ -832,8 +832,8 @@ bool FFeaturePackContentSource::ParseManifestString(const FString& ManifestStrin
 
 	FString CategoryString = ManifestObject->GetStringField("Category");
 	UEnum* Enum = FindObjectChecked<UEnum>(ANY_PACKAGE, TEXT("EContentSourceCategory"));
-	int32 Index = Enum->FindEnumIndex(FName(*CategoryString));
-	Category = Index != INDEX_NONE ? (EContentSourceCategory)Index : EContentSourceCategory::Unknown;
+	int32 EnumValue = Enum->GetValueByName(FName(*CategoryString));
+	Category = EnumValue != INDEX_NONE ? (EContentSourceCategory)EnumValue : EContentSourceCategory::Unknown;
 
 	// Thumbnail filename
 	IconFilename = ManifestObject->GetStringField("Thumbnail");
@@ -872,8 +872,8 @@ bool FFeaturePackContentSource::ParseManifestString(const FString& ManifestStrin
 			for (TSharedPtr<FJsonValue> DetailValue : EachAdditionalPack->GetArrayField("DetailLevels"))
 			{
 				const FString DetailString = DetailValue->AsString();
-				int32 eIndex = DetailEnum->FindEnumIndex(FName(*DetailString));
-				EFeaturePackDetailLevel EachLevel = Index != INDEX_NONE ? (EFeaturePackDetailLevel)eIndex : EFeaturePackDetailLevel::Standard;
+				int32 eValue = DetailEnum->GetValueByName(FName(*DetailString));
+				EFeaturePackDetailLevel EachLevel = eValue != INDEX_NONE ? (EFeaturePackDetailLevel)eValue : EFeaturePackDetailLevel::Standard;
 				Levels.AddUnique(EachLevel);
 			}
 			AdditionalFeaturePacks.Add(FFeaturePackLevelSet(MountName,Levels));
diff --git a/Engine/Source/Editor/BlueprintGraph/Private/K2Node_EaseFunction.cpp b/Engine/Source/Editor/BlueprintGraph/Private/K2Node_EaseFunction.cpp
--- a/Engine/Source/Editor/BlueprintGraph/Private/K2Node_EaseFunction.cpp
+++ b/Engine/Source/Editor/BlueprintGraph/Private/K2Node_EaseFunction.cpp
@@ -444,7 +444,7 @@ void UK2Node_EaseFunction::RefreshPinVisibility()
 	const auto EaseFuncPin = GetEaseFuncPin();
 	UEnum * Enum = FindObject<UEnum>(ANY_PACKAGE, TEXT("EEasingFunc"), true);
 	check(Enum != NULL);
-	const int32 NewEasingFunc = CanCustomizeCurve() ? Enum->FindEnumIndex(*EaseFuncPin->DefaultValue) : -1;
+	const int32 NewEasingFunc = CanCustomizeCurve() ? Enum->GetValueByName(*EaseFuncPin->DefaultValue) : -1;
 	
 	// Early exit in case no changes are required
 	const UEdGraphSchema_K2* K2Schema = Cast<const UEdGraphSchema_K2>(GetSchema());
diff --git a/Engine/Source/Editor/BlueprintGraph/Private/K2Node_GetEnumeratorName.cpp b/Engine/Source/Editor/BlueprintGraph/Private/K2Node_GetEnumeratorName.cpp
--- a/Engine/Source/Editor/BlueprintGraph/Private/K2Node_GetEnumeratorName.cpp
+++ b/Engine/Source/Editor/BlueprintGraph/Private/K2Node_GetEnumeratorName.cpp
@@ -158,7 +158,7 @@ void UK2Node_GetEnumeratorName::ExpandNode(class FKismetCompilerContext& Compile
 	{
 		//MAKE LITERAL BYTE FROM LITERAL ENUM
 		const FString EnumLiteral = IndexPin->GetDefaultAsString();
-		const int32 NumericValue = Enum->FindEnumIndex(*EnumLiteral);
+		const int32 NumericValue = Enum->GetValueByName(*EnumLiteral);
 		if (NumericValue == INDEX_NONE) 
 		{
 			CompilerContext.MessageLog.Error(*FString::Printf(*NSLOCTEXT("K2Node", "GetEnumeratorNam_Error_InvalidName", "@@ has invalid enum value '%s'").ToString(), *EnumLiteral), this);
diff --git a/Engine/Source/Editor/BlueprintGraph/Private/K2Node_Select.cpp b/Engine/Source/Editor/BlueprintGraph/Private/K2Node_Select.cpp
--- a/Engine/Source/Editor/BlueprintGraph/Private/K2Node_Select.cpp
+++ b/Engine/Source/Editor/BlueprintGraph/Private/K2Node_Select.cpp
@@ -217,8 +217,8 @@ public:
 
 				if (UEnum* NodeEnum = SelectNode->GetEnum())
 				{
-					int32 EnumIndex = NodeEnum->FindEnumIndex(FName(*OptionPins[OptionIdx]->PinName));
-					LiteralTerm->Name = FString::Printf(TEXT("%d"), EnumIndex);
+					int32 EnumValue = NodeEnum->GetValueByName(FName(*OptionPins[OptionIdx]->PinName));
+					LiteralTerm->Name = FString::Printf(TEXT("%d"), EnumValue);
 				}
 				else
 				{
diff --git a/Engine/Source/Editor/BlueprintGraph/Private/K2Node_SwitchEnum.cpp b/Engine/Source/Editor/BlueprintGraph/Private/K2Node_SwitchEnum.cpp
--- a/Engine/Source/Editor/BlueprintGraph/Private/K2Node_SwitchEnum.cpp
+++ b/Engine/Source/Editor/BlueprintGraph/Private/K2Node_SwitchEnum.cpp
@@ -171,10 +171,10 @@ UK2Node::ERedirectType UK2Node_SwitchEnum::DoPinsMatchForReconstruction(const UE
 	UK2Node::ERedirectType ReturnValue = Super::DoPinsMatchForReconstruction(NewPin, NewPinIndex, OldPin, OldPinIndex);
 	if (ReturnValue == UK2Node::ERedirectType_None && Enum && OldPinIndex > 2 && NewPinIndex > 2)
 	{
-		int32 OldIndex = Enum->FindEnumIndex(FName(*OldPin->PinName));
-		int32 NewIndex = Enum->FindEnumIndex(FName(*NewPin->PinName));
+		int32 OldValue = Enum->GetValueByName(FName(*OldPin->PinName));
+		int32 NewValue = Enum->GetValueByName(FName(*NewPin->PinName));
 		// This handles redirects properly
-		if (OldIndex == NewIndex && OldIndex != INDEX_NONE)
+		if (OldValue == NewValue && OldValue != INDEX_NONE)
 		{
 			ReturnValue = UK2Node::ERedirectType_Name;
 		}
diff --git a/Engine/Source/Editor/ContentBrowser/Private/AssetContextMenu.cpp b/Engine/Source/Editor/ContentBrowser/Private/AssetContextMenu.cpp
--- a/Engine/Source/Editor/ContentBrowser/Private/AssetContextMenu.cpp
+++ b/Engine/Source/Editor/ContentBrowser/Private/AssetContextMenu.cpp
@@ -572,7 +572,7 @@ bool FAssetContextMenu::AddDocumentationMenuOptions(FMenuBuilder& MenuBuilder)
 
 						UEnum* BlueprintTypeEnum = FindObject<UEnum>(ANY_PACKAGE, TEXT("EBlueprintType"), true);
 						auto EnumString = SelectedAssets[0].TagsAndValues.Find(GET_MEMBER_NAME_CHECKED(UBlueprint,BlueprintType));
-						EBlueprintType BlueprintType = (EnumString ? (EBlueprintType)BlueprintTypeEnum->FindEnumIndex(**EnumString) : BPTYPE_Normal);
+						EBlueprintType BlueprintType = (EnumString ? (EBlueprintType)BlueprintTypeEnum->GetValueByName(**EnumString) : BPTYPE_Normal);
 
 						switch (BlueprintType)
 						{
diff --git a/Engine/Source/Editor/KismetCompiler/Private/KismetCompilerMisc.cpp b/Engine/Source/Editor/KismetCompiler/Private/KismetCompilerMisc.cpp
--- a/Engine/Source/Editor/KismetCompiler/Private/KismetCompilerMisc.cpp
+++ b/Engine/Source/Editor/KismetCompiler/Private/KismetCompilerMisc.cpp
@@ -500,9 +500,9 @@ void FKismetCompilerUtilities::ValidateEnumProperties(UObject* DefaultObject, FC
 			const UEnum* Enum = ByteProperty->GetIntPropertyEnum();
 			if(Enum)
 			{		
-				const uint8 EnumIndex = ByteProperty->GetPropertyValue_InContainer(DefaultObject);
-				const int32 EnumAcceptableMax = Enum->NumEnums() - 1;
-				if(EnumIndex >= EnumAcceptableMax)
+				const uint8 EnumValue = ByteProperty->GetPropertyValue_InContainer(DefaultObject);
+				const int32 EnumAcceptableMax = Enum->GetMaxEnumValue();
+				if(EnumValue >= EnumAcceptableMax)
 				{
 					MessageLog.Warning(
 						*FString::Printf(
diff --git a/Engine/Source/Editor/KismetCompiler/Private/KismetCompilerVMBackend.cpp b/Engine/Source/Editor/KismetCompiler/Private/KismetCompilerVMBackend.cpp
--- a/Engine/Source/Editor/KismetCompiler/Private/KismetCompilerVMBackend.cpp
+++ b/Engine/Source/Editor/KismetCompiler/Private/KismetCompilerVMBackend.cpp
@@ -427,13 +427,13 @@ public:
 				//Check for valid enum object reference
 				if (ByteProp->Enum)
 				{
-					//Get index from enum string
-					Value = ByteProp->Enum->FindEnumIndex( *(Term->Name) );
+					//Get value from enum string
+					Value = ByteProp->Enum->GetValueByName(*(Term->Name));
 				}
 				// Allow enum literals to communicate with byte properties as literals
 				else if (UEnum* EnumPtr = Cast<UEnum>(Term->Type.PinSubCategoryObject.Get()))
 				{
-					Value = EnumPtr->FindEnumIndex( *(Term->Name) );
+					Value = EnumPtr->GetValueByName( *(Term->Name) );
 				}
 				else
 				{
diff --git a/Engine/Source/Editor/PropertyEditor/Private/PropertyHandleImpl.cpp b/Engine/Source/Editor/PropertyEditor/Private/PropertyHandleImpl.cpp
--- a/Engine/Source/Editor/PropertyEditor/Private/PropertyHandleImpl.cpp
+++ b/Engine/Source/Editor/PropertyEditor/Private/PropertyHandleImpl.cpp
@@ -68,15 +68,15 @@ FPropertyAccess::Result FPropertyValueImpl::GetPropertyValueString( FString& Out
 				UByteProperty* ByteProperty = Cast<UByteProperty>(Property);
 				if ( ByteProperty != nullptr && ByteProperty->Enum != nullptr )
 				{
-					const uint8 EnumValueIndex = ByteProperty->GetPropertyValue(ValueAddress);
+					const uint8 EnumValue = ByteProperty->GetPropertyValue(ValueAddress);
 
-					if (EnumValueIndex >= 0 && EnumValueIndex < ByteProperty->Enum->NumEnums())
+					if (EnumValue >= 0 && EnumValue < ByteProperty->Enum->GetMaxEnumValue() + 1)
 					{
 						// See if we specified an alternate name for this value using metadata
-						OutString = ByteProperty->Enum->GetDisplayNameText(EnumValueIndex).ToString();
+						OutString = ByteProperty->Enum->GetDisplayNameTextByValue(EnumValue).ToString();
 						if(!bAllowAlternateDisplayValue || OutString.Len() == 0) 
 						{
-							OutString = ByteProperty->Enum->GetEnumName(EnumValueIndex);
+							OutString = ByteProperty->Enum->GetEnumNameStringByValue(EnumValue);
 						}
 					}
 					else
diff --git a/Engine/Source/Editor/UnrealEd/Private/Kismet2/BlueprintEditorUtils.cpp b/Engine/Source/Editor/UnrealEd/Private/Kismet2/BlueprintEditorUtils.cpp
--- a/Engine/Source/Editor/UnrealEd/Private/Kismet2/BlueprintEditorUtils.cpp
+++ b/Engine/Source/Editor/UnrealEd/Private/Kismet2/BlueprintEditorUtils.cpp
@@ -7388,7 +7388,7 @@ bool FBlueprintEditorUtils::PropertyValueFromString(const UProperty* Property, c
 			int32 IntValue = 0;
 			if (const UEnum* Enum = ByteProperty->Enum)
 			{
-				IntValue = Enum->FindEnumIndex(FName(*Value));
+				IntValue = Enum->GetValueByName(FName(*Value));
 				bParseSucceeded = (INDEX_NONE != IntValue);
 
 				// If the parse did not succeed, clear out the int to keep the enum value valid
diff --git a/Engine/Source/Editor/UnrealEd/Private/Kismet2/EnumEditorUtils.cpp b/Engine/Source/Editor/UnrealEd/Private/Kismet2/EnumEditorUtils.cpp
--- a/Engine/Source/Editor/UnrealEd/Private/Kismet2/EnumEditorUtils.cpp
+++ b/Engine/Source/Editor/UnrealEd/Private/Kismet2/EnumEditorUtils.cpp
@@ -312,25 +312,25 @@ void FEnumEditorUtils::BroadcastChanges(const UUserDefinedEnum* Enum, const TArr
 	FEnumEditorManager::Get().PostChange(Enum, EEnumEditorChangeInfo::Changed);
 }
 
-int32 FEnumEditorUtils::ResolveEnumerator(const UEnum* Enum, FArchive& Ar, int32 EnumeratorIndex)
+int32 FEnumEditorUtils::ResolveEnumerator(const UEnum* Enum, FArchive& Ar, int32 EnumeratorValue)
 {
 	check(Ar.UseToResolveEnumerators());
 	const FArchiveEnumeratorResolver* EnumeratorResolver = (FArchiveEnumeratorResolver*)(&Ar);
 	if(Enum == EnumeratorResolver->Enum)
 	{
 		const auto& OldNames = EnumeratorResolver->OldNames;
-		if(EnumeratorIndex < OldNames.Num())
+		if(EnumeratorValue < OldNames.Num())
 		{
-			const FName EnumeratorName = OldNames[EnumeratorIndex].Key;
-			const int32 NewEnumIndex = Enum->FindEnumIndex(EnumeratorName);
+			const FName EnumeratorName = OldNames[EnumeratorValue].Key;
+			const int32 NewEnumIndex = Enum->GetValueByName(EnumeratorName);
 			if(INDEX_NONE != NewEnumIndex)
 			{
 				return NewEnumIndex;
 			}
 		}
-		return (Enum->NumEnums() - 1);
+		return Enum->GetMaxEnumValue();
 	}
-	return EnumeratorIndex;
+	return EnumeratorValue;
 }
 
 FString FEnumEditorUtils::GetEnumeratorDisplayName(const UUserDefinedEnum* Enum, const int32 EnumeratorIndex)
diff --git a/Engine/Source/Editor/UnrealEd/Public/Kismet2/EnumEditorUtils.h b/Engine/Source/Editor/UnrealEd/Public/Kismet2/EnumEditorUtils.h
--- a/Engine/Source/Editor/UnrealEd/Public/Kismet2/EnumEditorUtils.h
+++ b/Engine/Source/Editor/UnrealEd/Public/Kismet2/EnumEditorUtils.h
@@ -69,7 +69,7 @@ public:
 	 *
 	 *	@return new enum 
 	 */
-	static int32 ResolveEnumerator(const UEnum* Enum, FArchive& Ar, int32 EnumeratorIndex);
+	static int32 ResolveEnumerator(const UEnum* Enum, FArchive& Ar, int32 EnumeratorValue);
 
 	//DISPLAY NAME
 	static FString GetEnumeratorDisplayName(const UUserDefinedEnum* Enum, int32 EnumeratorIndex);
diff --git a/Engine/Source/Runtime/CoreUObject/Private/UObject/Class.cpp b/Engine/Source/Runtime/CoreUObject/Private/UObject/Class.cpp
--- a/Engine/Source/Runtime/CoreUObject/Private/UObject/Class.cpp
+++ b/Engine/Source/Runtime/CoreUObject/Private/UObject/Class.cpp
@@ -1055,10 +1055,10 @@ void UStruct::SerializeTaggedProperties(FArchive& Ar, uint8* Data, UStruct* Defa
 					else
 					{
 						Ar.Preload(Enum);
-						PreviousValue = Enum->FindEnumIndex(EnumValue);
-						if (Enum->NumEnums() < PreviousValue)
+						PreviousValue = Enum->GetValueByName(EnumValue);
+						if (Enum->GetNameByValue(PreviousValue) != NAME_None)
 						{
-							PreviousValue = Enum->NumEnums() - 1;
+							PreviousValue = Enum->GetMaxEnumValue();
 						}
 					}
 				}
@@ -1297,10 +1297,10 @@ void UStruct::SerializeTaggedProperties(FArchive& Ar, uint8* Data, UStruct* Defa
 					else
 					{
 						Ar.Preload(Enum);
-						PreviousValue = Enum->FindEnumIndex(EnumValue);
-						if (Enum->NumEnums() < PreviousValue)
+						PreviousValue = Enum->GetValueByName(EnumValue);
+						if (Enum->GetNameByValue(PreviousValue) != NAME_None)
 						{
-							PreviousValue = Enum->NumEnums() - 1;
+							PreviousValue = Enum->GetMaxEnumValue();
 						}
 					}
 				}
diff --git a/Engine/Source/Runtime/CoreUObject/Private/UObject/Enum.cpp b/Engine/Source/Runtime/CoreUObject/Private/UObject/Enum.cpp
--- a/Engine/Source/Runtime/CoreUObject/Private/UObject/Enum.cpp
+++ b/Engine/Source/Runtime/CoreUObject/Private/UObject/Enum.cpp
@@ -133,19 +133,6 @@ FName UEnum::GetNameByValue(int8 InValue) const
 	return NAME_None;
 }
 
-int8 UEnum::GetValueByName(FName InName)
-{
-	for (TPair<FName, int8> Kvp : Names)
-	{
-		if (Kvp.Key == InName)
-		{
-			return Kvp.Value;
-		}
-	}
-
-	return INDEX_NONE;
-}
-
 int32 UEnum::GetIndexByName(FName InName) const
 {
 	int32 Count = Names.Num();
@@ -439,6 +426,11 @@ bool UEnum::SetEnums(TArray<TPair<FName, int8>>& InNames, UEnum::ECppForm InCppF
 }
 
 #if WITH_EDITOR
+FText UEnum::GetDisplayNameTextByValue(int32 Value) const
+{
+	int32 Index = GetIndexByValue(Value);
+	return GetDisplayNameText(Index);
+}
 /**
  * Finds the localized display name or native display name as a fallback.
  *
@@ -638,39 +630,6 @@ void UEnum::RemoveMetaData( const TCHAR* Key, int32 NameIndex/*=INDEX_NONE*/) co
 
 #endif
 
-int32 UEnum::LookupEnumName(FName TestName, UEnum** FoundEnum)
-{
-	UEnum* TheEnum = AllEnumNames.FindRef(TestName);
-	if (FoundEnum != nullptr)
-	{
-		*FoundEnum = TheEnum;
-	}
-	return (TheEnum != nullptr) ? TheEnum->GetIndexByName(TestName) : INDEX_NONE;
-}
-
-int32 UEnum::LookupEnumNameSlow(const TCHAR* InTestShortName, UEnum** FoundEnum)
-{
-	int32 EnumIndex = LookupEnumName(InTestShortName, FoundEnum);
-	if (EnumIndex == INDEX_NONE)
-	{
-		FString TestShortName = FString(TEXT("::")) + InTestShortName;
-		UEnum* TheEnum = NULL;
-		for (TMap<FName, UEnum*>::TIterator It(AllEnumNames); It; ++It)
-		{
-			if (It.Key().ToString().Contains(TestShortName) )
-			{
-				TheEnum = It.Value();
-			}
-		}
-		if (FoundEnum != NULL)
-		{
-			*FoundEnum = TheEnum;
-		}
-		EnumIndex = (TheEnum != NULL) ? TheEnum->FindEnumIndex(InTestShortName) : INDEX_NONE;
-	}
-	return EnumIndex;
-}
-
 int32 UEnum::ParseEnum(const TCHAR*& Str)
 {
 	FString Token;
diff --git a/Engine/Source/Runtime/CoreUObject/Private/UObject/PropertyByte.cpp b/Engine/Source/Runtime/CoreUObject/Private/UObject/PropertyByte.cpp
--- a/Engine/Source/Runtime/CoreUObject/Private/UObject/PropertyByte.cpp
+++ b/Engine/Source/Runtime/CoreUObject/Private/UObject/PropertyByte.cpp
@@ -33,10 +33,10 @@ void UByteProperty::SerializeItem( FArchive& Ar, void* Value, void const* Defaul
 
 		// There's no guarantee EnumValueName is still present in Enum, in which case Value will be set to the enum's max value.
 		// On save, it will then be serialized as NAME_None.
-		*(uint8*)Value = Enum->FindEnumIndex(EnumValueName);
-		if ( Enum->NumEnums() < *(uint8*)Value )
+		*(uint8*)Value = Enum->GetValueByName(EnumValueName);
+		if ( Enum->GetMaxEnumValue() + 1 < *(uint8*)Value )
 		{
-			*(uint8*)Value = Enum->NumEnums() - 1;
+			*(uint8*)Value = Enum->GetMaxEnumValue();
 		}
 	}
 	// Saving
@@ -47,9 +47,9 @@ void UByteProperty::SerializeItem( FArchive& Ar, void* Value, void const* Defaul
 
 		// subtract 1 because the last entry in the enum's Names array
 		// is the _MAX entry
-		if ( ByteValue < Enum->NumEnums() - 1 )
+		if ( ByteValue < Enum->GetMaxEnumValue() )
 		{
-			EnumValueName = Enum->GetEnum(ByteValue);
+			EnumValueName = Enum->GetNameByValue(ByteValue);
 		}
 		else
 		{
@@ -113,16 +113,16 @@ void UByteProperty::ExportTextItem( FString& ValueStr, const void* PropertyValue
 	{
 		// if the value is the max value (the autogenerated *_MAX value), export as "INVALID", unless we're exporting text for copy/paste (for copy/paste,
 		// the property text value must actually match an entry in the enum's names array)
-		if (*(const uint8*)PropertyValue < Enum->NumEnums() - 1 || ((PortFlags&PPF_Copy) != 0 && *(const uint8*)PropertyValue < Enum->NumEnums()))
+		if (*(const uint8*)PropertyValue <= Enum->GetMaxEnumValue() - 1 || ((PortFlags&PPF_Copy) != 0 && *(const uint8*)PropertyValue <= Enum->GetMaxEnumValue()))
 		{
 			// We do not want to export the enum text for non-display uses, localization text is very dynamic and would cause issues on import
 			if (PortFlags & PPF_PropertyWindow)
 			{
-				ValueStr += Enum->GetEnumText(*(const uint8*)PropertyValue).ToString();
+				ValueStr += Enum->GetEnumTextByValue(*(const uint8*)PropertyValue).ToString();
 			}
 			else
 			{
-				ValueStr += Enum->GetEnumName(*(const uint8*)PropertyValue);
+				ValueStr += Enum->GetEnumNameStringByValue(*(const uint8*)PropertyValue);
 			}
 		}
 		else
@@ -143,10 +143,10 @@ const TCHAR* UByteProperty::ImportText_Internal( const TCHAR* InBuffer, void* Da
 		const TCHAR* Buffer = UPropertyHelpers::ReadToken( InBuffer, Temp, true );
 		if( Buffer != NULL )
 		{
-			const int32 EnumIndex = Enum->FindEnumIndex( *Temp );
-			if( EnumIndex != INDEX_NONE )
+			const int32 EnumValue = Enum->GetValueByName(*Temp);
+			if( EnumValue != INDEX_NONE )
 			{
-				*(uint8*)Data = EnumIndex;
+				*(uint8*)Data = EnumValue;
 				return Buffer;
 			}
 		}
diff --git a/Engine/Source/Runtime/CoreUObject/Public/UObject/Class.h b/Engine/Source/Runtime/CoreUObject/Public/UObject/Class.h
--- a/Engine/Source/Runtime/CoreUObject/Public/UObject/Class.h
+++ b/Engine/Source/Runtime/CoreUObject/Public/UObject/Class.h
@@ -1352,15 +1352,43 @@ public:
 
 	/** Gets enum name by index in Names array. Returns NAME_None if Index is not valid. */
 	FName GetNameByIndex(int8 Index) const;
-
+		
 	/** Gets enum value by index in Names array. */
 	int8 GetValueByIndex(int8 Index) const;
 
 	/** Gets enum name by value. Returns NAME_None if value is not found. */
 	FName GetNameByValue(int8 InValue) const;
-
+		
 	/** Gets enum value by name. Returns INDEX_NONE when name is not found. */
-	int8 GetValueByName(FName InName);
+	int8 GetValueByName(FName InName) const
+	{
+		FString InNameString = InName.ToString();
+		FString DoubleColon = TEXT("::");
+		int32 DoubleColonPosition = InNameString.Find(DoubleColon);
+		bool bIsInNameFullyQualified = DoubleColonPosition != INDEX_NONE;
+		bool bIsEnumFullyQualified = CppForm != ECppForm::Regular;
+
+		if (!bIsInNameFullyQualified && bIsEnumFullyQualified)
+		{
+			// Make InName fully qualified.
+			InName = FName(*GetName().Append(DoubleColon).Append(InNameString));
+		}
+		else if (bIsInNameFullyQualified && !bIsEnumFullyQualified)
+		{
+			// Make InName unqualified.
+			InName = FName(*InNameString.Mid(DoubleColonPosition + 2, InNameString.Len() - DoubleColonPosition + 2));
+		}
+
+		for (TPair<FName, int8> Kvp : Names)
+		{
+			if (Kvp.Key == InName)
+			{
+				return Kvp.Value;
+			}
+		}
+
+		return INDEX_NONE;
+	}
 
 	/** Gets index of name in enum. Returns INDEX_NONE when name is not found. */
 	int32 GetIndexByName(FName InName) const;
@@ -1441,12 +1469,41 @@ public:
 	/** searches the list of all enum value names for the specified name
 	 * @return the value the specified name represents if found, otherwise INDEX_NONE
 	 */
-	static int32 LookupEnumName(FName TestName, UEnum** FoundEnum = nullptr);
+	static int32 LookupEnumName(FName TestName, UEnum** FoundEnum = nullptr)
+	{
+		UEnum* TheEnum = AllEnumNames.FindRef(TestName);
+		if (FoundEnum != nullptr)
+		{
+			*FoundEnum = TheEnum;
+		}
+		return (TheEnum != nullptr) ? TheEnum->GetValueByName(TestName) : INDEX_NONE;
+	}
 
 	/** searches the list of all enum value names for the specified name
 	 * @return the value the specified name represents if found, otherwise INDEX_NONE
 	 */
-	static int32 LookupEnumNameSlow(const TCHAR* InTestShortName, UEnum** FoundEnum = nullptr);
+	static int32 LookupEnumNameSlow(const TCHAR* InTestShortName, UEnum** FoundEnum = NULL)
+	{
+		int32 EnumIndex = LookupEnumName(InTestShortName, FoundEnum);
+		if (EnumIndex == INDEX_NONE)
+		{
+			FString TestShortName = FString(TEXT("::")) + InTestShortName;
+			UEnum* TheEnum = NULL;
+			for (TMap<FName, UEnum*>::TIterator It(AllEnumNames); It; ++It)
+			{
+				if (It.Key().ToString().Contains(TestShortName) )
+				{
+					TheEnum = It.Value();
+				}
+			}
+			if (FoundEnum != NULL)
+			{
+				*FoundEnum = TheEnum;
+			}
+			EnumIndex = (TheEnum != NULL) ? TheEnum->GetValueByName(InTestShortName) : INDEX_NONE;
+		}
+		return EnumIndex;
+	}
 
 	/** parses the passed in string for a name, then searches for that name in any Enum (in any package)
 	 * @param Str	pointer to string to parse; if we successfully find an enum, this pointer is advanced past the name found
@@ -1475,11 +1532,35 @@ public:
 		return NAME_None;
 	}
 
+	int32 GetIndexByValue(int32 Value) const
+	{
+		for (int32 i = 0; i < Names.Num(); ++i)
+		{
+			if (Names[i].Value == Value)
+			{
+				return i;
+			}
+		}
+
+		return INDEX_NONE;
+	}
+	FString GetEnumNameStringByValue(int32 Value) const
+	{
+		int32 Index = GetIndexByValue(Value);
+		return GetEnumName(Index);
+	}
+
 	/**
 	 * @return	The short enum name at the specified Index.
 	 */
 	FString GetEnumName(int32 InIndex) const;
 
+	FText GetEnumTextByValue(int32 Value)
+	{
+		int32 Index = GetIndexByValue(Value);
+		return GetEnumText(Index);
+	}
+
 	/**
 	 * @return	The enum string at the specified index.
 	 */
@@ -1515,6 +1596,7 @@ public:
 	 * @return The display name for this object.
 	 */
 	FText GetDisplayNameText(int32 NameIndex=INDEX_NONE) const;
+	FText GetDisplayNameTextByValue(int32 Value = INDEX_NONE) const;
 
 	/**
 	 * Finds the localized tooltip or native tooltip as a fallback.
diff --git a/Engine/Source/Runtime/Engine/Classes/Engine/UserDefinedEnum.h b/Engine/Source/Runtime/Engine/Classes/Engine/UserDefinedEnum.h
--- a/Engine/Source/Runtime/Engine/Classes/Engine/UserDefinedEnum.h
+++ b/Engine/Source/Runtime/Engine/Classes/Engine/UserDefinedEnum.h
@@ -49,7 +49,7 @@ public:
 	 *	@param EnumeratorIndex	old index
 	 *	@return	new index
 	 */
-	virtual int32 ResolveEnumerator(FArchive& Ar, int32 EnumeratorIndex) const override;
+	virtual int32 ResolveEnumerator(FArchive& Ar, int32 EnumeratorValue) const override;
 
 	/**
 	 * @return	The enum string at the specified index.
diff --git a/Engine/Source/Runtime/Engine/Private/TextureDerivedData.cpp b/Engine/Source/Runtime/Engine/Private/TextureDerivedData.cpp
--- a/Engine/Source/Runtime/Engine/Private/TextureDerivedData.cpp
+++ b/Engine/Source/Runtime/Engine/Private/TextureDerivedData.cpp
@@ -1234,7 +1234,7 @@ static void SerializePlatformData(
 	{
 		FString PixelFormatString;
 		Ar << PixelFormatString;
-		PlatformData->PixelFormat = (EPixelFormat)PixelFormatEnum->FindEnumIndex(*PixelFormatString);
+		PlatformData->PixelFormat = (EPixelFormat)PixelFormatEnum->GetValueByName(*PixelFormatString);
 	}
 	else if (Ar.IsSaving())
 	{
@@ -1919,7 +1919,7 @@ void UTexture::SerializeCookedPlatformData(FArchive& Ar)
 		Ar << PixelFormatName;
 		while (PixelFormatName != NAME_None)
 		{
-			EPixelFormat PixelFormat = (EPixelFormat)PixelFormatEnum->FindEnumIndex(PixelFormatName);
+			EPixelFormat PixelFormat = (EPixelFormat)PixelFormatEnum->GetValueByName(PixelFormatName);
 			int32 SkipOffset = 0;
 			Ar << SkipOffset;
 			bool bFormatSupported = GPixelFormats[PixelFormat].Supported;
diff --git a/Engine/Source/Runtime/Engine/Private/UserDefinedEnum.cpp b/Engine/Source/Runtime/Engine/Private/UserDefinedEnum.cpp
--- a/Engine/Source/Runtime/Engine/Private/UserDefinedEnum.cpp
+++ b/Engine/Source/Runtime/Engine/Private/UserDefinedEnum.cpp
@@ -77,13 +77,13 @@ FString UUserDefinedEnum::GenerateNewEnumeratorName()
 
 #endif	// WITH_EDITOR
 
-int32 UUserDefinedEnum::ResolveEnumerator(FArchive& Ar, int32 EnumeratorIndex) const
+int32 UUserDefinedEnum::ResolveEnumerator(FArchive& Ar, int32 EnumeratorValue) const
 {
 #if WITH_EDITOR
-	return FEnumEditorUtils::ResolveEnumerator(this, Ar, EnumeratorIndex);
+	return FEnumEditorUtils::ResolveEnumerator(this, Ar, EnumeratorValue);
 #else // WITH_EDITOR
 	ensure(false);
-	return EnumeratorIndex;
+	return EnumeratorValue;
 #endif // WITH_EDITOR
 }
 
diff --git a/Engine/Source/Runtime/JsonUtilities/Private/JsonObjectConverter.cpp b/Engine/Source/Runtime/JsonUtilities/Private/JsonObjectConverter.cpp
--- a/Engine/Source/Runtime/JsonUtilities/Private/JsonObjectConverter.cpp
+++ b/Engine/Source/Runtime/JsonUtilities/Private/JsonObjectConverter.cpp
@@ -222,7 +222,7 @@ bool ConvertScalarJsonValueToUProperty(TSharedPtr<FJsonValue> JsonValue, UProper
 			const UEnum* EnumProperty = NumericProperty->GetIntPropertyEnum();
 			check(EnumProperty); // should be assured by IsEnum()
 			FString StrValue = JsonValue->AsString();
-			int32 IntValue = EnumProperty->FindEnumIndex(FName(*StrValue));
+			int32 IntValue = EnumProperty->GetValueByName(FName(*StrValue));
 			if (IntValue == INDEX_NONE)
 			{
 				UE_LOG(LogJson, Error, TEXT("JsonValueToUProperty - Unable import enum %s from string value %s for property %s"), *EnumProperty->CppType, *StrValue, *Property->GetNameCPP());
diff --git a/Engine/Source/Runtime/Online/OnlineSubsystem/Private/OnlineKeyValuePair.cpp b/Engine/Source/Runtime/Online/OnlineSubsystem/Private/OnlineKeyValuePair.cpp
--- a/Engine/Source/Runtime/Online/OnlineSubsystem/Private/OnlineKeyValuePair.cpp
+++ b/Engine/Source/Runtime/Online/OnlineSubsystem/Private/OnlineKeyValuePair.cpp
@@ -710,7 +710,7 @@ bool FVariantDataConverter::ConvertScalarVariantToUProperty(const FVariantData*
 			FString StrValue;
 			Variant->GetValue(StrValue);
 
-			int32 IntValue = EnumProperty->FindEnumIndex(FName(*StrValue));
+			int32 IntValue = EnumProperty->GetValueByName(FName(*StrValue));
 			if (IntValue == INDEX_NONE)
 			{
 				UE_LOG(LogOnline, Error, TEXT("ConvertScalarVariantToUProperty - Unable import enum %s from string value %s for property %s"), *EnumProperty->CppType, *StrValue, *Property->GetNameCPP());
diff --git a/Engine/Source/Runtime/Serialization/Private/Backends/JsonStructDeserializerBackend.cpp b/Engine/Source/Runtime/Serialization/Private/Backends/JsonStructDeserializerBackend.cpp
--- a/Engine/Source/Runtime/Serialization/Private/Backends/JsonStructDeserializerBackend.cpp
+++ b/Engine/Source/Runtime/Serialization/Private/Backends/JsonStructDeserializerBackend.cpp
@@ -269,14 +269,14 @@ bool FJsonStructDeserializerBackend::ReadProperty( UProperty* Property, UPropert
 			if (Property->GetClass() == UByteProperty::StaticClass())
 			{
 				UByteProperty* ByteProperty = Cast<UByteProperty>(Property);
-				int32 Index = ByteProperty->Enum->FindEnumIndex(*StringValue);
+				int32 Value = ByteProperty->Enum->GetValueByName(*StringValue);
 
-				if (Index == INDEX_NONE)
+				if (Value == INDEX_NONE)
 				{
 					return false;
 				}
 
-				return SetPropertyValue<UByteProperty, uint8>(Property, Outer, Data, ArrayIndex, (uint8)Index);
+				return SetPropertyValue<UByteProperty, uint8>(Property, Outer, Data, ArrayIndex, (uint8)Value);
 			}
 			
 			if (Property->GetClass() == UClassProperty::StaticClass())
