Implemented localized package loading based on trying to find the package located under "<ContentRoot>/L10N/<CultureCode>/<OriginalPackagePath>". The localized package content subdirectories are also filtered out of the content browser views. In the editor, package paths are not redirected in any additional manner, but outside the editor the package paths redirect to the calculated localized package location.

diff --git a/Engine/Source/Editor/ContentBrowser/Private/ContentBrowserUtils.cpp b/Engine/Source/Editor/ContentBrowser/Private/ContentBrowserUtils.cpp
--- a/Engine/Source/Editor/ContentBrowser/Private/ContentBrowserUtils.cpp
+++ b/Engine/Source/Editor/ContentBrowser/Private/ContentBrowserUtils.cpp
@@ -1112,6 +1112,12 @@ bool ContentBrowserUtils::IsPluginFolder( const FString& InPath )
 	return false;
 }
 
+bool ContentBrowserUtils::IsLocalizationFolder( const FString& InPath )
+{
+	return	InPath.StartsWith(TEXT("/Engine/L10N")) || InPath.StartsWith(TEXT("Engine/L10N")) ||
+			InPath.StartsWith(TEXT("/Game/L10N")) || InPath.StartsWith(TEXT("Game/L10N"));
+}
+
 void ContentBrowserUtils::GetObjectsInAssetData(const TArray<FAssetData>& AssetList, TArray<UObject*>& OutDroppedObjects)
 {	
 	for (int32 AssetIdx = 0; AssetIdx < AssetList.Num(); ++AssetIdx)
@@ -1594,6 +1600,12 @@ bool ContentBrowserUtils::IsValidFolderPathForCreate(const FString& InFolderPath
 
 	const FString NewFolderPath = InFolderPath / NewFolderName;
 
+	if (ContentBrowserUtils::IsLocalizationFolder(NewFolderPath))
+	{
+		OutErrorMessage = LOCTEXT("LocalizationFolderReserved", "The L10N folder is reserved for localized content.");
+		return false;
+	}
+
 	if (ContentBrowserUtils::DoesFolderExist(NewFolderPath))
 	{
 		OutErrorMessage = LOCTEXT("RenameFolderAlreadyExists", "A folder already exists at this location with this name.");
diff --git a/Engine/Source/Editor/ContentBrowser/Private/ContentBrowserUtils.h b/Engine/Source/Editor/ContentBrowser/Private/ContentBrowserUtils.h
--- a/Engine/Source/Editor/ContentBrowser/Private/ContentBrowserUtils.h
+++ b/Engine/Source/Editor/ContentBrowser/Private/ContentBrowserUtils.h
@@ -147,6 +147,9 @@ namespace ContentBrowserUtils
 	/** Returns true if the passed-in path is a plugin folder */
 	bool IsPluginFolder( const FString& InPath );
 
+	/** Returns true if the passed-in path is a localization folder */
+	bool IsLocalizationFolder( const FString& InPath );
+
 	/** Get all the objects in a list of asset data */
 	void GetObjectsInAssetData(const TArray<FAssetData>& AssetList, TArray<UObject*>& OutDroppedObjects);
 
diff --git a/Engine/Source/Editor/ContentBrowser/Private/SAssetView.cpp b/Engine/Source/Editor/ContentBrowser/Private/SAssetView.cpp
--- a/Engine/Source/Editor/ContentBrowser/Private/SAssetView.cpp
+++ b/Engine/Source/Editor/ContentBrowser/Private/SAssetView.cpp
@@ -1432,19 +1432,18 @@ void SAssetView::RefreshSourceItems()
 	for (int32 AssetIdx = Items.Num() - 1; AssetIdx >= 0; --AssetIdx)
 	{
 		const FAssetData& Item = Items[AssetIdx];
-		if ( Item.AssetClass == UObjectRedirector::StaticClass()->GetFName() && !Item.IsUAsset() )
+		// Do not show redirectors if they are not the main asset in the uasset file.
+		const bool IsMainlyARedirector = Item.AssetClass == UObjectRedirector::StaticClass()->GetFName() && !Item.IsUAsset();
+		// If this is an engine folder, and we don't want to show them, remove
+		const bool IsHiddenEngineFolder = !bDisplayEngine && ContentBrowserUtils::IsEngineFolder(Item.PackagePath.ToString());
+		// If this is a plugin folder, and we don't want to show them, remove
+		const bool IsAHiddenPluginFolder = !bDisplayPlugins && ContentBrowserUtils::IsPluginFolder(Item.PackagePath.ToString());
+		// Do not show localized content folders.
+		const bool IsLocalizedContentFolder = ContentBrowserUtils::IsLocalizationFolder(Item.PackagePath.ToString());
+
+		const bool ShouldFilterOut = IsMainlyARedirector || IsHiddenEngineFolder || IsAHiddenPluginFolder || IsLocalizedContentFolder;
+		if (ShouldFilterOut)
 		{
-			// Do not show redirectors if they are not the main asset in the uasset file.
-			Items.RemoveAtSwap(AssetIdx);
-		}
-		else if ( !bDisplayEngine && ContentBrowserUtils::IsEngineFolder(Item.PackagePath.ToString()) )
-		{
-			// If this is an engine folder, and we don't want to show them, remove
-			Items.RemoveAtSwap(AssetIdx);
-		}
-		else if ( !bDisplayPlugins && ContentBrowserUtils::IsPluginFolder(Item.PackagePath.ToString()) )
-		{
-			// If this is a plugin folder, and we don't want to show them, remove
 			Items.RemoveAtSwap(AssetIdx);
 		}
 	}
@@ -1692,6 +1691,11 @@ void SAssetView::RefreshFolders()
 					continue;
 				}
 
+				if (ContentBrowserUtils::IsLocalizationFolder(SubPath))
+				{
+					continue;
+				}
+
 				if(!Folders.Contains(SubPath))
 				{
 					FoldersToAdd.Add(SubPath);
diff --git a/Engine/Source/Editor/ContentBrowser/Private/SPathView.cpp b/Engine/Source/Editor/ContentBrowser/Private/SPathView.cpp
--- a/Engine/Source/Editor/ContentBrowser/Private/SPathView.cpp
+++ b/Engine/Source/Editor/ContentBrowser/Private/SPathView.cpp
@@ -1,4 +1,4 @@
-// Copyright 1998-2015 Epic Games, Inc. All Rights Reserved.
+﻿// Copyright 1998-2015 Epic Games, Inc. All Rights Reserved.
 
 #include "ContentBrowserPCH.h"
 
@@ -758,17 +758,15 @@ TSharedPtr<struct FTreeItem> SPathView::AddRootItem( const FString& InFolderName
 
 	// If this isn't an engine folder or we want to show them, add
 	const bool bDisplayEngine = GetDefault<UContentBrowserSettings>()->GetDisplayEngineFolder();
-	if ( bDisplayEngine || !ContentBrowserUtils::IsEngineFolder(InFolderName) )
-	{
 		const bool bDisplayPlugins = GetDefault<UContentBrowserSettings>()->GetDisplayPluginFolders();
-		if ( bDisplayPlugins || !ContentBrowserUtils::IsPluginFolder(InFolderName) )
+	if ( (bDisplayEngine || !ContentBrowserUtils::IsEngineFolder(InFolderName)) && 
+		 (bDisplayPlugins || !ContentBrowserUtils::IsPluginFolder(InFolderName)) )
 		{
 			const FText DisplayName = ContentBrowserUtils::GetRootDirDisplayName(InFolderName);
 			NewItem = MakeShareable( new FTreeItem(DisplayName, InFolderName, FString(TEXT("/")) + InFolderName, TSharedPtr<FTreeItem>()));
 			TreeRootItems.Add( NewItem );
 			TreeViewPtr->RequestTreeRefresh();
 		}
-	}
 
 	return NewItem;
 }
@@ -970,6 +968,12 @@ void SPathView::Populate()
 	const FString UserDeveloperFolder = FPackageName::FilenameToLongPackageName(FPaths::GameUserDeveloperDir().LeftChop(1));
 	PathList.Add(UserDeveloperFolder);
 
+	// Remove paths of localized assets.
+	PathList.RemoveAll([](const FString& Path) -> bool
+	{
+		return ContentBrowserUtils::IsLocalizationFolder(Path);
+	});
+
 	// we have a text filter, expand all parents of matching folders
 	for ( int32 PathIdx = 0; PathIdx < PathList.Num(); ++PathIdx)
 	{
@@ -1100,21 +1104,9 @@ FReply SPathView::OnFolderDragDetected(const FGeometry& Geometry, const FPointer
 	return FReply::Unhandled();
 }
 
-bool SPathView::VerifyFolderNameChanged(const FText& InName, FText& OutErrorMessage, const FString& InFolderPath) const
+bool SPathView::VerifyFolderNameChanged(const FString& InName, FText& OutErrorMessage, const FString& InFolderPath) const
 {	
-	if( !ContentBrowserUtils::IsValidFolderName(InName.ToString(), OutErrorMessage) )
-	{
-		return false;
-	}
-
-	const FString NewPath = FPaths::GetPath(InFolderPath) / InName.ToString();
-	if (ContentBrowserUtils::DoesFolderExist(NewPath))
-	{
-		OutErrorMessage = LOCTEXT("RenameFolderAlreadyExists", "A folder already exists at this location with this name.");
-		return false;
-	}
-
-	return true;
+	return ContentBrowserUtils::IsValidFolderPathForCreate(FPaths::GetPath(InFolderPath), InName, OutErrorMessage);
 }
 
 void SPathView::FolderNameChanged( const TSharedPtr< FTreeItem >& TreeItem, const FString& OldPath, const FVector2D& MessageLocation )
@@ -1382,8 +1374,12 @@ void SPathView::OnAssetRegistryPathAdded(const FString& Path)
 	// of successful hits in the filtered list. 
 	if ( SearchBoxFolderFilter->PassesFilter( Path ) )
 	{
+		// Do not add paths of localized assets.
+		if (!ContentBrowserUtils::IsLocalizationFolder(Path))
+		{
 		AddPath(Path);
 	}
+	}
 }
 
 void SPathView::OnAssetRegistryPathRemoved(const FString& Path)
diff --git a/Engine/Source/Editor/ContentBrowser/Private/SPathView.h b/Engine/Source/Editor/ContentBrowser/Private/SPathView.h
--- a/Engine/Source/Editor/ContentBrowser/Private/SPathView.h
+++ b/Engine/Source/Editor/ContentBrowser/Private/SPathView.h
@@ -168,7 +168,7 @@ private:
 	FReply OnFolderDragDetected(const FGeometry& Geometry, const FPointerEvent& MouseEvent);
 
 	/** Handler use dto verify the name of a new folder */
-	bool VerifyFolderNameChanged(const FText& InName, FText& OutErrorMessage, const FString& InFolderPath) const;
+	bool VerifyFolderNameChanged(const FString& InName, FText& OutErrorMessage, const FString& InFolderPath) const;
 
 	/** Handler for when a name was given to a new folder */
 	void FolderNameChanged( const TSharedPtr< FTreeItem >& TreeItem, const FString& OldPath, const FVector2D& MessageLocation );
diff --git a/Engine/Source/Editor/ContentBrowser/Private/SourcesViewWidgets.cpp b/Engine/Source/Editor/ContentBrowser/Private/SourcesViewWidgets.cpp
--- a/Engine/Source/Editor/ContentBrowser/Private/SourcesViewWidgets.cpp
+++ b/Engine/Source/Editor/ContentBrowser/Private/SourcesViewWidgets.cpp
@@ -188,7 +188,7 @@ bool SAssetTreeItem::VerifyNameChanged(const FText& InName, FText& OutError) con
 		TSharedPtr<FTreeItem> TreeItemPtr = TreeItem.Pin();
 		if(OnVerifyNameChanged.IsBound())
 		{
-			return OnVerifyNameChanged.Execute(InName, OutError, TreeItemPtr->FolderPath);
+			return OnVerifyNameChanged.Execute(InName.ToString(), OutError, TreeItemPtr->FolderPath);
 		}
 	}
 
diff --git a/Engine/Source/Editor/ContentBrowser/Private/SourcesViewWidgets.h b/Engine/Source/Editor/ContentBrowser/Private/SourcesViewWidgets.h
--- a/Engine/Source/Editor/ContentBrowser/Private/SourcesViewWidgets.h
+++ b/Engine/Source/Editor/ContentBrowser/Private/SourcesViewWidgets.h
@@ -7,7 +7,7 @@ class SAssetTreeItem : public SCompoundWidget
 {
 public:
 	DECLARE_DELEGATE_ThreeParams( FOnNameChanged, const TSharedPtr<FTreeItem>& /*TreeItem*/, const FString& /*OldPath*/, const FVector2D& /*MessageLocation*/);
-	DECLARE_DELEGATE_RetVal_ThreeParams( bool, FOnVerifyNameChanged, const FText& /*InName*/, FText& /*OutErrorMessage*/, const FString& /*FolderPath*/);
+	DECLARE_DELEGATE_RetVal_ThreeParams( bool, FOnVerifyNameChanged, const FString& /*InName*/, FText& /*OutErrorMessage*/, const FString& /*FolderPath*/);
 	DECLARE_DELEGATE_TwoParams( FOnAssetsDragDropped, const TArray<FAssetData>& /*AssetList*/, const TSharedPtr<FTreeItem>& /*TreeItem*/);
 	DECLARE_DELEGATE_TwoParams( FOnPathsDragDropped, const TArray<FString>& /*PathNames*/, const TSharedPtr<FTreeItem>& /*TreeItem*/);
 	DECLARE_DELEGATE_TwoParams( FOnFilesDragDropped, const TArray<FString>& /*FileNames*/, const TSharedPtr<FTreeItem>& /*TreeItem*/);
diff --git a/Engine/Source/Runtime/CoreUObject/Private/Misc/PackageName.cpp b/Engine/Source/Runtime/CoreUObject/Private/Misc/PackageName.cpp
--- a/Engine/Source/Runtime/CoreUObject/Private/Misc/PackageName.cpp
+++ b/Engine/Source/Runtime/CoreUObject/Private/Misc/PackageName.cpp
@@ -1,4 +1,4 @@
-// Copyright 1998-2015 Epic Games, Inc. All Rights Reserved.
+﻿// Copyright 1998-2015 Epic Games, Inc. All Rights Reserved.
 
 /*=============================================================================
 	UObjectHash.cpp: Unreal object name hashes
@@ -328,16 +328,20 @@ FString FPackageName::FilenameToLongPackageName(const FString& InFilename)
 	return Result;
 }
 
-FString FPackageName::LongPackageNameToFilename(const FString& InLongPackageName, const FString& InExtension)
+FString FPackageName::LongPackageNameToFilename(const FString& InLongPackageName, const FString& InExtension, const bool ShouldGetLocalizedPackage)
 {
 	const auto& Paths = FLongPackagePathsSingleton::Get();
 
+	const FCultureRef CurrentCulture = FInternationalization::Get().GetCurrentCulture();
+	const FString LocalizationPathParticle = FString::Printf(TEXT("L10N/%s/"), *CurrentCulture->GetName());
+
 	for (const auto& Pair : Paths.ContentRootToPath)
 	{
 		if (InLongPackageName.StartsWith(Pair.RootPath))
 		{
-			FString Result = Pair.ContentPath + InLongPackageName.Mid(Pair.RootPath.Len()) + InExtension;
-			return Result;
+			const FString RootRelativePath = InLongPackageName.Mid(Pair.RootPath.Len());
+			const bool IsLocalizedPath = RootRelativePath.StartsWith(LocalizationPathParticle);
+			return Pair.ContentPath + (ShouldGetLocalizedPackage && !IsLocalizedPath ? LocalizationPathParticle : TEXT("")) + RootRelativePath + InExtension;
 		}
 	}
 
@@ -599,7 +603,7 @@ bool FPackageName::FindPackageFileWithoutExtension(const FString& InPackageFilen
 	return false;
 }
 
-bool FPackageName::DoesPackageExist(const FString& LongPackageName, const FGuid* Guid /*= NULL*/, FString* OutFilename /*= NULL*/)
+bool FPackageName::DoesPackageExist(const FString& LongPackageName, const FGuid* Guid /*= NULL*/, FString* OutFilename /*= NULL*/, const bool ShouldGetLocalizedPackage /*= false*/)
 {
 	bool bFoundFile = false;
 
@@ -627,7 +631,7 @@ bool FPackageName::DoesPackageExist(const FString& LongPackageName, const FGuid*
 	}
 
 	// Convert to filename (no extension yet).
-	FString Filename = LongPackageNameToFilename(PackageName);
+	FString Filename = LongPackageNameToFilename(PackageName, TEXT(""), ShouldGetLocalizedPackage);
 	// Find the filename (with extension).
 	bFoundFile = FindPackageFileWithoutExtension(Filename, Filename);
 
diff --git a/Engine/Source/Runtime/CoreUObject/Private/UObject/Linker.cpp b/Engine/Source/Runtime/CoreUObject/Private/UObject/Linker.cpp
--- a/Engine/Source/Runtime/CoreUObject/Private/UObject/Linker.cpp
+++ b/Engine/Source/Runtime/CoreUObject/Private/UObject/Linker.cpp
@@ -467,8 +467,14 @@ FLinkerLoad* GetPackageLinker
 			LogGetPackageLinkerError(Result, InLongPackageName, ErrorText, ErrorText, InOuter, LoadFlags);
 			return nullptr;
 		}
-
-		if( !FPackageName::DoesPackageExist(InOuter->GetName(), CompatibleGuid, &NewFilename) )
+	
+		FString NativeFilename;
+		FString LocalizedFilename;
+		const bool DoesNativePackageExist = FPackageName::DoesPackageExist(InOuter->GetName(), CompatibleGuid, &NativeFilename, false);
+		const bool DoesLocalizedPackageExist = FPackageName::DoesPackageExist(InOuter->GetName(), CompatibleGuid, &LocalizedFilename, true);
+
+		// If we are the editor, we must have a native package. If we are the game, we must have a localized package or a native package.
+		if ( (GIsEditor && !DoesNativePackageExist) || (!GIsEditor && !DoesLocalizedPackageExist && !DoesNativePackageExist) )
 		{
 			// In memory-only packages have no linker and this is ok.
 			if (!(LoadFlags & LOAD_AllowDll) && !(InOuter->PackageFlags & PKG_InMemoryOnly))
@@ -486,6 +492,38 @@ FLinkerLoad* GetPackageLinker
 
 			return nullptr;
 		}
+
+		// The editor must not redirect packages for localization.
+		if (GIsEditor)
+		{
+			NewFilename = NativeFilename;
+		}
+		else
+		{
+			if (DoesLocalizedPackageExist)
+			{
+				NewFilename = LocalizedFilename;
+			}
+			// If we are the game, we can fallback to the native package, but must issue a warning.
+			else
+			{
+				// In memory-only packages have no linker and this is ok.
+				if (!(LoadFlags & LOAD_AllowDll) && !(InOuter->PackageFlags & PKG_InMemoryOnly))
+				{
+					FUObjectThreadContext& ThreadContext = FUObjectThreadContext::Get();
+					FFormatNamedArguments Arguments;
+					Arguments.Add(TEXT("AssetName"), FText::FromString(InOuter->GetName()));
+					Arguments.Add(TEXT("PackageName"), FText::FromString(ThreadContext.SerializedPackageLinker ? *(ThreadContext.SerializedPackageLinker->Filename) : TEXT("NULL")));
+					LogGetPackageLinkerError(Result, ThreadContext.SerializedPackageLinker ? *ThreadContext.SerializedPackageLinker->Filename : nullptr,
+						FText::Format(LOCTEXT("PackageNotFound", "Can't find localized file for asset '{AssetName}' while loading {PackageName}."), Arguments),
+						LOCTEXT("PackageNotFoundShort", "Can't find localized file for asset."),
+						InOuter,
+						LoadFlags);
+				}
+
+				NewFilename = NativeFilename;
+			}
+		}
 	}
 	else
 	{
@@ -509,7 +547,12 @@ FLinkerLoad* GetPackageLinker
 		}
 
 		// Verify that the file exists.
-		if( !FPackageName::DoesPackageExist( PackageName, CompatibleGuid, &NewFilename ) )
+		FString NativeFilename;
+		FString LocalizedFilename;
+		const bool DoesNativePackageExist = FPackageName::DoesPackageExist(PackageName, CompatibleGuid, &NativeFilename, false);
+		const bool DoesLocalizedPackageExist = FPackageName::DoesPackageExist(PackageName, CompatibleGuid, &LocalizedFilename, true);
+
+		if( (GIsEditor && !DoesNativePackageExist) || (!GIsEditor && !DoesLocalizedPackageExist && !DoesNativePackageExist) )
 		{
 			FFormatNamedArguments Arguments;
 			Arguments.Add(TEXT("Filename"), FText::FromString(InLongPackageName));
@@ -519,6 +562,25 @@ FLinkerLoad* GetPackageLinker
 			return nullptr;
 		}
 
+		// The editor must not redirect packages for localization.
+		if (GIsEditor)
+		{
+			NewFilename = NativeFilename;
+		}
+		else
+		{
+			// Use the localized package if possible.
+			if (DoesLocalizedPackageExist)
+			{
+				NewFilename = LocalizedFilename;
+			}
+			// If we are the game, we can fallback to the native package.
+			else
+			{
+				NewFilename = NativeFilename;
+			}
+		}
+
 		// Create the package with the provided long package name.
 		UPackage* FilenamePkg = (ExistingPackage ? ExistingPackage : CreatePackage(nullptr, *PackageName));
 		if (FilenamePkg != ExistingPackage && (LoadFlags & LOAD_PackageForPIE))
diff --git a/Engine/Source/Runtime/CoreUObject/Public/Misc/PackageName.h b/Engine/Source/Runtime/CoreUObject/Public/Misc/PackageName.h
--- a/Engine/Source/Runtime/CoreUObject/Public/Misc/PackageName.h
+++ b/Engine/Source/Runtime/CoreUObject/Public/Misc/PackageName.h
@@ -56,7 +56,7 @@ public:
 	 * @param InExtension Package extension.
 	 * @return Package filename.
 	 */
-	static FString LongPackageNameToFilename(const FString& InLongPackageName, const FString& InExtension = TEXT(""));
+	static FString LongPackageNameToFilename(const FString& InLongPackageName, const FString& InExtension = TEXT(""), const bool ShouldGetLocalizedPackage = false);
 	/** 
 	 * Returns the path to the specified package, excluding the short package name
 	 *
@@ -185,7 +185,7 @@ public:
 	 * @param OutFilename Package filename on disk.
 	 * @return true if the specified package name points to an existing package, false otherwise.
 	 **/
-	static bool DoesPackageExist(const FString& LongPackageName, const FGuid* Guid = NULL, FString* OutFilename = NULL);
+	static bool DoesPackageExist(const FString& LongPackageName, const FGuid* Guid = NULL, FString* OutFilename = NULL, const bool ShouldGetLocalizedPackage = false);
 
 	/**
 	 * Attempts to find a package given its short name on disk (very slow).
