Replacing a for loop with memcpy when filling FGCTask's objects to serialize

diff --git a/Engine/Source/Runtime/CoreUObject/Private/UObject/GarbageCollection.cpp b/Engine/Source/Runtime/CoreUObject/Private/UObject/GarbageCollection.cpp
--- a/Engine/Source/Runtime/CoreUObject/Private/UObject/GarbageCollection.cpp
+++ b/Engine/Source/Runtime/CoreUObject/Private/UObject/GarbageCollection.cpp
@@ -533,12 +533,8 @@ private:
 		FGCTask(FArchiveRealtimeGC* InOwner, const TArray<UObject*>* InObjectsToSerialize,int32 StartIndex, int32 NumObjects)
 			: Owner(InOwner)
 		{
-			ObjectsToSerialize.Reserve(NumObjects);
-			while (NumObjects)
-			{
-				ObjectsToSerialize.Add((*InObjectsToSerialize)[StartIndex++]);
-				NumObjects--;
-			}
+			ObjectsToSerialize.AddUninitialized(NumObjects);
+			FMemory::Memcpy(ObjectsToSerialize.GetData(), InObjectsToSerialize->GetData() + StartIndex, NumObjects * sizeof(UObject*));
 		}
 		FORCEINLINE TStatId GetStatId() const
 		{
