Cis fix

diff --git a/Engine/Source/Runtime/CoreUObject/Private/UObject/PropertyBool.cpp b/Engine/Source/Runtime/CoreUObject/Private/UObject/PropertyBool.cpp
--- a/Engine/Source/Runtime/CoreUObject/Private/UObject/PropertyBool.cpp
+++ b/Engine/Source/Runtime/CoreUObject/Private/UObject/PropertyBool.cpp
@@ -189,7 +189,7 @@ void UBoolProperty::ExportTextItem( FString& ValueStr, const void* PropertyValue
 	check(FieldSize != 0);
 	const uint8* ByteValue = (uint8*)PropertyValue + ByteOffset;
 	const bool bValue = 0 != ((*ByteValue) & FieldMask);
-	TCHAR* Temp = nullptr;
+	const TCHAR* Temp = nullptr;
 	if (0 != (PortFlags & PPF_ExportCpp))
 	{
 		Temp = (bValue ? TEXT("true") : TEXT("false"));
