Thread safety changes for object delete listener

diff --git a/Engine/Source/Runtime/CoreUObject/Private/UObject/UObjectArray.cpp b/Engine/Source/Runtime/CoreUObject/Private/UObject/UObjectArray.cpp
--- a/Engine/Source/Runtime/CoreUObject/Private/UObject/UObjectArray.cpp
+++ b/Engine/Source/Runtime/CoreUObject/Private/UObject/UObjectArray.cpp
@@ -147,6 +147,9 @@ void FUObjectArray::RemoveUObjectCreateListener(FUObjectCreateListener* Listener
  */
 void FUObjectArray::AddUObjectDeleteListener(FUObjectDeleteListener* Listener)
 {
+#if THREADSAFE_UOBJECTS
+	FScopeLock UObjectDeleteListenersLock(&UObjectDeleteListenersCritical);
+#endif
 	check(!UObjectDeleteListeners.Contains(Listener));
 	UObjectDeleteListeners.Add(Listener);
 }
@@ -158,6 +161,9 @@ void FUObjectArray::AddUObjectDeleteListener(FUObjectDeleteListener* Listener)
  */
 void FUObjectArray::RemoveUObjectDeleteListener(FUObjectDeleteListener* Listener)
 {
+#if THREADSAFE_UOBJECTS
+	FScopeLock UObjectDeleteListenersLock(&UObjectDeleteListenersCritical);
+#endif
 	UObjectDeleteListeners.RemoveSingleSwap(Listener);
 }
 
diff --git a/Engine/Source/Runtime/CoreUObject/Public/UObject/UObjectAnnotation.h b/Engine/Source/Runtime/CoreUObject/Public/UObject/UObjectAnnotation.h
--- a/Engine/Source/Runtime/CoreUObject/Public/UObject/UObjectAnnotation.h
+++ b/Engine/Source/Runtime/CoreUObject/Public/UObject/UObjectAnnotation.h
@@ -36,6 +36,7 @@ public:
 #if !(UE_BUILD_SHIPPING || UE_BUILD_TEST)
 		if (!bAutoRemove)
 		{
+			FScopeLock AnntationMapLock(&AnnotationMapCritical);
 			// in this case we are only verifying that the external assurances of removal are met
 			check(!AnnotationMap.Find(Object));
 		}
@@ -73,6 +74,7 @@ public:
 	void AddAnnotation(const UObjectBase *Object,TAnnotation Annotation)
 	{
 		check(Object);
+		FScopeLock AnntationMapLock(&AnnotationMapCritical);
 		AnnotationCacheKey = Object;
 		AnnotationCacheValue = Annotation;
 		if (Annotation.IsDefault())
@@ -101,11 +103,12 @@ public:
 	 * @return				Old annotation
 	 */
 	TAnnotation GetAndRemoveAnnotation(const UObjectBase *Object)
-	{
+	{		
 		check(Object);
+		FScopeLock AnntationMapLock(&AnnotationMapCritical);
 		AnnotationCacheKey = Object;
 		AnnotationCacheValue = TAnnotation();
-		bool bHadElements = (AnnotationMap.Num() > 0);
+		const bool bHadElements = (AnnotationMap.Num() > 0);
 		TAnnotation Result;
 		AnnotationMap.RemoveAndCopyValue(AnnotationCacheKey, Result);
 		if (bHadElements && AnnotationMap.Num() == 0)
@@ -128,9 +131,10 @@ public:
 	void RemoveAnnotation(const UObjectBase *Object)
 	{
 		check(Object);
+		FScopeLock AnntationMapLock(&AnnotationMapCritical);
 		AnnotationCacheKey = Object;
 		AnnotationCacheValue = TAnnotation();
-		bool bHadElements = (AnnotationMap.Num() > 0);
+		const bool bHadElements = (AnnotationMap.Num() > 0);
 		AnnotationMap.Remove(AnnotationCacheKey);
 		if (bHadElements && AnnotationMap.Num() == 0)
 		{
@@ -149,9 +153,10 @@ public:
 	 */
 	void RemoveAllAnnotations()
 	{
+		FScopeLock AnntationMapLock(&AnnotationMapCritical);
 		AnnotationCacheKey = NULL;
 		AnnotationCacheValue = TAnnotation();
-		bool bHadElements = (AnnotationMap.Num() > 0);
+		const bool bHadElements = (AnnotationMap.Num() > 0);
 		AnnotationMap.Empty();
 		if (bHadElements)
 		{
@@ -173,8 +178,9 @@ public:
 	FORCEINLINE TAnnotation GetAnnotation(const UObjectBase *Object)
 	{
 		check(Object);
+		FScopeLock AnntationMapLock(&AnnotationMapCritical);
 		if (Object != AnnotationCacheKey)
-		{
+		{			
 			AnnotationCacheKey = Object;
 			TAnnotation* Entry = AnnotationMap.Find(AnnotationCacheKey);
 			if (Entry)
@@ -203,6 +209,7 @@ public:
 	 */
 	void Reserve(int32 ExpectedNumElements)
 	{
+		FScopeLock AnntationMapLock(&AnnotationMapCritical);
 		AnnotationMap.Empty(ExpectedNumElements);
 	}
 
@@ -212,6 +219,7 @@ private:
 	 * Map from live objects to an annotation
 	 */
 	TMap<const UObjectBase *,TAnnotation> AnnotationMap;
+	FCriticalSection AnnotationMapCritical;
 
 	/**
 	 * Key for a one-item cache of the last lookup into AnnotationMap.
diff --git a/Engine/Source/Runtime/CoreUObject/Public/UObject/UObjectArray.h b/Engine/Source/Runtime/CoreUObject/Public/UObject/UObjectArray.h
--- a/Engine/Source/Runtime/CoreUObject/Public/UObject/UObjectArray.h
+++ b/Engine/Source/Runtime/CoreUObject/Public/UObject/UObjectArray.h
@@ -414,6 +414,9 @@ private:
 	 * Array of things to notify when a UObjectBase is destroyed
 	 */
 	TArray<FUObjectDeleteListener* > UObjectDeleteListeners;
+#if THREADSAFE_UOBJECTS
+	FCriticalSection UObjectDeleteListenersCritical;
+#endif
 };
 
 /** Global UObject allocator							*/
