[UE-17675] Deferred initialization of cyclic CDO objects now runs as if RF_NeedLoad was not set (as it wouldn't have been when it was deferred).

diff --git a/Engine/Source/Runtime/CoreUObject/Private/UObject/UObjectGlobals.cpp b/Engine/Source/Runtime/CoreUObject/Private/UObject/UObjectGlobals.cpp
--- a/Engine/Source/Runtime/CoreUObject/Private/UObject/UObjectGlobals.cpp
+++ b/Engine/Source/Runtime/CoreUObject/Private/UObject/UObjectGlobals.cpp
@@ -2324,7 +2324,18 @@ void FObjectInitializer::PostConstructInit()
 
 	bool bNeedInstancing = false;
 	// if HasAnyFlags(RF_NeedLoad), we do these steps later
+#if !USE_CIRCULAR_DEPENDENCY_LOAD_DEFERRING
 	if (!Obj->HasAnyFlags(RF_NeedLoad))
+#else 
+	// we defer this initialization in special set of cases (when Obj is a CDO 
+	// and its parent hasn't been serialized yet)... in those cases, Obj (the 
+	// CDO) wouldn't have had RF_NeedLoad set (not yet, because it is created 
+	// from Class->GetDefualtObject() without that flag); since we've deferred
+	// all this, it is likely that this flag is now present... these steps 
+	// (specifically sub-object instancing) is important for us to run on the
+	// CDO, so we allow all this when the bIsDeferredInitializer is true as well
+	if (!Obj->HasAnyFlags(RF_NeedLoad) || bIsDeferredInitializer)
+#endif // !USE_CIRCULAR_DEPENDENCY_LOAD_DEFERRING
 	{
 		if (bIsCDO || Class->HasAnyClassFlags(CLASS_PerObjectConfig))
 		{
@@ -2356,7 +2367,18 @@ void FObjectInitializer::PostConstructInit()
 		UE_LOG(LogUObjectGlobals, Fatal, TEXT("%s failed to route PostInitProperties. Call Super::PostInitProperties() in %s::PostInitProperties()."), *Obj->GetClass()->GetName(), *Obj->GetClass()->GetName());
 	}
 	// Check if all TSubobjectPtr properties have been initialized.
+#if !USE_CIRCULAR_DEPENDENCY_LOAD_DEFERRING
 	if (!Obj->HasAnyFlags(RF_NeedLoad))
+#else 
+	// we defer this initialization in special set of cases (when Obj is a CDO 
+	// and its parent hasn't been serialized yet)... in those cases, Obj (the 
+	// CDO) wouldn't have had RF_NeedLoad set (not yet, because it is created 
+	// from Class->GetDefualtObject() without that flag); since we've deferred
+	// all this, it is likely that this flag is now present... we want to run 
+	// all this as if the object was just created, so we check 
+	// bIsDeferredInitializer as well
+	if (!Obj->HasAnyFlags(RF_NeedLoad) || bIsDeferredInitializer)
+#endif // !USE_CIRCULAR_DEPENDENCY_LOAD_DEFERRING
 	{
 		for (UProperty* P = Class->RefLink; P; P = P->NextRef)
 		{
@@ -2381,10 +2403,22 @@ void FObjectInitializer::PostConstructInit()
 			}
 		}
 	}
-#endif
-	if (!Obj->HasAnyFlags(RF_NeedLoad) 
+#endif // !(UE_BUILD_SHIPPING || UE_BUILD_TEST)
+
+#if !USE_CIRCULAR_DEPENDENCY_LOAD_DEFERRING
+	if ( !Obj->HasAnyFlags(RF_NeedLoad)
+#else 
+	// we defer this initialization in special set of cases (when Obj is a CDO 
+	// and its parent hasn't been serialized yet)... in those cases, Obj (the 
+	// CDO) wouldn't have had RF_NeedLoad set (not yet, because it is created 
+	// from Class->GetDefualtObject() without that flag); since we've deferred
+	// all this, it is likely that this flag is now present... we want to run 
+	// all this as if the object was just created, so we check 
+	// bIsDeferredInitializer as well
+	if ( (!Obj->HasAnyFlags(RF_NeedLoad) || bIsDeferredInitializer)
+#endif // !USE_CIRCULAR_DEPENDENCY_LOAD_DEFERRING
 		// if component instancing is not enabled, then we leave the components in an invalid state, which will presumably be fixed by the caller
-		&& ((InstanceGraph == NULL) || InstanceGraph->IsSubobjectInstancingEnabled())) 
+		&& ((InstanceGraph == NULL) || InstanceGraph->IsSubobjectInstancingEnabled()) ) 
 	{
 		Obj->CheckDefaultSubobjects();
 	}
