GC will now spawn sub-tasks as soon as there's enough work without waiting for the parent task to finish processing all objects

diff --git a/Engine/Source/Runtime/Core/Public/Containers/Array.h b/Engine/Source/Runtime/Core/Public/Containers/Array.h
--- a/Engine/Source/Runtime/Core/Public/Containers/Array.h
+++ b/Engine/Source/Runtime/Core/Public/Containers/Array.h
@@ -1625,6 +1625,16 @@ public:
 		}
 	}
 
+	/**
+	 * Does nothing except setting the new number of elements in the array. Does not destruct items, does not de-allocate memory.
+	 * @param NewNum New number of elements in the array, must be <= the current number of elements in the array.
+	 */
+	void SetNumUnsafeInternal(int32 NewNum)
+	{
+		checkSlow(NewNum <= Num() && NewNum >= 0);
+		ArrayNum = NewNum;
+	}
+
 	/**
 	 * Appends the specified array to this array.
 	 *
diff --git a/Engine/Source/Runtime/CoreUObject/Private/UObject/GarbageCollection.cpp b/Engine/Source/Runtime/CoreUObject/Private/UObject/GarbageCollection.cpp
--- a/Engine/Source/Runtime/CoreUObject/Private/UObject/GarbageCollection.cpp
+++ b/Engine/Source/Runtime/CoreUObject/Private/UObject/GarbageCollection.cpp
@@ -676,7 +676,6 @@ public:
 
 		const int32 MinDesiredObjectsPerSubTask = 128; // sometimes there will be less, a lot less
 		const int32 NewObjectsArrayLength = InObjectsToSerializeArray.Num() * 2;
-		int32 TotalObjectsSerialized = InObjectsToSerializeArray.Num();
 
 		/** Growing array of objects that require serialization */
 		TArray<UObject*>	NewObjectsToSerializeArray;
@@ -875,6 +874,18 @@ public:
 				}
 				check(StackEntry == Stack.GetData());
 
+				if( GIsRunningParallelReachability && NewObjectsToSerialize.Num() >= MinDesiredObjectsPerSubTask )
+				{
+					// This will start queueing task with objects from the end of array until there's less objects than worth to queue
+					const int32 ObjectsPerSubTask = FMath::Max<int32>(MinDesiredObjectsPerSubTask, NewObjectsToSerialize.Num() / FTaskGraphInterface::Get().GetNumWorkerThreads());					
+					while (NewObjectsToSerialize.Num() >= MinDesiredObjectsPerSubTask)
+					{						
+						const int32 StartIndex = FMath::Max(0, NewObjectsToSerialize.Num() - ObjectsPerSubTask);
+						const int32 NumThisTask = NewObjectsToSerialize.Num() - StartIndex;
+						MyCompletionGraphEvent->DontCompleteUntil(TGraphTask<FGCTask>::CreateTask().ConstructAndDispatchWhenReady(this, &NewObjectsToSerialize, StartIndex, NumThisTask));
+						NewObjectsToSerialize.SetNumUnsafeInternal(StartIndex);
+					}
+				}
 #if PERF_DETAILED_PER_CLASS_GC_STATS
 				// Detailed per class stats should not be performed when parallel GC is running
 				check( !GIsRunningParallelReachability );
@@ -904,14 +915,15 @@ public:
 #endif
 			if( GIsRunningParallelReachability && NewObjectsToSerialize.Num() >= MinDesiredObjectsPerSubTask )
 			{			
-				int32 ObjectsPerSubTask = FMath::Max<int32>(MinDesiredObjectsPerSubTask,NewObjectsToSerialize.Num() / FTaskGraphInterface::Get().GetNumWorkerThreads());
+				const int32 ObjectsPerSubTask = FMath::Max<int32>(MinDesiredObjectsPerSubTask,NewObjectsToSerialize.Num() / FTaskGraphInterface::Get().GetNumWorkerThreads());
 				int32 StartIndex = 0;
 				while (StartIndex < NewObjectsToSerialize.Num())
 				{
-					int32 NumThisTask = FMath::Min<int32>(ObjectsPerSubTask, NewObjectsToSerialize.Num() - StartIndex);
+					const int32 NumThisTask = FMath::Min<int32>(ObjectsPerSubTask, NewObjectsToSerialize.Num() - StartIndex);
 					MyCompletionGraphEvent->DontCompleteUntil(TGraphTask<FGCTask>::CreateTask().ConstructAndDispatchWhenReady(this, &NewObjectsToSerialize, StartIndex, NumThisTask));
 					StartIndex += NumThisTask;
 				}
+				NewObjectsToSerialize.SetNumUnsafeInternal(0);
 			}
 			else if( NewObjectsToSerialize.Num() )
 			{
@@ -919,10 +931,9 @@ public:
 				// To avoid allocating and moving memory around swap ObjectsToSerialize and NewObjectsToSerialize arrays
 				Exchange( ObjectsToSerialize, NewObjectsToSerialize );
 				// Empty but don't free allocated memory
-				NewObjectsToSerialize.Reset();
+				NewObjectsToSerialize.SetNumUnsafeInternal(0);
 
 				CurrentIndex = 0;
-				TotalObjectsSerialized += NewObjectsToSerialize.Num();
 			}
 		}
 		while( CurrentIndex < ObjectsToSerialize.Num() );
@@ -1246,9 +1257,9 @@ void CollectGarbageInternal(EObjectFlags KeepFlags, bool bPerformFullPurge)
 
 	// Flush streaming before GC if requested
 	if (GFlushStreamingOnGC)
-		{
-			FlushAsyncLoading();
-		}
+	{
+		FlushAsyncLoading();
+	}
 
 	// Route callbacks so we can ensure that we are e.g. not in the middle of loading something by flushing
 	// the async loading, etc...
@@ -1269,8 +1280,9 @@ void CollectGarbageInternal(EObjectFlags KeepFlags, bool bPerformFullPurge)
 	check( !GObjPurgeIsRequired );
 
 #if VERIFY_DISREGARD_GC_ASSUMPTIONS
+	FUObjectArray& UObjectArray = GetUObjectArray();
 	// Only verify assumptions if option is enabled. This avoids false positives in the Editor or commandlets.
-	if( GetUObjectArray().DisregardForGCEnabled() && GShouldVerifyGCAssumptions )
+	if (UObjectArray.DisregardForGCEnabled() && GShouldVerifyGCAssumptions)
 	{
 		DECLARE_SCOPE_CYCLE_COUNTER( TEXT( "CollectGarbageInternal.VerifyGCAssumptions" ), STAT_CollectGarbageInternal_VerifyGCAssumptions, STATGROUP_GC );
 
@@ -1282,7 +1294,7 @@ void CollectGarbageInternal(EObjectFlags KeepFlags, bool bPerformFullPurge)
 			// Don't require UGCObjectReferencer's references to adhere to the assumptions.
 			// Although we want the referencer itself to sit in the disregard for gc set, most of the objects
 			// it's referencing will not be in the root set.
-			if (GetUObjectArray().IsDisregardForGC(Object) && !Object->IsA(UGCObjectReferencer::StaticClass()))
+			if (UObjectArray.IsDisregardForGC(Object) && !Object->IsA(UGCObjectReferencer::StaticClass()))
 			{
 				// Serialize object with reference collector.
 				TArray<UObject*> CollectedReferences;
@@ -1293,7 +1305,7 @@ void CollectGarbageInternal(EObjectFlags KeepFlags, bool bPerformFullPurge)
 				for( int32 ReferenceIndex=0; ReferenceIndex<CollectedReferences.Num(); ReferenceIndex++ )
 				{
 					UObject* ReferencedObject = CollectedReferences[ReferenceIndex];
-					if( ReferencedObject && !(ReferencedObject->HasAnyFlags(RF_RootSet) || GetUObjectArray().IsDisregardForGC(ReferencedObject)))
+					if (ReferencedObject && !(ReferencedObject->HasAnyFlags(RF_RootSet) || UObjectArray.IsDisregardForGC(ReferencedObject)))
 					{
 						UE_LOG(LogGarbage, Warning, TEXT("Disregard for GC object %s referencing %s which is not part of root set"),
 							*Object->GetFullName(),
@@ -1339,25 +1351,25 @@ void CollectGarbageInternal(EObjectFlags KeepFlags, bool bPerformFullPurge)
 	{
 		DECLARE_SCOPE_CYCLE_COUNTER( TEXT( "CollectGarbageInternal.UnhashUnreachable" ), STAT_CollectGarbageInternal_UnhashUnreachable, STATGROUP_GC );
 
-	// Unhash all unreachable objects.
-	const double StartTime = FPlatformTime::Seconds();
-	for ( FRawObjectIterator It(true); It; ++It )
-	{
-		//@todo UE4 - A prefetch was removed here. Re-add it. It wasn't right anyway, since it was ten items ahead and the consoles on have 8 prefetch slots
-
-		UObject* Object = *It;
-		if( Object->HasAnyFlags( RF_Unreachable ) )
+		// Unhash all unreachable objects.
+		const double StartTime = FPlatformTime::Seconds();
+		for ( FRawObjectIterator It(true); It; ++It )
 		{
-			// Begin the object's asynchronous destruction.
-			Object->ConditionalBeginDestroy();
-		}
+			//@todo UE4 - A prefetch was removed here. Re-add it. It wasn't right anyway, since it was ten items ahead and the consoles on have 8 prefetch slots
+
+			UObject* Object = *It;
+			if( Object->HasAnyFlags( RF_Unreachable ) )
+			{
+				// Begin the object's asynchronous destruction.
+				Object->ConditionalBeginDestroy();
+			}
 			else if (Object->HasAnyFlags(RF_NoStrongReference))
 			{
 				Object->ClearFlags(RF_NoStrongReference);
 				Object->SetFlags(RF_PendingKill);
 			}
-	}
-	UE_LOG(LogGarbage, Log, TEXT("%f ms for unhashing unreachable objects"), (FPlatformTime::Seconds() - StartTime) * 1000 );
+		}
+		UE_LOG(LogGarbage, Log, TEXT("%f ms for unhashing unreachable objects"), (FPlatformTime::Seconds() - StartTime) * 1000 );
 	}
 
 	// Set flag to indicate that we are relying on a purge to be performed.
