Native Specifier Flags are ignored when function signatures are compared.

diff --git a/Engine/Source/Runtime/CoreUObject/Public/UObject/Class.h b/Engine/Source/Runtime/CoreUObject/Public/UObject/Class.h
--- a/Engine/Source/Runtime/CoreUObject/Public/UObject/Class.h
+++ b/Engine/Source/Runtime/CoreUObject/Public/UObject/Class.h
@@ -1291,7 +1291,9 @@ public:
 	FORCEINLINE static uint64 GetDefaultIgnoredSignatureCompatibilityFlags()
 	{
 		//@TODO: UCREMOVAL: CPF_ConstParm added as a hack to get blueprints compiling with a const DamageType parameter.
-		const uint64 IgnoreFlags = CPF_PersistentInstance | CPF_ExportObject | CPF_InstancedReference | CPF_ContainsInstancedReference | CPF_ComputedFlags | CPF_ConstParm | CPF_HasGetValueTypeHash | CPF_UObjectWrapper;
+		const uint64 IgnoreFlags = CPF_PersistentInstance | CPF_ExportObject | CPF_InstancedReference 
+			| CPF_ContainsInstancedReference | CPF_ComputedFlags | CPF_ConstParm | CPF_HasGetValueTypeHash | CPF_UObjectWrapper
+			| CPF_NativeAccessSpecifiers;
 		return IgnoreFlags;
 	}
 
