Merging CL #2599282 by [REDACTED_PERSON][REDACTED_HANDLE] on 6/24/2015 2:34:03 PM

#UE4 Change the order of the async loading tick to process loaded packages before reading in new ones. This stops it from preloading all packages and exploding memory usage before it processes any of the loaded ones.
Change it so newly requested async load packages are placed after packages of equal priority, this stops them from interrupting in progress loads of the same priority

diff --git a/Engine/Source/Runtime/CoreUObject/Private/Serialization/AsyncLoading.cpp b/Engine/Source/Runtime/CoreUObject/Private/Serialization/AsyncLoading.cpp
--- a/Engine/Source/Runtime/CoreUObject/Private/Serialization/AsyncLoading.cpp
+++ b/Engine/Source/Runtime/CoreUObject/Private/Serialization/AsyncLoading.cpp
@@ -535,7 +535,7 @@ void FAsyncLoadingThread::InsertPackage(FAsyncPackage* Package)
 		FScopeLock LockAsyncPackages(&AsyncPackagesCritical);
 #endif
 		// Insert new package keeping descending priority order in AsyncPackages
-		auto InsertIndex = AsyncPackages.IndexOfByPredicate([Package](const FAsyncPackage* Element)
+		int32 InsertIndex = AsyncPackages.IndexOfByPredicate([Package](const FAsyncPackage* Element)
 		{
 			// NOTE: Switched to < because we want newly inserted packages to go AFTER all existing packages of the 
 			//		 same priority. If this is too slow, we could probably do "LastIndexOfByPredicate" where we iterate backwards??
@@ -700,18 +700,17 @@ EAsyncPackageState::Type FAsyncLoadingThread::TickAsyncLoading(bool bUseTimeLimi
 	if (!bLoadingSuspended)
 	{
 		double TickStartTime = FPlatformTime::Seconds();
-		double TimeLimitUsedForPreloading = 0;
+		double TimeLimitUsedForProcessLoaded = 0;
 
-		if (!bIsMultithreaded)
 		{
-			Result = TickAsyncThread(bUseTimeLimit, bUseFullTimeLimit, TimeLimit);
-			TimeLimitUsedForPreloading = FPlatformTime::Seconds() - TickStartTime;
+			Result = ProcessLoadedPackages(bUseTimeLimit, bUseFullTimeLimit, TimeLimit, WaitForRequestID);
+			TimeLimitUsedForProcessLoaded = FPlatformTime::Seconds() - TickStartTime;
 		}
 
-		if (Result != EAsyncPackageState::TimeOut && !IsTimeLimitExceeded(TickStartTime, bUseTimeLimit, TimeLimit, TEXT("Pre-ProcessLoadedPackages")))
+		if (!bIsMultithreaded && Result != EAsyncPackageState::TimeOut && !IsTimeLimitExceeded(TickStartTime, bUseTimeLimit, TimeLimit, TEXT("Pre-TickAsyncThread")))
 		{
-			double RemainingTimeLimit = FMath::Max(0.0, TimeLimit - TimeLimitUsedForPreloading);
-			Result = ProcessLoadedPackages(bUseTimeLimit, bUseFullTimeLimit, RemainingTimeLimit, WaitForRequestID);
+			double RemainingTimeLimit = FMath::Max(0.0, TimeLimit - TimeLimitUsedForProcessLoaded);
+			Result = TickAsyncThread(bUseTimeLimit, bUseFullTimeLimit, RemainingTimeLimit);			
 		}
 
 		if (Result != EAsyncPackageState::TimeOut && !IsTimeLimitExceeded(TickStartTime, bUseTimeLimit, TimeLimit, TEXT("Pre-EmptyReferencedObjects")))
