Moving linker serialization from UObject to UPackage to fix UObject mem reporting

diff --git a/Engine/Source/Runtime/CoreUObject/Private/UObject/Obj.cpp b/Engine/Source/Runtime/CoreUObject/Private/UObject/Obj.cpp
--- a/Engine/Source/Runtime/CoreUObject/Private/UObject/Obj.cpp
+++ b/Engine/Source/Runtime/CoreUObject/Private/UObject/Obj.cpp
@@ -895,12 +895,6 @@ void UObject::Serialize( FArchive& Ar )
 		{
 			Ar << ObjClass;
 		}
-		//@todo UE4 - This seems to be required and it should not be. Seems to be related to the texture streamer.
-		FLinkerLoad* LinkerLoad = GetLinker();
-		if (LinkerLoad)
-		{
-			LinkerLoad->Serialize(Ar);
-		}
 	}
 	// Special support for supporting undo/redo of renaming and changing Archetype.
 	else if( Ar.IsTransacting() )
diff --git a/Engine/Source/Runtime/CoreUObject/Private/UObject/Package.cpp b/Engine/Source/Runtime/CoreUObject/Private/UObject/Package.cpp
--- a/Engine/Source/Runtime/CoreUObject/Private/UObject/Package.cpp
+++ b/Engine/Source/Runtime/CoreUObject/Private/UObject/Package.cpp
@@ -86,6 +86,14 @@ void UPackage::Serialize( FArchive& Ar )
 	{
 		Ar << bDirty;
 	}
+	if (Ar.IsCountingMemory())
+	{		
+		if (LinkerLoad)
+		{
+			FLinker* Loader = LinkerLoad;
+			Loader->Serialize(Ar);
+		}
+	}
 }
 
 void UPackage::AddReferencedObjects(UObject* InThis, FReferenceCollector& Collector)
