Fixed improper pointer math for accessing the value of a fixed sized array property in gathering localization data from text properties.

diff --git a/Engine/Source/Runtime/CoreUObject/Private/Serialization/PropertyLocalizationDataGathering.cpp b/Engine/Source/Runtime/CoreUObject/Private/Serialization/PropertyLocalizationDataGathering.cpp
--- a/Engine/Source/Runtime/CoreUObject/Private/Serialization/PropertyLocalizationDataGathering.cpp
+++ b/Engine/Source/Runtime/CoreUObject/Private/Serialization/PropertyLocalizationDataGathering.cpp
@@ -75,7 +75,7 @@ void FPropertyLocalizationDataGatherer::GatherLocalizationDataFromTextProperty(c
 	for(int32 i = 0; i < TextProperty->ArrayDim; ++i)
 	{
 		const FString PathToElement = FString(PathToParent.IsEmpty() ? TEXT("") : PathToParent + TEXT(".")) + (IsFixedSizeArray ? TextProperty->GetName() + FString::Printf(TEXT("[%d]"), i) : TextProperty->GetName());
-		FText* const ElementValueAddress = reinterpret_cast<FText*>(ValueAddress) + TextProperty->ElementSize * i;
+		FText* const ElementValueAddress = reinterpret_cast<FText*>(reinterpret_cast<uint8*>(ValueAddress) + TextProperty->ElementSize * i);
 
 		UPackage* const Package = TextProperty->GetOutermost();
 		if( FTextInspector::GetFlags(*ElementValueAddress) & ETextFlag::ConvertedProperty )
