Properties names in FBox2D USTRUCT are the same like in struct FBox2D implementation.

diff --git a/Engine/Source/Runtime/CoreUObject/Classes/Object.h b/Engine/Source/Runtime/CoreUObject/Classes/Object.h
--- a/Engine/Source/Runtime/CoreUObject/Classes/Object.h
+++ b/Engine/Source/Runtime/CoreUObject/Classes/Object.h
@@ -545,7 +545,7 @@ struct FBox2D
 	FVector2D Max;
 
 	UPROPERTY()
-	uint8 IsValid;
+	bool bIsValid;
 
 };
 
