#UE4 - increased lan packet size to 1024 so more data can fit in the server response

diff --git a/Engine/Source/Runtime/CoreUObject/Public/UObject/CoreNet.h b/Engine/Source/Runtime/CoreUObject/Public/UObject/CoreNet.h
--- a/Engine/Source/Runtime/CoreUObject/Public/UObject/CoreNet.h
+++ b/Engine/Source/Runtime/CoreUObject/Public/UObject/CoreNet.h
@@ -434,7 +434,7 @@ COREUOBJECT_API void SerializeChecksum(FArchive &Ar, uint32 x, bool ErrorOK);
 * Values used for initializing UNetConnection and LanBeacon
 */
 enum { MAX_PACKET_SIZE = 512 }; // MTU for the connection
-enum { LAN_BEACON_MAX_PACKET_SIZE = 512 }; // MTU for the connection
+enum { LAN_BEACON_MAX_PACKET_SIZE = 1024 }; // MTU for the connection
 
 /**
  * Functions to assist in detecting errors during RPC calls
