Added FPackageName::TryConvertLongPackageNameToFilename

This allows you to attempt convert a package name (or path) into a filename (or path) without generating an assert. FPackageName::LongPackageNameToFilename now calls this function and asserts on failure.

diff --git a/Engine/Source/Runtime/CoreUObject/Private/Misc/PackageName.cpp b/Engine/Source/Runtime/CoreUObject/Private/Misc/PackageName.cpp
--- a/Engine/Source/Runtime/CoreUObject/Private/Misc/PackageName.cpp
+++ b/Engine/Source/Runtime/CoreUObject/Private/Misc/PackageName.cpp
@@ -328,7 +328,7 @@ FString FPackageName::FilenameToLongPackageName(const FString& InFilename)
 	return Result;
 }
 
-FString FPackageName::LongPackageNameToFilename(const FString& InLongPackageName, const FString& InExtension, const bool ShouldGetLocalizedPackage)
+bool FPackageName::TryConvertLongPackageNameToFilename(const FString& InLongPackageName, FString& OutFilename, const FString& InExtension, const bool ShouldGetLocalizedPackage)
 {
 	const auto& Paths = FLongPackagePathsSingleton::Get();
 
@@ -341,14 +341,24 @@ FString FPackageName::LongPackageNameToFilename(const FString& InLongPackageName
 		{
 			const FString RootRelativePath = InLongPackageName.Mid(Pair.RootPath.Len());
 			const bool IsLocalizedPath = RootRelativePath.StartsWith(LocalizationPathParticle);
-			return Pair.ContentPath + (ShouldGetLocalizedPackage && !IsLocalizedPath ? LocalizationPathParticle : TEXT("")) + RootRelativePath + InExtension;
+			OutFilename = Pair.ContentPath + (ShouldGetLocalizedPackage && !IsLocalizedPath ? LocalizationPathParticle : TEXT("")) + RootRelativePath + InExtension;
+			return true;
 		}
 	}
 
 	// This is not a long package name or the root folder is not handled in the above cases
-	UE_LOG(LogPackageName, Fatal,TEXT("LongPackageNameToFilename failed to convert '%s'. Path does not map to any roots."), *InLongPackageName);
+	return false;
+}
 
-	return InLongPackageName;
+FString FPackageName::LongPackageNameToFilename(const FString& InLongPackageName, const FString& InExtension, const bool ShouldGetLocalizedPackage)
+{
+	FString FailureReason;
+	FString Result;
+	if (!TryConvertLongPackageNameToFilename(InLongPackageName, Result, InExtension, ShouldGetLocalizedPackage))
+	{
+		UE_LOG(LogPackageName, Fatal,TEXT("LongPackageNameToFilename failed to convert '%s'. Path does not map to any roots."), *InLongPackageName);
+	}
+	return Result;
 }
 
 FString FPackageName::GetLongPackagePath(const FString& InLongPackageName)
diff --git a/Engine/Source/Runtime/CoreUObject/Public/Misc/PackageName.h b/Engine/Source/Runtime/CoreUObject/Public/Misc/PackageName.h
--- a/Engine/Source/Runtime/CoreUObject/Public/Misc/PackageName.h
+++ b/Engine/Source/Runtime/CoreUObject/Public/Misc/PackageName.h
@@ -49,6 +49,14 @@ public:
 	 * @return Long package name.
 	 */
 	static FString FilenameToLongPackageName(const FString& InFilename);
+	/** 
+	 * Tries to convert a long package name to a file name with the supplied extension.
+	 *
+	 * @param InLongPackageName Long Package Name
+	 * @param InExtension Package extension.
+	 * @return Package filename.
+	 */
+	static bool TryConvertLongPackageNameToFilename(const FString& InLongPackageName, FString& OutFilename, const FString& InExtension = TEXT(""), const bool ShouldGetLocalizedPackage = false);
 	/** 
 	 * Converts a long package name to a file name with the supplied extension.
 	 *
