Make UEnums store values as uint8 instead of int8.

diff --git a/Engine/Source/Editor/UnrealEd/Private/Kismet2/EnumEditorUtils.cpp b/Engine/Source/Editor/UnrealEd/Private/Kismet2/EnumEditorUtils.cpp
--- a/Engine/Source/Editor/UnrealEd/Private/Kismet2/EnumEditorUtils.cpp
+++ b/Engine/Source/Editor/UnrealEd/Private/Kismet2/EnumEditorUtils.cpp
@@ -34,7 +34,7 @@ UEnum* FEnumEditorUtils::CreateUserDefinedEnum(UObject* InParent, FName EnumName
 
 	if (NULL != Enum)
 	{
-		TArray<TPair<FName, int8>> EmptyNames;
+		TArray<TPair<FName, uint8>> EmptyNames;
 		Enum->SetEnums(EmptyNames, UEnum::ECppForm::Namespaced);
 		Enum->SetMetaData(TEXT("BlueprintType"), TEXT("true"));
 	}
@@ -51,13 +51,13 @@ void FEnumEditorUtils::UpdateAfterPathChanged(UEnum* Enum)
 {
 	check(NULL != Enum);
 
-	TArray<TPair<FName, int8>> NewEnumeratorsNames;
+	TArray<TPair<FName, uint8>> NewEnumeratorsNames;
 	const int32 EnumeratorsToCopy = Enum->NumEnums() - 1; // skip _MAX
 	for (int32 Index = 0; Index < EnumeratorsToCopy; Index++)
 	{
 		const FString ShortName = Enum->GetEnumName(Index);
 		const FString NewFullName = Enum->GenerateFullEnumName(*ShortName);
-		NewEnumeratorsNames.Add(TPairInitializer<FName, int8>(FName(*NewFullName), Enum->GetValueByIndex(Index)));
+		NewEnumeratorsNames.Add(TPairInitializer<FName, uint8>(FName(*NewFullName), Enum->GetValueByIndex(Index)));
 	}
 
 	Enum->SetEnums(NewEnumeratorsNames, UEnum::ECppForm::Namespaced);
@@ -66,7 +66,7 @@ void FEnumEditorUtils::UpdateAfterPathChanged(UEnum* Enum)
 //////////////////////////////////////////////////////////////////////////
 // Enumerators
 
-void FEnumEditorUtils::CopyEnumeratorsWithoutMax(const UEnum* Enum, TArray<TPair<FName, int8>>& OutEnumNames)
+void FEnumEditorUtils::CopyEnumeratorsWithoutMax(const UEnum* Enum, TArray<TPair<FName, uint8>>& OutEnumNames)
 {
 	if (Enum == nullptr)
 	{
@@ -90,13 +90,13 @@ void FEnumEditorUtils::AddNewEnumeratorForUserDefinedEnum(UUserDefinedEnum* Enum
 
 	PrepareForChange(Enum);
 
-	TArray<TPair<FName, int8>> OldNames, Names;
+	TArray<TPair<FName, uint8>> OldNames, Names;
 	CopyEnumeratorsWithoutMax(Enum, OldNames);
 	Names = OldNames;
 
 	FString EnumNameString = Enum->GenerateNewEnumeratorName();
 	const FString FullNameStr = Enum->GenerateFullEnumName(*EnumNameString);
-	Names.Add(TPairInitializer<FName, int8>(FName(*FullNameStr), Enum->GetMaxEnumValue() + 1));
+	Names.Add(TPairInitializer<FName, uint8>(FName(*FullNameStr), Enum->GetMaxEnumValue() + 1));
 
 	const UEnum::ECppForm EnumType = Enum->GetCppForm();
 	Enum->SetEnums(Names, EnumType);
@@ -116,7 +116,7 @@ void FEnumEditorUtils::RemoveEnumeratorFromUserDefinedEnum(UUserDefinedEnum* Enu
 
 	PrepareForChange(Enum);
 
-	TArray<TPair<FName, int8>> OldNames, Names;
+	TArray<TPair<FName, uint8>> OldNames, Names;
 	CopyEnumeratorsWithoutMax(Enum, OldNames);
 	Names = OldNames;
 
@@ -142,7 +142,7 @@ void FEnumEditorUtils::MoveEnumeratorInUserDefinedEnum(UUserDefinedEnum* Enum, i
 
 	PrepareForChange(Enum);
 
-	TArray<TPair<FName, int8>> OldNames, Names;
+	TArray<TPair<FName, uint8>> OldNames, Names;
 	CopyEnumeratorsWithoutMax(Enum, OldNames);
 	Names = OldNames;
 
@@ -189,9 +189,9 @@ class FArchiveEnumeratorResolver : public FArchiveUObject
 {
 public:
 	const UEnum* Enum;
-	const TArray<TPair<FName, int8>>& OldNames;
+	const TArray<TPair<FName, uint8>>& OldNames;
 
-	FArchiveEnumeratorResolver(const UEnum* InEnum, const TArray<TPair<FName, int8>>& InOldNames)
+	FArchiveEnumeratorResolver(const UEnum* InEnum, const TArray<TPair<FName, uint8>>& InOldNames)
 		: FArchiveUObject(), Enum(InEnum), OldNames(InOldNames)
 	{
 	}
@@ -207,7 +207,7 @@ void FEnumEditorUtils::PrepareForChange(const UUserDefinedEnum* Enum)
 	FEnumEditorManager::Get().PreChange(Enum, EEnumEditorChangeInfo::Changed);
 }
 
-void FEnumEditorUtils::BroadcastChanges(const UUserDefinedEnum* Enum, const TArray<TPair<FName, int8>>& OldNames, bool bResolveData)
+void FEnumEditorUtils::BroadcastChanges(const UUserDefinedEnum* Enum, const TArray<TPair<FName, uint8>>& OldNames, bool bResolveData)
 {
 	check(NULL != Enum);
 	if (bResolveData)
@@ -351,7 +351,7 @@ bool FEnumEditorUtils::SetEnumeratorDisplayName(UUserDefinedEnum* Enum, int32 En
 			PrepareForChange(Enum);
 			Enum->SetMetaData(FEnumEditorUtilsHelper::DisplayName(), *NewDisplayName, EnumeratorIndex);
 			EnsureAllDisplayNamesExist(Enum);
-			BroadcastChanges(Enum, TArray<TPair<FName, int8>>(), false);
+			BroadcastChanges(Enum, TArray<TPair<FName, uint8>>(), false);
 			return true;
 		}
 	}
diff --git a/Engine/Source/Editor/UnrealEd/Public/Kismet2/EnumEditorUtils.h b/Engine/Source/Editor/UnrealEd/Public/Kismet2/EnumEditorUtils.h
--- a/Engine/Source/Editor/UnrealEd/Public/Kismet2/EnumEditorUtils.h
+++ b/Engine/Source/Editor/UnrealEd/Public/Kismet2/EnumEditorUtils.h
@@ -12,10 +12,10 @@ class UUserDefinedEnum;
 class UNREALED_API FEnumEditorUtils
 {
 	static void PrepareForChange(const UUserDefinedEnum* Enum);
-	static void BroadcastChanges(const UUserDefinedEnum* Enum, const TArray<TPair<FName, int8>>& OldNames, bool bResolveData = true);
+	static void BroadcastChanges(const UUserDefinedEnum* Enum, const TArray<TPair<FName, uint8>>& OldNames, bool bResolveData = true);
 
 	/** copy full enumeratos names from given enum to OutEnumNames, the last '_MAX' enumerator is skipped */
-	static void CopyEnumeratorsWithoutMax(const UEnum* Enum, TArray<TPair<FName, int8>>& OutEnumNames);
+	static void CopyEnumeratorsWithoutMax(const UEnum* Enum, TArray<TPair<FName, uint8>>& OutEnumNames);
 public:
 
 	enum EEnumEditorChangeInfo
diff --git a/Engine/Source/Programs/UnrealHeaderTool/Private/CodeGenerator.cpp b/Engine/Source/Programs/UnrealHeaderTool/Private/CodeGenerator.cpp
--- a/Engine/Source/Programs/UnrealHeaderTool/Private/CodeGenerator.cpp
+++ b/Engine/Source/Programs/UnrealHeaderTool/Private/CodeGenerator.cpp
@@ -2946,10 +2946,10 @@ void FNativeClassHeaderGenerator::ExportGeneratedEnumsInitCode(const TArray<UEnu
 			GeneratedEnumRegisterFunctionText.Logf(TEXT("        {\r\n"));
 
 			GeneratedEnumRegisterFunctionText.Logf(TEXT("            ReturnEnum = new(EC_InternalUseOnlyConstructor, Outer, TEXT(\"%s\"), RF_Public|RF_Transient|RF_Native) UEnum(FObjectInitializer());\r\n"), *Enum->GetName());
-			GeneratedEnumRegisterFunctionText.Logf(TEXT("            TArray<TPair<FName, int8>> EnumNames;\r\n"));
+			GeneratedEnumRegisterFunctionText.Logf(TEXT("            TArray<TPair<FName, uint8>> EnumNames;\r\n"));
 			for (int32 Index = 0; Index < Enum->NumEnums(); Index++)
 			{
-				GeneratedEnumRegisterFunctionText.Logf(TEXT("            EnumNames.Add(TPairInitializer<FName, int8>(FName(TEXT(\"%s\")), %d));\r\n"), *Enum->GetNameByIndex(Index).ToString(), Enum->GetValueByIndex(Index));
+				GeneratedEnumRegisterFunctionText.Logf(TEXT("            EnumNames.Add(TPairInitializer<FName, uint8>(FName(TEXT(\"%s\")), %d));\r\n"), *Enum->GetNameByIndex(Index).ToString(), Enum->GetValueByIndex(Index));
 			}
 
 			FString EnumTypeStr;
diff --git a/Engine/Source/Programs/UnrealHeaderTool/Private/HeaderParser.cpp b/Engine/Source/Programs/UnrealHeaderTool/Private/HeaderParser.cpp
--- a/Engine/Source/Programs/UnrealHeaderTool/Private/HeaderParser.cpp
+++ b/Engine/Source/Programs/UnrealHeaderTool/Private/HeaderParser.cpp
@@ -1175,7 +1175,7 @@ UEnum* FHeaderParser::CompileEnum(FUnrealSourceFile& SourceFile)
 	FToken TagToken;
 
 	TArray<FScriptLocation> EnumTagLocations;
-	TArray<TPair<FName, int8>> EnumNames;
+	TArray<TPair<FName, uint8>> EnumNames;
 
 	int32 CurrentEnumValue = 0;
 
@@ -1217,7 +1217,7 @@ UEnum* FHeaderParser::CompileEnum(FUnrealSourceFile& SourceFile)
 			break;
 		}
 
-		TPair<FName, int8> CurrentEnum = TPair<FName, int8>(TPairInitializer<FName, int8>(NewTag, CurrentEnumValue));
+		TPair<FName, uint8> CurrentEnum = TPair<FName, uint8>(TPairInitializer<FName, uint8>(NewTag, CurrentEnumValue));
 
 		if (EnumNames.Find(CurrentEnum, iFound))
 		{
diff --git a/Engine/Source/Runtime/CoreUObject/Private/UObject/Enum.cpp b/Engine/Source/Runtime/CoreUObject/Private/UObject/Enum.cpp
--- a/Engine/Source/Runtime/CoreUObject/Private/UObject/Enum.cpp
+++ b/Engine/Source/Runtime/CoreUObject/Private/UObject/Enum.cpp
@@ -24,10 +24,10 @@ void UEnum::Serialize( FArchive& Ar )
 	{
 		TArray<FName> TempNames;
 		Ar << TempNames;
-		int8 Value = 0;
+		uint8 Value = 0;
 		for (const FName& TempName : TempNames)
 		{
-			Names.Add(TPairInitializer<FName, int8>(TempName, Value++));
+			Names.Add(TPairInitializer<FName, uint8>(TempName, Value++));
 		}
 	}
 	else
@@ -43,7 +43,7 @@ void UEnum::Serialize( FArchive& Ar )
 	}
 	else
 	{
-		int8 EnumTypeByte = (int8)CppForm;
+		uint8 EnumTypeByte = (uint8)CppForm;
 		Ar << EnumTypeByte;
 		CppForm = (ECppForm)EnumTypeByte;
 	}
@@ -86,7 +86,7 @@ void UEnum::RenameNamesAfterDuplication()
 	FString ThisName = GetName();
 
 	// Replace all usages of base class name to the duplicated one.
-	for (TPair<FName, int8>& Kvp : Names)
+	for (TPair<FName, uint8>& Kvp : Names)
 	{
 		FString NameString = Kvp.Key.ToString();
 		NameString.ReplaceInline(*BaseEnumName, *ThisName);
@@ -104,7 +104,7 @@ FString UEnum::GenerateFullEnumName(const TCHAR* InEnumName) const
 	return (CppForm != ECppForm::Regular) ? GenerateFullEnumName(this, InEnumName) : InEnumName;
 }
 
-FName UEnum::GetNameByIndex(int8 Index) const
+FName UEnum::GetNameByIndex(uint8 Index) const
 {
 	if (Names.IsValidIndex(Index))
 	{
@@ -114,15 +114,15 @@ FName UEnum::GetNameByIndex(int8 Index) const
 	return NAME_None;
 }
 
-int8 UEnum::GetValueByIndex(int8 Index) const
+uint8 UEnum::GetValueByIndex(uint8 Index) const
 {
 	check(Names.IsValidIndex(Index));
 	return Names[Index].Value;
 }
 
-FName UEnum::GetNameByValue(int8 InValue) const
+FName UEnum::GetNameByValue(uint8 InValue) const
 {
-	for (TPair<FName, int8> Kvp : Names)
+	for (TPair<FName, uint8> Kvp : Names)
 	{
 		if (Kvp.Value == InValue)
 		{
@@ -147,25 +147,75 @@ int32 UEnum::GetIndexByName(FName InName) const
 	return INDEX_NONE;
 }
 
-int8 UEnum::GetMaxEnumValue() const
+int32 UEnum::GetValueByName(FName InName) const
 {
-	int8 MaxValue = -1;
-	if (Names.Num() > 0)
+	FString InNameString = InName.ToString();
+	FString DoubleColon = TEXT("::");
+	int32 DoubleColonPosition = InNameString.Find(DoubleColon);
+	bool bIsInNameFullyQualified = DoubleColonPosition != INDEX_NONE;
+	bool bIsEnumFullyQualified = CppForm != ECppForm::Regular;
+
+	if (!bIsInNameFullyQualified && bIsEnumFullyQualified)
+	{
+		// Make InName fully qualified.
+		InName = FName(*GetName().Append(DoubleColon).Append(InNameString));
+	}
+	else if (bIsInNameFullyQualified && !bIsEnumFullyQualified)
+	{
+		// Make InName unqualified.
+		InName = FName(*InNameString.Mid(DoubleColonPosition + 2, InNameString.Len() - DoubleColonPosition + 2));
+	}
+
+	for (TPair<FName, uint8> Kvp : Names)
+	{
+		if (Kvp.Key == InName)
+		{
+			return Kvp.Value;
+		}
+	}
+
+	return INDEX_NONE;
+}
+
+
+uint8 UEnum::GetMaxEnumValue() const
+{
+	int32 NamesNum = Names.Num();
+	if (NamesNum == 0)
 	{
-		MaxValue = Names[0].Value;
+		return INDEX_NONE;
 	}
 
-	for (const auto& Pair : Names)
+	uint8 MaxValue = Names[0].Value;
+	for (int32 i = 0; i < NamesNum; ++i)
 	{
-		if (Pair.Value > MaxValue)
+		uint8 CurrentValue = Names[i].Value;
+		if (CurrentValue > MaxValue)
 		{
-			MaxValue = Pair.Value;
+			MaxValue = CurrentValue;
 		}
 	}
 
 	return MaxValue;
 }
 
+bool UEnum::IsValidEnumValue(uint8 InValue) const
+{
+	int32 NamesNum = Names.Num();
+	for (int32 i = 0; i < NamesNum; ++i)
+	{
+		uint8 CurrentValue = Names[i].Value;
+
+		// Last entry is autogenerated _MAX value
+		if (CurrentValue == InValue && i != NamesNum - 1)
+		{
+			return true;
+		}
+	}
+
+	return false;
+}
+
 FString UEnum::GenerateFullEnumName(const UEnum* InEnum, const TCHAR* InEnumName)
 {
 	if (InEnum->GetCppForm() == ECppForm::Regular || IsFullEnumName(InEnumName))
@@ -178,7 +228,7 @@ FString UEnum::GenerateFullEnumName(const UEnum* InEnum, const TCHAR* InEnumName
 
 void UEnum::AddNamesToMasterList()
 {
-	for (TPair<FName, int8> Kvp : Names)
+	for (TPair<FName, uint8> Kvp : Names)
 	{
 		UEnum* Enum = AllEnumNames.FindRef(Kvp.Key);
 		if (Enum == nullptr)
@@ -194,7 +244,7 @@ void UEnum::AddNamesToMasterList()
 
 void UEnum::RemoveNamesFromMasterList()
 {
-	for (TPair<FName, int8> Kvp : Names)
+	for (TPair<FName, uint8> Kvp : Names)
 	{
 		UEnum* Enum = AllEnumNames.FindRef(Kvp.Key);
 		if (Enum == this)
@@ -388,7 +438,7 @@ int32 UEnum::FindEnumRedirects(const UEnum* Enum, FName EnumEntryName)
  *
  * @return	true unless the MAX enum already exists and isn't the last enum.
  */
-bool UEnum::SetEnums(TArray<TPair<FName, int8>>& InNames, UEnum::ECppForm InCppForm)
+bool UEnum::SetEnums(TArray<TPair<FName, uint8>>& InNames, UEnum::ECppForm InCppForm)
 {
 	if (Names.Num() > 0)
 	{
@@ -411,7 +461,7 @@ bool UEnum::SetEnums(TArray<TPair<FName, int8>>& InNames, UEnum::ECppForm InCppF
 			return false;
 		}
 
-		Names.Add(TPairInitializer<FName, int8>(MaxEnumItem, GetMaxEnumValue() + 1));
+		Names.Add(TPairInitializer<FName, uint8>(MaxEnumItem, GetMaxEnumValue() + 1));
 	}
 
 	AddNamesToMasterList();
diff --git a/Engine/Source/Runtime/CoreUObject/Public/UObject/Class.h b/Engine/Source/Runtime/CoreUObject/Public/UObject/Class.h
--- a/Engine/Source/Runtime/CoreUObject/Public/UObject/Class.h
+++ b/Engine/Source/Runtime/CoreUObject/Public/UObject/Class.h
@@ -1353,55 +1353,29 @@ public:
 	FString CppType;
 
 	/** Gets enum name by index in Names array. Returns NAME_None if Index is not valid. */
-	FName GetNameByIndex(int8 Index) const;
+	FName GetNameByIndex(uint8 Index) const;
 		
 	/** Gets enum value by index in Names array. */
-	int8 GetValueByIndex(int8 Index) const;
+	uint8 GetValueByIndex(uint8 Index) const;
 
 	/** Gets enum name by value. Returns NAME_None if value is not found. */
-	FName GetNameByValue(int8 InValue) const;
+	FName GetNameByValue(uint8 InValue) const;
 		
 	/** Gets enum value by name. Returns INDEX_NONE when name is not found. */
-	int8 GetValueByName(FName InName) const
-	{
-		FString InNameString = InName.ToString();
-		FString DoubleColon = TEXT("::");
-		int32 DoubleColonPosition = InNameString.Find(DoubleColon);
-		bool bIsInNameFullyQualified = DoubleColonPosition != INDEX_NONE;
-		bool bIsEnumFullyQualified = CppForm != ECppForm::Regular;
-
-		if (!bIsInNameFullyQualified && bIsEnumFullyQualified)
-		{
-			// Make InName fully qualified.
-			InName = FName(*GetName().Append(DoubleColon).Append(InNameString));
-		}
-		else if (bIsInNameFullyQualified && !bIsEnumFullyQualified)
-		{
-			// Make InName unqualified.
-			InName = FName(*InNameString.Mid(DoubleColonPosition + 2, InNameString.Len() - DoubleColonPosition + 2));
-		}
-
-		for (TPair<FName, int8> Kvp : Names)
-		{
-			if (Kvp.Key == InName)
-			{
-				return Kvp.Value;
-			}
-		}
-
-		return INDEX_NONE;
-	}
+	int32 GetValueByName(FName InName) const;
 
 	/** Gets index of name in enum. Returns INDEX_NONE when name is not found. */
 	int32 GetIndexByName(FName InName) const;
 
 	/** Gets max value of Enum. Defaults to zero if there are no entries. */
-	int8 GetMaxEnumValue() const;
+	uint8 GetMaxEnumValue() const;
+
+	bool IsValidEnumValue(uint8 Value) const;
 
 protected:
 	// Variables.
 	/** List of pairs of all enum names and values. */
-	TArray<TPair<FName, int8>> Names;
+	TArray<TPair<FName, uint8>> Names;
 
 	/** How the enum was originally defined. */
 	ECppForm CppForm;
@@ -1520,7 +1494,7 @@ public:
 	 * @param InCppForm The form of enum.
 	 * @return	true unless the MAX enum already exists and isn't the last enum.
 	 */
-	virtual bool SetEnums(TArray<TPair<FName, int8>>& InNames, ECppForm InCppForm);
+	virtual bool SetEnums(TArray<TPair<FName, uint8>>& InNames, ECppForm InCppForm);
 
 	/**
 	 * @return	The enum name at the specified Index.
diff --git a/Engine/Source/Runtime/Engine/Classes/Engine/UserDefinedEnum.h b/Engine/Source/Runtime/Engine/Classes/Engine/UserDefinedEnum.h
--- a/Engine/Source/Runtime/Engine/Classes/Engine/UserDefinedEnum.h
+++ b/Engine/Source/Runtime/Engine/Classes/Engine/UserDefinedEnum.h
@@ -56,7 +56,7 @@ public:
 	 */
 	virtual FText GetEnumText(int32 InIndex) const override;
 
-	virtual bool SetEnums(TArray<TPair<FName, int8>>& InNames, ECppForm InCppForm) override;
+	virtual bool SetEnums(TArray<TPair<FName, uint8>>& InNames, ECppForm InCppForm) override;
 
 #if WITH_EDITOR
 	// Begin UObject interface
diff --git a/Engine/Source/Runtime/Engine/Private/UserDefinedEnum.cpp b/Engine/Source/Runtime/Engine/Private/UserDefinedEnum.cpp
--- a/Engine/Source/Runtime/Engine/Private/UserDefinedEnum.cpp
+++ b/Engine/Source/Runtime/Engine/Private/UserDefinedEnum.cpp
@@ -97,7 +97,7 @@ FText UUserDefinedEnum::GetEnumText(int32 InIndex) const
 	return Super::GetEnumText(InIndex);
 }
 
-bool UUserDefinedEnum::SetEnums(TArray<TPair<FName, int8>>& InNames, ECppForm InCppForm)
+bool UUserDefinedEnum::SetEnums(TArray<TPair<FName, uint8>>& InNames, ECppForm InCppForm)
 {
 	if (Names.Num() > 0)
 	{
@@ -117,7 +117,7 @@ bool UUserDefinedEnum::SetEnums(TArray<TPair<FName, int8>>& InNames, ECppForm In
 		const int32 MaxEnumItemIndex = GetValueByName(MaxEnumItem);
 		if ((MaxEnumItemIndex == INDEX_NONE) && (LookupEnumName(MaxEnumItem) == INDEX_NONE))
 		{
-			Names.Add(TPairInitializer<FName, int8>(MaxEnumItem, GetMaxEnumValue() + 1));
+			Names.Add(TPairInitializer<FName, uint8>(MaxEnumItem, GetMaxEnumValue() + 1));
 			AddNamesToMasterList();
 			return true;
 		}
@@ -126,4 +126,4 @@ bool UUserDefinedEnum::SetEnums(TArray<TPair<FName, int8>>& InNames, ECppForm In
 	UE_LOG(LogClass, Error, TEXT("Unable to generate enum MAX entry due to name collision. Enum: %s"), *GetPathName());
 
 	return false;
-}
\ No newline at end of file
+}
