Fix UEnum issues with incorrect _MAX enum value.

diff --git a/Engine/Source/Editor/KismetCompiler/Private/KismetCompilerMisc.cpp b/Engine/Source/Editor/KismetCompiler/Private/KismetCompilerMisc.cpp
--- a/Engine/Source/Editor/KismetCompiler/Private/KismetCompilerMisc.cpp
+++ b/Engine/Source/Editor/KismetCompiler/Private/KismetCompilerMisc.cpp
@@ -501,8 +501,7 @@ void FKismetCompilerUtilities::ValidateEnumProperties(UObject* DefaultObject, FC
 			if(Enum)
 			{		
 				const uint8 EnumValue = ByteProperty->GetPropertyValue_InContainer(DefaultObject);
-				const int32 EnumAcceptableMax = Enum->GetMaxEnumValue();
-				if(EnumValue >= EnumAcceptableMax)
+				if(!Enum->IsValidEnumValue(EnumValue))
 				{
 					MessageLog.Warning(
 						*FString::Printf(
diff --git a/Engine/Source/Editor/PropertyEditor/Private/PropertyHandleImpl.cpp b/Engine/Source/Editor/PropertyEditor/Private/PropertyHandleImpl.cpp
--- a/Engine/Source/Editor/PropertyEditor/Private/PropertyHandleImpl.cpp
+++ b/Engine/Source/Editor/PropertyEditor/Private/PropertyHandleImpl.cpp
@@ -70,7 +70,7 @@ FPropertyAccess::Result FPropertyValueImpl::GetPropertyValueString( FString& Out
 				{
 					const uint8 EnumValue = ByteProperty->GetPropertyValue(ValueAddress);
 
-					if (EnumValue >= 0 && EnumValue < ByteProperty->Enum->GetMaxEnumValue() + 1)
+					if (ByteProperty->Enum->IsValidEnumValue(EnumValue))
 					{
 						// See if we specified an alternate name for this value using metadata
 						OutString = ByteProperty->Enum->GetDisplayNameTextByValue(EnumValue).ToString();
diff --git a/Engine/Source/Runtime/AIModule/Classes/Tasks/AITask.h b/Engine/Source/Runtime/AIModule/Classes/Tasks/AITask.h
--- a/Engine/Source/Runtime/AIModule/Classes/Tasks/AITask.h
+++ b/Engine/Source/Runtime/AIModule/Classes/Tasks/AITask.h
@@ -14,7 +14,7 @@ enum class EAITaskPriority : uint8
 	Low = 64, //FGameplayTasks::DefaultPriority / 2,
 	AutonomousAI = 127, //FGameplayTasks::DefaultPriority,
 	High = 192, //(1.5 * FGameplayTasks::DefaultPriority),
-	Ultimate = 255,
+	Ultimate = 254,
 };
 
 UCLASS(Abstract)
diff --git a/Engine/Source/Runtime/CoreUObject/Private/UObject/Class.cpp b/Engine/Source/Runtime/CoreUObject/Private/UObject/Class.cpp
--- a/Engine/Source/Runtime/CoreUObject/Private/UObject/Class.cpp
+++ b/Engine/Source/Runtime/CoreUObject/Private/UObject/Class.cpp
@@ -1040,8 +1040,8 @@ void UStruct::SerializeTaggedProperties(FArchive& Ar, uint8* Data, UStruct* Defa
 				if (Tag.EnumName != NAME_None)
 				{
 					//@warning: mirrors loading code in UByteProperty::SerializeItem()
-					FName EnumValue;
-					Ar << EnumValue;
+					FName EnumName;
+					Ar << EnumName;
 					UEnum* Enum = FindField<UEnum>((DefaultsClass != NULL) ? DefaultsClass : DefaultsStruct->GetTypedOuter<UClass>(), Tag.EnumName);
 					if (Enum == NULL)
 					{
@@ -1055,8 +1055,8 @@ void UStruct::SerializeTaggedProperties(FArchive& Ar, uint8* Data, UStruct* Defa
 					else
 					{
 						Ar.Preload(Enum);
-						PreviousValue = Enum->GetValueByName(EnumValue);
-						if (Enum->GetNameByValue(PreviousValue) != NAME_None)
+						PreviousValue = Enum->GetValueByName(EnumName);
+						if (!Enum->IsValidEnumValue(PreviousValue))
 						{
 							PreviousValue = Enum->GetMaxEnumValue();
 						}
@@ -1282,8 +1282,8 @@ void UStruct::SerializeTaggedProperties(FArchive& Ar, uint8* Data, UStruct* Defa
 				{
 					// attempt to find the old enum and get the byte value from the serialized enum name
 					//@warning: mirrors loading code in UByteProperty::SerializeItem()
-					FName EnumValue;
-					Ar << EnumValue;
+					FName EnumName;
+					Ar << EnumName;
 					UEnum* Enum = FindField<UEnum>((DefaultsClass != NULL) ? DefaultsClass : DefaultsStruct->GetTypedOuter<UClass>(), Tag.EnumName);
 					if (Enum == NULL)
 					{
@@ -1297,8 +1297,8 @@ void UStruct::SerializeTaggedProperties(FArchive& Ar, uint8* Data, UStruct* Defa
 					else
 					{
 						Ar.Preload(Enum);
-						PreviousValue = Enum->GetValueByName(EnumValue);
-						if (Enum->GetNameByValue(PreviousValue) != NAME_None)
+						PreviousValue = Enum->GetValueByName(EnumName);
+						if (!Enum->IsValidEnumValue(PreviousValue))
 						{
 							PreviousValue = Enum->GetMaxEnumValue();
 						}
diff --git a/Engine/Source/Runtime/CoreUObject/Private/UObject/Enum.cpp b/Engine/Source/Runtime/CoreUObject/Private/UObject/Enum.cpp
--- a/Engine/Source/Runtime/CoreUObject/Private/UObject/Enum.cpp
+++ b/Engine/Source/Runtime/CoreUObject/Private/UObject/Enum.cpp
@@ -183,7 +183,7 @@ uint8 UEnum::GetMaxEnumValue() const
 	int32 NamesNum = Names.Num();
 	if (NamesNum == 0)
 	{
-		return INDEX_NONE;
+		return 0;
 	}
 
 	uint8 MaxValue = Names[0].Value;
@@ -205,9 +205,22 @@ bool UEnum::IsValidEnumValue(uint8 InValue) const
 	for (int32 i = 0; i < NamesNum; ++i)
 	{
 		uint8 CurrentValue = Names[i].Value;
+		if (CurrentValue == InValue)
+		{
+			return true;
+		}
+	}
 
-		// Last entry is autogenerated _MAX value
-		if (CurrentValue == InValue && i != NamesNum - 1)
+	return false;
+}
+
+bool UEnum::IsValidEnumName(FName InName) const
+{
+	int32 NamesNum = Names.Num();
+	for (int32 i = 0; i < NamesNum; ++i)
+	{
+		FName CurrentName = Names[i].Key;
+		if (CurrentName == InName)
 		{
 			return true;
 		}
@@ -436,7 +449,7 @@ int32 UEnum::FindEnumRedirects(const UEnum* Enum, FName EnumEntryName)
 /**
  * Sets the array of enums.
  *
- * @return	true unless the MAX enum already exists and isn't the last enum.
+ * @return	true unless the MAX enum already exists.
  */
 bool UEnum::SetEnums(TArray<TPair<FName, uint8>>& InNames, UEnum::ECppForm InCppForm)
 {
@@ -461,17 +474,18 @@ bool UEnum::SetEnums(TArray<TPair<FName, uint8>>& InNames, UEnum::ECppForm InCpp
 			return false;
 		}
 
+#if HACK_HEADER_GENERATOR
+		if (GetMaxEnumValue() == 255)
+		{
+			FString Name = GetName();
+			FError::Throwf(TEXT("Enum %s doesn't contain user defined %s_MAX entry and its maximum entry value equals 255. Autogenerated %s_MAX enum entry will equal 0. Change maximum entry value to 254 or define %s_MAX entry in enum."), *Name, *EnumPrefix, *EnumPrefix, *EnumPrefix, *EnumPrefix);
+		}
+#endif
 		Names.Add(TPairInitializer<FName, uint8>(MaxEnumItem, GetMaxEnumValue() + 1));
 	}
 
 	AddNamesToMasterList();
 
-	if (MaxEnumItemIndex != INDEX_NONE && MaxEnumItemIndex != Names.Num() - 1)
-	{
-		// The MAX enum already exists, but isn't the last enum.
-		return true;
-	}
-
 	return true;
 }
 
diff --git a/Engine/Source/Runtime/CoreUObject/Private/UObject/PropertyByte.cpp b/Engine/Source/Runtime/CoreUObject/Private/UObject/PropertyByte.cpp
--- a/Engine/Source/Runtime/CoreUObject/Private/UObject/PropertyByte.cpp
+++ b/Engine/Source/Runtime/CoreUObject/Private/UObject/PropertyByte.cpp
@@ -34,7 +34,7 @@ void UByteProperty::SerializeItem( FArchive& Ar, void* Value, void const* Defaul
 		// There's no guarantee EnumValueName is still present in Enum, in which case Value will be set to the enum's max value.
 		// On save, it will then be serialized as NAME_None.
 		*(uint8*)Value = Enum->GetValueByName(EnumValueName);
-		if ( Enum->GetMaxEnumValue() + 1 < *(uint8*)Value )
+		if ( !Enum->IsValidEnumValue(*(uint8*)Value) )
 		{
 			*(uint8*)Value = Enum->GetMaxEnumValue();
 		}
@@ -47,7 +47,7 @@ void UByteProperty::SerializeItem( FArchive& Ar, void* Value, void const* Defaul
 
 		// subtract 1 because the last entry in the enum's Names array
 		// is the _MAX entry
-		if ( ByteValue < Enum->GetMaxEnumValue() )
+		if ( Enum->IsValidEnumValue(ByteValue) )
 		{
 			EnumValueName = Enum->GetNameByValue(ByteValue);
 		}
@@ -115,7 +115,7 @@ void UByteProperty::ExportTextItem( FString& ValueStr, const void* PropertyValue
 		{
 			const int32 ActualValue = *(const uint8*)PropertyValue;
 			const int32 MaxValue = Enum->GetMaxEnumValue();
-			const int32 GoodValue = (ActualValue > MaxValue) ? MaxValue : ActualValue;
+			const int32 GoodValue = Enum->IsValidEnumValue(ActualValue) ? ActualValue : MaxValue;
 			ValueStr += FString::Printf(TEXT("%s::%s"), *Enum->GetName(),
 				*Enum->GetEnumName(Enum->GetIndexByValue(GoodValue)));
 		}
@@ -130,7 +130,9 @@ void UByteProperty::ExportTextItem( FString& ValueStr, const void* PropertyValue
 	{
 		// if the value is the max value (the autogenerated *_MAX value), export as "INVALID", unless we're exporting text for copy/paste (for copy/paste,
 		// the property text value must actually match an entry in the enum's names array)
-		if (*(const uint8*)PropertyValue <= Enum->GetMaxEnumValue() - 1 || ((PortFlags&PPF_Copy) != 0 && *(const uint8*)PropertyValue <= Enum->GetMaxEnumValue()))
+		bool bIsValid = Enum->IsValidEnumValue(*(const uint8*)PropertyValue);
+		bool bIsMax = *(const uint8*)PropertyValue == Enum->GetMaxEnumValue();
+		if (bIsValid && (!bIsMax || (PortFlags & PPF_Copy)))
 		{
 			// We do not want to export the enum text for non-display uses, localization text is very dynamic and would cause issues on import
 			if (PortFlags & PPF_PropertyWindow)
diff --git a/Engine/Source/Runtime/CoreUObject/Public/UObject/Class.h b/Engine/Source/Runtime/CoreUObject/Public/UObject/Class.h
--- a/Engine/Source/Runtime/CoreUObject/Public/UObject/Class.h
+++ b/Engine/Source/Runtime/CoreUObject/Public/UObject/Class.h
@@ -1370,7 +1370,11 @@ public:
 	/** Gets max value of Enum. Defaults to zero if there are no entries. */
 	uint8 GetMaxEnumValue() const;
 
-	bool IsValidEnumValue(uint8 Value) const;
+	/** Checks if enum has entry with given value. Includes autogenerated _MAX entry. */
+	bool IsValidEnumValue(uint8 InValue) const;
+
+	/** Checks if enum has entry with given name. Includes autogenerated _MAX entry. */
+	bool IsValidEnumName(FName InName) const;
 
 protected:
 	// Variables.
