#ue4 - Add "PostRepNotifies()" event to Objects, called after any rep notifies may be triggered.

diff --git a/Engine/Source/Runtime/CoreUObject/Public/UObject/UObject.h b/Engine/Source/Runtime/CoreUObject/Public/UObject/UObject.h
--- a/Engine/Source/Runtime/CoreUObject/Public/UObject/UObject.h
+++ b/Engine/Source/Runtime/CoreUObject/Public/UObject/UObject.h
@@ -624,6 +624,9 @@ public:
 	/** Called right after receiving a bunch */
 	virtual void PostNetReceive();
 
+	/** Called right after calling all OnRep notifies (called even when there are no notifies) */
+	virtual void PostRepNotifies() {}
+
 	/** Called right before being marked for destruction due to network replication */
 	virtual void PreDestroyFromReplication();
 
diff --git a/Engine/Source/Runtime/Engine/Private/DataReplication.cpp b/Engine/Source/Runtime/Engine/Private/DataReplication.cpp
--- a/Engine/Source/Runtime/Engine/Private/DataReplication.cpp
+++ b/Engine/Source/Runtime/Engine/Private/DataReplication.cpp
@@ -762,6 +762,11 @@ void FObjectReplicator::PostReceivedBunch()
 
 	// Call RepNotifies
 	CallRepNotifies(true);
+
+	if (!Object->IsPendingKill())
+	{
+		Object->PostRepNotifies();
+	}
 }
 
 static FORCEINLINE FPropertyRetirement ** UpdateAckedRetirements( FPropertyRetirement &	Retire, int32 OutAckPacketId )
