Backing out bad merge. Insert package predicate has already been changed to < instead of <= in main.

diff --git a/Engine/Source/Runtime/CoreUObject/Private/Serialization/AsyncLoading.cpp b/Engine/Source/Runtime/CoreUObject/Private/Serialization/AsyncLoading.cpp
--- a/Engine/Source/Runtime/CoreUObject/Private/Serialization/AsyncLoading.cpp
+++ b/Engine/Source/Runtime/CoreUObject/Private/Serialization/AsyncLoading.cpp
@@ -475,7 +475,7 @@ void FAsyncLoadingThread::ProcessAsyncPackageRequest(FAsyncPackageDesc* InReques
 		}
 #endif
 		// Add to queue according to priority.
-		InsertPackage(Package, false);
+		InsertPackage(Package);
 
 		// For all other cases this is handled in FindExistingPackageAndAddCompletionCallback
 		const int32 QueuedPackagesCount = QueuedPackagesCounter.Decrement();
@@ -520,7 +520,7 @@ int32 FAsyncLoadingThread::CreateAsyncPackagesFromQueue()
 	return NumCreated;
 }
 
-void FAsyncLoadingThread::InsertPackage(FAsyncPackage* Package, bool bIsImportPackage)
+void FAsyncLoadingThread::InsertPackage(FAsyncPackage* Package)
 {
 	checkSlow(IsInAsyncLoadThread());
 
@@ -535,24 +535,12 @@ void FAsyncLoadingThread::InsertPackage(FAsyncPackage* Package, bool bIsImportPa
 		FScopeLock LockAsyncPackages(&AsyncPackagesCritical);
 #endif
 		// Insert new package keeping descending priority order in AsyncPackages
-		int32 InsertIndex = INDEX_NONE;
-		
-		if (bIsImportPackage)
+		int32 InsertIndex = AsyncPackages.IndexOfByPredicate([Package](const FAsyncPackage* Element)
 		{
-			// This inserts it before packages of equal priority
-			InsertIndex = AsyncPackages.IndexOfByPredicate([Package](const FAsyncPackage* Element)
-			{
-				return Element->GetPriority() <= Package->GetPriority();
-			});
-		}
-		else
-		{
-			// This inserts it after packages of equal priority, so it doesn't take priority over nested loads
-			InsertIndex = AsyncPackages.IndexOfByPredicate([Package](const FAsyncPackage* Element)
-			{
-				return Element->GetPriority() < Package->GetPriority();
-			});
-		}
+			// NOTE: Switched to < because we want newly inserted packages to go AFTER all existing packages of the 
+			//		 same priority. If this is too slow, we could probably do "LastIndexOfByPredicate" where we iterate backwards??
+			return Element->GetPriority() < Package->GetPriority();
+		});
 
 		InsertIndex = InsertIndex == INDEX_NONE ? AsyncPackages.Num() : InsertIndex;
 
@@ -1395,7 +1383,7 @@ void FAsyncPackage::AddImportDependency(int32 CurrentPackageIndex, const FName&
 		{
 			PackageToStream->Desc.Priority = Desc.Priority;
 		}
-		FAsyncLoadingThread::Get().InsertPackage(PackageToStream, true);
+		FAsyncLoadingThread::Get().InsertPackage(PackageToStream);
 	}
 	else
 	{
diff --git a/Engine/Source/Runtime/CoreUObject/Private/Serialization/AsyncLoadingThread.h b/Engine/Source/Runtime/CoreUObject/Private/Serialization/AsyncLoadingThread.h
--- a/Engine/Source/Runtime/CoreUObject/Private/Serialization/AsyncLoadingThread.h
+++ b/Engine/Source/Runtime/CoreUObject/Private/Serialization/AsyncLoadingThread.h
@@ -207,9 +207,9 @@ public:
 	* [ASYNC THREAD] Inserts package to queue according to priority.
 	*
 	* @param PackageName async package name.
-	* @param bIsImportPackage true if this package was from another package's import, false if it was requested externally.
 	*/
-	void InsertPackage(FAsyncPackage* Package, bool bIsImportPackage);
+	void InsertPackage(FAsyncPackage* Package);
+
 	/**
 	* [ASYNC THREAD] Finds an existing async package in the LoadedPackages by its name.
 	*
