Content resave to get rid of short package name's from references during QAGame startup and a warning fix.

diff --git a/Engine/Source/Runtime/CoreUObject/Private/Misc/PackageName.cpp b/Engine/Source/Runtime/CoreUObject/Private/Misc/PackageName.cpp
--- a/Engine/Source/Runtime/CoreUObject/Private/Misc/PackageName.cpp
+++ b/Engine/Source/Runtime/CoreUObject/Private/Misc/PackageName.cpp
@@ -807,7 +807,7 @@ FString FPackageName::GetNormalizedObjectPath(const FString& ObjectPath)
 	{
 		FString LongPath;
 
-		UE_LOG(LogPackageName, Warning, TEXT("String asset reference \"%s\" is in short form, which is unsupported and -- even if valid -- resolving it will be really slow."), *ObjectPath);
+		UE_LOG(LogPackageName, Warning, TEXT("String asset reference \"%s\" is in short form, which is unsupported and -- even if valid -- resolving it will be really slow. Please consider resaving package in order to speed-up loading."), *ObjectPath);
 
 		if (!FPackageName::TryConvertShortPackagePathToLongInObjectPath(ObjectPath, LongPath))
 		{
