Fix UEnum serialization redirect issues.

diff --git a/Engine/Source/Runtime/CoreUObject/Private/UObject/PropertyByte.cpp b/Engine/Source/Runtime/CoreUObject/Private/UObject/PropertyByte.cpp
--- a/Engine/Source/Runtime/CoreUObject/Private/UObject/PropertyByte.cpp
+++ b/Engine/Source/Runtime/CoreUObject/Private/UObject/PropertyByte.cpp
@@ -33,11 +33,15 @@ void UByteProperty::SerializeItem( FArchive& Ar, void* Value, void const* Defaul
 
 		// There's no guarantee EnumValueName is still present in Enum, in which case Value will be set to the enum's max value.
 		// On save, it will then be serialized as NAME_None.
-		*(uint8*)Value = Enum->GetValueByName(EnumValueName);
-		if ( !Enum->IsValidEnumValue(*(uint8*)Value) )
+		int32 EnumIndex = Enum->FindEnumIndex(EnumValueName);
+		if (EnumIndex == INDEX_NONE)
 		{
 			*(uint8*)Value = Enum->GetMaxEnumValue();
 		}
+		else
+		{
+			*(uint8*)Value = Enum->GetValueByIndex(EnumIndex);
+		}
 	}
 	// Saving
 	else
@@ -162,10 +166,10 @@ const TCHAR* UByteProperty::ImportText_Internal( const TCHAR* InBuffer, void* Da
 		const TCHAR* Buffer = UPropertyHelpers::ReadToken( InBuffer, Temp, true );
 		if( Buffer != NULL )
 		{
-			const int32 EnumValue = Enum->GetValueByName(*Temp);
-			if( EnumValue != INDEX_NONE )
+			int32 EnumIndex = Enum->FindEnumIndex(*Temp);
+			if (EnumIndex != INDEX_NONE)
 			{
-				*(uint8*)Data = EnumValue;
+				*(uint8*)Data = Enum->GetValueByIndex(EnumIndex);
 				return Buffer;
 			}
 		}
