Cook-by-the-book will now skip saving packages that are referenced only through editor-only properties.

diff --git a/Engine/Source/Editor/UnrealEd/Classes/CookOnTheSide/CookOnTheFlyServer.h b/Engine/Source/Editor/UnrealEd/Classes/CookOnTheSide/CookOnTheFlyServer.h
--- a/Engine/Source/Editor/UnrealEd/Classes/CookOnTheSide/CookOnTheFlyServer.h
+++ b/Engine/Source/Editor/UnrealEd/Classes/CookOnTheSide/CookOnTheFlyServer.h
@@ -600,6 +600,34 @@ private:
 		FString StandardFilename;
 		FName StandardFileFName;
 	};
+
+	/** Simple thread safe proxy for TSet<FName> */
+	class FThreadSafeNameSet
+	{
+		TSet<FName> Names;
+		FCriticalSection NamesCritical;
+	public:
+		void Add(FName InName)
+		{
+			FScopeLock NamesLock(&NamesCritical);
+			Names.Add(InName);
+		}
+		bool AddUnique(FName InName)
+		{
+			FScopeLock NamesLock(&NamesCritical);
+			if (!Names.Contains(InName))
+			{
+				Names.Add(InName);
+				return true;
+			}
+			return false;
+		}
+		bool Contains(FName InName)
+		{
+			FScopeLock NamesLock(&NamesCritical);
+			return Names.Contains(InName);
+		}
+	};
 private:
 	/** Current cook mode the cook on the fly server is running in */
 	ECookMode::Type CurrentCookMode;
@@ -724,6 +752,7 @@ private:
 	FFilenameQueue CookRequests; // list of requested files
 	FThreadSafeUnsolicitedPackagesList UnsolicitedCookedPackages;
 	FThreadSafeFilenameSet CookedPackages; // set of files which have been cooked when needing to recook a file the entry will need to be removed from here
+	FThreadSafeNameSet UncookedEditorOnlyPackages; // set of packages that have been rejected due to being referenced by editor-only properties
 
 	FString GetCachedPackageFilename( const FName& PackageName ) const;
 	FString GetCachedStandardPackageFilename( const FName& PackageName ) const;
@@ -1285,9 +1314,9 @@ private:
 	 *	@param	SaveFlags			The flags to pass to the SavePackage function
 	 *	@param	bOutWasUpToDate		Upon return, if true then the cooked package was cached (up to date)
 	 *
-	 *	@return	bool			true if packages was cooked
+	 *	@return	ESavePackageResult::Success if packages was cooked
 	 */
-	bool SaveCookedPackage( UPackage* Package, uint32 SaveFlags, bool& bOutWasUpToDate );
+	ESavePackageResult SaveCookedPackage(UPackage* Package, uint32 SaveFlags, bool& bOutWasUpToDate);
 	/**
 	 *	Cook (save) the given package
 	 *
@@ -1297,9 +1326,9 @@ private:
 	 *	@param  TargetPlatformNames Only cook for target platforms which are included in this array (if empty cook for all target platforms specified on commandline options)
 	 *									TargetPlatformNames is in and out value returns the platforms which the SaveCookedPackage function saved for
 	 *
-	 *	@return	bool			true if packages was cooked
+	 *	@return	ESavePackageResult::Success if packages was cooked
 	 */
-	bool SaveCookedPackage( UPackage* Package, uint32 SaveFlags, bool& bOutWasUpToDate, TArray<FName> &TargetPlatformNames );
+	ESavePackageResult SaveCookedPackage(UPackage* Package, uint32 SaveFlags, bool& bOutWasUpToDate, TArray<FName> &TargetPlatformNames);
 
 
 	/**
diff --git a/Engine/Source/Editor/UnrealEd/Classes/Editor/EditorEngine.h b/Engine/Source/Editor/UnrealEd/Classes/Editor/EditorEngine.h
--- a/Engine/Source/Editor/UnrealEd/Classes/Editor/EditorEngine.h
+++ b/Engine/Source/Editor/UnrealEd/Classes/Editor/EditorEngine.h
@@ -2036,6 +2036,11 @@ public:
 		FOutputDevice* Error=GError, FLinkerLoad* Conform=NULL, bool bForceByteSwapping=false, bool bWarnOfLongFilename=true, 
 		uint32 SaveFlags=SAVE_None, const class ITargetPlatform* TargetPlatform = NULL, const FDateTime& FinalTimeStamp = FDateTime::MinValue(), bool bSlowTask = true );
 
+	/** The editor wrapper for UPackage::Save. Auto-adds files to source control when necessary */
+	ESavePackageResult Save(UPackage* InOuter, UObject* Base, EObjectFlags TopLevelFlags, const TCHAR* Filename,
+		FOutputDevice* Error = GError, FLinkerLoad* Conform = NULL, bool bForceByteSwapping = false, bool bWarnOfLongFilename = true,
+		uint32 SaveFlags = SAVE_None, const class ITargetPlatform* TargetPlatform = NULL, const FDateTime& FinalTimeStamp = FDateTime::MinValue(), bool bSlowTask = true);
+
 	/** Invoked before a UWorld is saved to update editor systems */
 	virtual void OnPreSaveWorld(uint32 SaveFlags, UWorld* World);
 
diff --git a/Engine/Source/Editor/UnrealEd/Private/CookOnTheFlyServer.cpp b/Engine/Source/Editor/UnrealEd/Private/CookOnTheFlyServer.cpp
--- a/Engine/Source/Editor/UnrealEd/Private/CookOnTheFlyServer.cpp
+++ b/Engine/Source/Editor/UnrealEd/Private/CookOnTheFlyServer.cpp
@@ -1969,8 +1969,12 @@ uint32 UCookOnTheFlyServer::TickCookOnTheSide( const float TimeSlice, uint32 &Co
 			SCOPE_TIMER(SavingPackages);
 			for ( int32 I = 0; I < PackagesToSave.Num(); ++I )
 			{
-
-				UPackage *Package = PackagesToSave[I];
+				UPackage* Package = PackagesToSave[I];
+				if (Package->IsLoadedByEditorPropertiesOnly() && UncookedEditorOnlyPackages.Contains(Package->GetFName()))
+				{
+					// We already attempted to cook this package and it's still not referenced by any non editor-only properties.
+					continue;
+				}
 
 				// if we are processing unsolicited packages we can optionally not save these right now
 				// the unsolicited packages which we missed now will be picked up on next run
@@ -2091,12 +2095,14 @@ uint32 UCookOnTheFlyServer::TickCookOnTheSide( const float TimeSlice, uint32 &Co
 					bShouldSaveAsync = false;
 				}
 
+				ESavePackageResult SavePackageResult = ESavePackageResult::Error;
 				END_PACKAGE_STAT(SavePackageCacheForCookedPlatformData);
 				{
 					SCOPE_PACKAGE_STAT(SaveCookedPackage);
 
 					SCOPE_TIMER(SaveCookedPackage);
-					if (SaveCookedPackage(Package, SAVE_KeepGUID | (bShouldSaveAsync ? SAVE_Async : SAVE_None) | (IsCookFlagSet(ECookInitializationFlags::Unversioned) ? SAVE_Unversioned : 0), bWasUpToDate, AllTargetPlatformNames))
+					SavePackageResult = SaveCookedPackage(Package, SAVE_KeepGUID | (bShouldSaveAsync ? SAVE_Async : SAVE_None) | (IsCookFlagSet(ECookInitializationFlags::Unversioned) ? SAVE_Unversioned : 0), bWasUpToDate, AllTargetPlatformNames);
+					if (SavePackageResult == ESavePackageResult::Success)
 					{
 						// Update flags used to determine garbage collection.
 						if (Package->ContainsMap())
@@ -2134,24 +2140,33 @@ uint32 UCookOnTheFlyServer::TickCookOnTheSide( const float TimeSlice, uint32 &Co
 
 				FName StandardFilename = GetCachedStandardPackageFileFName(Package);
 
-				if ( StandardFilename != NAME_None )
+				// We always want to mark package as processed unless it wasn't saved because it was referenced by editor-only data
+				// in which case we may still need to save it later when new content loads it through non editor-only references
+				if (StandardFilename != NAME_None)
 				{
 					// mark the package as cooked
-					FFilePlatformRequest FileRequest( StandardFilename, AllTargetPlatformNames);
-					CookedPackages.Add( FileRequest );
-
-					if ( (CurrentCookMode == ECookMode::CookOnTheFly) && (I >= FirstUnsolicitedPackage) ) 
+					FFilePlatformRequest FileRequest(StandardFilename, AllTargetPlatformNames);
+					if (SavePackageResult != ESavePackageResult::ReferencedOnlyByEditorOnlyData)
 					{
-						// this is an unsolicited package
-						if ((FPaths::FileExists(FileRequest.GetFilename().ToString()) == true) &&
-							(bWasUpToDate == false))
+						CookedPackages.Add(FileRequest);
+
+						if ((CurrentCookMode == ECookMode::CookOnTheFly) && (I >= FirstUnsolicitedPackage))
 						{
-							UnsolicitedCookedPackages.AddCookedPackage( FileRequest );
+							// this is an unsolicited package
+							if ((FPaths::FileExists(FileRequest.GetFilename().ToString()) == true) &&
+								(bWasUpToDate == false))
+							{
+								UnsolicitedCookedPackages.AddCookedPackage( FileRequest );
 #if DEBUG_COOKONTHEFLY
-							UE_LOG(LogCook, Display, TEXT("UnsolicitedCookedPackages: %s"), *FileRequest.GetFilename().ToString());
+								UE_LOG(LogCook, Display, TEXT("UnsolicitedCookedPackages: %s"), *FileRequest.GetFilename().ToString());
 #endif
+							}
 						}
 					}
+					else
+					{
+						UncookedEditorOnlyPackages.AddUnique(Package->GetFName());
+					}
 				}
 			}
 		}
@@ -2356,7 +2371,7 @@ bool UCookOnTheFlyServer::GetPackageTimestamp( const FString& InFilename, FDateT
 }
 
 
-bool UCookOnTheFlyServer::SaveCookedPackage( UPackage* Package, uint32 SaveFlags, bool& bOutWasUpToDate ) 
+ESavePackageResult UCookOnTheFlyServer::SaveCookedPackage(UPackage* Package, uint32 SaveFlags, bool& bOutWasUpToDate)
 {
 	TArray<FName> TargetPlatformNames; 
 	return SaveCookedPackage( Package, SaveFlags, bOutWasUpToDate, TargetPlatformNames );
@@ -2429,9 +2444,9 @@ bool UCookOnTheFlyServer::ShouldConsiderCompressedPackageFileLengthRequirements(
 	return bConsiderCompressedPackageFileLengthRequirements;
 }
 
-bool UCookOnTheFlyServer::SaveCookedPackage( UPackage* Package, uint32 SaveFlags, bool& bOutWasUpToDate, TArray<FName> &TargetPlatformNames )
+ESavePackageResult UCookOnTheFlyServer::SaveCookedPackage(UPackage* Package, uint32 SaveFlags, bool& bOutWasUpToDate, TArray<FName> &TargetPlatformNames)
 {
-	bool bSavedCorrectly = true;
+	ESavePackageResult Result = ESavePackageResult::Success;
 	check( bIsSavingPackage == false );
 	bIsSavingPackage = true;
 	FString Filename(GetCachedPackageFilename(Package));
@@ -2576,12 +2591,12 @@ bool UCookOnTheFlyServer::SaveCookedPackage( UPackage* Package, uint32 SaveFlags
 				{
 					LogCookerMessage( FString::Printf(TEXT("Couldn't save package, filename is too long: %s"), *PlatFilename), EMessageSeverity::Error );
 					UE_LOG( LogCook, Error, TEXT( "Couldn't save package, filename is too long :%s" ), *PlatFilename );
-					bSavedCorrectly = false;
+					Result = ESavePackageResult::Error;
 				}
 				else
 				{
 					SCOPE_TIMER(GEditorSavePackage);
-					bSavedCorrectly &= GEditor->SavePackage( Package, World, Flags, *PlatFilename, GError, NULL, bSwap, false, SaveFlags, Target, FDateTime::MinValue(), false );
+					Result = GEditor->Save(Package, World, Flags, *PlatFilename, GError, NULL, bSwap, false, SaveFlags, Target, FDateTime::MinValue(), false);
 					INC_INT_STAT(SavedPackage, 1);
 				}
 
@@ -2604,7 +2619,7 @@ bool UCookOnTheFlyServer::SaveCookedPackage( UPackage* Package, uint32 SaveFlags
 	bIsSavingPackage = false;
 
 	// return success
-	return bSavedCorrectly;
+	return Result;
 }
 
 void UCookOnTheFlyServer::Initialize( ECookMode::Type DesiredCookMode, ECookInitializationFlags InCookFlags, const FString &InOutputDirectoryOverride )
diff --git a/Engine/Source/Editor/UnrealEd/Private/EditorEngine.cpp b/Engine/Source/Editor/UnrealEd/Private/EditorEngine.cpp
--- a/Engine/Source/Editor/UnrealEd/Private/EditorEngine.cpp
+++ b/Engine/Source/Editor/UnrealEd/Private/EditorEngine.cpp
@@ -3686,7 +3686,7 @@ void UEditorEngine::OnSourceControlDialogClosed(bool bEnabled)
 	}
 }
 
-bool UEditorEngine::SavePackage( UPackage* InOuter, UObject* InBase, EObjectFlags TopLevelFlags, const TCHAR* Filename, 
+ESavePackageResult UEditorEngine::Save( UPackage* InOuter, UObject* InBase, EObjectFlags TopLevelFlags, const TCHAR* Filename, 
 				 FOutputDevice* Error, FLinkerLoad* Conform, bool bForceByteSwapping, bool bWarnOfLongFilename, 
 				 uint32 SaveFlags, const class ITargetPlatform* TargetPlatform, const FDateTime& FinalTimeStamp, bool bSlowTask )
 {
@@ -3769,13 +3769,13 @@ bool UEditorEngine::SavePackage( UPackage* InOuter, UObject* InBase, EObjectFlag
 
 	SlowTask.EnterProgressFrame(70);
 
-	bool bSuccess = UPackage::SavePackage(InOuter, Base, TopLevelFlags, Filename, Error, Conform, bForceByteSwapping, bWarnOfLongFilename, SaveFlags, TargetPlatform, FinalTimeStamp, bSlowTask);
+	const ESavePackageResult Result = UPackage::Save(InOuter, Base, TopLevelFlags, Filename, Error, Conform, bForceByteSwapping, bWarnOfLongFilename, SaveFlags, TargetPlatform, FinalTimeStamp, bSlowTask);
 
 	SlowTask.EnterProgressFrame(10);
 
 	// If the package is a valid candidate for being automatically-added to source control, go ahead and add it
 	// to the default changelist
-	if( bSuccess && bAutoAddPkgToSCC )
+	if (Result == ESavePackageResult::Success && bAutoAddPkgToSCC)
 	{
 		// IsPackageValidForAutoAdding should not return true if SCC is disabled
 		check(ISourceControlModule::Get().IsEnabled());
@@ -3802,7 +3802,7 @@ bool UEditorEngine::SavePackage( UPackage* InOuter, UObject* InBase, EObjectFlag
 			World->PersistentLevel->OwningWorld = OriginalOwningWorld;
 		}
 
-		OnPostSaveWorld(SaveFlags, World, OriginalPackageFlags, bSuccess);
+		OnPostSaveWorld(SaveFlags, World, OriginalPackageFlags, Result == ESavePackageResult::Success);
 
 		if (bInitializedPhysicsSceneForSave)
 		{
@@ -3816,7 +3816,17 @@ bool UEditorEngine::SavePackage( UPackage* InOuter, UObject* InBase, EObjectFlag
 		}
 	}
 
-	return bSuccess;
+	return Result;
+}
+
+bool UEditorEngine::SavePackage(UPackage* InOuter, UObject* InBase, EObjectFlags TopLevelFlags, const TCHAR* Filename,
+	FOutputDevice* Error, FLinkerLoad* Conform, bool bForceByteSwapping, bool bWarnOfLongFilename,
+	uint32 SaveFlags, const class ITargetPlatform* TargetPlatform, const FDateTime& FinalTimeStamp, bool bSlowTask)
+{
+	// Workaround to avoid function signature change while keeping both bool and ESavePackageResult versions of SavePackage
+	const ESavePackageResult Result = Save(InOuter, InBase, TopLevelFlags, Filename, Error, Conform, bForceByteSwapping,
+		bWarnOfLongFilename, SaveFlags, TargetPlatform, FinalTimeStamp, bSlowTask);
+	return Result == ESavePackageResult::Success;
 }
 
 void UEditorEngine::OnPreSaveWorld(uint32 SaveFlags, UWorld* World)
diff --git a/Engine/Source/Runtime/Core/Private/Serialization/Archive.cpp b/Engine/Source/Runtime/Core/Private/Serialization/Archive.cpp
--- a/Engine/Source/Runtime/Core/Private/Serialization/Archive.cpp
+++ b/Engine/Source/Runtime/Core/Private/Serialization/Archive.cpp
@@ -92,7 +92,9 @@ void FArchive::Reset()
 	ArIsSaveGame						= false;
 	CookingTargetPlatform = nullptr;
 	SerializedProperty = nullptr;
-
+#if WITH_EDITORONLY_DATA
+	EditorOnlyPropertyStack = 0;
+#endif
 	// Reset all custom versions to the current registered versions.
 	ResetCustomVersions();
 }
@@ -130,6 +132,10 @@ void FArchive::CopyTrivialFArchiveStatusMembers(const FArchive& ArchiveToCopy)
 	ArIsFilterEditorOnly                 = ArchiveToCopy.ArIsFilterEditorOnly;
 	ArIsSaveGame                         = ArchiveToCopy.ArIsSaveGame;
 	CookingTargetPlatform                = ArchiveToCopy.CookingTargetPlatform;
+	SerializedProperty = ArchiveToCopy.SerializedProperty;
+#if WITH_EDITORONLY_DATA
+	EditorOnlyPropertyStack = ArchiveToCopy.EditorOnlyPropertyStack;
+#endif
 }
 
 /**
diff --git a/Engine/Source/Runtime/Core/Public/Serialization/ArchiveBase.h b/Engine/Source/Runtime/Core/Public/Serialization/ArchiveBase.h
--- a/Engine/Source/Runtime/Core/Public/Serialization/ArchiveBase.h
+++ b/Engine/Source/Runtime/Core/Public/Serialization/ArchiveBase.h
@@ -904,6 +904,25 @@ public:
 		SerializedProperty = InProperty;
 	}
 
+#if WITH_EDITORONLY_DATA
+	/** Pushes editor-only marker to the stack of currently serialized properties */
+	FORCEINLINE void PushEditorOnlyProperty()
+	{
+		EditorOnlyPropertyStack++;
+	}
+	/** Pops editor-only marker from the stack of currently serialized properties */
+	FORCEINLINE void PopEditorOnlyProperty()
+	{
+		EditorOnlyPropertyStack--;
+		check(EditorOnlyPropertyStack >= 0);
+	}
+	/** Returns true if the stack of currently serialized properties contains an editor-only property */
+	FORCEINLINE bool IsEditorOnlyPropertyOnTheStack() const
+	{
+		return EditorOnlyPropertyStack > 0;
+	}
+#endif
+
 	/**
 	 * Gets the property that is currently being serialized
 	 *
@@ -1034,6 +1053,11 @@ private:
 	/** Holds the pointer to the property that is currently being serialized */
 	class UProperty* SerializedProperty;
 
+#if WITH_EDITORONLY_DATA
+	/** Non-zero if on the current stack of serialized properties there's an editor-only property. Needs to live in FArchive becasuse of SerializedProperty. */
+	int32 EditorOnlyPropertyStack;
+#endif
+
 	/**
 	 * Indicates if the custom versions container is in a 'reset' state.  This will be used to defer the choice about how to
 	 * populate the container until it is needed, where the read/write state will be known.
diff --git a/Engine/Source/Runtime/CoreUObject/Private/Serialization/ArchiveUObject.cpp b/Engine/Source/Runtime/CoreUObject/Private/Serialization/ArchiveUObject.cpp
--- a/Engine/Source/Runtime/CoreUObject/Private/Serialization/ArchiveUObject.cpp
+++ b/Engine/Source/Runtime/CoreUObject/Private/Serialization/ArchiveUObject.cpp
@@ -120,3 +120,19 @@ FArchive& FObjectAndNameAsStringProxyArchive::operator<<(class UObject*& Obj)
 	return *this;
 }
 
+#if WITH_EDITORONLY_DATA
+void FSerializedPropertyScope::PushEditorOnlyProperty()
+{
+	if (Property && Property->IsEditorOnlyProperty())
+	{
+		Ar.PushEditorOnlyProperty();
+	}
+}
+void FSerializedPropertyScope::PopEditorOnlyProperty()
+{
+	if (Property && Property->IsEditorOnlyProperty())
+	{
+		Ar.PopEditorOnlyProperty();
+	}
+}
+#endif
diff --git a/Engine/Source/Runtime/CoreUObject/Private/Serialization/AsyncLoading.cpp b/Engine/Source/Runtime/CoreUObject/Private/Serialization/AsyncLoading.cpp
--- a/Engine/Source/Runtime/CoreUObject/Private/Serialization/AsyncLoading.cpp
+++ b/Engine/Source/Runtime/CoreUObject/Private/Serialization/AsyncLoading.cpp
@@ -1270,6 +1270,10 @@ EAsyncPackageState::Type FAsyncPackage::CreateLinker()
 
 		// Always store package filename we loading from
 		Package->FileName = Desc.NameToLoad;
+#if WITH_EDITORONLY_DATA
+		// Assume all packages loaded through async loading are required by runtime
+		Package->SetLoadedByEditorPropertiesOnly(false);
+#endif
 
 		// if the linker already exists, we don't need to lookup the file (it may have been pre-created with
 		// a different filename)
diff --git a/Engine/Source/Runtime/CoreUObject/Private/UObject/LinkerLoad.cpp b/Engine/Source/Runtime/CoreUObject/Private/UObject/LinkerLoad.cpp
--- a/Engine/Source/Runtime/CoreUObject/Private/UObject/LinkerLoad.cpp
+++ b/Engine/Source/Runtime/CoreUObject/Private/UObject/LinkerLoad.cpp
@@ -26,6 +26,7 @@ DECLARE_CYCLE_STAT(TEXT("Linker Load Deferred"), STAT_LinkerLoadDeferred, STATGR
 DECLARE_STATS_GROUP( TEXT( "Linker Count" ), STATGROUP_LinkerCount, STATCAT_Advanced );
 DECLARE_DWORD_ACCUMULATOR_STAT(TEXT("Linker Count"), STAT_LinkerCount, STATGROUP_LinkerCount);
 DECLARE_DWORD_ACCUMULATOR_STAT(TEXT("Live Linker Count"), STAT_LiveLinkerCount, STATGROUP_LinkerCount);
+DECLARE_FLOAT_ACCUMULATOR_STAT(TEXT("Fixup editor-only flags time"), STAT_EditorOnlyFixupTime, STATGROUP_LinkerCount);
 
 /** Map that keeps track of any precached full package reads															*/
 TMap<FString, FLinkerLoad::FPackagePrecacheInfo> FLinkerLoad::PackagePrecacheMap;
@@ -2379,7 +2380,7 @@ bool FLinkerLoad::VerifyImportInner(const int32 ImportIndex, FString& WarningSuf
 
 			// we now fully load the package that we need a single export from - however, we still use CreatePackage below as it handles all cases when the package
 			// didn't exist (native only), etc		
-			TmpPkg = LoadPackageInternal(NULL, *Import.ObjectName.ToString(), InternalLoadFlags, this);
+			TmpPkg = LoadPackageInternal(NULL, *Import.ObjectName.ToString(), InternalLoadFlags | LOAD_IsVerifying, this);
 		}
 
 #if WITH_EDITOR
@@ -4188,6 +4189,37 @@ FArchive& FLinkerLoad::operator<<( UObject*& Object )
 	UObject* Temporary = NULL;
 	Temporary = IndexToObject( Index );
 
+#if WITH_EDITORONLY_DATA	
+	// When loading mark all packages that are accessed by non editor-only properties as being required at runtime.
+	if (Ar.IsLoading() && Temporary && !Ar.IsEditorOnlyPropertyOnTheStack())
+	{
+		FUObjectThreadContext& ThreadContext = FUObjectThreadContext::Get();
+		const bool bReferenceFromOutsideOfThePackage = Temporary->GetOutermost() != LinkerRoot;
+		const bool bIsAClass = Temporary->IsA(UClass::StaticClass());
+		const bool bReferencingPackageIsNotEditorOnly = bReferenceFromOutsideOfThePackage && !LinkerRoot->IsLoadedByEditorPropertiesOnly();
+		if (bReferencingPackageIsNotEditorOnly || bIsAClass)
+		{
+			// The package that caused this object to be loaded is not marked as editor-only, neighter is any of the referencing properties.
+			Temporary->GetOutermost()->SetLoadedByEditorPropertiesOnly(false);
+		}
+		else if (bReferenceFromOutsideOfThePackage && !bIsAClass)
+		{
+			// In this case the object is being accessed by object property from a package that's marked as editor-only, however
+			// since we're in the middle of loading, we can't be sure that the editor-only package will still be marked as editor-only
+			// after loading has finished (this is due to the fact how objects are being processed in EndLoad).
+			// So we need to remember which packages have been kept marked as editor-only by which package so that after all
+			// objects have been serialized we can go back and make sure the LinkerRoot package is still marked as editor-only and if not,
+			// remove the flag from all packages that are marked as such because of it.
+			FUObjectThreadContext& ThreadContext = FUObjectThreadContext::Get();
+			TSet<FName>& PackagesMarkedEditorOnly = ThreadContext.PackagesMarkedEditorOnlyByOtherPackage.FindOrAdd(LinkerRoot->GetFName());
+			if (!PackagesMarkedEditorOnly.Contains(Temporary->GetOutermost()->GetFName()))
+			{
+				PackagesMarkedEditorOnly.Add(Temporary->GetOutermost()->GetFName());
+			}
+		}
+	}
+#endif
+
 	FMemory::Memcpy(&Object, &Temporary, sizeof(UObject*));
 	return *this;
 }
@@ -4771,4 +4803,38 @@ void FLinkerLoad::FlushCache()
 	}
 }
 
+#if WITH_EDITORONLY_DATA
+/** Performs a fixup on packages' editor-only flag */
+void FixupPackageEditorOnlyFlag(FName PackageThatGotEditorOnlyFlagCleared, bool bRecursive)
+{
+	FUObjectThreadContext& ThreadContext = FUObjectThreadContext::Get();
+	STAT(double ThisTime = 0);
+	{
+		SCOPE_SECONDS_COUNTER(ThisTime);
+
+		// Now go through all packages that were marked as editor-only at load time
+		// and if they're no longer marked as such, make sure that all packages that
+		// were marked as editor-only because of that package, are now also marked as not editor-only.
+		TSet<FName>* PackagesMarkedEditorOnlyByThisPackage = ThreadContext.PackagesMarkedEditorOnlyByOtherPackage.Find(PackageThatGotEditorOnlyFlagCleared);
+		if (PackagesMarkedEditorOnlyByThisPackage)
+		{			
+			for (FName& PackageName : *PackagesMarkedEditorOnlyByThisPackage)
+			{
+				UPackage* EditorOnlyPackage = FindObjectFast<UPackage>(nullptr, PackageName);
+				if (EditorOnlyPackage && EditorOnlyPackage->IsLoadedByEditorPropertiesOnly())
+				{
+					// Now we will recursively unset the flag on all other packages
+					EditorOnlyPackage->SetLoadedByEditorPropertiesOnly(false, true);
+				}
+			}
+			ThreadContext.PackagesMarkedEditorOnlyByOtherPackage.Remove(PackageThatGotEditorOnlyFlagCleared);
+		}
+	}
+	if (!bRecursive)
+	{
+		INC_FLOAT_STAT_BY(STAT_EditorOnlyFixupTime, ThisTime);
+	}
+}
+#endif
+
 #undef LOCTEXT_NAMESPACE
diff --git a/Engine/Source/Runtime/CoreUObject/Private/UObject/Package.cpp b/Engine/Source/Runtime/CoreUObject/Private/UObject/Package.cpp
--- a/Engine/Source/Runtime/CoreUObject/Private/UObject/Package.cpp
+++ b/Engine/Source/Runtime/CoreUObject/Private/UObject/Package.cpp
@@ -2,6 +2,7 @@
 
 #include "CoreUObjectPrivate.h"
 #include "UObject/LinkerManager.h"
+#include "UObject/UObjectThreadContext.h"
 
 /*-----------------------------------------------------------------------------
 	UPackage.
@@ -38,6 +39,9 @@ void UPackage::PostInitProperties()
 #endif
 #if WITH_EDITORONLY_DATA
 	bIsCookedForEditor = false;
+	// Mark this package as editor-only by default. As soon as something in it is accessed through a non editor-only
+	// property the flag will be removed.
+	bLoadedByEditorPropertiesOnly = !HasAnyFlags(RF_ClassDefaultObject) && !(PackageFlags & PKG_CompiledIn);
 #endif
 }
 
@@ -212,6 +216,21 @@ void UPackage::BeginDestroy()
 	Super::BeginDestroy();
 }
 
+#if WITH_EDITORONLY_DATA
+void FixupPackageEditorOnlyFlag(FName PackageThatGotEditorOnlyFlagCleared, bool bRecursive);
+
+void UPackage::SetLoadedByEditorPropertiesOnly(bool bIsEditorOnly, bool bRecursive /*= false*/)
+{
+	const bool bWasEditorOnly = bLoadedByEditorPropertiesOnly;
+	bLoadedByEditorPropertiesOnly = bIsEditorOnly;
+	if (bWasEditorOnly && !bIsEditorOnly)
+	{
+		FixupPackageEditorOnlyFlag(GetFName(), bRecursive);
+	}
+}
+#endif
+
+
 IMPLEMENT_CORE_INTRINSIC_CLASS(UPackage, UObject,
 	{
 		Class->ClassAddReferencedObjects = &UPackage::AddReferencedObjects;
diff --git a/Engine/Source/Runtime/CoreUObject/Private/UObject/SavePackage.cpp b/Engine/Source/Runtime/CoreUObject/Private/UObject/SavePackage.cpp
--- a/Engine/Source/Runtime/CoreUObject/Private/UObject/SavePackage.cpp
+++ b/Engine/Source/Runtime/CoreUObject/Private/UObject/SavePackage.cpp
@@ -2972,32 +2972,12 @@ TMap<UObject*, UObject*> UnmarkExportTagFromDuplicates()
 }
 #endif // WITH_EDITOR
 
-/**
- * Save one specific object (along with any objects it references contained within the same Outer) into an Unreal package.
- * 
- * @param	InOuter							the outer to use for the new package
- * @param	Base							the object that should be saved into the package
- * @param	TopLevelFlags					For all objects which are not referenced [either directly, or indirectly] through Base, only objects
- *											that contain any of these flags will be saved.  If 0 is specified, only objects which are referenced
- *											by Base will be saved into the package.
- * @param	Filename						the name to use for the new package file
- * @param	Error							error output
- * @param	Conform							if non-NULL, all index tables for this will be sorted to match the order of the corresponding index table
- *											in the conform package
- * @param	bForceByteSwapping				whether we should forcefully byte swap before writing to disk
- * @param	bWarnOfLongFilename				[opt] If true (the default), warn when saving to a long filename.
- * @param	SaveFlags						Flags to control saving
- * @param	TargetPlatform					The platform being saved for
- * @param	FinalTimeStamp					If not FDateTime::MinValue(), the timestamp the saved file should be set to. (Intended for cooking only...)
- *
- * @return	true if the package was saved successfully.
- */
 
 #if WITH_EDITOR
 COREUOBJECT_API extern bool GOutputCookingWarnings;
 #endif
 
-bool UPackage::SavePackage( UPackage* InOuter, UObject* Base, EObjectFlags TopLevelFlags, const TCHAR* Filename, 
+ESavePackageResult UPackage::Save(UPackage* InOuter, UObject* Base, EObjectFlags TopLevelFlags, const TCHAR* Filename,
 	FOutputDevice* Error, FLinkerLoad* Conform, bool bForceByteSwapping, bool bWarnOfLongFilename, uint32 SaveFlags, 
 	const class ITargetPlatform* TargetPlatform, const FDateTime&  FinalTimeStamp, bool bSlowTask)
 {
@@ -3008,13 +2988,22 @@ bool UPackage::SavePackage( UPackage* InOuter, UObject* Base, EObjectFlags TopLe
 		if (GIsSavingPackage)
 		{
 			ensureMsgf(false, TEXT("Recursive SavePackage() is not supported"));
-			return false;
+			return ESavePackageResult::Error;
 		}
 
+		// Sanity checks
+		check(InOuter);
+		check(Filename);
 
-		
-		bool bIsCooking = TargetPlatform != NULL;
-
+		const bool bIsCooking = TargetPlatform != nullptr;
+#if WITH_EDITORONLY_DATA
+		if (bIsCooking && InOuter->IsLoadedByEditorPropertiesOnly())
+		{
+			// Don't save packages marked as editor-only.
+			UE_CLOG(!(SaveFlags & SAVE_NoError), LogSavePackage, Display, TEXT("Package loaded by editor-only properties %s. Package will not be saved."), *InOuter->GetName());
+			return ESavePackageResult::ReferencedOnlyByEditorOnlyData;
+		}
+#endif
 		// if we are cooking we should be doing it in the editor
 		// otherwise some other assumptions are bad
 		check( !bIsCooking || WITH_EDITOR );
@@ -3052,8 +3041,6 @@ bool UPackage::SavePackage( UPackage* InOuter, UObject* Base, EObjectFlags TopLe
 
 		UE_LOG_COOK_TIME(TEXT("BlockTillAllRequestsFinished"));
 
-		check(InOuter);
-		check(Filename);
 		uint32 Time=0; CLOCK_CYCLES(Time);
 
 		// Make sure package is fully loaded before saving. 
@@ -3077,7 +3064,7 @@ bool UPackage::SavePackage( UPackage* InOuter, UObject* Base, EObjectFlags TopLe
 				}
 				Error->Logf(ELogVerbosity::Warning, *ErrorText.ToString());
 			}
-			return false;
+			return ESavePackageResult::Error;
 		}
 
 		// Make sure package is allowed to be saved.
@@ -3103,7 +3090,7 @@ bool UPackage::SavePackage( UPackage* InOuter, UObject* Base, EObjectFlags TopLe
 					}
 					Error->Logf(ELogVerbosity::Warning, *ErrorText.ToString());
 				}
-				return false;
+				return ESavePackageResult::Error;
 			}
 		}
 
@@ -3127,7 +3114,7 @@ bool UPackage::SavePackage( UPackage* InOuter, UObject* Base, EObjectFlags TopLe
 				}
 				Error->Logf(ELogVerbosity::Error, *ErrorText.ToString());
 			}
-			return false;
+			return ESavePackageResult::Error;
 		}
 
 		bool FilterEditorOnly = (InOuter->PackageFlags & PKG_FilterEditorOnly) != 0;
@@ -3287,7 +3274,10 @@ bool UPackage::SavePackage( UPackage* InOuter, UObject* Base, EObjectFlags TopLe
 				InOuter->LinkerLicenseeVersion = Linker->LicenseeUE4Ver();
 				InOuter->LinkerCustomVersion = Linker->GetCustomVersions();
 
-				if ( EndSavingIfCancelled( Linker, TempFilename ) ) { return false; }
+				if (EndSavingIfCancelled(Linker, TempFilename)) 
+				{ 
+					return ESavePackageResult::Canceled; 
+				}
 				SlowTask.EnterProgressFrame();
 			
 				// keep a list of objects that would normally have gone into the dependency map, but since they are from cross-level dependencies, won't be found in the import map
@@ -3319,12 +3309,8 @@ bool UPackage::SavePackage( UPackage* InOuter, UObject* Base, EObjectFlags TopLe
 					GetObjectsWithAnyMarks(TagExpObjects, OBJECTMARK_TagExp);
 					if (TagExpObjects.Num() == 0)
 					{
-						if (!(SaveFlags & SAVE_NoError))
-						{
-							UE_LOG(LogSavePackage, Display, TEXT("No exports found (or all exports are editor-only) for %s. Package will not be saved."), *BaseFilename);
-						}
-
-						return false;
+						UE_CLOG(!(SaveFlags & SAVE_NoError), LogSavePackage, Display, TEXT("No exports found (or all exports are editor-only) for %s. Package will not be saved."), *BaseFilename);
+						return ESavePackageResult::ContainsEditorOnlyData;
 					}
 				}
 
@@ -3390,7 +3376,10 @@ bool UPackage::SavePackage( UPackage* InOuter, UObject* Base, EObjectFlags TopLe
 
 				UE_LOG_COOK_TIME(TEXT("SerializeImports"));
 				
-				if ( EndSavingIfCancelled( Linker, TempFilename ) ) { return false; }
+				if ( EndSavingIfCancelled( Linker, TempFilename ) ) 
+				{ 
+					return ESavePackageResult::Canceled;
+				}
 				SlowTask.EnterProgressFrame();
 
 				if ( !(Linker->Summary.PackageFlags & PKG_FilterEditorOnly) )
@@ -3415,7 +3404,10 @@ bool UPackage::SavePackage( UPackage* InOuter, UObject* Base, EObjectFlags TopLe
 
 				UE_LOG_COOK_TIME(TEXT("GatherLocalizableTextData"));
 				
-				if ( EndSavingIfCancelled( Linker, TempFilename ) ) { return false; }
+				if ( EndSavingIfCancelled( Linker, TempFilename ) ) 
+				{ 
+					return ESavePackageResult::Canceled;
+				}
 				SlowTask.EnterProgressFrame();
 
 				TArray<UObject*> PrivateObjects;
@@ -3492,7 +3484,10 @@ bool UPackage::SavePackage( UPackage* InOuter, UObject* Base, EObjectFlags TopLe
 
 				UE_LOG_COOK_TIME(TEXT("MarkNames"));
 
-				if ( EndSavingIfCancelled( Linker, TempFilename ) ) { return false; }
+				if ( EndSavingIfCancelled( Linker, TempFilename ) ) 
+				{ 
+					return ESavePackageResult::Canceled;
+				}
 				SlowTask.EnterProgressFrame();
 
 				if ( LevelObjects.Num() > 0 && ObjectsInOtherMaps.Num() == 0 )
@@ -3564,7 +3559,7 @@ bool UPackage::SavePackage( UPackage* InOuter, UObject* Base, EObjectFlags TopLe
 					{
 						Error->Logf(ELogVerbosity::Warning, TEXT("Can't save %s: Graph is linked to object %s in external map"), Filename, *CulpritString);
 					}
-					return false;
+					return ESavePackageResult::Error;
 				}
 
 				// The graph is linked to private objects!
@@ -3616,7 +3611,7 @@ bool UPackage::SavePackage( UPackage* InOuter, UObject* Base, EObjectFlags TopLe
 					{
 						Error->Logf(ELogVerbosity::Warning, TEXT("Can't save %s: Graph is linked to external private object %s"), Filename, *CulpritString);
 					}
-					return false;
+					return ESavePackageResult::Error;
 				}
 
 
@@ -3650,7 +3645,10 @@ bool UPackage::SavePackage( UPackage* InOuter, UObject* Base, EObjectFlags TopLe
 				*Linker << Linker->Summary;
 				int32 OffsetAfterPackageFileSummary = Linker->Tell();
 		
-				if ( EndSavingIfCancelled( Linker, TempFilename ) ) { return false; }
+				if ( EndSavingIfCancelled( Linker, TempFilename ) ) 
+				{ 
+					return ESavePackageResult::Canceled;
+				}
 				SlowTask.EnterProgressFrame();
 
 
@@ -3680,14 +3678,20 @@ bool UPackage::SavePackage( UPackage* InOuter, UObject* Base, EObjectFlags TopLe
 
 				UE_LOG_COOK_TIME(TEXT("SerializeSummary"));
 
-				if ( EndSavingIfCancelled( Linker, TempFilename ) ) { return false; }
+				if ( EndSavingIfCancelled( Linker, TempFilename ) ) 
+				{ 
+					return ESavePackageResult::Canceled;
+				}
 				SlowTask.EnterProgressFrame();
 
 				// Sort names.
 				FObjectNameSortHelper NameSortHelper;
 				NameSortHelper.SortNames( Linker, Conform );
 
-				if ( EndSavingIfCancelled( Linker, TempFilename ) ) { return false; }
+				if ( EndSavingIfCancelled( Linker, TempFilename ) ) 
+				{ 
+					return ESavePackageResult::Canceled;
+				}
 				SlowTask.EnterProgressFrame();
 
 				// Save names.
@@ -3700,7 +3704,10 @@ bool UPackage::SavePackage( UPackage* InOuter, UObject* Base, EObjectFlags TopLe
 
 				UE_LOG_COOK_TIME(TEXT("SerializeNames"));
 
-				if ( EndSavingIfCancelled( Linker, TempFilename ) ) { return false; }
+				if ( EndSavingIfCancelled( Linker, TempFilename ) ) 
+				{ 
+					return ESavePackageResult::Canceled;
+				}
 				SlowTask.EnterProgressFrame();
 
 				Linker->Summary.GatherableTextDataOffset = 0;
@@ -3719,7 +3726,10 @@ bool UPackage::SavePackage( UPackage* InOuter, UObject* Base, EObjectFlags TopLe
 
 				UE_LOG_COOK_TIME(TEXT("SerializeGatherableTextData"));
 
-				if ( EndSavingIfCancelled( Linker, TempFilename ) ) { return false; }
+				if ( EndSavingIfCancelled( Linker, TempFilename ) ) 
+				{ 
+					return ESavePackageResult::Canceled;
+				}
 				SlowTask.EnterProgressFrame();
 
 				// Build ImportMap.
@@ -3735,7 +3745,10 @@ bool UPackage::SavePackage( UPackage* InOuter, UObject* Base, EObjectFlags TopLe
 				}
 
 
-				if ( EndSavingIfCancelled( Linker, TempFilename ) ) { return false; }
+				if ( EndSavingIfCancelled( Linker, TempFilename ) ) 
+				{ 
+					return ESavePackageResult::Canceled;
+				}
 				SlowTask.EnterProgressFrame();
 
 				// sort and conform imports
@@ -3743,7 +3756,10 @@ bool UPackage::SavePackage( UPackage* InOuter, UObject* Base, EObjectFlags TopLe
 				ImportSortHelper.SortImports( Linker, Conform );
 				Linker->Summary.ImportCount = Linker->ImportMap.Num();
 
-				if ( EndSavingIfCancelled( Linker, TempFilename ) ) { return false; }
+				if ( EndSavingIfCancelled( Linker, TempFilename ) ) 
+				{ 
+					return ESavePackageResult::Canceled;
+				}
 				SlowTask.EnterProgressFrame();
 
 				
@@ -3781,7 +3797,10 @@ bool UPackage::SavePackage( UPackage* InOuter, UObject* Base, EObjectFlags TopLe
 
 				UE_LOG_COOK_TIME(TEXT("BuildExportMap"));
 
-				if ( EndSavingIfCancelled( Linker, TempFilename ) ) { return false; }
+				if ( EndSavingIfCancelled( Linker, TempFilename ) ) 
+				{ 
+					return ESavePackageResult::Canceled;
+				}
 				SlowTask.EnterProgressFrame();
 
 				// Sort exports alphabetically and conform the export table (if necessary)
@@ -3790,7 +3809,10 @@ bool UPackage::SavePackage( UPackage* InOuter, UObject* Base, EObjectFlags TopLe
 				
 				UE_LOG_COOK_TIME(TEXT("SortExportsNonSeekfree"));
 
-				if ( EndSavingIfCancelled( Linker, TempFilename ) ) { return false; }
+				if ( EndSavingIfCancelled( Linker, TempFilename ) ) 
+				{ 
+					return ESavePackageResult::Canceled;
+				}
 				SlowTask.EnterProgressFrame();
 
 				// Sort exports for seek-free loading.
@@ -3802,7 +3824,10 @@ bool UPackage::SavePackage( UPackage* InOuter, UObject* Base, EObjectFlags TopLe
 				
 				UE_LOG_COOK_TIME(TEXT("SortExportsSeekfree"));
 
-				if ( EndSavingIfCancelled( Linker, TempFilename ) ) { return false; }
+				if ( EndSavingIfCancelled( Linker, TempFilename ) ) 
+				{ 
+					return ESavePackageResult::Canceled;
+				}
 				SlowTask.EnterProgressFrame();
 
 				// Pre-size depends map.
@@ -3884,7 +3909,10 @@ bool UPackage::SavePackage( UPackage* InOuter, UObject* Base, EObjectFlags TopLe
 
 
 
-				if ( EndSavingIfCancelled( Linker, TempFilename ) ) { return false; }
+				if ( EndSavingIfCancelled( Linker, TempFilename ) ) 
+				{ 
+					return ESavePackageResult::Canceled;
+				}
 				SlowTask.EnterProgressFrame();
 
 				// Set linker reverse mappings.
@@ -3908,7 +3936,10 @@ bool UPackage::SavePackage( UPackage* InOuter, UObject* Base, EObjectFlags TopLe
 					}
 				}
 
-				if ( EndSavingIfCancelled( Linker, TempFilename ) ) { return false; }
+				if ( EndSavingIfCancelled( Linker, TempFilename ) ) 
+				{ 
+					return ESavePackageResult::Canceled;
+				}
 				SlowTask.EnterProgressFrame();
 
 				// If this is a map package, make sure there is a world or level in the export map.
@@ -3944,7 +3975,10 @@ bool UPackage::SavePackage( UPackage* InOuter, UObject* Base, EObjectFlags TopLe
 
 				UE_LOG_COOK_TIME(TEXT("BuildDependencyMap"));
 
-				if ( EndSavingIfCancelled( Linker, TempFilename ) ) { return false; }
+				if ( EndSavingIfCancelled( Linker, TempFilename ) ) 
+				{ 
+					return ESavePackageResult::Canceled;
+				}
 				SlowTask.EnterProgressFrame();
 
 				for( int32 i=0; i<Linker->ImportMap.Num(); i++ )
@@ -3967,7 +4001,10 @@ bool UPackage::SavePackage( UPackage* InOuter, UObject* Base, EObjectFlags TopLe
 
 				// Find components referenced by exports.
 
-				if ( EndSavingIfCancelled( Linker, TempFilename ) ) { return false; }
+				if ( EndSavingIfCancelled( Linker, TempFilename ) ) 
+				{ 
+					return ESavePackageResult::Canceled;
+				}
 				SlowTask.EnterProgressFrame();
 
 				// Save dummy import map, overwritten later.
@@ -3983,7 +4020,10 @@ bool UPackage::SavePackage( UPackage* InOuter, UObject* Base, EObjectFlags TopLe
 
 				UE_LOG_COOK_TIME(TEXT("SerializeImportMap"));
 
-				if ( EndSavingIfCancelled( Linker, TempFilename ) ) { return false; }
+				if ( EndSavingIfCancelled( Linker, TempFilename ) ) 
+				{ 
+					return ESavePackageResult::Canceled;
+				}
 				SlowTask.EnterProgressFrame();
 
 				// Save dummy export map, overwritten later.
@@ -3998,7 +4038,10 @@ bool UPackage::SavePackage( UPackage* InOuter, UObject* Base, EObjectFlags TopLe
 
 				UE_LOG_COOK_TIME(TEXT("SerializeExportMap"));
 
-				if ( EndSavingIfCancelled( Linker, TempFilename ) ) { return false; }
+				if ( EndSavingIfCancelled( Linker, TempFilename ) ) 
+				{ 
+					return ESavePackageResult::Canceled;
+				}
 				SlowTask.EnterProgressFrame();
 
 				// save depends map (no need for later patching)
@@ -4013,7 +4056,10 @@ bool UPackage::SavePackage( UPackage* InOuter, UObject* Base, EObjectFlags TopLe
 
 				UE_LOG_COOK_TIME(TEXT("SerializeDependencyMap"));
 
-				if (EndSavingIfCancelled(Linker, TempFilename)) { return false; }
+				if (EndSavingIfCancelled(Linker, TempFilename)) 
+				{ 
+					return ESavePackageResult::Canceled;
+				}
 				SlowTask.EnterProgressFrame();
 
 				// Save string asset reference map
@@ -4038,7 +4084,10 @@ bool UPackage::SavePackage( UPackage* InOuter, UObject* Base, EObjectFlags TopLe
 				
 				Linker->Summary.TotalHeaderSize	= Linker->Tell();
 
-				if ( EndSavingIfCancelled( Linker, TempFilename ) ) { return false; }
+				if ( EndSavingIfCancelled( Linker, TempFilename ) ) 
+				{ 
+					return ESavePackageResult::Canceled;
+				}
 				SlowTask.EnterProgressFrame(1, NSLOCTEXT("Core", "ProcessingExports", "ProcessingExports..."));
 
 				// look for this package in the list of packages to generate script SHA for 
@@ -4059,7 +4108,10 @@ bool UPackage::SavePackage( UPackage* InOuter, UObject* Base, EObjectFlags TopLe
 					int32 LastExportSaveStep = 0;
 					for( int32 i=0; i<Linker->ExportMap.Num(); i++ )
 					{
-						if ( EndSavingIfCancelled( Linker, TempFilename ) ) { return false; }
+						if ( EndSavingIfCancelled( Linker, TempFilename ) ) 
+						{ 
+							return ESavePackageResult::Canceled;
+						}
 						ExportScope.EnterProgressFrame();
 
 						FObjectExport& Export = Linker->ExportMap[i];
@@ -4149,7 +4201,10 @@ bool UPackage::SavePackage( UPackage* InOuter, UObject* Base, EObjectFlags TopLe
 				}
 
 				
-				if ( EndSavingIfCancelled( Linker, TempFilename ) ) { return false; }
+				if ( EndSavingIfCancelled( Linker, TempFilename ) ) 
+				{ 
+					return ESavePackageResult::Canceled;
+				}
 				SlowTask.EnterProgressFrame(1, NSLOCTEXT("Core", "SerializingBulkData", "Serializing bulk data"));
 
 				// now we write all the bulkdata that is supposed to be at the end of the package
@@ -4259,7 +4314,10 @@ bool UPackage::SavePackage( UPackage* InOuter, UObject* Base, EObjectFlags TopLe
 
 				UE_LOG_COOK_TIME(TEXT("SerializeExportMap"));
 
-				if ( EndSavingIfCancelled( Linker, TempFilename ) ) { return false; }
+				if ( EndSavingIfCancelled( Linker, TempFilename ) ) 
+				{ 
+					return ESavePackageResult::Canceled;
+				}
 				SlowTask.EnterProgressFrame();
 
 				FFormatNamedArguments NamedArgs;
@@ -4287,7 +4345,10 @@ bool UPackage::SavePackage( UPackage* InOuter, UObject* Base, EObjectFlags TopLe
 				*Linker << Linker->Summary;
 				check( Linker->Tell() == OffsetAfterPackageFileSummary );
 
-				if ( EndSavingIfCancelled( Linker, TempFilename ) ) { return false; }
+				if ( EndSavingIfCancelled( Linker, TempFilename ) ) 
+				{ 
+					return ESavePackageResult::Canceled;
+				}
 				SlowTask.EnterProgressFrame();
 
 				// Detach archive used for saving, closing file handle.
@@ -4298,7 +4359,10 @@ bool UPackage::SavePackage( UPackage* InOuter, UObject* Base, EObjectFlags TopLe
 				UNCLOCK_CYCLES(Time);
 				UE_LOG(LogSavePackage, Log,  TEXT("Save=%.2fms"), FPlatformTime::ToMilliseconds(Time) );
 		
-				if ( EndSavingIfCancelled( Linker, TempFilename ) ) { return false; }
+				if ( EndSavingIfCancelled( Linker, TempFilename ) ) 
+				{ 
+					return ESavePackageResult::Canceled;
+				}
 				SlowTask.EnterProgressFrame();
 
 				if( Success == true )
@@ -4491,15 +4555,22 @@ bool UPackage::SavePackage( UPackage* InOuter, UObject* Base, EObjectFlags TopLe
 
 		UE_LOG(LogSavePackage, Display, TEXT("Finished SavePackage %s"), Filename);
 
-		return Success;
+		return Success ? ESavePackageResult::Success : ESavePackageResult::Error;
 	}
 	else
 	{
-		return false;
+		return ESavePackageResult::Error;
 	}
 }
 
-
+bool UPackage::SavePackage(UPackage* InOuter, UObject* Base, EObjectFlags TopLevelFlags, const TCHAR* Filename,
+	FOutputDevice* Error, FLinkerLoad* Conform, bool bForceByteSwapping, bool bWarnOfLongFilename, uint32 SaveFlags,
+	const class ITargetPlatform* TargetPlatform, const FDateTime&  FinalTimeStamp, bool bSlowTask)
+{
+	const ESavePackageResult Result = Save(InOuter, Base, TopLevelFlags, Filename, Error, Conform, bForceByteSwapping, 
+		bWarnOfLongFilename, SaveFlags, TargetPlatform, FinalTimeStamp, bSlowTask);
+	return Result == ESavePackageResult::Success;
+}
 
 /**
  * Static: Saves thumbnail data for the specified package outer and linker
diff --git a/Engine/Source/Runtime/CoreUObject/Private/UObject/UObjectGlobals.cpp b/Engine/Source/Runtime/CoreUObject/Private/UObject/UObjectGlobals.cpp
--- a/Engine/Source/Runtime/CoreUObject/Private/UObject/UObjectGlobals.cpp
+++ b/Engine/Source/Runtime/CoreUObject/Private/UObject/UObjectGlobals.cpp
@@ -831,6 +831,12 @@ UObject* StaticLoadObjectInternal(UClass* ObjectClass, UObject* InOuter, const T
 		StrName += FPackageName::GetShortName(InName);
 		Result = StaticLoadObjectInternal(ObjectClass, InOuter, *StrName, Filename, LoadFlags, Sandbox, bAllowObjectReconciliation);
 	}
+#if WITH_EDITORONLY_DATA
+	else if (Result)
+	{
+		Result->GetOutermost()->SetLoadedByEditorPropertiesOnly(false);
+	}
+#endif
 
 	return Result;
 }
@@ -960,11 +966,21 @@ UPackage* LoadPackageInternal(UPackage* InOuter, const TCHAR* InLongPackageName,
 			Result->ThisRequiresLocalizationGather(Linker->RequiresLocalizationGather());
 		};
 
+#if WITH_EDITORONLY_DATA
+		if (!(LoadFlags & LOAD_IsVerifying) &&
+			(!ImportLinker || !ImportLinker->GetSerializedProperty() || !ImportLinker->GetSerializedProperty()->IsEditorOnlyProperty()))
+		{
+			// If this package hasn't been loaded as part of import verification and there's no import linker or the
+			// currently serialized property is not editor-only mark this package as runtime.
+			Result->SetLoadedByEditorPropertiesOnly(false);
+		}
+#endif
+
 		if (Result && Result->HasAnyFlags(RF_WasLoaded))
 		{
 			// The linker is associated with a package that has already been loaded.
 			// Loading packages that have already been loaded is unsupported.
-			EndLoadAndCopyLocalizationGatherFlag();
+			EndLoadAndCopyLocalizationGatherFlag();			
 			return Result;
 		}
 
@@ -1035,15 +1051,8 @@ UPackage* LoadPackageInternal(UPackage* InOuter, const TCHAR* InLongPackageName,
 		{
 			// Make sure we pass the property that's currently being serialized by the linker that owns the import 
 			// that triggered this LoadPackage call
-			UProperty* OldSerializedProperty = Linker->GetSerializedProperty();
-			if (ImportLinker)
-			{
-				Linker->SetSerializedProperty(ImportLinker->GetSerializedProperty());
-			}
-
+			FSerializedPropertyScope SerializedProperty(*Linker, ImportLinker ? ImportLinker->GetSerializedProperty() : Linker->GetSerializedProperty());
 			Linker->LoadAllObjects();
-
-			Linker->SetSerializedProperty(OldSerializedProperty);
 		}
 		else
 		{
diff --git a/Engine/Source/Runtime/CoreUObject/Public/Serialization/ArchiveUObjectBase.h b/Engine/Source/Runtime/CoreUObject/Public/Serialization/ArchiveUObjectBase.h
--- a/Engine/Source/Runtime/CoreUObject/Public/Serialization/ArchiveUObjectBase.h
+++ b/Engine/Source/Runtime/CoreUObject/Public/Serialization/ArchiveUObjectBase.h
@@ -19,3 +19,35 @@ public:
 	virtual FArchive& operator<< (struct FStringAssetReference& Value) override;
 	// End FArchive Interface
 };
+
+class UProperty;
+
+/** Helper class to set and restore serialized property on an archive */
+class FSerializedPropertyScope
+{
+	FArchive& Ar;
+	UProperty* Property;
+	UProperty* OldProperty;
+#if WITH_EDITORONLY_DATA
+	void PushEditorOnlyProperty();
+	void PopEditorOnlyProperty();
+#endif
+public:
+	FSerializedPropertyScope(FArchive& InAr, UProperty* InProperty)
+		: Ar(InAr)
+		, Property(InProperty)
+	{
+		OldProperty = Ar.GetSerializedProperty();
+		Ar.SetSerializedProperty(Property);
+#if WITH_EDITORONLY_DATA
+		PushEditorOnlyProperty();
+#endif
+	}
+	~FSerializedPropertyScope()
+	{
+#if WITH_EDITORONLY_DATA
+		PopEditorOnlyProperty();
+#endif
+		Ar.SetSerializedProperty(OldProperty);
+	}
+};
diff --git a/Engine/Source/Runtime/CoreUObject/Public/UObject/CoreObject.h b/Engine/Source/Runtime/CoreUObject/Public/UObject/CoreObject.h
--- a/Engine/Source/Runtime/CoreUObject/Public/UObject/CoreObject.h
+++ b/Engine/Source/Runtime/CoreUObject/Public/UObject/CoreObject.h
@@ -19,6 +19,23 @@ struct FLevelGuids
 	TArray<FGuid> Guids;
 };
 
+/**
+ * Represents the result of saving a package
+ */
+enum class ESavePackageResult
+{
+	/** Package was saved successfully */
+	Success, 
+	/** Unknown error occured when saving package */
+	Error,
+	/** Canceled by user */
+	Canceled,
+	/** [When cooking] Package was not saved because it contained editor-only data */
+	ContainsEditorOnlyData,
+	/** [When cooking] Package was not saved because it was referenced by editor-only properties */
+	ReferencedOnlyByEditorOnlyData
+};
+
 /**
  * A package.
  */
@@ -52,6 +69,10 @@ public:
 private:
 	/** Used by the editor to determine if a package has been changed.																							*/
 	bool	bDirty;
+#if WITH_EDITORONLY_DATA
+	/** True if this package is only referenced by editor-only properties */
+	bool bLoadedByEditorPropertiesOnly;
+#endif
 public:
 	/** Whether this package has been fully loaded (aka had all it's exports created) at some point.															*/
 	bool	bHasBeenFullyLoaded;
@@ -73,6 +94,13 @@ public:
 
 	virtual bool IsNameStableForNetworking() const override { return true; }		// For now, assume all packages have stable net names
 
+#if WITH_EDITORONLY_DATA
+	/** Sets the bLoadedByEditorPropertiesOnly flag */
+	void SetLoadedByEditorPropertiesOnly(bool bIsEditorOnly, bool bRecursive = false);
+	/** returns true when the package is only referenced by editor-only flag */
+	bool IsLoadedByEditorPropertiesOnly() const { return bLoadedByEditorPropertiesOnly; }
+#endif
+
 	/** Package flags, serialized.*/
 	uint32	PackageFlags;
 
@@ -366,12 +394,36 @@ public:
 	 * @param	TargetPlatform					The platform being saved for
 	 * @param	FinalTimeStamp					If not FDateTime::MinValue(), the timestamp the saved file should be set to. (Intended for cooking only...)
 	 *
-	 * @return	true if the package was saved successfully.
+	 * @return	ESavePackageResult enum value with the result of saving a package.
 	 */
-	static bool SavePackage( UPackage* InOuter, UObject* Base, EObjectFlags TopLevelFlags, const TCHAR* Filename, 
+	static ESavePackageResult Save(UPackage* InOuter, UObject* Base, EObjectFlags TopLevelFlags, const TCHAR* Filename,
 		FOutputDevice* Error=GError, FLinkerLoad* Conform=NULL, bool bForceByteSwapping=false, bool bWarnOfLongFilename=true, 
 		uint32 SaveFlags=SAVE_None, const class ITargetPlatform* TargetPlatform = NULL, const FDateTime& FinalTimeStamp = FDateTime::MinValue(), bool bSlowTask = true );
 
+	/**
+	* Save one specific object (along with any objects it references contained within the same Outer) into an Unreal package.
+	*
+	* @param	InOuter							the outer to use for the new package
+	* @param	Base							the object that should be saved into the package
+	* @param	TopLevelFlags					For all objects which are not referenced [either directly, or indirectly] through Base, only objects
+	*											that contain any of these flags will be saved.  If 0 is specified, only objects which are referenced
+	*											by Base will be saved into the package.
+	* @param	Filename						the name to use for the new package file
+	* @param	Error							error output
+	* @param	Conform							if non-NULL, all index tables for this will be sorted to match the order of the corresponding index table
+	*											in the conform package
+	* @param	bForceByteSwapping				whether we should forcefully byte swap before writing to disk
+	* @param	bWarnOfLongFilename				[opt] If true (the default), warn when saving to a long filename.
+	* @param	SaveFlags						Flags to control saving
+	* @param	TargetPlatform					The platform being saved for
+	* @param	FinalTimeStamp					If not FDateTime::MinValue(), the timestamp the saved file should be set to. (Intended for cooking only...)
+	*
+	* @return	true if the package was saved successfully.
+	*/
+	static bool SavePackage(UPackage* InOuter, UObject* Base, EObjectFlags TopLevelFlags, const TCHAR* Filename,
+		FOutputDevice* Error = GError, FLinkerLoad* Conform = NULL, bool bForceByteSwapping = false, bool bWarnOfLongFilename = true,
+		uint32 SaveFlags = SAVE_None, const class ITargetPlatform* TargetPlatform = NULL, const FDateTime& FinalTimeStamp = FDateTime::MinValue(), bool bSlowTask = true);
+
 	/** Wait for any SAVE_Async file writes to complete **/
 	static void WaitForAsyncFileWrites();
 
diff --git a/Engine/Source/Runtime/CoreUObject/Public/UObject/ObjectBase.h b/Engine/Source/Runtime/CoreUObject/Public/UObject/ObjectBase.h
--- a/Engine/Source/Runtime/CoreUObject/Public/UObject/ObjectBase.h
+++ b/Engine/Source/Runtime/CoreUObject/Public/UObject/ObjectBase.h
@@ -50,7 +50,7 @@ enum ELoadFlags
 	LOAD_AllowDll					= 0x00000020,	// Allow plain DLLs.
 //	LOAD_Unused						= 0x00000040
 	LOAD_NoVerify					= 0x00000080,   // Don't verify imports yet.
-//	LOAD_Unused						= 0x00000100,
+	LOAD_IsVerifying			= 0x00000100,		// Is verifying imports
 //	LOAD_Unused						= 0x00000200,
 //	LOAD_Unused						= 0x00000400,
 //	LOAD_Unused						= 0x00000800,
diff --git a/Engine/Source/Runtime/CoreUObject/Public/UObject/PropertyTag.h b/Engine/Source/Runtime/CoreUObject/Public/UObject/PropertyTag.h
--- a/Engine/Source/Runtime/CoreUObject/Public/UObject/PropertyTag.h
+++ b/Engine/Source/Runtime/CoreUObject/Public/UObject/PropertyTag.h
@@ -137,12 +137,8 @@ struct FPropertyTag
 		}
 		else
 		{
-			UProperty* OldSerializedProperty = Ar.GetSerializedProperty();
-			Ar.SetSerializedProperty(Property);
-
+			FSerializedPropertyScope SerializedProperty(Ar, Property);
 			Property->SerializeItem( Ar, Value, Defaults );
-
-			Ar.SetSerializedProperty(OldSerializedProperty);
 		}
 	}
 };
diff --git a/Engine/Source/Runtime/CoreUObject/Public/UObject/UObjectThreadContext.h b/Engine/Source/Runtime/CoreUObject/Public/UObject/UObjectThreadContext.h
--- a/Engine/Source/Runtime/CoreUObject/Public/UObject/UObjectThreadContext.h
+++ b/Engine/Source/Runtime/CoreUObject/Public/UObject/UObjectThreadContext.h
@@ -51,4 +51,8 @@ public:
 	/** Used to verify that the Super::PostLoad chain is intact.			*/
 	TArray<UObject*, TInlineAllocator<16> > DebugPostLoad;
 #endif
+#if WITH_EDITORONLY_DATA
+	/** Maps a package name to all packages marked as editor-only due to the fact it was marked as editor-only */
+	TMap<FName, TSet<FName>> PackagesMarkedEditorOnlyByOtherPackage;
+#endif
 };
\ No newline at end of file
diff --git a/Engine/Source/Runtime/CoreUObject/Public/UObject/UnrealType.h b/Engine/Source/Runtime/CoreUObject/Public/UObject/UnrealType.h
--- a/Engine/Source/Runtime/CoreUObject/Public/UObject/UnrealType.h
+++ b/Engine/Source/Runtime/CoreUObject/Public/UObject/UnrealType.h
@@ -218,13 +218,13 @@ public:
 	{
 		if( ShouldSerializeValue(Ar) )
 		{
-			UProperty* OldSerializedProperty = Ar.GetSerializedProperty();
+			FSerializedPropertyScope SerializedProperty(Ar, this);
 			for (int32 Idx = 0; Idx < ArrayDim; Idx++)
 			{
+				// Keep setting the property in case something inside of SerializeItem changes it
 				Ar.SetSerializedProperty(this);
 				SerializeItem( Ar, ContainerPtrToValuePtr<void>(Data, Idx) );
 			}
-			Ar.SetSerializedProperty(OldSerializedProperty);
 		}
 	}
 	/**
@@ -245,10 +245,8 @@ public:
 				void const* Default = ContainerPtrToValuePtrForDefaults<void>(DefaultStruct, DefaultData, Idx);
 				if ( !Identical(Target, Default, Ar.GetPortFlags()) )
 				{
-					UProperty* OldSerializedProperty = Ar.GetSerializedProperty();
-					Ar.SetSerializedProperty(this);
+					FSerializedPropertyScope SerializedProperty(Ar, this);
 					SerializeItem( Ar, Target, Default );
-					Ar.SetSerializedProperty(OldSerializedProperty);
 				}
 			}
 		}
