Fixing ExportMapSorter condition

diff --git a/Engine/Source/Runtime/CoreUObject/Private/UObject/LinkerLoad.cpp b/Engine/Source/Runtime/CoreUObject/Private/UObject/LinkerLoad.cpp
--- a/Engine/Source/Runtime/CoreUObject/Private/UObject/LinkerLoad.cpp
+++ b/Engine/Source/Runtime/CoreUObject/Private/UObject/LinkerLoad.cpp
@@ -4524,7 +4524,7 @@ bool ExportMapSorter(const FObjectExport& Lhs, const FObjectExport& Rhs)
 		return true;
 	}
 
-	if (Rhs.ObjectName > Rhs.ObjectName)
+	if (Lhs.ObjectName > Rhs.ObjectName)
 	{
 		return false;
 	}
