Fix shadowed local variable. Mac CIS fix.

diff --git a/Engine/Source/Runtime/CoreUObject/Private/UObject/LinkerLoad.cpp b/Engine/Source/Runtime/CoreUObject/Private/UObject/LinkerLoad.cpp
--- a/Engine/Source/Runtime/CoreUObject/Private/UObject/LinkerLoad.cpp
+++ b/Engine/Source/Runtime/CoreUObject/Private/UObject/LinkerLoad.cpp
@@ -4210,7 +4210,6 @@ FArchive& FLinkerLoad::operator<<( UObject*& Object )
 			// So we need to remember which packages have been kept marked as editor-only by which package so that after all
 			// objects have been serialized we can go back and make sure the LinkerRoot package is still marked as editor-only and if not,
 			// remove the flag from all packages that are marked as such because of it.
-			FUObjectThreadContext& ThreadContext = FUObjectThreadContext::Get();
 			TSet<FName>& PackagesMarkedEditorOnly = ThreadContext.PackagesMarkedEditorOnlyByOtherPackage.FindOrAdd(LinkerRoot->GetFName());
 			if (!PackagesMarkedEditorOnly.Contains(Temporary->GetOutermost()->GetFName()))
 			{
