UE4 - Added a few blueprint related stats to the UObject core.

diff --git a/Engine/Source/Runtime/CoreUObject/Private/UObject/UObjectGlobals.cpp b/Engine/Source/Runtime/CoreUObject/Private/UObject/UObjectGlobals.cpp
--- a/Engine/Source/Runtime/CoreUObject/Private/UObject/UObjectGlobals.cpp
+++ b/Engine/Source/Runtime/CoreUObject/Private/UObject/UObjectGlobals.cpp
@@ -1580,6 +1580,7 @@ UObject* StaticDuplicateObject(UObject const* SourceObject, UObject* DestOuter,
 
 UObject* StaticDuplicateObjectEx( FObjectDuplicationParameters& Parameters )
 {
+	QUICK_SCOPE_CYCLE_COUNTER(STAT_StaticDuplicateObject);
 	// make sure the two classes are the same size, as this hopefully will mean they are serialization
 	// compatible. It's not a guarantee, but will help find errors
 	checkf( (Parameters.DestClass->GetPropertiesSize() >= Parameters.SourceObject->GetClass()->GetPropertiesSize()),
@@ -2533,6 +2534,7 @@ void FObjectInitializer::InitProperties(UObject* Obj, UClass* DefaultsClass, UOb
 		{
 			if (Class->GetDefaultObject(false) != DefaultData)
 			{
+				QUICK_SCOPE_CYCLE_COUNTER(STAT_InitProperties_FromTemplate);
 				for (UProperty* P = Class->PropertyLink; P; P = P->PropertyLinkNext)
 				{
 					P->CopyCompleteValue_InContainer(Obj, DefaultData);
@@ -2540,6 +2542,7 @@ void FObjectInitializer::InitProperties(UObject* Obj, UClass* DefaultsClass, UOb
 			}
 			else
 			{
+				QUICK_SCOPE_CYCLE_COUNTER(STAT_InitProperties_ConfigEtcOnly);
 				// Copy all properties that require additional initialization (e.g. CPF_Config).
 				for (UProperty* P = Class->PostConstructLink; P; P = P->PostConstructLinkNext)
 				{
@@ -2550,6 +2553,7 @@ void FObjectInitializer::InitProperties(UObject* Obj, UClass* DefaultsClass, UOb
 	}
 	else
 	{
+		QUICK_SCOPE_CYCLE_COUNTER(STAT_InitProperties_Blueprint);
 		UObject* ClassDefaults = bCopyTransientsFromClassDefaults ? DefaultsClass->GetDefaultObject() : NULL;		
 		for (UProperty* P = Class->PropertyLink; P; P = P->PropertyLinkNext)
 		{
