Minimizing GC allocations by pooling objects to serialize arrays.

Added gc.MinDesiredObjectsPerSubTask to tweak GC sub-tasks

diff --git a/Engine/Source/Runtime/CoreUObject/Private/UObject/GarbageCollection.cpp b/Engine/Source/Runtime/CoreUObject/Private/UObject/GarbageCollection.cpp
--- a/Engine/Source/Runtime/CoreUObject/Private/UObject/GarbageCollection.cpp
+++ b/Engine/Source/Runtime/CoreUObject/Private/UObject/GarbageCollection.cpp
@@ -182,6 +182,79 @@ bool IsGarbageCollectionLocked()
 	return GGarbageCollectionGuardCritical.IsAsyncLocked();
 }
 
+/**
+ * Pool for reducing GC allocations
+ */
+class FGCArrayPool
+{
+public:
+
+	/**
+	* Gets the singleton instance of the FObjectArrayPool
+	* @return Pool singleton.
+	*/
+	static FGCArrayPool& Get()
+	{
+		static FGCArrayPool Singleton;
+		return Singleton;
+	}
+
+	/**
+	* Gets an event from the pool or creates one if necessary.
+	*
+	* @return The array.
+	* @see ReturnToPool
+	*/
+	TArray<UObject*>* GetArrayFromPool()
+	{
+		TArray<UObject*>* Result = Pool.Pop();
+		if (!Result)
+		{
+			Result = new TArray<UObject*>();
+		}
+		check(Result);
+		return Result;
+	}
+
+	/**
+	* Returns an array to the pool.
+	*
+	* @param Array The array to return.
+	* @see GetArrayFromPool
+	*/
+	void ReturnToPool(TArray<UObject*>* Array)
+	{
+		check(Array);
+		Array->Reset();
+		Pool.Push(Array);
+	}
+
+	/** Performs memory cleanup */
+	void Cleanup()
+	{
+		uint32 FreedMemory = 0;
+		TArray< TArray<UObject*>* > AllArrays;
+		Pool.PopAll(AllArrays);
+		for (TArray<UObject*>* Array : AllArrays)
+		{
+			FreedMemory += Array->GetAllocatedSize();
+			delete Array;
+		}
+		UE_LOG(LogGarbage, Log, TEXT("Freed %ub from %d GC array pools."), FreedMemory, AllArrays.Num());
+	}
+
+private:
+
+	/** Holds the collection of recycled arrays. */
+	TLockFreePointerList< TArray<UObject*> > Pool;
+};
+
+/** Called on shutdown to free GC memory */
+void CleanupGCArrayPools()
+{
+	FGCArrayPool::Get().Cleanup();
+}
+
 /**
  * If set and VERIFY_DISREGARD_GC_ASSUMPTIONS is true, we verify GC assumptions about "Disregard For GC" objects. We also
  * verify that no unreachable actors/ components are referenced if VERIFY_NO_UNREACHABLE_OBJECTS_ARE_REFERENCED
@@ -189,6 +262,15 @@ bool IsGarbageCollectionLocked()
  */
 COREUOBJECT_API bool	GShouldVerifyGCAssumptions				= !(UE_BUILD_SHIPPING != 0 && WITH_EDITOR != 0);
 
+// Minimum number of objects to spawn a GC sub-task for
+static int32 GMinDesiredObjectsPerSubTask = 128;
+static FAutoConsoleVariableRef CVarMinDesiredObjectsPerSubTask(
+	TEXT("gc.MinDesiredObjectsPerSubTask"),
+	GMinDesiredObjectsPerSubTask,
+	TEXT("Minimum number of objects to spawn a GC sub-task for."),
+	ECVF_Default
+	);
+
 #if PERF_DETAILED_PER_CLASS_GC_STATS
 /** Map from a UClass' FName to the number of objects that were purged during the last purge phase of this class.	*/
 static TMap<const FName,uint32> GClassToPurgeCountMap;
@@ -527,14 +609,19 @@ private:
 	class FGCTask
 	{
 		FArchiveRealtimeGC*	Owner;
-		TArray<UObject*>	ObjectsToSerialize;
+		TArray<UObject*>*	ObjectsToSerialize;
 
 	public:
 		FGCTask(FArchiveRealtimeGC* InOwner, const TArray<UObject*>* InObjectsToSerialize,int32 StartIndex, int32 NumObjects)
 			: Owner(InOwner)
+			, ObjectsToSerialize(FGCArrayPool::Get().GetArrayFromPool())
+		{
+			ObjectsToSerialize->AddUninitialized(NumObjects);
+			FMemory::Memcpy(ObjectsToSerialize->GetData(), InObjectsToSerialize->GetData() + StartIndex, NumObjects * sizeof(UObject*));
+		}
+		~FGCTask()
 		{
-			ObjectsToSerialize.AddUninitialized(NumObjects);
-			FMemory::Memcpy(ObjectsToSerialize.GetData(), InObjectsToSerialize->GetData() + StartIndex, NumObjects * sizeof(UObject*));
+			FGCArrayPool::Get().ReturnToPool(ObjectsToSerialize);
 		}
 		FORCEINLINE TStatId GetStatId() const
 		{
@@ -550,7 +637,7 @@ private:
 		}
 		void DoTask(ENamedThreads::Type CurrentThread, FGraphEventRef& MyCompletionGraphEvent)
 		{
-			Owner->ProcessObjectArray(ObjectsToSerialize, MyCompletionGraphEvent);
+			Owner->ProcessObjectArray(*ObjectsToSerialize, MyCompletionGraphEvent);			
 		}
 	};
 
@@ -571,7 +658,7 @@ public:
 		UObject* CurrentObject = NULL;
 
 		/** Growing array of objects that require serialization */
-		TArray<UObject*>	ObjectsToSerialize;
+		TArray<UObject*>& ObjectsToSerialize = *FGCArrayPool::Get().GetArrayFromPool();
 
 		// Reset object count.
 		GObjectCountDuringLastMarkPhase = 0;
@@ -662,6 +749,8 @@ public:
 				GIsRunningParallelReachability = false;
 			}
 		}
+
+		FGCArrayPool::Get().ReturnToPool(&ObjectsToSerialize);
 	}
 
 	void DispatchObjectTasks(TArray<UObject*>& ObjectsToSerialize)
@@ -674,7 +763,7 @@ public:
 
 		UObject* CurrentObject = NULL;
 
-		const int32 MinDesiredObjectsPerSubTask = 128; // sometimes there will be less, a lot less
+		const int32 MinDesiredObjectsPerSubTask = GMinDesiredObjectsPerSubTask; // sometimes there will be less, a lot less
 		const int32 NewObjectsArrayLength = InObjectsToSerializeArray.Num() * 2;
 
 		/** Growing array of objects that require serialization */
@@ -874,10 +963,10 @@ public:
 				}
 				check(StackEntry == Stack.GetData());
 
-				if( GIsRunningParallelReachability && NewObjectsToSerialize.Num() >= MinDesiredObjectsPerSubTask )
+				if (GIsRunningParallelReachability && NewObjectsToSerialize.Num() >= MinDesiredObjectsPerSubTask)
 				{
 					// This will start queueing task with objects from the end of array until there's less objects than worth to queue
-					const int32 ObjectsPerSubTask = FMath::Max<int32>(MinDesiredObjectsPerSubTask, NewObjectsToSerialize.Num() / FTaskGraphInterface::Get().GetNumWorkerThreads());					
+					const int32 ObjectsPerSubTask = FMath::Max<int32>(MinDesiredObjectsPerSubTask, NewObjectsToSerialize.Num() / FTaskGraphInterface::Get().GetNumWorkerThreads());
 					while (NewObjectsToSerialize.Num() >= MinDesiredObjectsPerSubTask)
 					{						
 						const int32 StartIndex = FMath::Max(0, NewObjectsToSerialize.Num() - ObjectsPerSubTask);
diff --git a/Engine/Source/Runtime/CoreUObject/Private/UObject/Obj.cpp b/Engine/Source/Runtime/CoreUObject/Private/UObject/Obj.cpp
--- a/Engine/Source/Runtime/CoreUObject/Private/UObject/Obj.cpp
+++ b/Engine/Source/Runtime/CoreUObject/Private/UObject/Obj.cpp
@@ -3650,6 +3650,8 @@ void StaticUObjectInit()
 	UE_LOG(LogInit, Log, TEXT("Object subsystem initialized") );
 }
 
+void CleanupGCArrayPools();
+
 //
 // Shut down the object manager.
 //
@@ -3741,6 +3743,7 @@ void StaticExit()
 	// Empty arrays to prevent falsely-reported memory leaks.
 	FUObjectThreadContext::Get().ObjLoaded.Empty();
 	FDeferredMessageLog::Cleanup();
+	CleanupGCArrayPools();
 
 	UE_LOG(LogExit, Log, TEXT("Object subsystem successfully closed.") );
 }
