Fixing potential issues with LinkerAnnotations and the order of static destruction of singletons.

diff --git a/Engine/Source/Runtime/CoreUObject/Private/UObject/Obj.cpp b/Engine/Source/Runtime/CoreUObject/Private/UObject/Obj.cpp
--- a/Engine/Source/Runtime/CoreUObject/Private/UObject/Obj.cpp
+++ b/Engine/Source/Runtime/CoreUObject/Private/UObject/Obj.cpp
@@ -3650,7 +3650,9 @@ void StaticUObjectInit()
 	UE_LOG(LogInit, Log, TEXT("Object subsystem initialized") );
 }
 
+// Internal cleanup functions
 void CleanupGCArrayPools();
+void CleanupLinkerAnnotations();
 
 //
 // Shut down the object manager.
@@ -3744,6 +3746,7 @@ void StaticExit()
 	FUObjectThreadContext::Get().ObjLoaded.Empty();
 	FDeferredMessageLog::Cleanup();
 	CleanupGCArrayPools();
+	CleanupLinkerAnnotations();
 
 	UE_LOG(LogExit, Log, TEXT("Object subsystem successfully closed.") );
 }
diff --git a/Engine/Source/Runtime/CoreUObject/Private/UObject/UObjectLinker.cpp b/Engine/Source/Runtime/CoreUObject/Private/UObject/UObjectLinker.cpp
--- a/Engine/Source/Runtime/CoreUObject/Private/UObject/UObjectLinker.cpp
+++ b/Engine/Source/Runtime/CoreUObject/Private/UObject/UObjectLinker.cpp
@@ -84,6 +84,12 @@ template <> struct TIsPODType<FLinkerIndexPair> { enum { Value = false }; };
  */
 static FUObjectAnnotationDense<FLinkerIndexPair,false> LinkerAnnotation;
 
+/** Remove all annotations on exit. This is to prevent issues with the order of static destruction of singletons. */
+void CleanupLinkerAnnotations()
+{
+	// Remove all annotations and remove LinkerAnnotation from UObject delete listeners.
+	LinkerAnnotation.RemoveAllAnnotations();
+}
 
 
 void UObject::SetLinker( FLinkerLoad* LinkerLoad, int32 LinkerIndex, bool bShouldDetachExisting )
