Added verbose message with referencing property info for string asset references when resolving them

diff --git a/Engine/Source/Runtime/CoreUObject/Private/Misc/RedirectCollector.cpp b/Engine/Source/Runtime/CoreUObject/Private/Misc/RedirectCollector.cpp
--- a/Engine/Source/Runtime/CoreUObject/Private/Misc/RedirectCollector.cpp
+++ b/Engine/Source/Runtime/CoreUObject/Private/Misc/RedirectCollector.cpp
@@ -46,16 +46,20 @@ void FRedirectCollector::OnRedirectorFollowed(const FString& InString, UObject*
 void FRedirectCollector::OnStringAssetReferenceLoaded(const FString& InString)
 {
 	FUObjectThreadContext& ThreadContext = FUObjectThreadContext::Get();
-	FString ContainingPackage;
+	FPackagePropertyPair ContainingPackageAndProperty;
 	if (ThreadContext.SerializedObject)
 	{
-		auto Linker = ThreadContext.SerializedObject->GetLinker();
+		FLinkerLoad* Linker = ThreadContext.SerializedObject->GetLinker();
 		if (Linker)
 		{
-			ContainingPackage = Linker->Filename;
+			ContainingPackageAndProperty.Package = Linker->Filename;
+			if (Linker->GetSerializedProperty())
+			{
+				ContainingPackageAndProperty.Property = FString::Printf(TEXT("%s:%s"), *ThreadContext.SerializedObject->GetPathName(), *Linker->GetSerializedProperty()->GetName());
+			}
 		}
 	}
-	StringAssetReferences.Add(InString, ContainingPackage);
+	StringAssetReferences.Add(InString, ContainingPackageAndProperty);
 }
 
 FString FRedirectCollector::OnStringAssetReferenceSaved(const FString& InString)
@@ -72,15 +76,17 @@ void FRedirectCollector::ResolveStringAssetReference()
 {
 	while (StringAssetReferences.Num())
 	{
-		TMultiMap<FString, FString>::TIterator First(StringAssetReferences);
+		TMultiMap<FString, FPackagePropertyPair>::TIterator First(StringAssetReferences);
 		FString ToLoad = First.Key();
-		FString RefFilename = First.Value();
+		FPackagePropertyPair RefFilenameAndProperty = First.Value();
 		First.RemoveCurrent();
 
 		if (ToLoad.Len() > 0)
 		{
 			UE_LOG(LogRedirectors, Log, TEXT("String Asset Reference '%s'"), *ToLoad);
-			StringAssetRefFilenameStack.Push(RefFilename);
+			UE_CLOG(RefFilenameAndProperty.Property.Len(), LogRedirectors, Verbose, TEXT("    Referenced by '%s'"), *RefFilenameAndProperty.Property);
+
+			StringAssetRefFilenameStack.Push(RefFilenameAndProperty.Package);
 
 			UObject *Loaded = LoadObject<UObject>(NULL, *ToLoad, NULL, LOAD_None, NULL);
 			StringAssetRefFilenameStack.Pop();
@@ -90,7 +96,7 @@ void FRedirectCollector::ResolveStringAssetReference()
 			{
 				UE_LOG(LogRedirectors, Log, TEXT("    Found redir '%s'"), *Redirector->GetFullName());
 				FRedirection Redir;
-				Redir.PackageFilename = RefFilename;
+				Redir.PackageFilename = RefFilenameAndProperty.Package;
 				Redir.RedirectorName = Redirector->GetFullName();
 				Redir.RedirectorPackageFilename = Redirector->GetLinker()->Filename;
 				CA_SUPPRESS(28182)
diff --git a/Engine/Source/Runtime/CoreUObject/Public/Misc/RedirectCollector.h b/Engine/Source/Runtime/CoreUObject/Public/Misc/RedirectCollector.h
--- a/Engine/Source/Runtime/CoreUObject/Public/Misc/RedirectCollector.h
+++ b/Engine/Source/Runtime/CoreUObject/Public/Misc/RedirectCollector.h
@@ -44,6 +44,19 @@ struct FRedirection
 class COREUOBJECT_API FRedirectCollector
 {
 public:
+	
+	/** Helper struct for string asset reference tracking */
+	struct FPackagePropertyPair
+	{
+		FPackagePropertyPair() {}
+		FPackagePropertyPair(const FString& InPackage, const FString& InProperty)
+		: Package(InPackage)
+		, Property(InProperty)
+		{}
+		FString Package;
+		FString Property;
+	};
+
 	/**
 	 * Responds to FCoreDelegates::RedirectorFollowed. Records all followed redirections
 	 * so they can be cleaned later.
@@ -78,7 +91,7 @@ public:
 	TArray<FRedirection> Redirections;
 
 	/** A gathered list string asset references , with the key being the string reference (GetPathName()) and the value equal to the package with the reference */
-	TMultiMap<FString, FString> StringAssetReferences;
+	TMultiMap<FString, FPackagePropertyPair> StringAssetReferences;
 
 	/** When saving, apply this remapping to all string asset references */
 	TMap<FString, FString> StringAssetRemap;
