Core: Changed most log messages about String Asset References to be verbose, but made a failed resolve a warning

[INTEGRATE] Change 2662014 by [REDACTED_PERSON][REDACTED_HANDLE] on 2015/08/19 20:20:18

diff --git a/Engine/Source/Runtime/CoreUObject/Private/Misc/RedirectCollector.cpp b/Engine/Source/Runtime/CoreUObject/Private/Misc/RedirectCollector.cpp
--- a/Engine/Source/Runtime/CoreUObject/Private/Misc/RedirectCollector.cpp
+++ b/Engine/Source/Runtime/CoreUObject/Private/Misc/RedirectCollector.cpp
@@ -83,7 +83,7 @@ void FRedirectCollector::ResolveStringAssetReference()
 
 		if (ToLoad.Len() > 0)
 		{
-			UE_LOG(LogRedirectors, Log, TEXT("String Asset Reference '%s'"), *ToLoad);
+			UE_LOG(LogRedirectors, Verbose, TEXT("String Asset Reference '%s'"), *ToLoad);
 			UE_CLOG(RefFilenameAndProperty.Property.Len(), LogRedirectors, Verbose, TEXT("    Referenced by '%s'"), *RefFilenameAndProperty.Property);
 
 			StringAssetRefFilenameStack.Push(RefFilenameAndProperty.Package);
@@ -94,7 +94,7 @@ void FRedirectCollector::ResolveStringAssetReference()
 			UObjectRedirector* Redirector = dynamic_cast<UObjectRedirector*>(Loaded);
 			if (Redirector)
 			{
-				UE_LOG(LogRedirectors, Log, TEXT("    Found redir '%s'"), *Redirector->GetFullName());
+				UE_LOG(LogRedirectors, Verbose, TEXT("    Found redir '%s'"), *Redirector->GetFullName());
 				FRedirection Redir;
 				Redir.PackageFilename = RefFilenameAndProperty.Package;
 				Redir.RedirectorName = Redirector->GetFullName();
@@ -107,7 +107,7 @@ void FRedirectCollector::ResolveStringAssetReference()
 			if (Loaded)
 			{
 				FString Dest = Loaded->GetPathName();
-				UE_LOG(LogRedirectors, Log, TEXT("    Resolved to '%s'"), *Dest);
+				UE_LOG(LogRedirectors, Verbose, TEXT("    Resolved to '%s'"), *Dest);
 				if (Dest != ToLoad)
 				{
 					StringAssetRemap.Add(ToLoad, Dest);
@@ -115,7 +115,7 @@ void FRedirectCollector::ResolveStringAssetReference()
 			}
 			else
 			{
-				UE_LOG(LogRedirectors, Log, TEXT("    Not Found!"));
+				UE_LOG(LogRedirectors, Warning, TEXT("String Asset Reference '%s' was not found!"), *ToLoad);
 			}
 		}
 	}
