When cooking saving package will now skip creating exports and imports that are editor-only (are instances of an editor class).

diff --git a/Engine/Config/BaseEngine.ini b/Engine/Config/BaseEngine.ini
--- a/Engine/Config/BaseEngine.ini
+++ b/Engine/Config/BaseEngine.ini
@@ -1431,6 +1431,7 @@ CutdownPaths=%GAMEDIR%CutdownPackages
 ZeroEngineVersionWarning=True
 UseStrictEngineVersioning=True
 UseVTableConstructors=True
+CanStripEditorOnlyExportsAndImports=True
 
 [/Script/Engine.StreamingSettings]
 s.AsyncIOBandwidthLimit=0
diff --git a/Engine/Source/Programs/UnrealHeaderTool/Private/CodeGenerator.cpp b/Engine/Source/Programs/UnrealHeaderTool/Private/CodeGenerator.cpp
--- a/Engine/Source/Programs/UnrealHeaderTool/Private/CodeGenerator.cpp
+++ b/Engine/Source/Programs/UnrealHeaderTool/Private/CodeGenerator.cpp
@@ -1269,7 +1269,7 @@ void FNativeClassHeaderGenerator::ExportGeneratedPackageInitCode(UPackage* InPac
 		GeneratedFunctionText.Logf(TEXT("#endif\r\n"));
 	}
 
-	GeneratedFunctionText.Logf(TEXT("            ReturnPackage->PackageFlags |= PKG_CompiledIn | 0x%08X;\r\n"), InPackage->PackageFlags & (PKG_ClientOptional|PKG_ServerSideOnly));
+	GeneratedFunctionText.Logf(TEXT("            ReturnPackage->PackageFlags |= PKG_CompiledIn | 0x%08X;\r\n"), InPackage->PackageFlags & (PKG_ClientOptional | PKG_ServerSideOnly | PKG_EditorOnly));
 	TheFlagAudit.Add(InPackage, TEXT("PackageFlags"), InPackage->PackageFlags);
 	{
 		FGuid Guid;
@@ -5416,6 +5416,10 @@ ECompilationResult::Type UnrealHeaderTool_Main(const FString& ModuleInfoFilename
 		Package->PackageFlags |= PKG_ContainsScript;
 		Package->PackageFlags &= ~(PKG_ClientOptional | PKG_ServerSideOnly);
 		Package->PackageFlags |= PKG_Compiling;
+		if (Module.ModuleType == EBuildModuleType::Editor)
+		{
+			Package->PackageFlags |= PKG_EditorOnly;
+		}
 
 		GPackageToManifestModuleMap.Add(Package, &Module);
 
diff --git a/Engine/Source/Runtime/CoreUObject/Private/UObject/SavePackage.cpp b/Engine/Source/Runtime/CoreUObject/Private/UObject/SavePackage.cpp
--- a/Engine/Source/Runtime/CoreUObject/Private/UObject/SavePackage.cpp
+++ b/Engine/Source/Runtime/CoreUObject/Private/UObject/SavePackage.cpp
@@ -345,6 +345,92 @@ public:
 	}
 };
 
+/** Returns if true if a class comes from an editor-only package */
+static bool IsEditorOnlyStruct(UStruct* InStruct)
+{
+	// If any of the classes in this class' hierarchy comes from editor only package
+	// this class is also classified as editor-only.
+	for (UStruct* Struct = InStruct; Struct; Struct = Struct->GetSuperStruct())
+	{
+		UPackage* StructPackage = Struct->GetOutermost();
+		check(StructPackage);
+		if (StructPackage->PackageFlags & PKG_EditorOnly)
+		{
+			return true;
+		}
+	}
+	return false;
+}
+
+/** 
+ * Returns if true if the object is editor-only:
+ * - it's a package marked as PKG_EditorOnly
+ * or
+ * - it's a class from a package marked as PKG_EditorOnly
+ * or
+ * - its class is from a package marked as PKG_EditorOnly
+ * or
+ * - its outer is editor-only
+*/
+static bool IsEditorOnlyObject(UObject* InObject)
+{
+	DECLARE_SCOPE_CYCLE_COUNTER(TEXT("IsEditorOnlyObject"), STAT_IsEditorOnlyObject, STATGROUP_LoadTime);
+
+	// Configurable via ini setting
+	static struct FCanStripEditorOnlyExportsAndImports
+	{
+		bool bCanStripEditorOnlyObjects;
+		FCanStripEditorOnlyExportsAndImports()
+			: bCanStripEditorOnlyObjects(true)
+		{
+			GConfig->GetBool(TEXT("Core.System"), TEXT("CanStripEditorOnlyExportsAndImports"), bCanStripEditorOnlyObjects, GEngineIni);
+		}
+		FORCEINLINE operator bool() const { return bCanStripEditorOnlyObjects; }
+	} CanStripEditorOnlyExportsAndImports;
+	if (!CanStripEditorOnlyExportsAndImports)
+	{
+		return false;
+	}
+
+	if (InObject->HasAnyMarks(OBJECTMARK_EditorOnly))
+	{
+		return true;
+	}
+
+	bool bResult = false;
+	check(InObject);
+	// If this is a package that is editor only or the object is in editor-only package,
+	// the object is editor-only too.
+	UPackage* Package = Cast<UPackage>(InObject);
+	if (!Package)
+	{
+		Package = InObject->GetOutermost();
+	}
+	if (Package && (Package->PackageFlags & PKG_EditorOnly))
+	{
+		bResult = true;
+	}
+	if (!bResult && !InObject->IsA(UPackage::StaticClass()))
+	{
+		// Otherwise the object is editor-only if its class is editor-only
+		UStruct* Struct = InObject->IsA(UStruct::StaticClass()) ? CastChecked<UStruct>(InObject) : InObject->GetClass();
+		bResult = IsEditorOnlyStruct(Struct);
+		if (!bResult && InObject->GetOuter())
+		{
+			// Now check if the outer is editor only
+			bResult = IsEditorOnlyObject(InObject->GetOuter());
+		}
+	}
+	return bResult;
+}
+
+/**
+* Returns true if cook target is set and it doesn't support editor-only data
+**/
+static bool IsCookingForNoEditorDataPlatform(FArchive& Ar)
+{
+	return Ar.IsCooking() && !Ar.CookingTarget()->HasEditorOnlyData();
+}
 
 /**
  * Archive for tagging objects and names that must be exported
@@ -409,7 +495,7 @@ FArchive& FArchiveSaveTagExports::operator<<( UObject*& Obj )
 {
 	check(Outer);
 	CheckObjectPriorToSave(*this, Obj, Outer);
-	if( Obj && Obj->IsIn(Outer) && !Obj->HasAnyFlags(RF_Transient) && !Obj->HasAnyMarks(OBJECTMARK_TagExp) )
+	if (Obj && Obj->IsIn(Outer) && !Obj->HasAnyFlags(RF_Transient) && !Obj->HasAnyMarks((EObjectMark)(OBJECTMARK_TagExp | OBJECTMARK_EditorOnly)))
 	{
 #if 0
 		// the following line should be used to track down the cause behind
@@ -423,8 +509,18 @@ FArchive& FArchiveSaveTagExports::operator<<( UObject*& Obj )
 			UE_LOG(LogSavePackage, Log, TEXT(""));
 		}
 #endif
-		// Set flags.
-		Obj->Mark(OBJECTMARK_TagExp);
+
+		const bool bIsEditorOnly = IsCookingForNoEditorDataPlatform(*this) && IsEditorOnlyObject(Obj);
+		if (bIsEditorOnly)
+		{
+			Obj->Mark(OBJECTMARK_EditorOnly);
+			UE_LOG(LogSavePackage, Verbose, TEXT("Skipping editor-only export %s"), *Obj->GetPathName());
+		}
+		else
+		{
+			// Set flags.
+			Obj->Mark(OBJECTMARK_TagExp);
+		}
 
 		// first, serialize this object's archetype so that if the archetype's load flags are set correctly if this object
 		// is encountered by the SaveTagExports archive before its archetype.  This is necessary for the code below which verifies
@@ -659,7 +755,6 @@ public:
 	 * This is overridden for the specific Archive Types
 	 **/
 	virtual FString GetArchiveName() const override;
-
 };
 
 
@@ -686,15 +781,22 @@ FArchive& FArchiveSaveTagImports::operator<<( UObject*& Obj )
 	const EObjectMark ObjectMarks = UPackage::GetObjectMarksForTargetPlatform( CookingTarget(), IsCooking() );
 	
 	// Skip PendingKill objects and objects that are both not for client and not for server when cooking.
-	if( Obj && !Obj->IsPendingKill() && (!IsCooking() || !Obj->HasAllMarks(ObjectMarks)))
+	if (Obj && !Obj->IsPendingKill() && (!IsCooking() || !Obj->HasAllMarks(ObjectMarks)) && !Obj->HasAnyMarks(OBJECTMARK_EditorOnly))
 	{
 		if( !Obj->HasAnyFlags(RF_Transient) || Obj->HasAllFlags(RF_Native) )
 		{
 			// remember it as a dependency, unless it's a top level package or native
-			bool bIsTopLevelPackage = Obj->GetOuter() == NULL && dynamic_cast<UPackage*>(Obj);
+			const bool bIsTopLevelPackage = Obj->GetOuter() == NULL && dynamic_cast<UPackage*>(Obj);
 			bool bIsNative = Obj->HasAnyFlags(RF_Native);
 			UObject* Outer = Obj->GetOuter();
-			
+
+			const bool bIsEditorOnly = IsCookingForNoEditorDataPlatform(*this) && IsEditorOnlyObject(Obj);
+			if (bIsEditorOnly)
+			{
+				Obj->Mark(OBJECTMARK_EditorOnly);
+				UE_LOG(LogSavePackage, Verbose, TEXT("Skipping editor-only import %s"), *Obj->GetPathName());
+			}
+
 			// go up looking for native classes
 			while (!bIsNative && Outer)
 			{
@@ -714,7 +816,10 @@ FArchive& FArchiveSaveTagImports::operator<<( UObject*& Obj )
 			if( !Obj->HasAnyMarks(OBJECTMARK_TagExp) )  
 			{
 				// mark this object as an import
-				Obj->Mark(OBJECTMARK_TagImp);
+				if (!bIsEditorOnly)
+				{
+					Obj->Mark(OBJECTMARK_TagImp);
+				}
 
 				if ( Obj->HasAnyFlags(RF_ClassDefaultObject) )
 				{
@@ -2997,11 +3102,19 @@ ESavePackageResult UPackage::Save(UPackage* InOuter, UObject* Base, EObjectFlags
 
 		const bool bIsCooking = TargetPlatform != nullptr;
 #if WITH_EDITORONLY_DATA
-		if (bIsCooking && InOuter->IsLoadedByEditorPropertiesOnly())
+		if (bIsCooking)
 		{
 			// Don't save packages marked as editor-only.
-			UE_CLOG(!(SaveFlags & SAVE_NoError), LogSavePackage, Display, TEXT("Package loaded by editor-only properties %s. Package will not be saved."), *InOuter->GetName());
-			return ESavePackageResult::ReferencedOnlyByEditorOnlyData;
+			if (InOuter->IsLoadedByEditorPropertiesOnly())
+			{				
+				UE_CLOG(!(SaveFlags & SAVE_NoError), LogSavePackage, Display, TEXT("Package loaded by editor-only properties: %s. Package will not be saved."), *InOuter->GetName());
+				return ESavePackageResult::ReferencedOnlyByEditorOnlyData;
+			}
+			else if (InOuter->PackageFlags & PKG_EditorOnly)
+			{
+				UE_CLOG(!(SaveFlags & SAVE_NoError), LogSavePackage, Display, TEXT("Package marked as editor-only: %s. Package will not be saved."), *InOuter->GetName());
+				return ESavePackageResult::ReferencedOnlyByEditorOnlyData;
+			}
 		}
 #endif
 		// if we are cooking we should be doing it in the editor
@@ -3194,7 +3307,7 @@ ESavePackageResult UPackage::Save(UPackage* InOuter, UObject* Base, EObjectFlags
 
 			
 				// Clear OBJECTMARK_TagExp again as we need to redo tagging below.
-				UnMarkAllObjects(OBJECTMARK_TagExp);
+				UnMarkAllObjects((EObjectMark)(OBJECTMARK_TagExp|OBJECTMARK_EditorOnly));
 			
 
 				// We need to serialize objects yet again to tag objects that were created by PreSave as OBJECTMARK_TagExp.
diff --git a/Engine/Source/Runtime/CoreUObject/Public/UObject/ObjectBase.h b/Engine/Source/Runtime/CoreUObject/Public/UObject/ObjectBase.h
--- a/Engine/Source/Runtime/CoreUObject/Public/UObject/ObjectBase.h
+++ b/Engine/Source/Runtime/CoreUObject/Public/UObject/ObjectBase.h
@@ -91,7 +91,7 @@ enum EPackageFlags
 	PKG_ServerSideOnly				= 0x00000004,   // Only needed on the server side.
 	PKG_CompiledIn					= 0x00000010,   // This package is from "compiled in" classes.
 	PKG_ForDiffing					= 0x00000020,	// This package was loaded just for the purposes of diff'ing
-//	PKG_Unused						= 0x00000040
+	PKG_EditorOnly					= 0x00000040, // This is editor-only package (for example: editor module script package)
 //  PKG_Unused						= 0x00000080,
 //	PKG_Unused						= 0x00000100,
 //	PKG_Unused						= 0x00000200,
diff --git a/Engine/Source/Runtime/CoreUObject/Public/UObject/UObjectMarks.h b/Engine/Source/Runtime/CoreUObject/Public/UObject/UObjectMarks.h
--- a/Engine/Source/Runtime/CoreUObject/Public/UObject/UObjectMarks.h
+++ b/Engine/Source/Runtime/CoreUObject/Public/UObject/UObjectMarks.h
@@ -48,7 +48,7 @@ enum EObjectMark
 	OBJECTMARK_NotForClient			= 0x00000020,		// Temporary save tag for client load flag.
 	OBJECTMARK_NotForServer			= 0x00000040,		// Temporary save tag for server load flag.
 	OBJECTMARK_NotForEditorGame		= 0x00000080,		// Temporary save tag for editorgame load flag.
-
+	OBJECTMARK_EditorOnly = 0x00000100, // Temporary editor only flag
 	OBJECTMARK_ALLMARKS				= 0xffffffff,		// -1, all possible marks
 };
 
