Merging CL#2664160 by [REDACTED_PERSON][REDACTED_HANDLE]

Adding IsValidLowLevel checks to root set objects when initializing GC.

diff --git a/Engine/Source/Runtime/CoreUObject/Private/UObject/GarbageCollection.cpp b/Engine/Source/Runtime/CoreUObject/Private/UObject/GarbageCollection.cpp
--- a/Engine/Source/Runtime/CoreUObject/Private/UObject/GarbageCollection.cpp
+++ b/Engine/Source/Runtime/CoreUObject/Private/UObject/GarbageCollection.cpp
@@ -686,7 +686,7 @@ public:
 			// Special case handling for objects that are part of the root set.
 			if( Object->HasAnyFlags( RF_RootSet ) )
 			{
-				checkSlow( Object->IsValidLowLevel() );
+				check(Object->IsValidLowLevel());
 				// We cannot use RF_PendingKill on objects that are part of the root set.
 				checkCode( if( Object->HasAnyFlags( RF_PendingKill ) ) { UE_LOG(LogGarbage, Fatal, TEXT("Object %s is part of root set though has been marked RF_PendingKill!"), *Object->GetFullName() ); } );
 				ObjectsToSerialize.Add( Object );
@@ -696,7 +696,8 @@ public:
 			{
 				// Mark objects as unreachable unless they have any of the passed in KeepFlags set and it's not marked for elimination..
 				if( Object->HasAnyFlags( KeepFlags ) && !Object->HasAnyFlags( RF_PendingKill ) )
-				{	
+				{
+					check(Object->IsValidLowLevel());
 					ObjectsToSerialize.Add( Object );
 				}
 				else
