Early version of redirecting system for "Blueprints converted to C++".
It's supposed to be used while cooking to fix references to objects (BPGC, UDS, UDE), that where converted into a native code.
Converted items have "ReplaceConverted" meta data generated. The MD stores original object path (and all it's redirected paths).

diff --git a/Engine/Config/BaseEditor.ini b/Engine/Config/BaseEditor.ini
--- a/Engine/Config/BaseEditor.ini
+++ b/Engine/Config/BaseEditor.ini
@@ -345,6 +345,9 @@ bBlueprintIsNotBlueprintType=false
 SkipByteCodeSerialization=true
 BreakSerializationRecursion=true
 
+[BlueprintCppConversion]
+ReplaceConvertedAssetWhileCooking=false
+
 [ScriptErrorLog]
 ; per object
 MaxNumOfAccessViolation=16
diff --git a/Engine/Source/Developer/BlueprintCompilerCppBackend/Private/BlueprintCompilerCppBackendUtils.cpp b/Engine/Source/Developer/BlueprintCompilerCppBackend/Private/BlueprintCompilerCppBackendUtils.cpp
--- a/Engine/Source/Developer/BlueprintCompilerCppBackend/Private/BlueprintCompilerCppBackendUtils.cpp
+++ b/Engine/Source/Developer/BlueprintCompilerCppBackend/Private/BlueprintCompilerCppBackendUtils.cpp
@@ -213,20 +213,38 @@ FString FEmitHelper::GenerateReplaceConvertedMD(UObject* Obj)
 	if (Obj)
 	{
 		Result = TEXT("ReplaceConverted=\"");
+
+		// 1. Current object
 		Result += Obj->GetPathName();
-		UPackage* Package = Obj->GetOutermost();
-		check(Package);
-		TArray<UObject*> AllObjects;
-		GetObjectsWithOuter(Package, AllObjects, true);
-		for (auto LocalObj : AllObjects)
+
+		// 2. Loaded Redirectors
 		{
-			auto Redirector = Cast<UObjectRedirector>(LocalObj);
-			if (Redirector && (Obj == Redirector->DestinationObject))
+			struct FLocalHelper
+			{
+				static UObject* FindFinalObject(UObjectRedirector* Redirector)
+				{
+					auto DestObj = Redirector ? Redirector->DestinationObject : nullptr;
+					auto InnerRedir = Cast<UObjectRedirector>(DestObj);
+					return InnerRedir ? FindFinalObject(InnerRedir) : DestObj;
+				}
+			};
+
+			TArray<UObject*> AllObjects;
+			GetObjectsOfClass(UObjectRedirector::StaticClass(), AllObjects);
+			for (auto LocalObj : AllObjects)
 			{
-				Result += TEXT(",");
-				Result += Redirector->GetPathName();
+				auto Redirector = CastChecked<UObjectRedirector>(LocalObj);
+				if (Obj == FLocalHelper::FindFinalObject(Redirector))
+				{
+					Result += TEXT(",");
+					Result += Redirector->GetPathName();
+				}
 			}
 		}
+
+		// 3. Unloaded Redirectors
+		// TODO: It would be better to load all redirectors before compiling. Than checking here AssetRegistry.
+
 		Result += TEXT("\"");
 	}
 
diff --git a/Engine/Source/Editor/UnrealEd/Private/Commandlets/CookCommandlet.cpp b/Engine/Source/Editor/UnrealEd/Private/Commandlets/CookCommandlet.cpp
--- a/Engine/Source/Editor/UnrealEd/Private/Commandlets/CookCommandlet.cpp
+++ b/Engine/Source/Editor/UnrealEd/Private/Commandlets/CookCommandlet.cpp
@@ -24,6 +24,7 @@
 #include "CookerSettings.h"
 #include "ShaderCompiler.h"
 #include "MemoryMisc.h"
+#include "BlueprintSupport.h"
 
 DEFINE_LOG_CATEGORY_STATIC(LogCookCommandlet, Log, All);
 
@@ -516,6 +517,8 @@ int32 UCookCommandlet::Main(const FString& CmdLineParams)
 		}
 	}
 
+	static const FBoolConfigValueHelper ReplaceConvertedAsset(TEXT("BlueprintCppConversion"), TEXT("ReplaceConvertedAssetWhileCooking"));
+	FReplaceConvertedAssetManager::Get().SetEnabled(ReplaceConvertedAsset);
 
 	if ( bCookOnTheFly )
 	{
diff --git a/Engine/Source/Runtime/CoreUObject/Private/Blueprint/BlueprintSupport.cpp b/Engine/Source/Runtime/CoreUObject/Private/Blueprint/BlueprintSupport.cpp
--- a/Engine/Source/Runtime/CoreUObject/Private/Blueprint/BlueprintSupport.cpp
+++ b/Engine/Source/Runtime/CoreUObject/Private/Blueprint/BlueprintSupport.cpp
@@ -8,6 +8,7 @@
 #include "PropertyTag.h"
 #include "UObject/StructScriptLoader.h"
 #include "UObject/UObjectThreadContext.h"
+#include "ModuleManager.h"
 
 /** 
  * Defined in BlueprintSupport.cpp
@@ -1639,3 +1640,86 @@ void FDeferredObjInitializerTracker::ResolveDeferredSubClassObjects(UClass* Supe
 
 // don't want other files ending up with this internal define
 #undef DEFERRED_DEPENDENCY_CHECK
+
+#if WITH_EDITOR
+
+FReplaceConvertedAssetManager::FReplaceConvertedAssetManager() 
+	: bIsEnabled(false)
+{
+	FModuleManager::Get().OnModulesChanged().AddStatic(&FReplaceConvertedAssetManager::OnModulesChanged);
+
+	// FOR DEVELOPMENT/TEST ONLY:
+	// SetEnabled(true);
+}
+
+void FReplaceConvertedAssetManager::OnModulesChanged(FName ModuleThatChanged, EModuleChangeReason ReasonForChange)
+{
+	Get().GatherOriginalPathsOfConvertedAssets();
+}
+
+FReplaceConvertedAssetManager& FReplaceConvertedAssetManager::Get()
+{
+	static FReplaceConvertedAssetManager ReplaceConvertedAssetManager;
+	return ReplaceConvertedAssetManager;
+}
+
+UObject* FReplaceConvertedAssetManager::FindReplacement(const FString& OriginalPathName) const
+{
+	UObject* const* ObjPtr = ReplaceMap.Find(OriginalPathName);
+	return ObjPtr ? *ObjPtr : nullptr;
+}
+
+UPackage* FReplaceConvertedAssetManager::FindPackageReplacement(const FString& OriginalPathName) const
+{
+	for (auto It = ReplaceMap.CreateConstIterator(); It; ++It)
+	{
+		if (It.Key().StartsWith(OriginalPathName))
+		{
+			auto Val = It.Value();
+			return ensure(Val) ? Val->GetOutermost() : nullptr;
+		}
+	}
+	return nullptr;
+}
+
+void FReplaceConvertedAssetManager::GatherOriginalPathsOfConvertedAssets()
+{
+	if (!IsEnabled())
+	{
+		return;
+	}
+
+	ReplaceMap.Reset();
+	const FName ReplaceConverted(TEXT("ReplaceConverted"));
+
+	auto FillMap = [&](UField* Field)
+	{
+		if (Field && Field->HasMetaData(ReplaceConverted))
+		{
+			auto CombinedPaths = Field->GetMetaData(ReplaceConverted);
+			TArray<FString> Paths;
+			CombinedPaths.ParseIntoArray(Paths, TEXT(","));
+			for (auto& Path : Paths)
+			{
+				ReplaceMap.Add(Path, Field);
+			}
+		}
+	};
+
+	for (UClass* LocalClass : TObjectRange<UClass>())
+	{
+		FillMap(LocalClass);
+	}
+
+	for (UScriptStruct* LocalScriptStruct : TObjectRange<UScriptStruct>())
+	{
+		FillMap(LocalScriptStruct);
+	}
+
+	for (UEnum* LocalEnum : TObjectRange<UEnum>())
+	{
+		FillMap(LocalEnum);
+	}
+}
+
+#endif //WITH_EDITOR
\ No newline at end of file
diff --git a/Engine/Source/Runtime/CoreUObject/Private/UObject/LinkerLoad.cpp b/Engine/Source/Runtime/CoreUObject/Private/UObject/LinkerLoad.cpp
--- a/Engine/Source/Runtime/CoreUObject/Private/UObject/LinkerLoad.cpp
+++ b/Engine/Source/Runtime/CoreUObject/Private/UObject/LinkerLoad.cpp
@@ -13,6 +13,7 @@
 #include "UObject/UObjectThreadContext.h"
 #include "GatherableTextData.h"
 #include "Serialization/AsyncLoading.h"
+#include "ModuleManager.h"
 
 #define LOCTEXT_NAMESPACE "LinkerLoad"
 
@@ -1308,6 +1309,36 @@ FLinkerLoad::ELinkerStatus FLinkerLoad::FixupImportMap()
 				for( int32 i=0; i<ImportMap.Num(); i++ )
 				{
 					FObjectImport& Import = ImportMap[i];
+
+#if WITH_EDITOR
+					// Fix references to assets that were converted into native code
+					
+					if (FReplaceConvertedAssetManager::Get().IsEnabled() && !Import.XObject)
+					{
+						const FString ImportPathName = GetImportPathName(i);
+						const FName PackageClassName(TEXT("Package"));
+						const bool bIsPackage = Import.ClassName == PackageClassName; // there must be a better way
+						if (bIsPackage)
+						{
+							if (auto NewNativePackage = FReplaceConvertedAssetManager::Get().FindPackageReplacement(ImportPathName))
+							{
+								Import.XObject = NewNativePackage;
+								Import.ObjectName = NewNativePackage->GetFName();
+								continue;
+							}
+						}
+						else if (auto ConvertedObject = FReplaceConvertedAssetManager::Get().FindReplacement(ImportPathName))
+						{
+							Import.ObjectName = ConvertedObject->GetFName();
+							Import.XObject = ConvertedObject;
+							Import.ClassName = ConvertedObject->GetClass()->GetFName();
+							auto ClassPackage = ConvertedObject->GetClass()->GetOuterUPackage();
+							Import.ClassPackage = ClassPackage ? ClassPackage->GetFName() : FName();
+							continue;
+						}
+					}
+#endif //WITH_EDITOR
+
 					{
 						FSubobjectRedirect* Redirect = SubobjectNameRedirects.Find(Import.ObjectName);
 						if (Redirect)
@@ -2131,7 +2162,6 @@ void FLinkerLoad::GatherImportDependencies(int32 ImportIndex, TSet<FDependencyRe
 	}
 }
 
-
 FLinkerLoad::EVerifyResult FLinkerLoad::VerifyImport(int32 ImportIndex)
 {
 	FObjectImport& Import = ImportMap[ImportIndex];
diff --git a/Engine/Source/Runtime/CoreUObject/Public/Blueprint/BlueprintSupport.h b/Engine/Source/Runtime/CoreUObject/Public/Blueprint/BlueprintSupport.h
--- a/Engine/Source/Runtime/CoreUObject/Public/Blueprint/BlueprintSupport.h
+++ b/Engine/Source/Runtime/CoreUObject/Public/Blueprint/BlueprintSupport.h
@@ -63,6 +63,47 @@ private:
 
 	FScopedClassDependencyGather();
 };
+
+/**
+* Stores info about BPCG, UDS, UDE, etc,, converted into native code.
+* Native entities (generated from BP items) have "ReplaceConverted" with original object path.
+* This path is used to fix linker.
+*/
+struct COREUOBJECT_API FReplaceConvertedAssetManager
+{
+private:
+	TMap<FString, UObject*> ReplaceMap;
+	bool bIsEnabled;
+
+	FReplaceConvertedAssetManager();
+
+	void GatherOriginalPathsOfConvertedAssets();
+
+public:
+
+	static FReplaceConvertedAssetManager& Get();
+
+	void SetEnabled(bool InEnabled)
+	{
+		bIsEnabled = InEnabled;
+		if (bIsEnabled)
+		{
+			GatherOriginalPathsOfConvertedAssets();
+		}
+	}
+
+	bool IsEnabled()
+	{
+		return bIsEnabled;
+	}
+
+	UPackage* FindPackageReplacement(const FString& OriginalPathName) const;
+
+	UObject* FindReplacement(const FString& OriginalPathName) const;
+
+	static void OnModulesChanged(FName ModuleThatChanged, EModuleChangeReason ReasonForChange);
+};
+
 #endif //WITH_EDITOR
 
 /** 
