Core: avoid leaking ICppStructOps.

- Leaking those structures (get allocated in autogenerated code) polluted valgrind's memcheck with thousands of fake leaks.

diff --git a/Engine/Source/Runtime/CoreUObject/Private/UObject/Class.cpp b/Engine/Source/Runtime/CoreUObject/Private/UObject/Class.cpp
--- a/Engine/Source/Runtime/CoreUObject/Private/UObject/Class.cpp
+++ b/Engine/Source/Runtime/CoreUObject/Private/UObject/Class.cpp
@@ -1835,7 +1835,17 @@ struct TStructOpsTypeTraits<FTestStruct> : public TStructOpsTypeTraitsBase
 **/
 static TMap<FName,UScriptStruct::ICppStructOps*>& GetDeferredCppStructOps()
 {
-	static TMap<FName,UScriptStruct::ICppStructOps*> DeferredCppStructOps;
+	static struct TMapWithAutoCleanup : public TMap<FName, UScriptStruct::ICppStructOps*>
+	{
+		~TMapWithAutoCleanup()
+		{
+			for (PairSetType::TConstIterator It(Pairs); It; ++It)
+			{
+				delete It->Value;
+			}
+		}
+	}
+	DeferredCppStructOps;
 	return DeferredCppStructOps;
 }
 
