Fix access violation crash in GarbageCollection.cpp - ProcessObjectArray.

diff --git a/Engine/Source/Runtime/CoreUObject/Private/UObject/GarbageCollection.cpp b/Engine/Source/Runtime/CoreUObject/Private/UObject/GarbageCollection.cpp
--- a/Engine/Source/Runtime/CoreUObject/Private/UObject/GarbageCollection.cpp
+++ b/Engine/Source/Runtime/CoreUObject/Private/UObject/GarbageCollection.cpp
@@ -795,7 +795,11 @@ public:
 #endif
 				CurrentObject = ObjectsToSerialize[CurrentIndex++];
 
-				const UObject * const NextObject = ObjectsToSerialize.GetData()[CurrentIndex]; // special syntax avoiding out of bounds checking
+				// GetData() used to avoiding bounds checking (min and max)
+				// FMath::Min used to avoid out of bounds (without branching) on last iteration. Though anything can be passed into PrefetchBlock, 
+				// reading ObjectsToSerialize out of bounds is not safe since ObjectsToSerialize[Num()] may be an unallocated/unsafe address.
+				const UObject * const NextObject = ObjectsToSerialize.GetData()[FMath::Min<int32>(CurrentIndex, ObjectsToSerialize.Num() - 1)];
+
 				// Prefetch the next object assuming that the property size of the next object is the same as the current one.
 				// This allows us to avoid a branch here.
 				FPlatformMisc::PrefetchBlock(NextObject, CurrentObject->GetClass()->GetPropertiesSize());
