Editor: Some more minor formatting changes to save package log output timings, and fixed an indentation issue
[INTEGRATE] Change 2667231 by [REDACTED_PERSON][REDACTED_HANDLE] on 2015/08/25 03:45:48

diff --git a/Engine/Source/Runtime/CoreUObject/Private/UObject/SavePackage.cpp b/Engine/Source/Runtime/CoreUObject/Private/UObject/SavePackage.cpp
--- a/Engine/Source/Runtime/CoreUObject/Private/UObject/SavePackage.cpp
+++ b/Engine/Source/Runtime/CoreUObject/Private/UObject/SavePackage.cpp
@@ -47,33 +47,31 @@ static uint32 GSavePackageTransactionNumber = 0;
 
 	// static ICookingStats* CookingStats = GetCookingStats();
 
-void CookStatsStartStat(const FName& Key, const TCHAR *Filename)
-{
+	void CookStatsStartStat(const FName& Key, const TCHAR *Filename)
+	{
 #if UE_OUTPUTSTATSTOLOG
-	UE_LOG(LogSavePackage, Log, TEXT("Starting save package for %s"), Filename);
+		UE_LOG(LogSavePackage, Log, TEXT("Starting save package for %s"), Filename);
 #endif
+
 		ICookingStats* CookingStats = GetCookingStats();
 		if (CookingStats)
 		{
 			CookingStats->AddTagValue(Key, TEXT("Filename"), FString(Filename));
 		}
-
-}
+	}
 
 	void CookStatsAddStat(const FName& Key, const FName& Tag, const float TimeMilliseconds)
-{
+	{
 #if UE_OUTPUTSTATSTOLOG
-	UE_LOG(LogSavePackage, Log, TEXT("TIMING: %s took %fms"), Tag, TimeMilliseconds);
+		UE_LOG(LogSavePackage, Log, TEXT("TIMING: %s took %.2fms"), Tag, TimeMilliseconds);
 #endif
+
 		ICookingStats* CookingStats = GetCookingStats();
 		if (CookingStats)
 		{
 			CookingStats->AddTagValue(Key, Tag, FString::Printf(TEXT("%fms"), TimeMilliseconds));
 		}
 	}
-
-
-
 }; // namespace SavePackageStats
 
 
@@ -92,7 +90,7 @@ void CookStatsStartStat(const FName& Key, const TCHAR *Filename)
 	UE_LOG_COOK_TIME(TEXT("UnaccountedTime")); \
 	{\
 		double CurrentTime = FPlatformTime::Seconds(); \
-		UE_LOG(LogSavePackage, Log, TEXT("Total save time: %fms"), (CurrentTime - StartTime) * 1000.0f); \
+		UE_LOG(LogSavePackage, Log, TEXT("Total save time: %.2fms"), (CurrentTime - StartTime) * 1000.0f); \
 	}
 
 #else
