Merging 2667825 using Main-Branch-Map
#UE4 #CoreUObject #Redirectors
- Allow the active class redirectors map to be loaded from all config files so plugins can define redirectors

diff --git a/Engine/Source/Runtime/CoreUObject/Private/UObject/LinkerLoad.cpp b/Engine/Source/Runtime/CoreUObject/Private/UObject/LinkerLoad.cpp
--- a/Engine/Source/Runtime/CoreUObject/Private/UObject/LinkerLoad.cpp
+++ b/Engine/Source/Runtime/CoreUObject/Private/UObject/LinkerLoad.cpp
@@ -72,15 +72,9 @@ void FLinkerLoad::AddGameNameRedirect(const FName OldName, const FName NewName)
  * Add redirects to FLinkerLoad static map
  */
 void FLinkerLoad::CreateActiveRedirectsMap(const FString& GEngineIniName)
-{		
-	if (bActiveRedirectsMapInitialized)
-	{
-		return;
-	}
-	else
-	{
-		bActiveRedirectsMapInitialized = true;
-	}
+{
+	// mark that this has been done at least once
+	bActiveRedirectsMapInitialized = true;
 
 	if (GConfig)
 	{
@@ -195,10 +189,6 @@ void FLinkerLoad::CreateActiveRedirectsMap(const FString& GEngineIniName)
 				}
 			}
 		}
-		else
-		{
-			UE_LOG(LogLinker, Log, TEXT("Active class redirects did not initialize because /Script/Engine.Engine was not available in the engine ini."));
-		}
 	}
 	else
 	{
diff --git a/Engine/Source/Runtime/CoreUObject/Private/UObject/Obj.cpp b/Engine/Source/Runtime/CoreUObject/Private/UObject/Obj.cpp
--- a/Engine/Source/Runtime/CoreUObject/Private/UObject/Obj.cpp
+++ b/Engine/Source/Runtime/CoreUObject/Private/UObject/Obj.cpp
@@ -3590,7 +3590,10 @@ void PreInitUObject()
 void InitUObject()
 {
 	// Initialize redirects map
-	FLinkerLoad::CreateActiveRedirectsMap(GEngineIni);
+	for (const auto& It : *GConfig)
+	{
+		FLinkerLoad::CreateActiveRedirectsMap(It.Key);
+	}
 
 	FCoreDelegates::OnShutdownAfterError.AddStatic(StaticShutdownAfterError);
 	FCoreDelegates::OnExit.AddStatic(StaticExit);
