Networking: Reduce incremental allocations in FClassNetCacheMgr::GetClassNetCache
[INTEGRATE] Change 2668348 by [REDACTED_PERSON][REDACTED_HANDLE] on 2015/08/25 16:46:24

diff --git a/Engine/Source/Runtime/CoreUObject/Private/UObject/CoreNet.cpp b/Engine/Source/Runtime/CoreUObject/Private/UObject/CoreNet.cpp
--- a/Engine/Source/Runtime/CoreUObject/Private/UObject/CoreNet.cpp
+++ b/Engine/Source/Runtime/CoreUObject/Private/UObject/CoreNet.cpp
@@ -161,9 +161,9 @@ uint32 FClassNetCacheMgr::GetFieldChecksum( const UField* Field, uint32 Checksum
 	return Checksum;
 }
 
-const FClassNetCache * FClassNetCacheMgr::GetClassNetCache( const UClass* Class )
+const FClassNetCache* FClassNetCacheMgr::GetClassNetCache( const UClass* Class )
 {
-	FClassNetCache * Result = ClassFieldIndices.FindRef( Class );
+	FClassNetCache* Result = ClassFieldIndices.FindRef( Class );
 
 	if ( !Result )
 	{
@@ -182,11 +182,12 @@ const FClassNetCache * FClassNetCacheMgr::GetClassNetCache( const UClass* Class
 		Result->Fields.Empty( Class->NetFields.Num() );
 
 		TArray< UProperty* > Properties;
+		Properties.Empty( Class->NetFields.Num() );
 
 		for( int32 i = 0; i < Class->NetFields.Num(); i++ )
 		{
 			// Add each net field to cache, and assign index/checksum
-			UField * Field = Class->NetFields[i];
+			UField* Field = Class->NetFields[i];
 
 			UProperty* Property = Cast< UProperty >( Field );
 			
