Support for RF_Async and RF_AsyncLoading in FReferenceChainSearch::PrintReferencers.

diff --git a/Engine/Source/Runtime/CoreUObject/Private/UObject/ReferenceChainSearch.cpp b/Engine/Source/Runtime/CoreUObject/Private/UObject/ReferenceChainSearch.cpp
--- a/Engine/Source/Runtime/CoreUObject/Private/UObject/ReferenceChainSearch.cpp
+++ b/Engine/Source/Runtime/CoreUObject/Private/UObject/ReferenceChainSearch.cpp
@@ -163,6 +163,16 @@ void FReferenceChainSearch::PrintReferencers( FReferenceChain& Referencer )
 			ObjectReachability += TEXT("(standalone) ");
 		}
 
+		if( RefInfo.ReferencedBy->HasAnyFlags(RF_Async) )
+		{
+			ObjectReachability += TEXT("(async) ");
+		}
+
+		if( RefInfo.ReferencedBy->HasAnyFlags(RF_AsyncLoading) )
+		{
+			ObjectReachability += TEXT("(asyncloading) ");
+		}
+
 		if (GetUObjectArray().IsDisregardForGC(RefInfo.ReferencedBy))
 		{
 			ObjectReachability += TEXT("(NeverGCed) ");
