AsyncLoading: Fixing FAsyncPackageDesc memory leaks, reducing TArray memory allocations.

diff --git a/Engine/Source/Runtime/CoreUObject/Private/Serialization/AsyncLoading.cpp b/Engine/Source/Runtime/CoreUObject/Private/Serialization/AsyncLoading.cpp
--- a/Engine/Source/Runtime/CoreUObject/Private/Serialization/AsyncLoading.cpp
+++ b/Engine/Source/Runtime/CoreUObject/Private/Serialization/AsyncLoading.cpp
@@ -137,7 +137,7 @@ public:
 			Obj->AtomicallyClearFlags(AsyncFlags);
 			check(!Obj->HasAnyFlags(AsyncFlags))
 		}
-		ReferencedObjects.Empty(ReferencedObjects.Num());
+		ReferencedObjects.Reset();
 	}
 	/** Removes all referenced objects and markes them for GC */
 	void EmptyReferencedObjectsAndCancelLoading()
@@ -161,7 +161,7 @@ public:
 			}
 			check(!Object->HasAnyFlags(AsyncFlags | LoadFlags));
 		}
-		ReferencedObjects.Empty(ReferencedObjects.Num());
+		ReferencedObjects.Reset();
 	}
 
 #if !UE_BUILD_SHIPPING
@@ -329,19 +329,23 @@ void FAsyncLoadingThread::CancelAsyncLoadingInternal()
 #if THREADSAFE_UOBJECTS
 		FScopeLock QueueLock(&QueueCritical);
 #endif
-		QueuedPackages.Empty();
+		for (FAsyncPackageDesc* PackageDesc : QueuedPackages)
+		{
+			delete PackageDesc;
+		}
+		QueuedPackages.Reset();
 	}
 
 	{
 		// Packages we started processing, need to be canceled.
 		// Accessed only in async thread, no need to protect region.
-		for (auto AsyncPackage : AsyncPackages)
+		for (FAsyncPackage* AsyncPackage : AsyncPackages)
 		{
 			AsyncPackage->Cancel();
 			delete AsyncPackage;
 		}
 
-		AsyncPackages.Empty();
+		AsyncPackages.Reset();
 	}
 
 	{
@@ -349,23 +353,23 @@ void FAsyncLoadingThread::CancelAsyncLoadingInternal()
 #if THREADSAFE_UOBJECTS
 		FScopeLock LoadedLock(&LoadedPackagesCritical);
 #endif
-		for (auto LoadedPackage : LoadedPackages)
+		for (FAsyncPackage* LoadedPackage : LoadedPackages)
 		{
 			LoadedPackage->Cancel();
 			delete LoadedPackage;
 		}
-		LoadedPackages.Empty();
+		LoadedPackages.Reset();
 	}
 	{
 #if THREADSAFE_UOBJECTS
 		FScopeLock LoadedLock(&LoadedPackagesToProcessCritical);
 #endif
-		for (auto LoadedPackage : LoadedPackagesToProcess)
+		for (FAsyncPackage* LoadedPackage : LoadedPackagesToProcess)
 		{
 			LoadedPackage->Cancel();
 			delete LoadedPackage;
 		}
-		LoadedPackagesToProcess.Empty();
+		LoadedPackagesToProcess.Reset();
 	}
 
 	AsyncLoadingCounter.Reset();
@@ -504,7 +508,7 @@ int32 FAsyncLoadingThread::CreateAsyncPackagesFromQueue()
 		FScopeLock QueueLock(&QueueCritical);
 #endif
 		QueueCopy = QueuedPackages;
-		QueuedPackages.Empty();
+		QueuedPackages.Reset();
 	}
 
 	if (QueueCopy.Num() > 0)
@@ -512,10 +516,11 @@ int32 FAsyncLoadingThread::CreateAsyncPackagesFromQueue()
 		double Timer = 0;
 		{
 			SCOPE_SECONDS_COUNTER(Timer);
-			for (auto PackageRequest : QueueCopy)
+			for (FAsyncPackageDesc* PackageRequest : QueueCopy)
 			{
-				DependencyTracker.Empty();
+				DependencyTracker.Reset();
 				ProcessAsyncPackageRequest(PackageRequest, nullptr, DependencyTracker);
+				delete PackageRequest;
 			}
 		}
 		UE_LOG(LogStreaming, Verbose, TEXT("Async package requests inserted in %fms"), Timer * 1000.0);
@@ -629,7 +634,7 @@ EAsyncPackageState::Type FAsyncLoadingThread::ProcessLoadedPackages(bool bUseTim
 		FScopeLock LoadedPackagesToProcessLock(&LoadedPackagesToProcessCritical);
 #endif
 		LoadedPackagesToProcess.Append(LoadedPackages);
-		LoadedPackages.Empty();
+		LoadedPackages.Reset();
 	}
 		
 	for (int32 PackageIndex = 0; PackageIndex < LoadedPackagesToProcess.Num() && !IsAsyncLoadingSuspended(); ++PackageIndex)
