Optimize FindFunctionByName

diff --git a/Engine/Source/Runtime/CoreUObject/Private/UObject/Class.cpp b/Engine/Source/Runtime/CoreUObject/Private/UObject/Class.cpp
--- a/Engine/Source/Runtime/CoreUObject/Private/UObject/Class.cpp
+++ b/Engine/Source/Runtime/CoreUObject/Private/UObject/Class.cpp
@@ -3641,6 +3641,8 @@ void UClass::PurgeClass(bool bRecompilingOnLoad)
 	ScriptObjectReferences.Empty();
 
 	FuncMap.Empty();
+	ParentFuncMap.Empty();
+	InterfaceFuncMap.Empty();
 	PropertyLink = NULL;
 }
 
@@ -3963,22 +3965,48 @@ void UClass::AddNativeFunction(const ANSICHAR* InName,Native InPointer)
 
 UFunction* UClass::FindFunctionByName(FName InName, EIncludeSuperFlag::Type IncludeSuper) const
 {
-	if (!IncludeSuper)
-		return FuncMap.FindRef(InName);
-
-	for (const UClass* SearchClass = this; SearchClass; SearchClass = SearchClass->GetSuperClass())
+	UFunction* Result = FuncMap.FindRef(InName);
+	if (Result == nullptr && IncludeSuper == EIncludeSuperFlag::IncludeSuper)
 	{
-		if (UFunction* Result = SearchClass->FuncMap.FindRef(InName))
-			return Result;
+		if (Interfaces.Num())
+		{
+			if (UFunction** InterfaceResult = InterfaceFuncMap.Find(InName))
+			{
+				Result = *InterfaceResult;
+			}
+			else
+			{
+				for (const FImplementedInterface& Inter : Interfaces)
+				{
+					Result = Inter.Class->FindFunctionByName(InName);
+					if (Result)
+					{
+						break;
+					}
+				}
 
-		for (auto& Inter : SearchClass->Interfaces)
+				InterfaceFuncMap.Add(InName, Result);
+			}
+		}
+
+		if (Result == nullptr)
 		{
-			if (UFunction* Result = Inter.Class->FindFunctionByName(InName))
-				return Result;
+			if (UClass* SuperClass = GetSuperClass())
+			{
+				if (UFunction** ParentResult = ParentFuncMap.Find(InName))
+				{
+					Result = *ParentResult;
+				}
+				else
+				{
+					Result = SuperClass->FindFunctionByName(InName);
+					ParentFuncMap.Add(InName, Result);
+				}
+			}
 		}
 	}
 
-	return NULL;
+	return Result;
 }
 
 const FString UClass::GetConfigName() const
diff --git a/Engine/Source/Runtime/CoreUObject/Public/UObject/Class.h b/Engine/Source/Runtime/CoreUObject/Public/UObject/Class.h
--- a/Engine/Source/Runtime/CoreUObject/Public/UObject/Class.h
+++ b/Engine/Source/Runtime/CoreUObject/Public/UObject/Class.h
@@ -1906,9 +1906,15 @@ public:
 	UObject* ClassDefaultObject;
 
 private:
-	/** Map of all functions by name contained in this state */
+	/** Map of all functions by name contained in this class */
 	TMap<FName, UFunction*> FuncMap;
 
+	/** A cache of all functions by name that exist in a parent context */
+	mutable TMap<FName, UFunction*> ParentFuncMap;
+
+	/** A cache of all functions by name that exist in an interface context */
+	mutable TMap<FName, UFunction*> InterfaceFuncMap;
+
 public:
 	/**
 	 * The list of interfaces which this class implements, along with the pointer property that is located at the offset of the interface's vtable.
