Further reducing the number of allocations during GC.

Force inlined some of the functions as they were not inlined by the compiler.

Added GC array pool leak checks

diff --git a/Engine/Source/Runtime/CoreUObject/Private/UObject/GarbageCollection.cpp b/Engine/Source/Runtime/CoreUObject/Private/UObject/GarbageCollection.cpp
--- a/Engine/Source/Runtime/CoreUObject/Private/UObject/GarbageCollection.cpp
+++ b/Engine/Source/Runtime/CoreUObject/Private/UObject/GarbageCollection.cpp
@@ -193,7 +193,7 @@ public:
 	* Gets the singleton instance of the FObjectArrayPool
 	* @return Pool singleton.
 	*/
-	static FGCArrayPool& Get()
+	FORCEINLINE static FGCArrayPool& Get()
 	{
 		static FGCArrayPool Singleton;
 		return Singleton;
@@ -205,7 +205,7 @@ public:
 	* @return The array.
 	* @see ReturnToPool
 	*/
-	TArray<UObject*>* GetArrayFromPool()
+	FORCEINLINE TArray<UObject*>* GetArrayFromPool()
 	{
 		TArray<UObject*>* Result = Pool.Pop();
 		if (!Result)
@@ -213,6 +213,9 @@ public:
 			Result = new TArray<UObject*>();
 		}
 		check(Result);
+#if UE_BUILD_DEBUG
+		NumberOfUsedArrays.Increment();
+#endif // UE_BUILD_DEBUG
 		return Result;
 	}
 
@@ -222,8 +225,12 @@ public:
 	* @param Array The array to return.
 	* @see GetArrayFromPool
 	*/
-	void ReturnToPool(TArray<UObject*>* Array)
+	FORCEINLINE void ReturnToPool(TArray<UObject*>* Array)
 	{
+#if UE_BUILD_DEBUG
+		const int32 CheckUsedArrays = NumberOfUsedArrays.Decrement();
+		checkSlow(CheckUsedArrays >= 0);
+#endif // UE_BUILD_DEBUG
 		check(Array);
 		Array->Reset();
 		Pool.Push(Array);
@@ -232,6 +239,11 @@ public:
 	/** Performs memory cleanup */
 	void Cleanup()
 	{
+#if UE_BUILD_DEBUG
+		const int32 CheckUsedArrays = NumberOfUsedArrays.GetValue();
+		checkSlow(CheckUsedArrays == 0);
+#endif // UE_BUILD_DEBUG
+
 		uint32 FreedMemory = 0;
 		TArray< TArray<UObject*>* > AllArrays;
 		Pool.PopAll(AllArrays);
@@ -243,10 +255,25 @@ public:
 		UE_LOG(LogGarbage, Log, TEXT("Freed %ub from %d GC array pools."), FreedMemory, AllArrays.Num());
 	}
 
+#if UE_BUILD_DEBUG
+	void CheckLeaks()
+	{
+		// This function is called after GC has finished so at this point there should be no
+		// arrays used by GC and all should be returned to the pool
+		const int32 LeakedGCPoolArrays = NumberOfUsedArrays.GetValue();
+		checkSlow(LeakedGCPoolArrays == 0);
+	}
+#endif
+
 private:
 
 	/** Holds the collection of recycled arrays. */
 	TLockFreePointerList< TArray<UObject*> > Pool;
+
+#if UE_BUILD_DEBUG
+	/** Number of arrays currently acquired from the pool by GC */
+	FThreadSafeCounter NumberOfUsedArrays;
+#endif // UE_BUILD_DEBUG
 };
 
 /** Called on shutdown to free GC memory */
@@ -664,7 +691,7 @@ public:
 		GObjectCountDuringLastMarkPhase = 0;
 
 		// Presize array and add a bit of extra slack for prefetching.
-		ObjectsToSerialize.Empty( GetUObjectArray().GetObjectArrayNumMinusPermanent() + 3 );
+		ObjectsToSerialize.Reset( GetUObjectArray().GetObjectArrayNumMinusPermanent() + 3 );
 		// Make sure GC referencer object is checked for references to other objects even if it resides in permanent object pool
 		if (FPlatformProperties::RequiresCookedData() && FGCObject::GGCObjectReferencer && GetUObjectArray().IsDisregardForGC(FGCObject::GGCObjectReferencer))
 		{
@@ -752,6 +779,10 @@ public:
 		}
 
 		FGCArrayPool::Get().ReturnToPool(&ObjectsToSerialize);
+
+#if UE_BUILD_DEBUG
+		FGCArrayPool::Get().CheckLeaks();
+#endif
 	}
 
 	void DispatchObjectTasks(TArray<UObject*>& ObjectsToSerialize)
@@ -765,11 +796,9 @@ public:
 		UObject* CurrentObject = NULL;
 
 		const int32 MinDesiredObjectsPerSubTask = GMinDesiredObjectsPerSubTask; // sometimes there will be less, a lot less
-		const int32 NewObjectsArrayLength = InObjectsToSerializeArray.Num() * 2;
 
 		/** Growing array of objects that require serialization */
-		TArray<UObject*>	NewObjectsToSerializeArray;
-		NewObjectsToSerializeArray.Empty( NewObjectsArrayLength );
+		TArray<UObject*>&	NewObjectsToSerializeArray = *FGCArrayPool::Get().GetArrayFromPool();
 
 		// Ping-pong between these two arrays if there's not enough objects to spawn a new task
 		TArray<UObject*>& ObjectsToSerialize = InObjectsToSerializeArray;
@@ -1031,6 +1060,8 @@ public:
 			}
 		}
 		while( CurrentIndex < ObjectsToSerialize.Num() );
+
+		FGCArrayPool::Get().ReturnToPool(&NewObjectsToSerializeArray);
 	}
 };
 
