Minor fix for checking available object count in UObject array.

diff --git a/Engine/Source/Runtime/CoreUObject/Private/UObject/UObjectArray.cpp b/Engine/Source/Runtime/CoreUObject/Private/UObject/UObjectArray.cpp
--- a/Engine/Source/Runtime/CoreUObject/Private/UObject/UObjectArray.cpp
+++ b/Engine/Source/Runtime/CoreUObject/Private/UObject/UObjectArray.cpp
@@ -55,8 +55,8 @@ void FUObjectArray::AllocateUObjectIndex(UObjectBase* Object, bool bMergingThrea
 		if (AvailableIndex)
 		{
 #if WITH_EDITOR
-			ObjAvailableCount.Decrement();
-			checkSlow(ObjAvailableCount.GetValue() >= 0);
+			const int32 AvailableCount = ObjAvailableCount.Decrement();
+			checkSlow(AvailableCount >= 0);
 #endif
 			Index = (int32)(uintptr_t)AvailableIndex;
 			check(ObjObjects[Index]==nullptr);
