Basic support for native blueprints in async loading and linker (WIP)

diff --git a/Engine/Source/Runtime/CoreUObject/Private/Serialization/AsyncLoading.cpp b/Engine/Source/Runtime/CoreUObject/Private/Serialization/AsyncLoading.cpp
--- a/Engine/Source/Runtime/CoreUObject/Private/Serialization/AsyncLoading.cpp
+++ b/Engine/Source/Runtime/CoreUObject/Private/Serialization/AsyncLoading.cpp
@@ -1467,6 +1467,13 @@ void FAsyncPackage::AddDependencyTree(int32 CurrentPackageIndex, FAsyncPackage&
 	SearchedPackages.Add(&ImportedPackage);
 }
 
+/** Get all assets paths necessary for the class with the given class name. */
+void FConvertedBlueprintsDependencies_GetAssets(FName ClassName, TSet<FName>& OutPackagePaths)
+{
+	// @todo: needs to be implemented by BP team
+}
+
+
 /** 
  * Load imports till time limit is exceeded.
  *
@@ -1484,74 +1491,89 @@ EAsyncPackageState::Type FAsyncPackage::LoadImports()
 	// GC can't run in here
 	FGCScopeGuard GCGuard;
 
+	// Array that will hold dependency packages for the import
+	TSet<FName> ImportDependencyPackages;
+
 	// Create imports.
 	while (LoadImportIndex < Linker->ImportMap.Num() && !IsTimeLimitExceeded())
 	{
 		// Get the package for this import
 		const FObjectImport* Import = &Linker->ImportMap[LoadImportIndex++];
-		while (Import->OuterIndex.IsImport())
+		ImportDependencyPackages.Reset();
+		if (FPlatformProperties::RequiresCookedData())
 		{
-			Import = &Linker->Imp(Import->OuterIndex);
+			// Try to get dependencies for native blueprint classes
+			FConvertedBlueprintsDependencies_GetAssets(Import->ObjectName, ImportDependencyPackages);
 		}
-		check(Import->OuterIndex.IsNull());
-
-		// @todo: why do we need this? some UFunctions have null outer in the linker.
-		if (Import->ClassName != NAME_Package)
+		if (ImportDependencyPackages.Num() == 0)
 		{
-			continue;
-		}
+			while (Import->OuterIndex.IsImport())
+			{
+				Import = &Linker->Imp(Import->OuterIndex);
+			}
+			check(Import->OuterIndex.IsNull());
 
-		// Our import package name is the import name
-		const FName ImportPackageFName(Import->ObjectName);
+			// @todo: why do we need this? some UFunctions have null outer in the linker.
+			if (Import->ClassName != NAME_Package)
+			{
+				continue;
+			}
+			
+			// Our import package name is the import name
+			ImportDependencyPackages.Add(Import->ObjectName);
+		}
 
-		// Handle circular dependencies - try to find existing packages.
-		UPackage* ExistingPackage = dynamic_cast<UPackage*>(StaticFindObjectFast(UPackage::StaticClass(), nullptr, ImportPackageFName, true));
-		if (ExistingPackage && !(ExistingPackage->PackageFlags & PKG_CompiledIn) && !ExistingPackage->bHasBeenFullyLoaded)//!ExistingPackage->HasAnyFlags(RF_WasLoaded))
+		for (const FName& ImportPackageFName : ImportDependencyPackages)
 		{
-			// The import package already exists. Check if it's currently being streamed as well. If so, make sure
-			// we add all dependencies that don't yet have linkers created otherwise we risk that if the current package
-			// doesn't depend on any other packages that have not yet started streaming, creating imports is going
-			// to load packages blocking the main thread.
-			int32 PendingAsyncPackageIndex = FAsyncLoadingThread::Get().FindAsyncPackage(ImportPackageFName);
-			if (PendingAsyncPackageIndex != INDEX_NONE)
+			// Handle circular dependencies - try to find existing packages.
+			UPackage* ExistingPackage = dynamic_cast<UPackage*>(StaticFindObjectFast(UPackage::StaticClass(), nullptr, ImportPackageFName, true));
+			if (ExistingPackage && !(ExistingPackage->PackageFlags & PKG_CompiledIn) && !ExistingPackage->bHasBeenFullyLoaded)//!ExistingPackage->HasAnyFlags(RF_WasLoaded))
 			{
-				FAsyncPackage& PendingPackage = *FAsyncLoadingThread::Get().GetPackage(PendingAsyncPackageIndex);
-				FLinkerLoad* PendingPackageLinker = PendingPackage.Linker;
-				if (PendingPackageLinker == nullptr || !PendingPackageLinker->HasFinishedInitialization())
+				// The import package already exists. Check if it's currently being streamed as well. If so, make sure
+				// we add all dependencies that don't yet have linkers created otherwise we risk that if the current package
+				// doesn't depend on any other packages that have not yet started streaming, creating imports is going
+				// to load packages blocking the main thread.
+				int32 PendingAsyncPackageIndex = FAsyncLoadingThread::Get().FindAsyncPackage(ImportPackageFName);
+				if (PendingAsyncPackageIndex != INDEX_NONE)
 				{
-					// Add this import to the dependency list.
-					AddUniqueLinkerDependencyPackage(AsyncQueueIndex, PendingPackage);
+					FAsyncPackage& PendingPackage = *FAsyncLoadingThread::Get().GetPackage(PendingAsyncPackageIndex);
+					FLinkerLoad* PendingPackageLinker = PendingPackage.Linker;
+					if (PendingPackageLinker == nullptr || !PendingPackageLinker->HasFinishedInitialization())
+					{
+						// Add this import to the dependency list.
+						AddUniqueLinkerDependencyPackage(AsyncQueueIndex, PendingPackage);
+					}
+					else
+					{
+						UE_LOG(LogStreaming, Verbose, TEXT("FAsyncPackage::LoadImports for %s: Linker exists for %s"), *Desc.NameToLoad.ToString(), *ImportPackageFName.ToString());
+						// Only keep a reference to this package so that its linker doesn't go away too soon
+						PendingPackage.DependencyRefCount.Increment();
+						ReferencedImports.Add(&PendingPackage);
+						// Check if we need to add its dependencies too.
+						TSet<FAsyncPackage*> SearchedPackages;
+						AddDependencyTree(AsyncQueueIndex, PendingPackage, SearchedPackages);
+					}
+				}
+			}
+
+			if (!ExistingPackage && ContainsDependencyPackage(PendingImportedPackages, ImportPackageFName) == INDEX_NONE)
+			{
+				const FString ImportPackageName(Import->ObjectName.ToString());
+				// The package doesn't exist and this import is not in the dependency list so add it now.
+				if (!FPackageName::IsShortPackageName(ImportPackageName))
+				{
+					UE_LOG(LogStreaming, Verbose, TEXT("FAsyncPackage::LoadImports for %s: Loading %s"), *Desc.NameToLoad.ToString(), *ImportPackageName);
+					AddImportDependency(AsyncQueueIndex, ImportPackageFName);
 				}
 				else
 				{
-					UE_LOG(LogStreaming, Verbose, TEXT("FAsyncPackage::LoadImports for %s: Linker exists for %s"), *Desc.NameToLoad.ToString(), *ImportPackageFName.ToString());
-					// Only keep a reference to this package so that its linker doesn't go away too soon
-					PendingPackage.DependencyRefCount.Increment();
-					ReferencedImports.Add(&PendingPackage);
-					// Check if we need to add its dependencies too.
-					TSet<FAsyncPackage*> SearchedPackages;
-					AddDependencyTree(AsyncQueueIndex, PendingPackage, SearchedPackages);
+					// This usually means there's a reference to a script package from another project
+					UE_LOG(LogStreaming, Warning, TEXT("FAsyncPackage::LoadImports for %s: Short package name in imports list: %s"), *Desc.NameToLoad.ToString(), *ImportPackageName);
 				}
 			}
-		}
 
-		if (!ExistingPackage && ContainsDependencyPackage(PendingImportedPackages, ImportPackageFName) == INDEX_NONE)
-		{			
-			const FString ImportPackageName(Import->ObjectName.ToString());
-			// The package doesn't exist and this import is not in the dependency list so add it now.
-			if (!FPackageName::IsShortPackageName(ImportPackageName))
-			{
-				UE_LOG(LogStreaming, Verbose, TEXT("FAsyncPackage::LoadImports for %s: Loading %s"), *Desc.NameToLoad.ToString(), *ImportPackageName);
-				AddImportDependency(AsyncQueueIndex, ImportPackageFName);
-			}
-			else
-			{
-				// This usually means there's a reference to a script package from another project
-				UE_LOG(LogStreaming, Warning, TEXT("FAsyncPackage::LoadImports for %s: Short package name in imports list: %s"), *Desc.NameToLoad.ToString(), *ImportPackageName);
-			}
+			UpdateLoadPercentage();
 		}
-
-		UpdateLoadPercentage();
 	}
 
 	if (PendingImportedPackages.Num())
diff --git a/Engine/Source/Runtime/CoreUObject/Private/UObject/LinkerLoad.cpp b/Engine/Source/Runtime/CoreUObject/Private/UObject/LinkerLoad.cpp
--- a/Engine/Source/Runtime/CoreUObject/Private/UObject/LinkerLoad.cpp
+++ b/Engine/Source/Runtime/CoreUObject/Private/UObject/LinkerLoad.cpp
@@ -3913,6 +3913,12 @@ bool FLinkerLoad::IsImportNative(const int32 Index) const
 	return bIsImportNative;
 }
 
+// Temporary placeholder: constructs native blieprint class given its name
+void ConstructNativeBlueprintClass(FName ClassName)
+{
+	// @todo: this needs to find a generated function by class name.
+}
+
 // Return the loaded object corresponding to an import index; any errors are fatal.
 UObject* FLinkerLoad::CreateImport( int32 Index )
 {
@@ -3948,11 +3954,19 @@ UObject* FLinkerLoad::CreateImport( int32 Index )
 					// Import is a toplevel package.
 					if( Import.OuterIndex.IsNull() )
 					{
-						FindObject = CreatePackage( NULL, *Import.ObjectName.ToString() );
+						FindObject = CreatePackage(NULL, *Import.ObjectName.ToString());
 					}
 					// Import is regular import/ export.
 					else
 					{
+						if (FPlatformProperties::RequiresCookedData())
+						{
+							// At this point we know the class package and class of the import exist, this can potentially be a native blueprint
+							// (class is just a UClass and the package is /Script/CoreUObject)
+							// @todo: maybe we can further filter out classes by checking FindClass somehow (e.g. once it's a specialized UCookedBlueprintClass or something)
+							ConstructNativeBlueprintClass(Import.ObjectName);
+						}
+						
 						// Find the imports' outer.
 						UObject* FindOuter = NULL;
 						// Import.
