Use utility function instead of doing work directly

diff --git a/Engine/Source/Runtime/CoreUObject/Private/UObject/Class.cpp b/Engine/Source/Runtime/CoreUObject/Private/UObject/Class.cpp
--- a/Engine/Source/Runtime/CoreUObject/Private/UObject/Class.cpp
+++ b/Engine/Source/Runtime/CoreUObject/Private/UObject/Class.cpp
@@ -3641,8 +3641,7 @@ void UClass::PurgeClass(bool bRecompilingOnLoad)
 	ScriptObjectReferences.Empty();
 
 	FuncMap.Empty();
-	ParentFuncMap.Empty();
-	InterfaceFuncMap.Empty();
+	ClearFunctionMapsCaches();
 	PropertyLink = NULL;
 }
 
