Integrate CL 2676785 to main

Fixes/improvements for dependency preloading
 - Move dependency lookups earlier in synchronous package load so that they happen BEFORE the source package has its linker created
 - Put loose (non-map based) packages FIRST in the PAK file. Matches general access patterns better and gets higher pak seek contiguity.

diff --git a/Engine/Source/Editor/UnrealEd/Private/Commandlets/ChunkManifestGenerator.cpp b/Engine/Source/Editor/UnrealEd/Private/Commandlets/ChunkManifestGenerator.cpp
--- a/Engine/Source/Editor/UnrealEd/Private/Commandlets/ChunkManifestGenerator.cpp
+++ b/Engine/Source/Editor/UnrealEd/Private/Commandlets/ChunkManifestGenerator.cpp
@@ -556,9 +556,9 @@ bool FChunkManifestGenerator::SaveAssetRegistry(const FString& SandboxPath, cons
 			if (ContainsMap(AssetData.PackageName))
 			{
 				MapList.Add(AssetData.PackageName);
+			}
 		}
 	}
-	}
 
 	AssetRegistry.SaveRegistryData(SerializedAssetRegistry, GeneratedAssetRegistryData, &MapList);
 	UE_LOG(LogChunkManifestGenerator, Display, TEXT("Generated asset registry num assets %d, size is %5.2fkb"), GeneratedAssetRegistryData.Num(), (float)SerializedAssetRegistry.Num() / 1024.f);
@@ -675,12 +675,12 @@ FString FChunkManifestGenerator::CreateCookerFileOrderString(const TMap<FName, F
 
 	TArray<FName> FileOrder;
 	TArray<FName> EncounteredNames;
-	for (auto Asset : TopLevelMapNodes)
+	for (auto Asset : TopLevelNodes)
 	{
 		AddAssetToFileOrderRecursive(Asset, FileOrder, EncounteredNames, InAssetData, InMaps);
 	}
 
-	for (auto Asset : TopLevelNodes)
+	for (auto Asset : TopLevelMapNodes)
 	{
 		AddAssetToFileOrderRecursive(Asset, FileOrder, EncounteredNames, InAssetData, InMaps);
 	}
diff --git a/Engine/Source/Runtime/CoreUObject/Private/UObject/UObjectGlobals.cpp b/Engine/Source/Runtime/CoreUObject/Private/UObject/UObjectGlobals.cpp
--- a/Engine/Source/Runtime/CoreUObject/Private/UObject/UObjectGlobals.cpp
+++ b/Engine/Source/Runtime/CoreUObject/Private/UObject/UObjectGlobals.cpp
@@ -89,23 +89,23 @@ static FORCEINLINE bool IsGarbageCollectingOnGameThread()
 namespace
 {
 	/**
-	* Legacy static find object helper, that helps to find reflected types, that
-	* are no longer a subobjects of UCLASS defined in the same header.
-	*
-	* If the class looked for is of one of the relocated types (or theirs subclass)
-	* then it performs another search in containing package.
-	*
-	* If the class match wasn't exact (i.e. either nullptr or subclass of allowed
-	* ones) and we've found an object we're revalidating it to make sure the
-	* legacy search was valid.
-	*
-	* @param ObjectClass Class of the object to find.
-	* @param ObjectPackage Package of the object to find.
-	* @param ObjectName Name of the object to find.
-	* @param ExactClass If the class match has to be exact. I.e. ObjectClass == FoundObjects.GetClass()
-	*
-	* @returns Found object.
-	*/
+	 * Legacy static find object helper, that helps to find reflected types, that
+	 * are no longer a subobjects of UCLASS defined in the same header.
+	 *
+	 * If the class looked for is of one of the relocated types (or theirs subclass)
+	 * then it performs another search in containing package.
+	 *
+	 * If the class match wasn't exact (i.e. either nullptr or subclass of allowed
+	 * ones) and we've found an object we're revalidating it to make sure the
+	 * legacy search was valid.
+	 *
+	 * @param ObjectClass Class of the object to find.
+	 * @param ObjectPackage Package of the object to find.
+	 * @param ObjectName Name of the object to find.
+	 * @param ExactClass If the class match has to be exact. I.e. ObjectClass == FoundObjects.GetClass()
+	 *
+	 * @returns Found object.
+	 */
 	UObject* StaticFindObjectWithChangedLegacyPath(UClass* ObjectClass, UObject* ObjectPackage, FName ObjectName, bool ExactClass)
 	{
 		UObject* MatchingObject = nullptr;
@@ -587,7 +587,7 @@ FString ResolveIniObjectsReference(const FString& ObjectReference, const FString
 	if (!GConfig->GetString(*Section, *Key, Output, *IniFilename))
 	{
 		if (bThrow == true)
-		{
+{
 			UE_LOG(LogUObjectGlobals, Error, TEXT(" %s %s "), *FString::Printf(TEXT("Can't find '%s' in configuration file section=%s key=%s"), *ObjectReference, *Section, *Key), **IniFilename);
 		}
 	}
@@ -622,11 +622,11 @@ const FString* GetIniFilenameFromObjectsReference(const FString& Name)
 // Resolve a package and name.
 //
 bool ResolveName( UObject*& InPackage, FString& InOutName, bool Create, bool Throw )
-{
+	{
 	const FString* IniFilename = GetIniFilenameFromObjectsReference(InOutName);
 
 	if (IniFilename && InOutName.Contains(TEXT("."), ESearchCase::CaseSensitive))
-	{
+		{
 		InOutName = ResolveIniObjectsReference(InOutName, IniFilename, Throw);
 	}
 
@@ -893,7 +893,7 @@ UClass* StaticLoadClass( UClass* BaseClass, UObject* InOuter, const TCHAR* InNam
 * @param	ImportLinker	Linker that requests this package through one of its imports
 * @return	Loaded package if successful, NULL otherwise
 */
-UPackage* LoadPackageInternal(UPackage* InOuter, const TCHAR* InLongPackageName, uint32 LoadFlags, FLinkerLoad* ImportLinker)
+UPackage* LoadPackageInternal(UPackage* InOuter, const TCHAR* InLongPackageName, uint32 LoadFlags, FLinkerLoad* ImportLinker, TSet<FName>& InDependencyTracker, IAssetRegistryInterface* InAssetRegistry)
 {
 	DECLARE_SCOPE_CYCLE_COUNTER(TEXT("LoadPackageInternal"), STAT_LoadPackageInternal, STATGROUP_ObjectVerbose);
 
@@ -936,6 +936,22 @@ UPackage* LoadPackageInternal(UPackage* InOuter, const TCHAR* InLongPackageName,
 	
 	SlowTask.EnterProgressFrame(10);
 
+	if (InAssetRegistry)
+	{
+		FName PackageName(InLongPackageName);
+		TArray<FName> PackageDependencies;
+		InAssetRegistry->GetDependencies(PackageName, PackageDependencies, EAssetRegistryDependencyType::Hard);
+
+		for (auto Dependency : PackageDependencies)
+		{
+			if (!InDependencyTracker.Contains(Dependency) && FindObjectFast<UPackage>(NULL, Dependency, false, false) == nullptr)
+			{
+				InDependencyTracker.Add(Dependency);
+				LoadPackageInternal(InOuter, *Dependency.ToString(), LoadFlags, nullptr, InDependencyTracker, InAssetRegistry);
+			}
+		}
+	}
+
 	// Try to load.
 	BeginLoad();
 
@@ -980,44 +996,13 @@ UPackage* LoadPackageInternal(UPackage* InOuter, const TCHAR* InLongPackageName,
 		{
 			// The linker is associated with a package that has already been loaded.
 			// Loading packages that have already been loaded is unsupported.
-			EndLoadAndCopyLocalizationGatherFlag();			
+			EndLoadAndCopyLocalizationGatherFlag();
 			return Result;
 		}
 
 		// The time tracker keeps track of time spent in LoadPackage.
 		FExclusiveLoadPackageTimeTracker::FScopedPackageTracker Tracker(Result);
 
-#if !WITH_EDITOR
-		bool bAllowDependencyPreloading = ((LoadFlags & LOAD_DisableDependencyPreloading) == 0);
-		static auto CVarPreloadDependencies = IConsoleManager::Get().FindConsoleVariable(TEXT("s.PreloadPackageDependencies"));
-		if (bAllowDependencyPreloading && CVarPreloadDependencies && CVarPreloadDependencies->GetInt() != 0)
-		{
-			if (Result->PackageFlags & PKG_ProcessingDependencies)
-			{
-				// We've currently already processing the dependencies of this package, so there is a circular dependency
-				EndLoad();
-				return nullptr;
-			}
-
-			auto AssetRegistry = IAssetRegistryInterface::GetPtr();
-
-			if (AssetRegistry)
-			{
-				TArray<FName> PackageDependencies;
-				FName PackageName(InLongPackageName);
-
-				AssetRegistry->GetDependencies(PackageName, PackageDependencies, EAssetRegistryDependencyType::Hard);
-
-				Result->PackageFlags |= PKG_ProcessingDependencies;
-				for (auto Dependency : PackageDependencies)
-				{
-					LoadPackage(InOuter, *Dependency.ToString(), LoadFlags);
-				}
-				Result->PackageFlags &= ~PKG_ProcessingDependencies;
-			}
-		}
-#endif
-
 		// If we are loading a package for diff'ing, set the package flag
 		if(LoadFlags & LOAD_ForDiff)
 		{
@@ -1146,6 +1131,23 @@ UPackage* LoadPackageInternal(UPackage* InOuter, const TCHAR* InLongPackageName,
 	return Result;
 }
 
+UPackage* LoadPackageInternal(UPackage* InOuter, const TCHAR* InLongPackageName, uint32 LoadFlags, FLinkerLoad* ImportLinker)
+{
+	IAssetRegistryInterface* AssetRegistry = nullptr;
+#if !WITH_EDITOR
+	bool bAllowDependencyPreloading = ((LoadFlags & LOAD_DisableDependencyPreloading) == 0);
+	static auto CVarPreloadDependencies = IConsoleManager::Get().FindConsoleVariable(TEXT("s.PreloadPackageDependencies"));
+	
+	if (bAllowDependencyPreloading && CVarPreloadDependencies && CVarPreloadDependencies->GetInt() != 0)
+	{
+		AssetRegistry = IAssetRegistryInterface::GetPtr();
+	}
+#endif
+
+	TSet<FName> DependencyTracker;
+	return LoadPackageInternal(InOuter, InLongPackageName, LoadFlags, ImportLinker, DependencyTracker, AssetRegistry);
+}
+
 UPackage* LoadPackage(UPackage* InOuter, const TCHAR* InLongPackageName, uint32 LoadFlags)
 {
 	// Change to 1 if you want more detailed stats for loading packages, but at the cost of adding dynamic stats.
@@ -2417,7 +2419,7 @@ void FObjectInitializer::PostConstructInit()
 #endif // !(UE_BUILD_SHIPPING || UE_BUILD_TEST)
 
 #if !USE_CIRCULAR_DEPENDENCY_LOAD_DEFERRING
-	if ( !Obj->HasAnyFlags(RF_NeedLoad)
+	if (!Obj->HasAnyFlags(RF_NeedLoad) 
 #else 
 	// we defer this initialization in special set of cases (when Obj is a CDO 
 	// and its parent hasn't been serialized yet)... in those cases, Obj (the 
@@ -2429,7 +2431,7 @@ void FObjectInitializer::PostConstructInit()
 	if ( (!Obj->HasAnyFlags(RF_NeedLoad) || bIsDeferredInitializer)
 #endif // !USE_CIRCULAR_DEPENDENCY_LOAD_DEFERRING
 		// if component instancing is not enabled, then we leave the components in an invalid state, which will presumably be fixed by the caller
-		&& ((InstanceGraph == NULL) || InstanceGraph->IsSubobjectInstancingEnabled()) ) 
+		&& ((InstanceGraph == NULL) || InstanceGraph->IsSubobjectInstancingEnabled())) 
 	{
 		Obj->CheckDefaultSubobjects();
 	}
diff --git a/Engine/Source/Runtime/CoreUObject/Public/UObject/ObjectBase.h b/Engine/Source/Runtime/CoreUObject/Public/UObject/ObjectBase.h
--- a/Engine/Source/Runtime/CoreUObject/Public/UObject/ObjectBase.h
+++ b/Engine/Source/Runtime/CoreUObject/Public/UObject/ObjectBase.h
@@ -107,7 +107,7 @@ enum EPackageFlags
 	PKG_DisallowLazyLoading			= 0x00080000,	// Set if the archive serializing this package cannot use lazy loading
 	PKG_PlayInEditor				= 0x00100000,	// Set if the package was created for the purpose of PIE
 	PKG_ContainsScript				= 0x00200000,	// Package is allowed to contain UClass objects
-	PKG_ProcessingDependencies		= 0x00400000,
+//	PKG_Unused						= 0x00400000,
 //	PKG_Unused						= 0x00800000,
 //	PKG_Unused						= 0x01000000,	
 	PKG_StoreCompressed				= 0x02000000,	// Package is being stored compressed, requires archive support for compression
