FAsyncPackage::LoadImports uses FConvertedBlueprintsDependencies

diff --git a/Engine/Source/Runtime/CoreUObject/Private/Serialization/AsyncLoading.cpp b/Engine/Source/Runtime/CoreUObject/Private/Serialization/AsyncLoading.cpp
--- a/Engine/Source/Runtime/CoreUObject/Private/Serialization/AsyncLoading.cpp
+++ b/Engine/Source/Runtime/CoreUObject/Private/Serialization/AsyncLoading.cpp
@@ -13,6 +13,8 @@
 #include "Serialization/AsyncLoadingThread.h"
 #include "ExclusiveLoadPackageTimeTracker.h"
 #include "AssetRegistryInterface.h"
+#include "BlueprintSupport.h"
+
 
 /*-----------------------------------------------------------------------------
 	Async loading stats.
@@ -1467,13 +1469,6 @@ void FAsyncPackage::AddDependencyTree(int32 CurrentPackageIndex, FAsyncPackage&
 	SearchedPackages.Add(&ImportedPackage);
 }
 
-/** Get all assets paths necessary for the class with the given class name. */
-void FConvertedBlueprintsDependencies_GetAssets(FName ClassName, TSet<FName>& OutPackagePaths)
-{
-	// @todo: needs to be implemented by BP team
-}
-
-
 /** 
  * Load imports till time limit is exceeded.
  *
@@ -1492,7 +1487,7 @@ EAsyncPackageState::Type FAsyncPackage::LoadImports()
 	FGCScopeGuard GCGuard;
 
 	// Array that will hold dependency packages for the import
-	TSet<FName> ImportDependencyPackages;
+	TArray<FName> ImportDependencyPackages;
 
 	// Create imports.
 	while (LoadImportIndex < Linker->ImportMap.Num() && !IsTimeLimitExceeded())
@@ -1503,7 +1498,7 @@ EAsyncPackageState::Type FAsyncPackage::LoadImports()
 		if (FPlatformProperties::RequiresCookedData())
 		{
 			// Try to get dependencies for native blueprint classes
-			FConvertedBlueprintsDependencies_GetAssets(Import->ObjectName, ImportDependencyPackages);
+			FConvertedBlueprintsDependencies::Get().GetAssets(Import->ObjectName, ImportDependencyPackages);
 		}
 		if (ImportDependencyPackages.Num() == 0)
 		{
