Fixed Data Table CSV/JSON importing not working correctly with FName properties that contain a space

UE-20839 - Data Table cannot import FName properties that contain a space

UNameProperty::ImportText_Internal would only import a name with a space if it was surrounded by quotes, but it will only be surrounded by quotes if PPF_Delimited is set.

This changes UNameProperty::ImportText_Internal to work more like UStrProperty::ImportText_Internal, and perform different import logic based on the PPF_Delimited flag.

diff --git a/Engine/Source/Runtime/CoreUObject/Private/UObject/PropertyName.cpp b/Engine/Source/Runtime/CoreUObject/Private/UObject/PropertyName.cpp
--- a/Engine/Source/Runtime/CoreUObject/Private/UObject/PropertyName.cpp
+++ b/Engine/Source/Runtime/CoreUObject/Private/UObject/PropertyName.cpp
@@ -26,12 +26,22 @@ void UNameProperty::ExportTextItem( FString& ValueStr, const void* PropertyValue
 }
 const TCHAR* UNameProperty::ImportText_Internal( const TCHAR* Buffer, void* Data, int32 PortFlags, UObject* Parent, FOutputDevice* ErrorText ) const
 {
-	FString Temp;
-	Buffer = UPropertyHelpers::ReadToken( Buffer, Temp, true );
-	if( !Buffer )
-		return NULL;
+	if (!(PortFlags & PPF_Delimited))
+	{
+		*(FName*)Data = FName(Buffer, FNAME_Add, true);
 
-	*(FName*)Data = FName(*Temp, FNAME_Add, true);
+		// in order to indicate that the value was successfully imported, advance the buffer past the last character that was imported
+		Buffer += FCString::Strlen(Buffer);
+	}
+	else
+	{
+		FString Temp;
+		Buffer = UPropertyHelpers::ReadToken(Buffer, Temp, true);
+		if (!Buffer)
+			return NULL;
+
+		*(FName*)Data = FName(*Temp, FNAME_Add, true);
+	}
 	return Buffer;
 }
 
