Adding UDynamicClass for dynamic UObject types

diff --git a/Engine/Source/Runtime/CoreUObject/Private/UObject/Class.cpp b/Engine/Source/Runtime/CoreUObject/Private/UObject/Class.cpp
--- a/Engine/Source/Runtime/CoreUObject/Private/UObject/Class.cpp
+++ b/Engine/Source/Runtime/CoreUObject/Private/UObject/Class.cpp
@@ -4525,6 +4525,74 @@ IMPLEMENT_CORE_INTRINSIC_CLASS(UDelegateFunction, UFunction,
 	}
 );
 
+/*-----------------------------------------------------------------------------
+UDynamicClass constructors.
+-----------------------------------------------------------------------------*/
+
+/**
+* Internal constructor.
+*/
+UDynamicClass::UDynamicClass(const FObjectInitializer& ObjectInitializer)
+: UClass(ObjectInitializer)
+{
+	// If you add properties here, please update the other constructors and PurgeClass()
+}
+
+/**
+* Create a new UDynamicClass given its superclass.
+*/
+UDynamicClass::UDynamicClass(const FObjectInitializer& ObjectInitializer, UClass* InBaseClass)
+: UClass(ObjectInitializer, InBaseClass)
+{
+}
+
+/**
+* Called when dynamically linked.
+*/
+UDynamicClass::UDynamicClass(
+	EStaticConstructor,
+	FName			InName,
+	uint32			InSize,
+	uint32			InClassFlags,
+	EClassCastFlags	InClassCastFlags,
+	const TCHAR*    InConfigName,
+	EObjectFlags	InFlags,
+	ClassConstructorType InClassConstructor,
+#if WITH_HOT_RELOAD_CTORS
+	ClassVTableHelperCtorCallerType InClassVTableHelperCtorCaller,
+#endif // WITH_HOT_RELOAD_CTORS
+	ClassAddReferencedObjectsType InClassAddReferencedObjects)
+: UClass(
+  EC_StaticConstructor
+, InName
+, InSize
+, InClassFlags
+, InClassCastFlags
+, InConfigName
+, InFlags
+, InClassConstructor
+#if WITH_HOT_RELOAD_CTORS
+, InClassVTableHelperCtorCaller
+#endif // WITH_HOT_RELOAD_CTORS
+, InClassAddReferencedObjects)
+{
+}
+
+void UDynamicClass::AddReferencedObjects(UObject* InThis, FReferenceCollector& Collector)
+{
+	UClass* This = CastChecked<UDynamicClass>(InThis);
+
+	//Collector.AddReferencedObjects(This->ConvertedSubobjectsFromBPGC, This);
+
+	Super::AddReferencedObjects(This, Collector);
+}
+
+
+IMPLEMENT_CORE_INTRINSIC_CLASS(UDynamicClass, UClass,
+{
+	Class->ClassAddReferencedObjects = &UDynamicClass::AddReferencedObjects;
+}
+);
 
 #if _MSC_VER == 1900
 	#ifdef PRAGMA_ENABLE_SHADOW_VARIABLE_WARNINGS
diff --git a/Engine/Source/Runtime/CoreUObject/Public/UObject/Class.h b/Engine/Source/Runtime/CoreUObject/Public/UObject/Class.h
--- a/Engine/Source/Runtime/CoreUObject/Public/UObject/Class.h
+++ b/Engine/Source/Runtime/CoreUObject/Public/UObject/Class.h
@@ -2372,6 +2372,32 @@ protected:
 	virtual UObject* CreateDefaultObject();
 };
 
+/**
+* Dynamic class (can be constructed after initial startup)
+*/
+class COREUOBJECT_API UDynamicClass : public UClass
+{
+	DECLARE_CASTED_CLASS_INTRINSIC_NO_CTOR(UDynamicClass, UClass, 0, CoreUObject, CASTCLASS_None, NO_API)
+	DECLARE_WITHIN(UPackage)
+
+public:
+
+	UDynamicClass(const FObjectInitializer& ObjectInitializer = FObjectInitializer::Get());
+	explicit UDynamicClass(const FObjectInitializer& ObjectInitializer, UClass* InSuperClass);
+	UDynamicClass(EStaticConstructor, FName InName, uint32 InSize, uint32 InClassFlags, EClassCastFlags InClassCastFlags,
+		const TCHAR* InClassConfigName, EObjectFlags InFlags, ClassConstructorType InClassConstructor,
+#if WITH_HOT_RELOAD_CTORS
+		ClassVTableHelperCtorCallerType InClassVTableHelperCtorCaller,
+#endif // WITH_HOT_RELOAD_CTORS
+		ClassAddReferencedObjectsType InClassAddReferencedObjects);
+
+	// UObject interface.
+	static void AddReferencedObjects(UObject* InThis, FReferenceCollector& Collector);
+
+	/** Misc objects owned by the class. Mostly templates from BP-generated native class */
+	//TArray<UObject*> ConvertedSubobjectsFromBPGC;
+};
+
 /**
  * Helper template to call the default constructor for a class
  */
@@ -2451,40 +2477,60 @@ void GetPrivateStaticClassBody( const TCHAR* PackageName, const TCHAR* Name, UCl
 	}
 #endif
 
-	EObjectFlags Flags = EObjectFlags(RF_Public | RF_Standalone | RF_Transient | RF_Native);
 	if (!bIsDynamic)
 	{
-		Flags = EObjectFlags(Flags | RF_RootSet);
+		ReturnClass = ::new (GUObjectAllocator.AllocateUObject(sizeof(UClass), ALIGNOF(UClass), true))
+			UClass
+			(
+			EC_StaticConstructor,
+			Name,
+			sizeof(TClass),
+			TClass::StaticClassFlags,
+			TClass::StaticClassCastFlags(),
+			TClass::StaticConfigName(),
+			EObjectFlags(RF_Public | RF_Standalone | RF_Transient | RF_Native | RF_RootSet),
+			(UClass::ClassConstructorType)InternalConstructor<TClass>,
+#if WITH_HOT_RELOAD_CTORS
+			(UClass::ClassVTableHelperCtorCallerType)InternalVTableHelperCtorCaller<TClass>,
+#endif // WITH_HOT_RELOAD_CTORS
+			&TClass::AddReferencedObjects
+			);
+		check(ReturnClass);
+		InitializePrivateStaticClass(
+			TClass::Super::StaticClass(),
+			ReturnClass,
+			TClass::WithinClass::StaticClass(),
+			PackageName,
+			Name
+			);
 	}
 	else
 	{
-		Flags = EObjectFlags(Flags | RF_Dynamic);
-	}
-
-	ReturnClass = ::new (GUObjectAllocator.AllocateUObject(sizeof(UClass),ALIGNOF(UClass),true)) 
-		UClass
-		(
-		EC_StaticConstructor,
-		Name,
-		sizeof(TClass),
-		TClass::StaticClassFlags,
-		TClass::StaticClassCastFlags(),
-		TClass::StaticConfigName(),
-		Flags,
-		(UClass::ClassConstructorType)InternalConstructor<TClass>,
+		ReturnClass = ::new (GUObjectAllocator.AllocateUObject(sizeof(UDynamicClass), ALIGNOF(UDynamicClass), true))
+			UDynamicClass
+			(
+			EC_StaticConstructor,
+			Name,
+			sizeof(TClass),
+			TClass::StaticClassFlags,
+			TClass::StaticClassCastFlags(),
+			TClass::StaticConfigName(),
+			EObjectFlags(RF_Public | RF_Standalone | RF_Transient | RF_Native | RF_Dynamic),
+			(UClass::ClassConstructorType)InternalConstructor<TClass>,
 #if WITH_HOT_RELOAD_CTORS
-		(UClass::ClassVTableHelperCtorCallerType)InternalVTableHelperCtorCaller<TClass>,
+			(UClass::ClassVTableHelperCtorCallerType)InternalVTableHelperCtorCaller<TClass>,
 #endif // WITH_HOT_RELOAD_CTORS
-		&TClass::AddReferencedObjects
-		);
-	check(ReturnClass);
-	InitializePrivateStaticClass(
-		TClass::Super::StaticClass(),
-		ReturnClass,
-		TClass::WithinClass::StaticClass(),
-		PackageName,
-		Name
-		);
+			&TClass::AddReferencedObjects
+			);
+		check(ReturnClass);
+		InitializePrivateStaticClass(
+			TClass::Super::StaticClass(),
+			ReturnClass,
+			TClass::WithinClass::StaticClass(),
+			PackageName,
+			Name
+			);
+	}
 
 	// Register the class's native functions.
 	RegisterNativeFunc();
