#ue4 - Improved FWeakObjectPtr::Get() performance. Avoiding redundant work (namely multiple TUObjectArray accesses). Inlined some internal functions that were not being inlined. Removed some branches, especially in the common case when bEvenIfPendingKill=false.

--------

diff --git a/Engine/Source/Runtime/Core/Public/UObject/NameTypes.h b/Engine/Source/Runtime/Core/Public/UObject/NameTypes.h
--- a/Engine/Source/Runtime/Core/Public/UObject/NameTypes.h
+++ b/Engine/Source/Runtime/Core/Public/UObject/NameTypes.h
@@ -304,7 +304,7 @@ class TStaticIndirectArrayThreadSafeRead
 	 * Return a pointer to the pointer to a given element
 	 * @param Index The Index of an element we want to retrieve the pointer-to-pointer for
 	 **/
-	ElementType const* const* GetItemPtr(int32 Index) const
+	FORCEINLINE_DEBUGGABLE ElementType const* const* GetItemPtr(int32 Index) const
 	{
 		int32 ChunkIndex = Index / ElementsPerChunk;
 		int32 WithinChunkIndex = Index % ElementsPerChunk;
@@ -327,7 +327,7 @@ public:
 	 * Thread safe, but you know, someone might have added more elements before this even returns
 	 * @return	the number of elements in the array
 	**/
-	int32 Num() const
+	FORCEINLINE int32 Num() const
 	{
 		return NumElements;
 	}
@@ -337,7 +337,7 @@ public:
 	 * @param	Index	Index to test
 	 * @return	true, if this is a valid
 	**/
-	bool IsValidIndex(int32 Index) const
+	FORCEINLINE bool IsValidIndex(int32 Index) const
 	{
 		return Index < Num() && Index >= 0;
 	}
diff --git a/Engine/Source/Runtime/CoreUObject/Private/UObject/WeakObjectPtr.cpp b/Engine/Source/Runtime/CoreUObject/Private/UObject/WeakObjectPtr.cpp
--- a/Engine/Source/Runtime/CoreUObject/Private/UObject/WeakObjectPtr.cpp
+++ b/Engine/Source/Runtime/CoreUObject/Private/UObject/WeakObjectPtr.cpp
@@ -128,7 +128,7 @@ public:
 	 * @param Index - UObject Index
 	 * @return - the serial number for this UObject
 	 */
-	int32 GetSerialNumber(int32 Index)
+	FORCEINLINE_DEBUGGABLE int32 GetSerialNumber(int32 Index)
 	{
 		int32 Block, SubIndex;
 		GetBlock( Index, Block, SubIndex );
@@ -212,7 +212,7 @@ void FWeakObjectPtr::operator=(const class UObject *Object)
  * internal function to test for serial number matches
  * @return true if the serial number in this matches the central table
 **/
-bool FWeakObjectPtr::SerialNumbersMatch() const
+FORCEINLINE_DEBUGGABLE bool FWeakObjectPtr::SerialNumbersMatch() const
 {
 	checkSlow(ObjectSerialNumber > FSerialNumberBlock::START_SERIAL_NUMBER && ObjectIndex >= 0 && ObjectIndex < FSerialNumberBlock::MAX_UOBJECTS); // otherwise this is a corrupted weak pointer
 	int32 ActualSerialNumber = GSerialNumberManager.GetSerialNumber(ObjectIndex);
@@ -222,6 +222,12 @@ bool FWeakObjectPtr::SerialNumbersMatch() const
 
 
 bool FWeakObjectPtr::IsValid(bool bEvenIfPendingKill, bool bThreadsafeTest) const
+{
+	// This is the external function, so we just pass through to the internal inlined method.
+	return Internal_IsValid(bEvenIfPendingKill, bThreadsafeTest);
+}
+
+FORCEINLINE_DEBUGGABLE bool FWeakObjectPtr::Internal_IsValid(bool bEvenIfPendingKill, bool bThreadsafeTest) const
 {
 	if (ObjectSerialNumber == 0)
 	{
@@ -265,20 +271,32 @@ bool FWeakObjectPtr::IsStale(bool bEvenIfPendingKill, bool bThreadsafeTest) cons
 	return GetUObjectArray().IsStale(ObjectIndex, bEvenIfPendingKill);
 }
 
-UObject* FWeakObjectPtr::Get(bool bEvenIfPendingKill) const
+FORCEINLINE_DEBUGGABLE UObject* FWeakObjectPtr::Internal_Get(bool bEvenIfPendingKill) const
 {
 	UObject* Result = nullptr;
-	if (IsValid(true))
+
+	if (Internal_IsValid(true, true))
 	{
-		Result = (UObject*)(GetUObjectArray().IndexToObject(GetObjectIndex(), bEvenIfPendingKill));
+		Result = (UObject*)(GetUObjectArray().IndexToValidObject(GetObjectIndex(), bEvenIfPendingKill));
 	}
 	return Result;
 }
 
+UObject* FWeakObjectPtr::Get(/*bool bEvenIfPendingKill = false*/) const
+{
+	// Using a literal here allows the optimizer to remove branches later down the chain.
+	return Internal_Get(false);
+}
+
+UObject* FWeakObjectPtr::Get(bool bEvenIfPendingKill) const
+{
+	return Internal_Get(bEvenIfPendingKill);
+}
+
 UObject* FWeakObjectPtr::GetEvenIfUnreachable() const
 {
 	UObject* Result = nullptr;
-	if (IsValid(true, true))
+	if (Internal_IsValid(true, true))
 	{
 		Result = static_cast<UObject*>(GetUObjectArray().IndexToObject(GetObjectIndex(), true));
 	}
diff --git a/Engine/Source/Runtime/CoreUObject/Public/UObject/UObjectArray.h b/Engine/Source/Runtime/CoreUObject/Public/UObject/UObjectArray.h
--- a/Engine/Source/Runtime/CoreUObject/Public/UObject/UObjectArray.h
+++ b/Engine/Source/Runtime/CoreUObject/Public/UObject/UObjectArray.h
@@ -160,26 +160,28 @@ public:
 		return Object;
 	}
 
+	FORCEINLINE class UObjectBase* IndexToValidObject(int32 Index, bool bEvenIfPendingKill)
+	{
+		UObjectBaseUtility* Object = static_cast<UObjectBaseUtility*>(IndexToObject(Index));
+		return Internal_IsValid(Object, bEvenIfPendingKill) ? Object : nullptr;
+	}
+
 	FORCEINLINE bool IsValid(int32 Index, bool bEvenIfPendingKill)
 	{
 		// This method assumes Index points to a valid object.
 		UObjectBaseUtility* Object = static_cast<UObjectBaseUtility*>(IndexToObject(Index));
-		if(Object == NULL)
-		{
-			return false;
-		}
-		return !Object->HasAnyFlags(RF_Unreachable) && (bEvenIfPendingKill || !Object->IsPendingKill());
+		return Internal_IsValid(Object, bEvenIfPendingKill);
 	}
 
 	FORCEINLINE bool IsStale(int32 Index, bool bEvenIfPendingKill)
 	{
 		// This method assumes Index points to a valid object.
 		UObjectBaseUtility* Object = static_cast<UObjectBaseUtility*>(IndexToObject(Index));
-		if(Object == NULL)
+		if (Object)
 		{
-			return true;
+			return bEvenIfPendingKill ? Object->HasAnyFlags(RF_Unreachable | RF_PendingKill) : Object->HasAnyFlags(RF_Unreachable);
 		}
-		return Object->HasAnyFlags(RF_Unreachable) || (bEvenIfPendingKill && Object->IsPendingKill());
+		return true;
 	}
 
 	/**
@@ -381,6 +383,15 @@ public:
 
 private:
 
+	FORCEINLINE bool Internal_IsValid(UObjectBaseUtility* Object, bool bEvenIfPendingKill)
+	{
+		if (Object)
+		{
+			return bEvenIfPendingKill ? !Object->HasAnyFlags(RF_Unreachable) : !Object->HasAnyFlags(RF_Unreachable | RF_PendingKill);
+		}
+		return false;
+	}
+
 	typedef TStaticIndirectArrayThreadSafeRead<UObjectBase, 8 * 1024 * 1024 /* Max 8M UObjects */, 16384 /* allocated in 64K/128K chunks */ > TUObjectArray;
 
 	/**
diff --git a/Engine/Source/Runtime/CoreUObject/Public/UObject/WeakObjectPtr.h b/Engine/Source/Runtime/CoreUObject/Public/UObject/WeakObjectPtr.h
--- a/Engine/Source/Runtime/CoreUObject/Public/UObject/WeakObjectPtr.h
+++ b/Engine/Source/Runtime/CoreUObject/Public/UObject/WeakObjectPtr.h
@@ -99,11 +99,17 @@ public:
 
 
 	/**  
-	 * Dereference the weak pointer
+	 * Dereference the weak pointer.
 	 * @param bEvenIfPendingKill, if this is true, pendingkill objects are considered valid
 	 * @return NULL if this object is gone or the weak pointer was NULL, otherwise a UObject pointer
-	**/
-	COREUOBJECT_API class UObject *Get(bool bEvenIfPendingKill = false) const;
+	 */
+	COREUOBJECT_API class UObject *Get(bool bEvenIfPendingKill) const;
+
+	/**  
+	 * Dereference the weak pointer. This is an optimized version implying bEvenIfPendingKill=false.
+	 * @return NULL if this object is gone or the weak pointer was NULL, otherwise a UObject pointer
+	 */
+	COREUOBJECT_API class UObject *Get(/*bool bEvenIfPendingKill = false*/) const;
 
 	/** Dereference the weak pointer even if it is RF_PendingKill or RF_Unreachable */
 	COREUOBJECT_API class UObject *GetEvenIfUnreachable() const;
@@ -141,7 +147,7 @@ public:
 
 protected:
 
-	int32 GetObjectIndex() const
+	FORCEINLINE int32 GetObjectIndex() const
 	{
 		return ObjectIndex;
 	}
@@ -155,6 +161,12 @@ private:
 	**/
 	bool SerialNumbersMatch() const;
 
+	/** Private (inlined) version for internal use only. */
+	bool Internal_IsValid(bool bEvenIfPendingKill, bool bThreadsafeTest) const;
+
+	/** Private (inlined) version for internal use only. */
+	class UObject* Internal_Get(bool bEvenIfPendingKill) const;
+
 	int32		ObjectIndex;
 	int32		ObjectSerialNumber;
 };
