comment spelling fix

diff --git a/Engine/Source/Runtime/CoreUObject/Public/UObject/UObjectBase.h b/Engine/Source/Runtime/CoreUObject/Public/UObject/UObjectBase.h
--- a/Engine/Source/Runtime/CoreUObject/Public/UObject/UObjectBase.h
+++ b/Engine/Source/Runtime/CoreUObject/Public/UObject/UObjectBase.h
@@ -90,7 +90,7 @@ public:
 
 	/**
 	 * Faster version of IsValidLowLevel.
-	 * Checks to see if the object appears to be valid by checking pointers and their alingment.
+	 * Checks to see if the object appears to be valid by checking pointers and their alignment.
 	 * Name and InternalIndex checks are less accurate than IsValidLowLevel.
 	 * @param bRecursive true if the Class pointer should be checked with IsValidLowLevelFast
 	 * @return true if this appears to be a valid object
