Core: Added support for tracking the total number of objects outside of WITH_EDITOR by defining UE_GC_TRACK_OBJ_AVAILABLE to 1
[INTEGRATE] Change 2684069 by [REDACTED_PERSON][REDACTED_HANDLE] on 2015/09/08 21:44:59

diff --git a/Engine/Source/Runtime/CoreUObject/Private/UObject/UObjectArray.cpp b/Engine/Source/Runtime/CoreUObject/Private/UObject/UObjectArray.cpp
--- a/Engine/Source/Runtime/CoreUObject/Private/UObject/UObjectArray.cpp
+++ b/Engine/Source/Runtime/CoreUObject/Private/UObject/UObjectArray.cpp
@@ -54,7 +54,7 @@ void FUObjectArray::AllocateUObjectIndex(UObjectBase* Object, bool bMergingThrea
 		int32* AvailableIndex = ObjAvailableList.Pop();
 		if (AvailableIndex)
 		{
-#if WITH_EDITOR
+#if UE_GC_TRACK_OBJ_AVAILABLE
 			const int32 AvailableCount = ObjAvailableCount.Decrement();
 			checkSlow(AvailableCount >= 0);
 #endif
@@ -112,7 +112,7 @@ void FUObjectArray::FreeUObjectIndex(UObjectBase* Object)
 	if (Index > ObjLastNonGCIndex && !GExitPurge)  
 	{
 		ObjAvailableList.Push((int32*)(uintptr_t)Index);
-#if WITH_EDITOR
+#if UE_GC_TRACK_OBJ_AVAILABLE
 		ObjAvailableCount.Increment();
 #endif
 	}
diff --git a/Engine/Source/Runtime/CoreUObject/Private/UObject/UObjectHash.cpp b/Engine/Source/Runtime/CoreUObject/Private/UObject/UObjectHash.cpp
--- a/Engine/Source/Runtime/CoreUObject/Private/UObject/UObjectHash.cpp
+++ b/Engine/Source/Runtime/CoreUObject/Private/UObject/UObjectHash.cpp
@@ -8,7 +8,6 @@
 
 DEFINE_LOG_CATEGORY_STATIC(LogUObjectHash, Log, All);
 
-DECLARE_STATS_GROUP(TEXT("UObject Hash"), STATGROUP_UObjectHash, STATCAT_Advanced);
 DECLARE_CYCLE_STAT(TEXT("StaticFindObjectFastInternal"), STAT_Hash_StaticFindObjectFastInternal, STATGROUP_UObjectHash);
 DECLARE_CYCLE_STAT( TEXT( "StaticFindObjectFastExplicit" ), STAT_Hash_StaticFindObjectFastExplicit, STATGROUP_UObjectHash );
 DECLARE_CYCLE_STAT( TEXT( "GetObjectsWithOuter" ), STAT_Hash_GetObjectsWithOuter, STATGROUP_UObjectHash );
@@ -17,6 +16,10 @@ DECLARE_CYCLE_STAT( TEXT( "GetObjectsOfClass" ), STAT_Hash_GetObjectsOfClass, ST
 DECLARE_CYCLE_STAT( TEXT( "HashObject" ), STAT_Hash_HashObject, STATGROUP_UObjectHash );
 DECLARE_CYCLE_STAT( TEXT( "UnhashObject" ), STAT_Hash_UnhashObject, STATGROUP_UObjectHash );
 
+#if UE_GC_TRACK_OBJ_AVAILABLE
+DEFINE_STAT( STAT_Hash_NumObjects );
+#endif
+
 /**
  * This implementation will use more space than the UE3 implementation. The goal was to make UObjects smaller to save L2 cache space. 
  * The hash is rarely used at runtime. A more space-efficient implementation is possible.
diff --git a/Engine/Source/Runtime/CoreUObject/Public/UObject/UObjectArray.h b/Engine/Source/Runtime/CoreUObject/Public/UObject/UObjectArray.h
--- a/Engine/Source/Runtime/CoreUObject/Public/UObject/UObjectArray.h
+++ b/Engine/Source/Runtime/CoreUObject/Public/UObject/UObjectArray.h
@@ -266,7 +266,7 @@ public:
 		return ObjLastNonGCIndex + 1;
 	}
 
-#if WITH_EDITOR
+#if UE_GC_TRACK_OBJ_AVAILABLE
 	/**
 	 * Returns the number of actual object indices that are claimed (the total size of the global object array minus
 	 * the number of available object array elements
@@ -413,7 +413,7 @@ private:
 	FCriticalSection ObjObjectsCritical;
 	/** Available object indices.											*/
 	TLockFreePointerListUnordered<int32> ObjAvailableList;
-#if WITH_EDITOR
+#if UE_GC_TRACK_OBJ_AVAILABLE
 	/** Available object index count.										*/
 	FThreadSafeCounter ObjAvailableCount;
 #endif
diff --git a/Engine/Source/Runtime/CoreUObject/Public/UObject/UObjectHash.h b/Engine/Source/Runtime/CoreUObject/Public/UObject/UObjectHash.h
--- a/Engine/Source/Runtime/CoreUObject/Public/UObject/UObjectHash.h
+++ b/Engine/Source/Runtime/CoreUObject/Public/UObject/UObjectHash.h
@@ -7,6 +7,17 @@
 #ifndef __UOBJECTHASH_H__
 #define __UOBJECTHASH_H__
 
+DECLARE_STATS_GROUP(TEXT("UObject Hash"), STATGROUP_UObjectHash, STATCAT_Advanced);
+
+/**
+ * Controls whether the number of available elements is being tracked in the ObjObjects array.
+ * By default it is only tracked in WITH_EDITOR builds as it adds a small amount of tracking overhead
+ */
+#define UE_GC_TRACK_OBJ_AVAILABLE (WITH_EDITOR)
+
+#if UE_GC_TRACK_OBJ_AVAILABLE
+DECLARE_DWORD_COUNTER_STAT_EXTERN(TEXT("NumObjects"), STAT_Hash_NumObjects, STATGROUP_UObjectHash, COREUOBJECT_API);
+#endif
 
 /**
  * Private internal version of StaticFindObjectFast that allows using 0 exclusion flags.
diff --git a/Engine/Source/Runtime/Launch/Private/LaunchEngineLoop.cpp b/Engine/Source/Runtime/Launch/Private/LaunchEngineLoop.cpp
--- a/Engine/Source/Runtime/Launch/Private/LaunchEngineLoop.cpp
+++ b/Engine/Source/Runtime/Launch/Private/LaunchEngineLoop.cpp
@@ -2604,6 +2604,11 @@ void FEngineLoop::Tick()
 		const FCPUTime CPUTime = FPlatformTime::GetCPUTime();
 		SET_FLOAT_STAT( STAT_CPUTimePct, CPUTime.CPUTimePct );
 		SET_FLOAT_STAT( STAT_CPUTimePctRelative, CPUTime.CPUTimePctRelative );
+
+		// Set the UObject count stat
+#if UE_GC_TRACK_OBJ_AVAILABLE
+		SET_DWORD_STAT(STAT_Hash_NumObjects, GetUObjectArray().GetObjectArrayNumMinusAvailable());
+#endif
 	}
 }
 
