#ue4 - Inline FSerialNumberManager::GetBlock()

--------

diff --git a/Engine/Source/Runtime/CoreUObject/Private/UObject/WeakObjectPtr.cpp b/Engine/Source/Runtime/CoreUObject/Private/UObject/WeakObjectPtr.cpp
--- a/Engine/Source/Runtime/CoreUObject/Private/UObject/WeakObjectPtr.cpp
+++ b/Engine/Source/Runtime/CoreUObject/Private/UObject/WeakObjectPtr.cpp
@@ -68,7 +68,7 @@ public:
 	 * @param OutBlock - Serial Number Block
 	 * @param OutSubIndex - Serial Number SubIndex
 	 */
-	void GetBlock(const int32& InIndex, int32& OutBlock, int32& OutSubIndex) const
+	FORCEINLINE_DEBUGGABLE void GetBlock(const int32& InIndex, int32& OutBlock, int32& OutSubIndex) const
 	{
 		OutBlock = InIndex / FSerialNumberBlock::SERIAL_NUMBER_BLOCK_SIZE;
 		OutSubIndex = InIndex - OutBlock * FSerialNumberBlock::SERIAL_NUMBER_BLOCK_SIZE;
