[UE-20878] Ensuring that the parent CDO is FULLY serialized by the time we load the child's (handles cyclic dependency load scenarios where the super is in the midst of serializing its CDO, and in turn seeks to load a sub-class's CDO).

diff --git a/Engine/Source/Runtime/CoreUObject/Private/Blueprint/BlueprintSupport.cpp b/Engine/Source/Runtime/CoreUObject/Private/Blueprint/BlueprintSupport.cpp
--- a/Engine/Source/Runtime/CoreUObject/Private/Blueprint/BlueprintSupport.cpp
+++ b/Engine/Source/Runtime/CoreUObject/Private/Blueprint/BlueprintSupport.cpp
@@ -1061,6 +1061,19 @@ void FLinkerLoad::FinalizeBlueprint(UClass* LoadClass)
 		FLinkerLoad* SuperLinker = SuperClass->GetLinker();
 		if ((SuperLinker != nullptr) && SuperLinker->IsBlueprintFinalizationPending())
 		{
+			DEFERRED_DEPENDENCY_CHECK(SuperLinker->DeferredCDOIndex != INDEX_NONE);
+			UObject* SuperCDO = SuperLinker->ExportMap[SuperLinker->DeferredCDOIndex].Object;
+			// we MUST have the super fully serialized before we can finalize  
+			// this (class and CDO); if the SuperCDO is already in the midst of 
+			// serializing somewhere up the stack (and a cyclic dependency has  
+			// landed us here, finalizing one of it's children), then it is 
+			// paramount that we force it through serialization (so we reset the 
+			// RF_NeedLoad guard, and leave it to ResolveDeferredExports, for it
+			// to re-run the serialization)
+			if ( (SuperCDO != nullptr) && !SuperCDO->HasAnyFlags(RF_NeedLoad|RF_LoadCompleted) )
+			{
+				SuperCDO->SetFlags(RF_NeedLoad);
+			}
 			SuperLinker->FinalizeBlueprint(SuperClass);
 		}
 	}
@@ -1497,7 +1510,8 @@ FObjectInitializer* FDeferredObjInitializerTracker::Add(const FObjectInitializer
 		SuperClassMap.AddUnique(SuperClass, LoadClass);
 #if USE_DEFERRED_DEPENDENCY_CHECK_VERIFICATION_TESTS
 		UObject* SuperCDO = SuperClass->GetDefaultObject(/*bCreateIfNeeded =*/false);
-		DEFERRED_DEPENDENCY_CHECK( SuperCDO && (SuperCDO->HasAnyFlags(RF_NeedLoad) || SuperCDO->HasAnyFlags(RF_LoadCompleted) || IsCdoDeferred(SuperClass)) );
+		DEFERRED_DEPENDENCY_CHECK( SuperCDO && (SuperCDO->HasAnyFlags(RF_NeedLoad) ||
+			(SuperClass->GetLinker() && SuperClass->GetLinker()->IsBlueprintFinalizationPending()) || IsCdoDeferred(SuperClass)) );
 #endif // USE_DEFERRED_DEPENDENCY_CHECK_VERIFICATION_TESTS
 		
 		auto& InitializersCache = ThreadInst.DeferredInitializers;
diff --git a/Engine/Source/Runtime/CoreUObject/Private/UObject/LinkerPlaceholderBase.cpp b/Engine/Source/Runtime/CoreUObject/Private/UObject/LinkerPlaceholderBase.cpp
--- a/Engine/Source/Runtime/CoreUObject/Private/UObject/LinkerPlaceholderBase.cpp
+++ b/Engine/Source/Runtime/CoreUObject/Private/UObject/LinkerPlaceholderBase.cpp
@@ -253,7 +253,7 @@ int32 FLinkerPlaceholderBase::ResolvePlaceholderPropertyValues(UObject* NewObjec
 		for (const UObjectProperty* Property : ReferencingPair.Value)
 		{
 #if USE_DEFERRED_DEPENDENCY_CHECK_VERIFICATION_TESTS
-			check(Property->GetOwnerClass() == Container->GetClass());
+			check(Container->GetClass()->IsChildOf(Property->GetOwnerClass()));
 #endif // USE_DEFERRED_DEPENDENCY_CHECK_VERIFICATION_TESTS
 
 			TArray<const UProperty*> PropertyChain;
