Merging CL #2685024

Object iterator optimizations

diff --git a/Engine/Source/Runtime/Core/Public/UObject/NameTypes.h b/Engine/Source/Runtime/Core/Public/UObject/NameTypes.h
--- a/Engine/Source/Runtime/Core/Public/UObject/NameTypes.h
+++ b/Engine/Source/Runtime/Core/Public/UObject/NameTypes.h
@@ -347,7 +347,7 @@ public:
 	 * @return	a reference to the pointer to the element
 	 * Thread safe, if it is valid now, it is valid forever. This might return nullptr, but by then, some other thread might have made it non-nullptr.
 	**/
-	ElementType const* const& operator[](int32 Index) const
+	FORCEINLINE ElementType const* const& operator[](int32 Index) const
 	{
 		ElementType const* const* ItemPtr = GetItemPtr(Index);
 		check(ItemPtr);
diff --git a/Engine/Source/Runtime/CoreUObject/Public/UObject/UObjectArray.h b/Engine/Source/Runtime/CoreUObject/Public/UObject/UObjectArray.h
--- a/Engine/Source/Runtime/CoreUObject/Public/UObject/UObjectArray.h
+++ b/Engine/Source/Runtime/CoreUObject/Public/UObject/UObjectArray.h
@@ -286,7 +286,7 @@ public:
 
 
 	/**
-	 * Low level iterator
+	 * Low level iterator.
 	 */
 	class TIterator
 	{
@@ -304,7 +304,8 @@ public:
 		 */
 		TIterator( const FUObjectArray& InArray, bool bOnlyGCedObjects = false ) :	
 			Array(InArray),
-			Index(-1)
+			Index(-1),
+			CurrentObject(nullptr)
 		{
 			if (bOnlyGCedObjects)
 			{
@@ -339,7 +340,7 @@ public:
 		/** Conversion to "bool" returning true if the iterator is valid. */
 		FORCEINLINE_EXPLICIT_OPERATOR_BOOL() const
 		{ 
-			return Array.ObjObjects.IsValidIndex(Index); 
+			return !!CurrentObject;
 		}
 		/** inverse of the "bool" operator */
 		FORCEINLINE bool operator !() const 
@@ -356,7 +357,7 @@ public:
 		 */
 		FORCEINLINE UObjectBase* GetObject() const 
 		{ 
-			return (UObjectBase*)Array.ObjObjects[Index];
+			return CurrentObject;
 		}
 		/**
 		 * Iterator advance with ordinary name for clarity in subclasses
@@ -365,9 +366,11 @@ public:
 		FORCEINLINE bool Advance()
 		{
 			//@todo UE4 check this for LHS on Index on consoles
+			CurrentObject = nullptr;
 			while(++Index < Array.GetObjectArrayNum())
 			{
-				if (GetObject())
+				CurrentObject = (UObjectBase*)Array.ObjObjects[Index];
+				if (CurrentObject)
 				{
 					return true;
 				}
@@ -379,6 +382,8 @@ public:
 		const FUObjectArray& Array;
 		/** index of the current element in the object array */
 		int32 Index;
+		/** Current object */
+		mutable UObjectBase* CurrentObject;
 	};
 
 private:
