missed file

diff --git a/Engine/Source/Runtime/CoreUObject/Public/UObject/Class.h b/Engine/Source/Runtime/CoreUObject/Public/UObject/Class.h
--- a/Engine/Source/Runtime/CoreUObject/Public/UObject/Class.h
+++ b/Engine/Source/Runtime/CoreUObject/Public/UObject/Class.h
@@ -2022,6 +2022,14 @@ public:
 		FuncMap.Add(NewFunction->GetFName(), NewFunction);
 	}
 
+	// This is used by the code generator, which instantiates UFunctions with a name that is later overridden. Overridden names
+	// are needed to support generated versions of blueprint classes, properties of which do not have the same naming 
+	// restrictions as native C++ properties.
+	void AddFunctionToFunctionMapWithOverriddenName(UFunction* NewFunction, FName OverriddenName)
+	{
+		FuncMap.Add(OverriddenName, NewFunction);
+	}
+
 	// Remove a function from the function map
 	void RemoveFunctionFromFunctionMap(UFunction* Function)
 	{
