Name mangling logic for blueprint generated classes. This solves problems with ascii characters like '-', '+' and '=', unicode characters, and that match existing c++ identifiers (IE naming a blueprint variable bool or U8) showing up in code generated from blueprints.

diff --git a/Engine/Source/Developer/BlueprintCompilerCppBackend/Private/BlueprintCompilerCppBackendUtils.cpp b/Engine/Source/Developer/BlueprintCompilerCppBackend/Private/BlueprintCompilerCppBackendUtils.cpp
--- a/Engine/Source/Developer/BlueprintCompilerCppBackend/Private/BlueprintCompilerCppBackendUtils.cpp
+++ b/Engine/Source/Developer/BlueprintCompilerCppBackend/Private/BlueprintCompilerCppBackendUtils.cpp
@@ -96,8 +96,20 @@ FString FEmitHelper::GetCppName(const UField* Field)
 	}
 	else if (auto AsProperty = Cast<UProperty>(Field))
 	{
+		if (const UStruct* Owner = AsProperty->GetOwnerStruct())
+		{
+			if ((Cast<UBlueprintGeneratedClass>(Owner) ||
+				!Owner->HasAnyFlags(RF_Native)))
+			{
+				return ::UnicodeToCPPIdentifier(AsProperty->GetName(), AsProperty->HasAnyPropertyFlags(CPF_Deprecated), TEXT("bpv__"));
+			}
+		}
 		return AsProperty->GetNameCPP();
 	}
+	if (!Field->HasAnyFlags(RF_Native))
+	{
+		return ::UnicodeToCPPIdentifier(Field->GetName(), false, TEXT("bpf__"));
+	}
 	return Field->GetName();
 }
 
@@ -177,6 +189,10 @@ FString FEmitHelper::HandleMetaData(const UField* Field, bool AddCategory, TArra
 	{
 		MetaDataStrings.Append(*AdditinalMetaData);
 	}
+	if (Field)
+	{
+		MetaDataStrings.Emplace(FString::Printf(TEXT("OverrideNativeName=\"%s\""), *Field->GetName()));
+	}
 	MetaDataStrings.Remove(FString());
 	if (MetaDataStrings.Num())
 	{
diff --git a/Engine/Source/Runtime/CoreUObject/Private/UObject/Property.cpp b/Engine/Source/Runtime/CoreUObject/Private/UObject/Property.cpp
--- a/Engine/Source/Runtime/CoreUObject/Private/UObject/Property.cpp
+++ b/Engine/Source/Runtime/CoreUObject/Private/UObject/Property.cpp
@@ -485,7 +485,18 @@ void UProperty::ExportCppDeclaration(FOutputDevice& Out, EExportedDeclaration::T
 		}
 	}
 
-	FString NameCpp = bSkipParameterName ? FString() : GetNameCPP();
+	FString NameCpp;
+	if (!bSkipParameterName)
+	{
+		if (AdditionalExportCPPFlags & CPPF_BlueprintCppBackend)
+		{
+			NameCpp = UnicodeToCPPIdentifier(GetName(), HasAnyPropertyFlags(CPF_Deprecated), TEXT("bpv__"));
+		}
+		else
+		{
+			NameCpp = GetNameCPP();
+		}
+	}
 	if (DeclarationType == EExportedDeclaration::MacroParameter)
 	{
 		NameCpp = FString(TEXT(", ")) + NameCpp;
diff --git a/Engine/Source/Runtime/CoreUObject/Private/UObject/ScriptCore.cpp b/Engine/Source/Runtime/CoreUObject/Private/UObject/ScriptCore.cpp
--- a/Engine/Source/Runtime/CoreUObject/Private/UObject/ScriptCore.cpp
+++ b/Engine/Source/Runtime/CoreUObject/Private/UObject/ScriptCore.cpp
@@ -142,6 +142,69 @@ FEditorScriptExecutionGuard::~FEditorScriptExecutionGuard()
 	GAllowActorScriptExecutionInEditor = bOldGAllowScriptExecutionInEditor;
 }
 
+FString UnicodeToCPPIdentifier(const FString& InName, bool bDeprecated, const TCHAR* Prefix)
+{
+	// FName's can contain unicode characters or collide with other CPP identifiers or keywords. This function 
+	// returns a string that will have a prefix which is unlikely to collide with existing identifiers and
+	// converts unicode characters in place to valid ascii characters. Strictly speaking a C++ compiler *could*
+	// support unicode identifiers in source files, but I am not comfortable relying on this behavior.
+	
+	auto IsValidIdentifierChar = [](TCHAR Char)
+	{
+		return Char == TCHAR('_')
+			|| (Char >= TCHAR('a') && Char <= TCHAR('z'))
+			|| (Char >= TCHAR('A') && Char <= TCHAR('Z'))
+			|| (Char >= TCHAR('0') && Char <= TCHAR('9'));
+	};
+
+	auto ToValidIdentifier = [](int32 Val) -> TCHAR
+	{
+		if (Val < 26)
+		{
+			return TCHAR(TCHAR('a') + (26 - Val));
+		}
+		else if (Val < 52)
+		{
+			return TCHAR(TCHAR('A') + (52 - Val));
+		}
+		else if (Val < 62)
+		{
+			return TCHAR(TCHAR('0') + (62 - Val));
+		}
+		else
+		{
+			check(Val == 62);
+			return TCHAR('_');
+		}
+	};
+
+	FString Ret = InName;
+	// Initialize postfix with a unique identifier. This prevents potential collisions between names that have unicode
+	// characters and those that do not. The drawback is that it is not safe to put '__pf' in a blueprint name.
+	FString Postfix = TEXT("__pf");
+	for (auto& Char : Ret)
+	{
+		// if the character is not a valid character for a c++ identifier, then we need to encode it using valid characters:
+		if (!IsValidIdentifierChar(Char))
+		{
+			// deterministically map char to a valid ascii character, we have 63 characters available (aA-zZ, 0-9, and _)
+			// so the optimal encoding would be base 63:
+			int32 RawValue = Char;
+			int32 Counter = 0;
+			while (RawValue != 0)
+			{
+				int32 Digit = RawValue % 63;
+				RawValue = (RawValue - Digit) / 63;
+				Postfix.AppendChar(ToValidIdentifier(Digit));
+			}
+			Char = TCHAR('x');
+		}
+	}
+
+	Ret = FString(Prefix) + Ret + Postfix;
+	return bDeprecated ? Ret + TEXT("_DEPRECATED") : Ret;
+}
+
 /*-----------------------------------------------------------------------------
 	FFrame implementation.
 -----------------------------------------------------------------------------*/
diff --git a/Engine/Source/Runtime/CoreUObject/Public/UObject/Script.h b/Engine/Source/Runtime/CoreUObject/Public/UObject/Script.h
--- a/Engine/Source/Runtime/CoreUObject/Public/UObject/Script.h
+++ b/Engine/Source/Runtime/CoreUObject/Public/UObject/Script.h
@@ -323,3 +323,5 @@ public:
 private:
 	bool bOldGAllowScriptExecutionInEditor;
 };
+
+COREUOBJECT_API FString UnicodeToCPPIdentifier(const FString& InName, bool bDeprecated, const TCHAR* Prefix);
