Merging CL #2686175

Changing GetUObjectArray() -> GUObjectArray to avoid redundant function calls in places like weak pointers/GC/object iterators

diff --git a/Engine/Plugins/Experimental/LiveEditor/Source/LiveEditor/Private/LiveEditor.cpp b/Engine/Plugins/Experimental/LiveEditor/Source/LiveEditor/Private/LiveEditor.cpp
--- a/Engine/Plugins/Experimental/LiveEditor/Source/LiveEditor/Private/LiveEditor.cpp
+++ b/Engine/Plugins/Experimental/LiveEditor/Source/LiveEditor/Private/LiveEditor.cpp
@@ -195,10 +195,10 @@ void FLiveEditor::InstallHooks()
 	FEdGraphUtilities::RegisterVisualPinFactory(PinFactory);
 
 	ObjectCreationListener = new FLiveEditorObjectCreateListener();
-	GetUObjectArray().AddUObjectCreateListener(ObjectCreationListener);
+	GUObjectArray.AddUObjectCreateListener(ObjectCreationListener);
 
 	ObjectDeletionListener = new FLiveEditorObjectDeleteListener();
-	GetUObjectArray().AddUObjectDeleteListener(ObjectDeletionListener);
+	GUObjectArray.AddUObjectDeleteListener(ObjectDeletionListener);
 
 	FGlobalTabmanager::Get()->RegisterNomadTabSpawner(LiveEditorModule::LiveEditorApp, FOnSpawnTab::CreateStatic(&SpawnLiveEditorTab))
 		.SetDisplayName(NSLOCTEXT("LiveEditorPlugin", "TabTitle", "Live Editor"))
@@ -212,8 +212,8 @@ void FLiveEditor::RemoveHooks()
 
 	FEdGraphUtilities::UnregisterVisualPinFactory(PinFactory);
 
-	GetUObjectArray().RemoveUObjectCreateListener(ObjectCreationListener);
-	GetUObjectArray().RemoveUObjectDeleteListener(ObjectDeletionListener);
+	GUObjectArray.RemoveUObjectCreateListener(ObjectCreationListener);
+	GUObjectArray.RemoveUObjectDeleteListener(ObjectDeletionListener);
 
 	FGlobalTabmanager::Get()->UnregisterNomadTabSpawner(LiveEditorModule::LiveEditorApp);
 }
diff --git a/Engine/Plugins/Runtime/LiveEditorListenServer/Source/LiveEditorListenServer/Private/LiveEditorListenServer.cpp b/Engine/Plugins/Runtime/LiveEditorListenServer/Source/LiveEditorListenServer/Private/LiveEditorListenServer.cpp
--- a/Engine/Plugins/Runtime/LiveEditorListenServer/Source/LiveEditorListenServer/Private/LiveEditorListenServer.cpp
+++ b/Engine/Plugins/Runtime/LiveEditorListenServer/Source/LiveEditorListenServer/Private/LiveEditorListenServer.cpp
@@ -221,10 +221,10 @@ void FLiveEditorListenServer::ShutdownModule()
 void FLiveEditorListenServer::InstallHooks()
 {
 	ObjectCreateListener = new nLiveEditorListenServer::FCreateListener(this);
-	GetUObjectArray().AddUObjectCreateListener(ObjectCreateListener);
+	GUObjectArray.AddUObjectCreateListener(ObjectCreateListener);
 
 	ObjectDeleteListener = new nLiveEditorListenServer::FDeleteListener(this);
-	GetUObjectArray().AddUObjectDeleteListener(ObjectDeleteListener);
+	GUObjectArray.AddUObjectDeleteListener(ObjectDeleteListener);
 
 	TickObject = new nLiveEditorListenServer::FTickObject(this);
 
diff --git a/Engine/Source/Editor/MainFrame/Private/MainFrameModule.cpp b/Engine/Source/Editor/MainFrame/Private/MainFrameModule.cpp
--- a/Engine/Source/Editor/MainFrame/Private/MainFrameModule.cpp
+++ b/Engine/Source/Editor/MainFrame/Private/MainFrameModule.cpp
@@ -272,7 +272,7 @@ TSharedRef<SWidget> FMainFrameModule::MakeDeveloperTools() const
 
 		static FText GetUObjectCountAsString() 
 		{
-			return FText::AsNumber(GetUObjectArray().GetObjectArrayNumMinusAvailable());
+			return FText::AsNumber(GUObjectArray.GetObjectArrayNumMinusAvailable());
 		}
 
 		static void OpenVideo( FString SourceFilePath )
diff --git a/Engine/Source/Runtime/CoreUObject/Private/Serialization/ArchiveObjectGraph.cpp b/Engine/Source/Runtime/CoreUObject/Private/Serialization/ArchiveObjectGraph.cpp
--- a/Engine/Source/Runtime/CoreUObject/Private/Serialization/ArchiveObjectGraph.cpp
+++ b/Engine/Source/Runtime/CoreUObject/Private/Serialization/ArchiveObjectGraph.cpp
@@ -17,8 +17,8 @@ FArchiveObjectGraph::FArchiveObjectGraph(bool IncludeTransients, EObjectFlags	Ke
 	TArray<UObject*> RootObjects;
 
 	// allocate enough memory for all objects
-	ObjectGraph.Empty(GetUObjectArray().GetObjectArrayNum());
-	RootObjects.Empty(GetUObjectArray().GetObjectArrayNum());
+	ObjectGraph.Empty(GUObjectArray.GetObjectArrayNum());
+	RootObjects.Empty(GUObjectArray.GetObjectArrayNum());
 
 	// search for objects that have the right flags and add them to the list of objects that we're going to start with
 	// all other objects need to be tagged so that we can tell whether they've been serialized or not.
diff --git a/Engine/Source/Runtime/CoreUObject/Private/Serialization/ArchiveTraceRoute.cpp b/Engine/Source/Runtime/CoreUObject/Private/Serialization/ArchiveTraceRoute.cpp
--- a/Engine/Source/Runtime/CoreUObject/Private/Serialization/ArchiveTraceRoute.cpp
+++ b/Engine/Source/Runtime/CoreUObject/Private/Serialization/ArchiveTraceRoute.cpp
@@ -140,8 +140,8 @@ FArchiveTraceRoute::FArchiveTraceRoute( UObject* TargetObject, TMap<UObject*,FTr
 	TSparseArray<UObject*> RootObjects;
 
 	// allocate enough memory for all objects
-	ObjectGraph.Empty(GetUObjectArray().GetObjectArrayNum());
-	RootObjects.Empty(GetUObjectArray().GetObjectArrayNum() / 2);
+	ObjectGraph.Empty(GUObjectArray.GetObjectArrayNum());
+	RootObjects.Empty(GUObjectArray.GetObjectArrayNum() / 2);
 
 	// search for objects that have the right flags and add them to the list of objects that we're going to start with
 	// all other objects need to be tagged so that we can tell whether they've been serialized or not.
diff --git a/Engine/Source/Runtime/CoreUObject/Private/UObject/GarbageCollection.cpp b/Engine/Source/Runtime/CoreUObject/Private/UObject/GarbageCollection.cpp
--- a/Engine/Source/Runtime/CoreUObject/Private/UObject/GarbageCollection.cpp
+++ b/Engine/Source/Runtime/CoreUObject/Private/UObject/GarbageCollection.cpp
@@ -691,9 +691,9 @@ public:
 		GObjectCountDuringLastMarkPhase = 0;
 
 		// Presize array and add a bit of extra slack for prefetching.
-		ObjectsToSerialize.Reset( GetUObjectArray().GetObjectArrayNumMinusPermanent() + 3 );
+		ObjectsToSerialize.Reset( GUObjectArray.GetObjectArrayNumMinusPermanent() + 3 );
 		// Make sure GC referencer object is checked for references to other objects even if it resides in permanent object pool
-		if (FPlatformProperties::RequiresCookedData() && FGCObject::GGCObjectReferencer && GetUObjectArray().IsDisregardForGC(FGCObject::GGCObjectReferencer))
+		if (FPlatformProperties::RequiresCookedData() && FGCObject::GGCObjectReferencer && GUObjectArray.IsDisregardForGC(FGCObject::GGCObjectReferencer))
 		{
 			ObjectsToSerialize.Add(FGCObject::GGCObjectReferencer);
 		}
@@ -1081,7 +1081,7 @@ void IncrementalPurgeGarbage( bool bUseTimeLimit, float TimeLimit )
 	if (GExitPurge)
 	{
 		GObjPurgeIsRequired = true;
-		GetUObjectArray().DisableDisregardForGC();
+		GUObjectArray.DisableDisregardForGC();
 		GObjCurrentPurgeObjectIndexNeedsReset = true;
 		GObjCurrentPurgeObjectIndexResetPastPermanent = false;
 	}
@@ -1405,7 +1405,7 @@ void CollectGarbageInternal(EObjectFlags KeepFlags, bool bPerformFullPurge)
 	check( !GObjPurgeIsRequired );
 
 #if VERIFY_DISREGARD_GC_ASSUMPTIONS
-	FUObjectArray& UObjectArray = GetUObjectArray();
+	FUObjectArray& UObjectArray = GUObjectArray;
 	// Only verify assumptions if option is enabled. This avoids false positives in the Editor or commandlets.
 	if (UObjectArray.DisregardForGCEnabled() && GShouldVerifyGCAssumptions)
 	{
diff --git a/Engine/Source/Runtime/CoreUObject/Private/UObject/LinkerLoad.cpp b/Engine/Source/Runtime/CoreUObject/Private/UObject/LinkerLoad.cpp
--- a/Engine/Source/Runtime/CoreUObject/Private/UObject/LinkerLoad.cpp
+++ b/Engine/Source/Runtime/CoreUObject/Private/UObject/LinkerLoad.cpp
@@ -2083,7 +2083,7 @@ void FLinkerLoad::GatherImportDependencies(int32 ImportIndex, TSet<FDependencyRe
 	if (Import.SourceLinker == NULL || Import.SourceIndex == INDEX_NONE)
 	{
 #if DO_CHECK
-		int32 NumObjectsBefore = GetUObjectArray().GetObjectArrayNum();
+		int32 NumObjectsBefore = GUObjectArray.GetObjectArrayNum();
 #endif
 
 		// temp storage we can ignore
@@ -2120,9 +2120,9 @@ void FLinkerLoad::GatherImportDependencies(int32 ImportIndex, TSet<FDependencyRe
 
 #if DO_CHECK && !NO_LOGGING
 		// only object we should create are one FLinkerLoad for source linker
-		if (GetUObjectArray().GetObjectArrayNum() - NumObjectsBefore > 2)
+		if (GUObjectArray.GetObjectArrayNum() - NumObjectsBefore > 2)
 		{
-			UE_LOG(LogLinker, Warning, TEXT("Created %d objects checking %s"), GetUObjectArray().GetObjectArrayNum() - NumObjectsBefore, *GetImportFullName(ImportIndex));
+			UE_LOG(LogLinker, Warning, TEXT("Created %d objects checking %s"), GUObjectArray.GetObjectArrayNum() - NumObjectsBefore, *GetImportFullName(ImportIndex));
 		}
 #endif
 	}
diff --git a/Engine/Source/Runtime/CoreUObject/Private/UObject/Obj.cpp b/Engine/Source/Runtime/CoreUObject/Private/UObject/Obj.cpp
--- a/Engine/Source/Runtime/CoreUObject/Private/UObject/Obj.cpp
+++ b/Engine/Source/Runtime/CoreUObject/Private/UObject/Obj.cpp
@@ -3791,7 +3791,7 @@ void MarkObjectsToDisregardForGC()
 	}
 
 	UE_LOG(LogObj, Log, TEXT("%i objects as part of root set at end of initial load."), NumAlwaysLoadedObjects);
-	if (GetUObjectArray().DisregardForGCEnabled())
+	if (GUObjectArray.DisregardForGCEnabled())
 	{
 		UE_LOG(LogObj, Log, TEXT("%i objects are not in the root set, but can never be destroyed because they are in the DisregardForGC set."), NumAlwaysLoadedObjects - NumRootObjects);
 	}
diff --git a/Engine/Source/Runtime/CoreUObject/Private/UObject/ReferenceChainSearch.cpp b/Engine/Source/Runtime/CoreUObject/Private/UObject/ReferenceChainSearch.cpp
--- a/Engine/Source/Runtime/CoreUObject/Private/UObject/ReferenceChainSearch.cpp
+++ b/Engine/Source/Runtime/CoreUObject/Private/UObject/ReferenceChainSearch.cpp
@@ -173,7 +173,7 @@ void FReferenceChainSearch::PrintReferencers( FReferenceChain& Referencer )
 			ObjectReachability += TEXT("(asyncloading) ");
 		}
 
-		if (GetUObjectArray().IsDisregardForGC(RefInfo.ReferencedBy))
+		if (GUObjectArray.IsDisregardForGC(RefInfo.ReferencedBy))
 		{
 			ObjectReachability += TEXT("(NeverGCed) ");
 		}
@@ -845,7 +845,7 @@ FString FReferenceChainSearch::FReferenceChainLink::ToString() const
 		ObjectReachability += TEXT("(standalone) ");
 	}
 
-	if (GetUObjectArray().IsDisregardForGC(ReferencedBy))
+	if (GUObjectArray.IsDisregardForGC(ReferencedBy))
 	{
 		ObjectReachability += TEXT("(NeverGCed) ");
 	}
diff --git a/Engine/Source/Runtime/CoreUObject/Private/UObject/StructScriptLoader.cpp b/Engine/Source/Runtime/CoreUObject/Private/UObject/StructScriptLoader.cpp
--- a/Engine/Source/Runtime/CoreUObject/Private/UObject/StructScriptLoader.cpp
+++ b/Engine/Source/Runtime/CoreUObject/Private/UObject/StructScriptLoader.cpp
@@ -260,7 +260,7 @@ bool FStructScriptLoader::LoadStructWithScript(UStruct* DestScriptContainer, FAr
 	ensure(ScriptEndOffset == Ar.Tell());
 	checkf(BytecodeIndex == BytecodeBufferSize, TEXT("'%s' script expression-count mismatch; Expected: %i, Got: %i"), *DestScriptContainer->GetName(), BytecodeBufferSize, BytecodeIndex);
 
-	if (!GetUObjectArray().IsDisregardForGC(DestScriptContainer))
+	if (!GUObjectArray.IsDisregardForGC(DestScriptContainer))
 	{
 		DestScriptContainer->ScriptObjectReferences.Empty();
 		FArchiveScriptReferenceCollector ObjRefCollector(DestScriptContainer->ScriptObjectReferences);
diff --git a/Engine/Source/Runtime/CoreUObject/Private/UObject/UObjectArray.cpp b/Engine/Source/Runtime/CoreUObject/Private/UObject/UObjectArray.cpp
--- a/Engine/Source/Runtime/CoreUObject/Private/UObject/UObjectArray.cpp
+++ b/Engine/Source/Runtime/CoreUObject/Private/UObject/UObjectArray.cpp
@@ -211,5 +211,5 @@ void FUObjectArray::ShutdownUObjectArray()
 
 UObjectBase*** FUObjectArray::GetObjectArrayForDebugVisualizers()
 {
-	return GetUObjectArray().ObjObjects.GetRootBlockForDebuggerVisualizers();
+	return GUObjectArray.ObjObjects.GetRootBlockForDebuggerVisualizers();
 }
diff --git a/Engine/Source/Runtime/CoreUObject/Private/UObject/UObjectBase.cpp b/Engine/Source/Runtime/CoreUObject/Private/UObject/UObjectBase.cpp
--- a/Engine/Source/Runtime/CoreUObject/Private/UObject/UObjectBase.cpp
+++ b/Engine/Source/Runtime/CoreUObject/Private/UObject/UObjectBase.cpp
@@ -91,7 +91,7 @@ UObjectBase::~UObjectBase()
 		// Validate it.
 		check(IsValidLowLevel());
 		LowLevelRename(NAME_None);
-		GetUObjectArray().FreeUObjectIndex(this);
+		GUObjectArray.FreeUObjectIndex(this);
 	}
 }
 
@@ -143,7 +143,7 @@ void UObjectBase::DeferredRegister(UClass *UClassStaticClass,const TCHAR* Packag
 	AddObject(FName(InName));
 
 	// Make sure that objects disregarded for GC are part of root set.
-	check(!GetUObjectArray().IsDisregardForGC(this) || (GetFlags() & RF_RootSet) );
+	check(!GUObjectArray.IsDisregardForGC(this) || (GetFlags() & RF_RootSet) );
 }
 
 /**
@@ -215,7 +215,7 @@ bool UObjectBase::IsValidLowLevel() const
 		UE_LOG(LogUObjectBase, Warning, TEXT("Object is not registered") );
 		return false;
 	}
-	return GetUObjectArray().IsValid(this);
+	return GUObjectArray.IsValid(this);
 }
 
 bool UObjectBase::IsValidLowLevelFast(bool bRecursive /*= true*/) const
@@ -260,7 +260,7 @@ bool UObjectBase::IsValidLowLevelFast(bool bRecursive /*= true*/) const
 		return false;
 	}
 	// Lightweight versions of index checks.
-	if (!GetUObjectArray().IsValidIndex(this) || !Name.IsValidIndexFast())
+	if (!GUObjectArray.IsValidIndex(this) || !Name.IsValidIndexFast())
 	{
 		UE_LOG(LogUObjectBase, Error, TEXT("Object array index or name index is invalid."));
 		return false;
@@ -904,7 +904,7 @@ void UObjectBaseInit()
 	UE_LOG(LogInit, Log, TEXT("Presizing for %i objects not considered by GC, pre-allocating %i bytes."), MaxObjectsNotConsideredByGC, SizeOfPermanentObjectPool );
 
 	GUObjectAllocator.AllocatePermanentObjectPool(SizeOfPermanentObjectPool);
-	GetUObjectArray().AllocatePermanentObjectPool(MaxObjectsNotConsideredByGC);
+	GUObjectArray.AllocatePermanentObjectPool(MaxObjectsNotConsideredByGC);
 
 	void InitAsyncThread();
 	InitAsyncThread();
@@ -920,7 +920,7 @@ void UObjectBaseInit()
  */
 void UObjectBaseShutdown()
 {
-	GetUObjectArray().ShutdownUObjectArray();
+	GUObjectArray.ShutdownUObjectArray();
 	Internal::GObjInitialized = false;
 }
 
diff --git a/Engine/Source/Runtime/CoreUObject/Private/UObject/UObjectHash.cpp b/Engine/Source/Runtime/CoreUObject/Private/UObject/UObjectHash.cpp
--- a/Engine/Source/Runtime/CoreUObject/Private/UObject/UObjectHash.cpp
+++ b/Engine/Source/Runtime/CoreUObject/Private/UObject/UObjectHash.cpp
@@ -20,6 +20,9 @@ DECLARE_CYCLE_STAT( TEXT( "UnhashObject" ), STAT_Hash_UnhashObject, STATGROUP_UO
 DEFINE_STAT( STAT_Hash_NumObjects );
 #endif
 
+// Global UObject array instance
+FUObjectArray GUObjectArray;
+
 /**
  * This implementation will use more space than the UE3 implementation. The goal was to make UObjects smaller to save L2 cache space. 
  * The hash is rarely used at runtime. A more space-efficient implementation is possible.
@@ -307,13 +310,6 @@ public:
 	}
 };
 
-FUObjectArray& GetUObjectArray()
-{
-	static FUObjectArray GlobalUObjectArray;
-	return GlobalUObjectArray;
-}
-
-
 /**
  * Calculates the object's hash just using the object's name index
  *
@@ -622,7 +618,7 @@ void GetObjectsWithOuter(const class UObjectBase* Outer, TArray<UObject *>& Resu
 				Results.Add(Object);
 			}
 		}
-		int32 MaxResults = GetUObjectArray().GetObjectArrayNum();
+		int32 MaxResults = GUObjectArray.GetObjectArrayNum();
 		while (StartNum != Results.Num() && bIncludeNestedObjects)
 		{
 			int32 RangeStart = StartNum;
@@ -788,7 +784,7 @@ void GetObjectsOfClass(UClass* ClassToLookFor, TArray<UObject *>& Results, bool
 
 	GetObjectsOfClassThreadSafe( FUObjectHashTables::Get(), ClassesToSearch, Results, ExclusionFlags );
 
-	check( Results.Num() <= GetUObjectArray().GetObjectArrayNum() ); // otherwise we have a cycle in the outer chain, which should not be possible
+	check( Results.Num() <= GUObjectArray.GetObjectArrayNum() ); // otherwise we have a cycle in the outer chain, which should not be possible
 }
 
 void ForEachObjectOfClass(UClass* ClassToLookFor, TFunctionRef<void (UObject*)> Operation, bool bIncludeDerivedClasses, EObjectFlags AdditionalExcludeFlags)
@@ -852,7 +848,7 @@ void GetDerivedClasses(UClass* ClassToLookFor, TArray<UClass *>& Results, bool b
 
 void AllocateUObjectIndexForCurrentThread(UObjectBase* Object)
 {
-	GetUObjectArray().AllocateUObjectIndex(Object);
+	GUObjectArray.AllocateUObjectIndex(Object);
 }
 
 void HashObject(UObjectBase* Object)
diff --git a/Engine/Source/Runtime/CoreUObject/Private/UObject/WeakObjectPtr.cpp b/Engine/Source/Runtime/CoreUObject/Private/UObject/WeakObjectPtr.cpp
--- a/Engine/Source/Runtime/CoreUObject/Private/UObject/WeakObjectPtr.cpp
+++ b/Engine/Source/Runtime/CoreUObject/Private/UObject/WeakObjectPtr.cpp
@@ -186,7 +186,7 @@ static FSerialNumberManager GSerialNumberManager;
 **/
 void FWeakObjectPtr::Init()
 {
-	GetUObjectArray().AddUObjectDeleteListener(&GSerialNumberManager);
+	GUObjectArray.AddUObjectDeleteListener(&GSerialNumberManager);
 }
 
 /**  
@@ -198,7 +198,7 @@ void FWeakObjectPtr::operator=(const class UObject *Object)
 	if (Object // && UObjectInitialized() we might need this at some point, but it is a speed hit we would prefer to avoid
 		)
 	{
-		ObjectIndex = GetUObjectArray().ObjectToIndex((UObjectBase*)Object);
+		ObjectIndex = GUObjectArray.ObjectToIndex((UObjectBase*)Object);
 		ObjectSerialNumber = GSerialNumberManager.GetAndAllocateSerialNumber(ObjectIndex);
 		checkSlow(SerialNumbersMatch());
 	}
@@ -246,7 +246,7 @@ FORCEINLINE_DEBUGGABLE bool FWeakObjectPtr::Internal_IsValid(bool bEvenIfPending
 	{
 		return true;
 	}
-	return GetUObjectArray().IsValid(ObjectIndex, bEvenIfPendingKill);
+	return GUObjectArray.IsValid(ObjectIndex, bEvenIfPendingKill);
 }
 
 bool FWeakObjectPtr::IsStale(bool bEvenIfPendingKill, bool bThreadsafeTest) const
@@ -268,7 +268,7 @@ bool FWeakObjectPtr::IsStale(bool bEvenIfPendingKill, bool bThreadsafeTest) cons
 	{
 		return false;
 	}
-	return GetUObjectArray().IsStale(ObjectIndex, bEvenIfPendingKill);
+	return GUObjectArray.IsStale(ObjectIndex, bEvenIfPendingKill);
 }
 
 FORCEINLINE_DEBUGGABLE UObject* FWeakObjectPtr::Internal_Get(bool bEvenIfPendingKill) const
@@ -277,7 +277,7 @@ FORCEINLINE_DEBUGGABLE UObject* FWeakObjectPtr::Internal_Get(bool bEvenIfPending
 
 	if (Internal_IsValid(true, true))
 	{
-		Result = (UObject*)(GetUObjectArray().IndexToValidObject(GetObjectIndex(), bEvenIfPendingKill));
+		Result = (UObject*)(GUObjectArray.IndexToValidObject(GetObjectIndex(), bEvenIfPendingKill));
 	}
 	return Result;
 }
@@ -298,7 +298,7 @@ UObject* FWeakObjectPtr::GetEvenIfUnreachable() const
 	UObject* Result = nullptr;
 	if (Internal_IsValid(true, true))
 	{
-		Result = static_cast<UObject*>(GetUObjectArray().IndexToObject(GetObjectIndex(), true));
+		Result = static_cast<UObject*>(GUObjectArray.IndexToObject(GetObjectIndex(), true));
 	}
 	return Result;
 }
diff --git a/Engine/Source/Runtime/CoreUObject/Public/UObject/UObjectAnnotation.h b/Engine/Source/Runtime/CoreUObject/Public/UObject/UObjectAnnotation.h
--- a/Engine/Source/Runtime/CoreUObject/Public/UObject/UObjectAnnotation.h
+++ b/Engine/Source/Runtime/CoreUObject/Public/UObject/UObjectAnnotation.h
@@ -90,7 +90,7 @@ public:
 				if (bAutoRemove)
 #endif
 				{
-					GetUObjectArray().AddUObjectDeleteListener(this);
+					GUObjectArray.AddUObjectDeleteListener(this);
 				}
 			}
 			AnnotationMap.Add(AnnotationCacheKey,AnnotationCacheValue);
@@ -118,7 +118,7 @@ public:
 			if (bAutoRemove)
 #endif
 			{
-				GetUObjectArray().RemoveUObjectDeleteListener(this);
+				GUObjectArray.RemoveUObjectDeleteListener(this);
 			}
 		}
 		return Result;
@@ -143,7 +143,7 @@ public:
 			if (bAutoRemove)
 #endif
 			{
-				GetUObjectArray().RemoveUObjectDeleteListener(this);
+				GUObjectArray.RemoveUObjectDeleteListener(this);
 			}
 		}
 	}
@@ -165,7 +165,7 @@ public:
 			if (bAutoRemove)
 #endif
 			{
-				GetUObjectArray().RemoveUObjectDeleteListener(this);
+				GUObjectArray.RemoveUObjectDeleteListener(this);
 			}
 		}
 	}
@@ -489,7 +489,7 @@ public:
 	void AddAnnotation(const UObjectBase *Object,TAnnotation Annotation)
 	{
 		check(Object);
-		AddAnnotation(GetUObjectArray().ObjectToIndex(Object),Annotation);
+		AddAnnotation(GUObjectArray.ObjectToIndex(Object),Annotation);
 	}
 	/**
 	 * Add an annotation to the annotation list. If the Annotation is the default, then the annotation is removed.
@@ -514,7 +514,7 @@ public:
 				if (bAutoRemove)
 #endif
 				{
-					GetUObjectArray().AddUObjectDeleteListener(this);
+					GUObjectArray.AddUObjectDeleteListener(this);
 				}
 			}
 			if (Index >= AnnotationArray.Num())
@@ -537,7 +537,7 @@ public:
 	void RemoveAnnotation(const UObjectBase *Object)
 	{
 		check(Object);
-		RemoveAnnotation(GetUObjectArray().ObjectToIndex(Object));
+		RemoveAnnotation(GUObjectArray.ObjectToIndex(Object));
 	}
 	/**
 	 * Removes an annotation from the annotation list. 
@@ -569,7 +569,7 @@ public:
 			if (bAutoRemove)
 #endif
 			{
-				GetUObjectArray().RemoveUObjectDeleteListener(this);
+				GUObjectArray.RemoveUObjectDeleteListener(this);
 			}
 		}
 	}
@@ -582,7 +582,7 @@ public:
 	FORCEINLINE TAnnotation GetAnnotation(const UObjectBase *Object)
 	{
 		check(Object);
-		return GetAnnotation(GetUObjectArray().ObjectToIndex(Object));
+		return GetAnnotation(GUObjectArray.ObjectToIndex(Object));
 	}
 
 	/**
@@ -610,7 +610,7 @@ public:
 	FORCEINLINE TAnnotation& GetAnnotationRef(const UObjectBase *Object)
 	{
 		check(Object);
-		return GetAnnotationRef(GetUObjectArray().ObjectToIndex(Object));
+		return GetAnnotationRef(GUObjectArray.ObjectToIndex(Object));
 	}
 
 	/**
@@ -676,11 +676,11 @@ public:
 	FORCEINLINE void Set(const UObjectBase *Object)
 	{
 		checkSlow(Object);
-		int32 Index = GetUObjectArray().ObjectToIndex(Object);
+		int32 Index = GUObjectArray.ObjectToIndex(Object);
 		checkSlow(Index >= 0);
 		if (AnnotationArray.Num() == 0)
 		{
-			GetUObjectArray().AddUObjectDeleteListener(this);
+			GUObjectArray.AddUObjectDeleteListener(this);
 		}
 		if (Index >= AnnotationArray.Num() * BitsPerElement)
 		{
@@ -700,7 +700,7 @@ public:
 	FORCEINLINE void Clear(const UObjectBase *Object)
 	{
 		checkSlow(Object);
-		int32 Index = GetUObjectArray().ObjectToIndex(Object);
+		int32 Index = GUObjectArray.ObjectToIndex(Object);
 		RemoveAnnotation(Index);
 	}
 
@@ -721,7 +721,7 @@ public:
 	FORCEINLINE bool Get(const UObjectBase *Object)
 	{
 		checkSlow(Object);
-		int32 Index = GetUObjectArray().ObjectToIndex(Object);
+		int32 Index = GUObjectArray.ObjectToIndex(Object);
 		checkSlow(Index >= 0);
 		if (Index < AnnotationArray.Num() * BitsPerElement)
 		{
@@ -754,7 +754,7 @@ private:
 		AnnotationArray.Empty();
 		if (bHadElements)
 		{
-			GetUObjectArray().RemoveUObjectDeleteListener(this);
+			GUObjectArray.RemoveUObjectDeleteListener(this);
 		}
 	}
 
diff --git a/Engine/Source/Runtime/CoreUObject/Public/UObject/UObjectArray.h b/Engine/Source/Runtime/CoreUObject/Public/UObject/UObjectArray.h
--- a/Engine/Source/Runtime/CoreUObject/Public/UObject/UObjectArray.h
+++ b/Engine/Source/Runtime/CoreUObject/Public/UObject/UObjectArray.h
@@ -378,7 +378,7 @@ public:
 			return false;
 		}
 	private:
-		/** the array that we are iterating on, probably always GetUObjectArray() */
+		/** the array that we are iterating on, probably always GUObjectArray */
 		const FUObjectArray& Array;
 		/** index of the current element in the object array */
 		int32 Index;
@@ -436,7 +436,7 @@ private:
 };
 
 /** Global UObject allocator							*/
-COREUOBJECT_API FUObjectArray& GetUObjectArray();
+extern COREUOBJECT_API FUObjectArray GUObjectArray;
 
 /**
 	* Static version of IndexToObject for use with TWeakObjectPtr.
@@ -445,7 +445,7 @@ struct FIndexToObject
 {
 	static FORCEINLINE class UObjectBase* IndexToObject(int32 Index, bool bEvenIfPendingKill)
 	{
-		return GetUObjectArray().IndexToObject(Index, bEvenIfPendingKill);
+		return GUObjectArray.IndexToObject(Index, bEvenIfPendingKill);
 	}
 };
 
diff --git a/Engine/Source/Runtime/CoreUObject/Public/UObject/UObjectIterator.h b/Engine/Source/Runtime/CoreUObject/Public/UObject/UObjectIterator.h
--- a/Engine/Source/Runtime/CoreUObject/Public/UObject/UObjectIterator.h
+++ b/Engine/Source/Runtime/CoreUObject/Public/UObject/UObjectIterator.h
@@ -19,7 +19,7 @@ public:
 	 * @param	bOnlyGCedObjects	if true, skip all of the permanent objects
 	 */
 	FRawObjectIterator(bool bOnlyGCedObjects = false) :
-		FUObjectArray::TIterator( GetUObjectArray(), bOnlyGCedObjects )
+		FUObjectArray::TIterator( GUObjectArray, bOnlyGCedObjects )
 	{
 	}
 	/**
@@ -59,7 +59,7 @@ public:
 	 * @param	AdditionalExclusionFlags	RF_* flags that should not be included in results
 	 */
 	FObjectIterator( UClass* InClass=UObject::StaticClass(), bool bOnlyGCedObjects = false, EObjectFlags AdditionalExclusionFlags = RF_NoFlags ) :
-		FUObjectArray::TIterator( GetUObjectArray(), bOnlyGCedObjects ),
+		FUObjectArray::TIterator( GUObjectArray, bOnlyGCedObjects ),
 		Class( InClass ),
 		ExclusionFlags(AdditionalExclusionFlags)
 	{
diff --git a/Engine/Source/Runtime/Engine/Private/GameViewportClient.cpp b/Engine/Source/Runtime/Engine/Private/GameViewportClient.cpp
--- a/Engine/Source/Runtime/Engine/Private/GameViewportClient.cpp
+++ b/Engine/Source/Runtime/Engine/Private/GameViewportClient.cpp
@@ -2946,7 +2946,7 @@ bool UGameViewportClient::HandleDisplayAllCommand( const TCHAR* Cmd, FOutputDevi
 					// so then we only have to iterate over dynamic things each frame
 					for (TObjectIterator<UObject> It; It; ++It)
 					{
-						if (!GetUObjectArray().IsDisregardForGC(*It))
+						if (!GUObjectArray.IsDisregardForGC(*It))
 						{
 							break;
 						}
@@ -2993,7 +2993,7 @@ bool UGameViewportClient::HandleDisplayAllLocationCommand( const TCHAR* Cmd, FOu
 			// so then we only have to iterate over dynamic things each frame
 			for (TObjectIterator<UObject> It(true); It; ++It)
 			{
-				if (!GetUObjectArray().IsDisregardForGC(*It))
+				if (!GUObjectArray.IsDisregardForGC(*It))
 				{
 					break;
 				}
@@ -3031,7 +3031,7 @@ bool UGameViewportClient::HandleDisplayAllRotationCommand( const TCHAR* Cmd, FOu
 			// so then we only have to iterate over dynamic things each frame
 			for (TObjectIterator<UObject> It(true); It; ++It)
 			{
-				if (!GetUObjectArray().IsDisregardForGC(*It))
+				if (!GUObjectArray.IsDisregardForGC(*It))
 				{
 					break;
 				}
diff --git a/Engine/Source/Runtime/Engine/Private/UnrealEngine.cpp b/Engine/Source/Runtime/Engine/Private/UnrealEngine.cpp
--- a/Engine/Source/Runtime/Engine/Private/UnrealEngine.cpp
+++ b/Engine/Source/Runtime/Engine/Private/UnrealEngine.cpp
@@ -5530,7 +5530,7 @@ bool UEngine::HandleObjCommand( const TCHAR* Cmd, FOutputDevice& Ar )
 					continue;
 				}
 
-				if ( bOnlyListGCObjects && GetUObjectArray().IsDisregardForGC(*It) )
+				if ( bOnlyListGCObjects && GUObjectArray.IsDisregardForGC(*It) )
 				{
 					continue;
 				}
diff --git a/Engine/Source/Runtime/Launch/Private/LaunchEngineLoop.cpp b/Engine/Source/Runtime/Launch/Private/LaunchEngineLoop.cpp
--- a/Engine/Source/Runtime/Launch/Private/LaunchEngineLoop.cpp
+++ b/Engine/Source/Runtime/Launch/Private/LaunchEngineLoop.cpp
@@ -1440,7 +1440,7 @@ int32 FEngineLoop::PreInit( const TCHAR* CmdLine )
 	// Setup GC optimizations
 	if (bIsSeekFreeDedicatedServer || bHasEditorToken)
 	{
-		GetUObjectArray().DisableDisregardForGC();
+		GUObjectArray.DisableDisregardForGC();
 	}
 	
 	SlowTask.EnterProgressFrame(10);
@@ -1527,7 +1527,7 @@ int32 FEngineLoop::PreInit( const TCHAR* CmdLine )
 	}
 
 	MarkObjectsToDisregardForGC(); 
-	GetUObjectArray().CloseDisregardForGC();
+	GUObjectArray.CloseDisregardForGC();
 
 	SetIsServerForOnlineSubsystemsDelegate(FQueryIsRunningServer::CreateStatic(&IsServerDelegateForOSS));
 
@@ -2607,7 +2607,7 @@ void FEngineLoop::Tick()
 
 		// Set the UObject count stat
 #if UE_GC_TRACK_OBJ_AVAILABLE
-		SET_DWORD_STAT(STAT_Hash_NumObjects, GetUObjectArray().GetObjectArrayNumMinusAvailable());
+		SET_DWORD_STAT(STAT_Hash_NumObjects, GUObjectArray.GetObjectArrayNumMinusAvailable());
 #endif
 	}
 }
