Dynamic Class InitializePrivateStaticClass  fix

diff --git a/Engine/Source/Runtime/CoreUObject/Private/UObject/Class.cpp b/Engine/Source/Runtime/CoreUObject/Private/UObject/Class.cpp
--- a/Engine/Source/Runtime/CoreUObject/Private/UObject/Class.cpp
+++ b/Engine/Source/Runtime/CoreUObject/Private/UObject/Class.cpp
@@ -60,7 +60,7 @@ COREUOBJECT_API void InitializePrivateStaticClass(
 	else
 	{
 		// Register immediately
-		TClass_PrivateStaticClass->DeferredRegister(UClass::StaticClass(), PackageName, Name);
+		TClass_PrivateStaticClass->DeferredRegister(UDynamicClass::StaticClass(), PackageName, Name);
 	}
 }
 
diff --git a/Engine/Source/Runtime/CoreUObject/Private/UObject/UObjectBase.cpp b/Engine/Source/Runtime/CoreUObject/Private/UObject/UObjectBase.cpp
--- a/Engine/Source/Runtime/CoreUObject/Private/UObject/UObjectBase.cpp
+++ b/Engine/Source/Runtime/CoreUObject/Private/UObject/UObjectBase.cpp
@@ -657,7 +657,7 @@ void UClassCompiledInDefer(FFieldCompiledInInfo* ClassInfo, const TCHAR* Name, S
 	GetDeferredClassRegistration().Add(ClassInfo);
 }
 
-TMap<FName, UClass *(*)()> GetDynamicClassMap()
+TMap<FName, UClass *(*)()>& GetDynamicClassMap()
 {
 	static TMap<FName, UClass *(*)()> DynamicClassMap;
 	return DynamicClassMap;
