added const version of UProperty::GetOwnerProperty()

diff --git a/Engine/Source/Runtime/CoreUObject/Public/UObject/UnrealType.h b/Engine/Source/Runtime/CoreUObject/Public/UObject/UnrealType.h
--- a/Engine/Source/Runtime/CoreUObject/Public/UObject/UnrealType.h
+++ b/Engine/Source/Runtime/CoreUObject/Public/UObject/UnrealType.h
@@ -716,6 +716,16 @@ public:
 		return Result;
 	}
 
+	const UProperty* GetOwnerProperty() const
+	{
+		const UProperty* Result = this;
+		for (UProperty* PropBase = dynamic_cast<UProperty*>(GetOuter()); PropBase; PropBase = dynamic_cast<UProperty*>(PropBase->GetOuter()))
+		{
+			Result = PropBase;
+		}
+		return Result;
+	}
+
 	/**
 	 * Returns this property's propertyflags
 	 */
