De-templatizing GetPrivateStaticClassBody

diff --git a/Engine/Source/Runtime/CoreUObject/Private/UObject/Class.cpp b/Engine/Source/Runtime/CoreUObject/Private/UObject/Class.cpp
--- a/Engine/Source/Runtime/CoreUObject/Private/UObject/Class.cpp
+++ b/Engine/Source/Runtime/CoreUObject/Private/UObject/Class.cpp
@@ -59,7 +59,7 @@ COREUOBJECT_API void InitializePrivateStaticClass(
 	}
 	else
 	{
-		// Register immediately
+		// Register immediately (don't let the function name mistake you!)
 		TClass_PrivateStaticClass->DeferredRegister(UDynamicClass::StaticClass(), PackageName, Name);
 	}
 }
@@ -4166,6 +4166,115 @@ IMPLEMENT_CORE_INTRINSIC_CLASS(UClass, UStruct,
 	}
 );
 
+void GetPrivateStaticClassBody(
+	const TCHAR* PackageName,
+	const TCHAR* Name,
+	UClass*& ReturnClass,
+	void(*RegisterNativeFunc)(),
+	uint32 InSize,
+	uint32 InClassFlags,
+	EClassCastFlags InClassCastFlags,
+	const TCHAR* InConfigName,
+	UClass::ClassConstructorType InClassConstructor,
+	UClass::ClassVTableHelperCtorCallerType InClassVTableHelperCtorCaller,
+	UClass::ClassAddReferencedObjectsType InClassAddReferencedObjects,
+	UClass::StaticClassFunctionType InSuperClassFn,
+	UClass::StaticClassFunctionType InWithinClassFn,
+	bool bIsDynamic /*= false*/
+	)
+{
+#if WITH_HOT_RELOAD
+	if (GIsHotReload)
+	{
+		check(!bIsDynamic);
+		UPackage* Package = FindPackage(NULL, PackageName);
+		if (!Package)
+		{
+			UE_LOG(LogClass, Log, TEXT("Could not find existing package %s for HotReload."), PackageName);
+			return;
+		}
+		ReturnClass = FindObject<UClass>((UObject *)Package, Name);
+		if (ReturnClass)
+		{
+			if (ReturnClass->HotReloadPrivateStaticClass(
+				InSize,
+				InClassFlags,
+				InClassCastFlags,
+				InConfigName,
+				InClassConstructor,
+#if WITH_HOT_RELOAD_CTORS
+				InClassVTableHelperCtorCaller,
+#endif // WITH_HOT_RELOAD_CTORS
+				InClassAddReferencedObjects,
+				InSuperClassFn(),
+				InWithinClassFn()
+				))
+			{
+				// Register the class's native functions.
+				RegisterNativeFunc();
+			}
+			return;
+		}
+		else
+		{
+			UE_LOG(LogClass, Log, TEXT("Could not find existing class %s in package %s for HotReload, assuming new class"), Name, PackageName);
+		}
+	}
+#endif
+
+	if (!bIsDynamic)
+	{
+		ReturnClass = (UClass*)GUObjectAllocator.AllocateUObject(sizeof(UClass), ALIGNOF(UClass), true);
+		ReturnClass = ::new (ReturnClass)
+			UClass
+			(
+			EC_StaticConstructor,
+			Name,
+			InSize,
+			InClassFlags,
+			InClassCastFlags,
+			InConfigName,
+			EObjectFlags(RF_Public | RF_Standalone | RF_Transient | RF_Native | RF_RootSet),
+			InClassConstructor,
+#if WITH_HOT_RELOAD_CTORS
+			InClassVTableHelperCtorCaller,
+#endif // WITH_HOT_RELOAD_CTORS
+			InClassAddReferencedObjects
+			);
+		check(ReturnClass);
+	}
+	else
+	{
+		ReturnClass = (UClass*)GUObjectAllocator.AllocateUObject(sizeof(UDynamicClass), ALIGNOF(UDynamicClass), true);
+		ReturnClass = ::new (ReturnClass)
+			UDynamicClass
+			(
+			EC_StaticConstructor,
+			Name,
+			InSize,
+			InClassFlags,
+			InClassCastFlags,
+			InConfigName,
+			EObjectFlags(RF_Public | RF_Standalone | RF_Transient | RF_Native | RF_Dynamic),
+			InClassConstructor,
+#if WITH_HOT_RELOAD_CTORS
+			InClassVTableHelperCtorCaller,
+#endif // WITH_HOT_RELOAD_CTORS
+			InClassAddReferencedObjects
+			);
+		check(ReturnClass);
+	}
+	InitializePrivateStaticClass(
+		InSuperClassFn(),
+		ReturnClass,
+		InWithinClassFn(),
+		PackageName,
+		Name
+		);
+
+	// Register the class's native functions.
+	RegisterNativeFunc();
+}
 
 /*-----------------------------------------------------------------------------
 	UFunction.
diff --git a/Engine/Source/Runtime/CoreUObject/Public/UObject/Class.h b/Engine/Source/Runtime/CoreUObject/Public/UObject/Class.h
--- a/Engine/Source/Runtime/CoreUObject/Public/UObject/Class.h
+++ b/Engine/Source/Runtime/CoreUObject/Public/UObject/Class.h
@@ -1838,6 +1838,7 @@ public:
 	typedef UObject*	(*ClassVTableHelperCtorCallerType)	(FVTableHelper& Helper);
 #endif // WITH_HOT_RELOAD_CTORS
 	typedef void		(*ClassAddReferencedObjectsType)	(UObject*, class FReferenceCollector&);
+	typedef UClass* (*StaticClassFunctionType)();
 
 	ClassConstructorType ClassConstructor;
 #if WITH_HOT_RELOAD_CTORS
@@ -2426,16 +2427,19 @@ void InternalConstructor( const FObjectInitializer& X )
 	T::__DefaultConstructor(X);
 }
 
-#if WITH_HOT_RELOAD_CTORS
+
 /**
  * Helper template to call the vtable ctor caller for a class
  */
 template<class T>
 UObject* InternalVTableHelperCtorCaller(FVTableHelper& Helper)
 {
+#if WITH_HOT_RELOAD_CTORS
 	return T::__VTableCtorCaller(Helper);
-}
+#else
+	return nullptr;
 #endif // WITH_HOT_RELOAD_CTORS
+}
 
 COREUOBJECT_API void InitializePrivateStaticClass(
 	class UClass* TClass_Super_StaticClass,
@@ -2451,110 +2455,33 @@ COREUOBJECT_API void InitializePrivateStaticClass(
  * @param PackageName name of the package this class will be inside
  * @param Name of the class
  * @param ReturnClass reference to pointer to result. This must be PrivateStaticClass.
- * @param RegisterNativeFunc Native function registration funcfion pointer.
+ * @param RegisterNativeFunc Native function registration function pointer.
+ * @param InSize Size of the class
+ * @param InClassFlags Class flags
+ * @param InClassCastFlags Class cast flags
+ * @param InConfigName Class config name
+ * @param InClassConstructor Class constructor function pointer
+ * @param InClassVTableHelperCtorCaller Class constructor function for vtable pointer
+ * @param InClassAddReferencedObjects Class AddReferencedObjects function pointer
+ * @param InSuperClassFn Super class function pointer
+ * @param WithinClass Within class
  * @param bIsDynamic true if the class can be constructed dynamically at runtime
  */
-template<class TClass>
-void GetPrivateStaticClassBody( const TCHAR* PackageName, const TCHAR* Name, UClass*& ReturnClass, void (*RegisterNativeFunc)(), bool bIsDynamic = false )
-{ 
-#if WITH_HOT_RELOAD
-	if (GIsHotReload)
-	{
-		check(!bIsDynamic);
-		UPackage* Package = FindPackage(NULL, PackageName);
-		if (!Package)
-		{
-			UE_LOG(LogClass, Log, TEXT("Could not find existing package %s for HotReload."),PackageName);
-			return;
-		}
-		ReturnClass = FindObject<UClass>((UObject *)Package, Name);
-		if (ReturnClass)
-		{
-			if (ReturnClass->HotReloadPrivateStaticClass(
-				sizeof(TClass),
-				TClass::StaticClassFlags,
-				TClass::StaticClassCastFlags(),
-				TClass::StaticConfigName(),
-				(UClass::ClassConstructorType)InternalConstructor<TClass>,
-#if WITH_HOT_RELOAD_CTORS
-				(UClass::ClassVTableHelperCtorCallerType)InternalVTableHelperCtorCaller<TClass>,
-#endif // WITH_HOT_RELOAD_CTORS
-				&TClass::AddReferencedObjects,
-				TClass::Super::StaticClass(),
-				TClass::WithinClass::StaticClass()
-				))
-			{
-				// Register the class's native functions.
-				RegisterNativeFunc();
-			}
-			return;
-		}
-		else
-		{
-			UE_LOG(LogClass, Log, TEXT("Could not find existing class %s in package %s for HotReload, assuming new class"),Name,PackageName);
-		}
-	}
-#endif
-
-	if (!bIsDynamic)
-	{
-		ReturnClass = ::new (GUObjectAllocator.AllocateUObject(sizeof(UClass), ALIGNOF(UClass), true))
-			UClass
-			(
-			EC_StaticConstructor,
-			Name,
-			sizeof(TClass),
-			TClass::StaticClassFlags,
-			TClass::StaticClassCastFlags(),
-			TClass::StaticConfigName(),
-			EObjectFlags(RF_Public | RF_Standalone | RF_Transient | RF_Native | RF_RootSet),
-			(UClass::ClassConstructorType)InternalConstructor<TClass>,
-#if WITH_HOT_RELOAD_CTORS
-			(UClass::ClassVTableHelperCtorCallerType)InternalVTableHelperCtorCaller<TClass>,
-#endif // WITH_HOT_RELOAD_CTORS
-			&TClass::AddReferencedObjects
-			);
-		check(ReturnClass);
-		InitializePrivateStaticClass(
-			TClass::Super::StaticClass(),
-			ReturnClass,
-			TClass::WithinClass::StaticClass(),
-			PackageName,
-			Name
-			);
-	}
-	else
-	{
-		ReturnClass = ::new (GUObjectAllocator.AllocateUObject(sizeof(UDynamicClass), ALIGNOF(UDynamicClass), true))
-			UDynamicClass
-			(
-			EC_StaticConstructor,
-			Name,
-			sizeof(TClass),
-			TClass::StaticClassFlags,
-			TClass::StaticClassCastFlags(),
-			TClass::StaticConfigName(),
-			EObjectFlags(RF_Public | RF_Standalone | RF_Transient | RF_Native | RF_Dynamic),
-			(UClass::ClassConstructorType)InternalConstructor<TClass>,
-#if WITH_HOT_RELOAD_CTORS
-			(UClass::ClassVTableHelperCtorCallerType)InternalVTableHelperCtorCaller<TClass>,
-#endif // WITH_HOT_RELOAD_CTORS
-			&TClass::AddReferencedObjects
-			);
-		check(ReturnClass);
-		InitializePrivateStaticClass(
-			TClass::Super::StaticClass(),
-			ReturnClass,
-			TClass::WithinClass::StaticClass(),
-			PackageName,
-			Name
-			);
-	}
-
-	// Register the class's native functions.
-	RegisterNativeFunc();
-}
-
+COREUOBJECT_API void GetPrivateStaticClassBody(
+	const TCHAR* PackageName,
+	const TCHAR* Name,
+	UClass*& ReturnClass,
+	void(*RegisterNativeFunc)(),
+	uint32 InSize,
+	uint32 InClassFlags,
+	EClassCastFlags InClassCastFlags,
+	const TCHAR* InConfigName,
+	UClass::ClassConstructorType InClassConstructor,
+	UClass::ClassVTableHelperCtorCallerType InClassVTableHelperCtorCaller,
+	UClass::ClassAddReferencedObjectsType InClassAddReferencedObjects,
+	UClass::StaticClassFunctionType InSuperClassFn,
+	UClass::StaticClassFunctionType InWithinClassFn,
+	bool bIsDynamic = false);
 
 /*-----------------------------------------------------------------------------
 	FObjectInstancingGraph.
diff --git a/Engine/Source/Runtime/CoreUObject/Public/UObject/ObjectBase.h b/Engine/Source/Runtime/CoreUObject/Public/UObject/ObjectBase.h
--- a/Engine/Source/Runtime/CoreUObject/Public/UObject/ObjectBase.h
+++ b/Engine/Source/Runtime/CoreUObject/Public/UObject/ObjectBase.h
@@ -1448,11 +1448,20 @@ public: \
 		if (!PrivateStaticClass) \
 		{ \
 			/* this could be handled with templates, but we want it external to avoid code bloat */ \
-			GetPrivateStaticClassBody<TClass>( \
+			GetPrivateStaticClassBody( \
 				Package, \
 				(TCHAR*)TEXT(#TClass) + 1 + ((StaticClassFlags & CLASS_Deprecated) ? 11 : 0), \
 				PrivateStaticClass, \
-				StaticRegisterNatives##TClass \
+				StaticRegisterNatives##TClass, \
+				sizeof(TClass), \
+				TClass::StaticClassFlags, \
+				TClass::StaticClassCastFlags(), \
+				TClass::StaticConfigName(), \
+				(UClass::ClassConstructorType)InternalConstructor<TClass>, \
+				(UClass::ClassVTableHelperCtorCallerType)InternalVTableHelperCtorCaller<TClass>, \
+				&TClass::AddReferencedObjects, \
+				&TClass::Super::StaticClass, \
+				&TClass::WithinClass::StaticClass \
 			); \
 		} \
 		return PrivateStaticClass; \
@@ -1492,11 +1501,20 @@ public: \
 		if (!PrivateStaticClass) \
 		{ \
 			/* this could be handled with templates, but we want it external to avoid code bloat */ \
-			GetPrivateStaticClassBody<TClass>(\
+			GetPrivateStaticClassBody( \
 			Package, \
 			(TCHAR*)TEXT(#TClass) + 1 + ((StaticClassFlags & CLASS_Deprecated) ? 11 : 0), \
 			PrivateStaticClass, \
 			StaticRegisterNatives##TClass, \
+			sizeof(TClass), \
+			TClass::StaticClassFlags, \
+			TClass::StaticClassCastFlags(), \
+			TClass::StaticConfigName(), \
+			(UClass::ClassConstructorType)InternalConstructor<TClass>, \
+			(UClass::ClassVTableHelperCtorCallerType)InternalVTableHelperCtorCaller<TClass>, \
+			&TClass::AddReferencedObjects, \
+			&TClass::Super::StaticClass, \
+			&TClass::WithinClass::StaticClass, \
 			true \
 			); \
 		} \
