Dynamic class fixes

diff --git a/Engine/Source/Runtime/CoreUObject/Private/UObject/UObjectBase.cpp b/Engine/Source/Runtime/CoreUObject/Private/UObject/UObjectBase.cpp
--- a/Engine/Source/Runtime/CoreUObject/Private/UObject/UObjectBase.cpp
+++ b/Engine/Source/Runtime/CoreUObject/Private/UObject/UObjectBase.cpp
@@ -375,7 +375,7 @@ static TArray<FPendingStructRegistrant>& GetDeferredCompiledInStructRegistration
 	return DeferredCompiledInRegistration;
 }
 
-TMap<FName, UScriptStruct *(*)()> GetDynamicStructMap()
+TMap<FName, UScriptStruct *(*)()>& GetDynamicStructMap()
 {
 	static TMap<FName, UScriptStruct *(*)()> DynamicStructMap;
 	return DynamicStructMap;
@@ -498,7 +498,7 @@ static TArray<FPendingEnumRegistrant>& GetDeferredCompiledInEnumRegistration()
 	return DeferredCompiledInRegistration;
 }
 
-TMap<FName, UEnum *(*)()> GetDynamicEnumMap()
+TMap<FName, UEnum *(*)()>& GetDynamicEnumMap()
 {
 	static TMap<FName, UEnum *(*)()> DynamicEnumMap;
 	return DynamicEnumMap;
@@ -678,7 +678,8 @@ void UObjectCompiledInDefer(class UClass *(*InRegister)(), const TCHAR* Name, bo
 	}
 	else
 	{
-		GetDynamicClassMap().Add(FName(Name), InRegister);
+		// Name is with the CPP prefix, skip it when adding to the map
+		GetDynamicClassMap().Add(FName(Name + 1), InRegister);
 	}
 }
 
