Dynamic classes can now be constructed by their StaticClass() functions (no need to call ConstructDynamicType)

diff --git a/Engine/Source/Programs/UnrealHeaderTool/Private/CodeGenerator.cpp b/Engine/Source/Programs/UnrealHeaderTool/Private/CodeGenerator.cpp
--- a/Engine/Source/Programs/UnrealHeaderTool/Private/CodeGenerator.cpp
+++ b/Engine/Source/Programs/UnrealHeaderTool/Private/CodeGenerator.cpp
@@ -1513,8 +1513,8 @@ void FNativeClassHeaderGenerator::ExportNativeGeneratedInitCode(FClass* Class, F
 	{	
 		const TCHAR* ClassNameCPP = NameLookupCPP.GetNameCPP(Class);
 
-		GeneratedFunctionText.Logf(TEXT("    static FCompiledInDefer Z_CompiledInDefer_UClass_%s(%s, TEXT(\"%s\"), %s);\r\n"),
-			ClassNameCPP, *SingletonName, ClassNameCPP,
+		GeneratedFunctionText.Logf(TEXT("    static FCompiledInDefer Z_CompiledInDefer_UClass_%s(%s, &%s::StaticClass, TEXT(\"%s\"), %s);\r\n"),
+			ClassNameCPP, *SingletonName, ClassNameCPP, ClassNameCPP,
 			bIsDynamic ? TEXT("true") : TEXT("false"));
 		
 		// Append base class' CRC at the end of the generated code, this will force update derived classes
diff --git a/Engine/Source/Runtime/CoreUObject/Private/UObject/Class.cpp b/Engine/Source/Runtime/CoreUObject/Private/UObject/Class.cpp
--- a/Engine/Source/Runtime/CoreUObject/Private/UObject/Class.cpp
+++ b/Engine/Source/Runtime/CoreUObject/Private/UObject/Class.cpp
@@ -4166,6 +4166,8 @@ IMPLEMENT_CORE_INTRINSIC_CLASS(UClass, UStruct,
 	}
 );
 
+UClass::StaticClassFunctionType GetDynamicClassConstructFn(FName ClassName);
+
 void GetPrivateStaticClassBody(
 	const TCHAR* PackageName,
 	const TCHAR* Name,
@@ -4274,6 +4276,14 @@ void GetPrivateStaticClassBody(
 
 	// Register the class's native functions.
 	RegisterNativeFunc();
+
+	if (bIsDynamic)
+	{
+		// Now call the UHT-generated Z_Construct* function for the dynamic class
+		UClass::StaticClassFunctionType ZConstructDynamicClassFn = GetDynamicClassConstructFn(Name);
+		check(ZConstructDynamicClassFn);
+		ZConstructDynamicClassFn();
+	}
 }
 
 /*-----------------------------------------------------------------------------
diff --git a/Engine/Source/Runtime/CoreUObject/Private/UObject/UObjectBase.cpp b/Engine/Source/Runtime/CoreUObject/Private/UObject/UObjectBase.cpp
--- a/Engine/Source/Runtime/CoreUObject/Private/UObject/UObjectBase.cpp
+++ b/Engine/Source/Runtime/CoreUObject/Private/UObject/UObjectBase.cpp
@@ -657,13 +657,21 @@ void UClassCompiledInDefer(FFieldCompiledInInfo* ClassInfo, const TCHAR* Name, S
 	GetDeferredClassRegistration().Add(ClassInfo);
 }
 
-TMap<FName, UClass *(*)()>& GetDynamicClassMap()
+struct FClassConstructFunctions
 {
-	static TMap<FName, UClass *(*)()> DynamicClassMap;
+	/** Autogenerated Z_Construct* function pointer */
+	UClass::StaticClassFunctionType ZConstructFn;
+	/** StaticClass() function pointer */
+	UClass::StaticClassFunctionType StaticClassFn;
+};
+
+TMap<FName, FClassConstructFunctions>& GetDynamicClassMap()
+{
+	static TMap<FName, FClassConstructFunctions> DynamicClassMap;
 	return DynamicClassMap;
 }
 
-void UObjectCompiledInDefer(class UClass *(*InRegister)(), const TCHAR* Name, bool bDynamic)
+void UObjectCompiledInDefer(UClass *(*InRegister)(), UClass *(*InStaticClass)(), const TCHAR* Name, bool bDynamic)
 {
 	if (!bDynamic)
 	{
@@ -679,7 +687,10 @@ void UObjectCompiledInDefer(class UClass *(*InRegister)(), const TCHAR* Name, bo
 	else
 	{
 		// Name is with the CPP prefix, skip it when adding to the map
-		GetDynamicClassMap().Add(FName(Name + 1), InRegister);
+		FClassConstructFunctions ClassFunctions;
+		ClassFunctions.ZConstructFn = InRegister;
+		ClassFunctions.StaticClassFn = InStaticClass;
+		GetDynamicClassMap().Add(FName(Name + 1), ClassFunctions);
 	}
 }
 
@@ -1027,14 +1038,27 @@ const TCHAR* DebugFullName(UObject* Object)
 UClass* ConstructDynamicClass(FName ClassName)
 {
 	UClass* Result = nullptr;
-	UClass *(**StaticClassFNPtr)() = GetDynamicClassMap().Find(ClassName);
-	if (StaticClassFNPtr)
+	FClassConstructFunctions* ClassConstructFn = GetDynamicClassMap().Find(ClassName);
+	if (ClassConstructFn)
 	{
-		Result = (*StaticClassFNPtr)();
+		// This calls the class' StaticClass() function which calls GetPrivateStaticClassBody() which calls Z_Construct* function.
+		Result = ClassConstructFn->StaticClassFn();
 	}
 	return Result;
 }
 
+UClass::StaticClassFunctionType GetDynamicClassConstructFn(FName ClassName)
+{
+	FClassConstructFunctions* ClassConstructFn = GetDynamicClassMap().Find(ClassName);
+	UE_CLOG(!ClassConstructFn, LogUObjectBase, Fatal, TEXT("Unable to find construct function pointer for dynamic class %s. Make sure dynamic class exists."), *ClassName.ToString());
+	if (ClassConstructFn)
+	{
+		return ClassConstructFn->StaticClassFn;
+	}
+	// We should never get here. We either find the class or assert.
+	return nullptr;
+}
+
 UScriptStruct* ConstructDynamicStruct(FName StructName)
 {
 	UScriptStruct* Result = nullptr;
diff --git a/Engine/Source/Runtime/CoreUObject/Public/UObject/UObjectBase.h b/Engine/Source/Runtime/CoreUObject/Public/UObject/UObjectBase.h
--- a/Engine/Source/Runtime/CoreUObject/Public/UObject/UObjectBase.h
+++ b/Engine/Source/Runtime/CoreUObject/Public/UObject/UObjectBase.h
@@ -326,13 +326,13 @@ struct TClassCompiledInDefer : public FFieldCompiledInInfo
 /**
  * Stashes the singleton function that builds a compiled in class. Later, this is executed.
  */
-COREUOBJECT_API void UObjectCompiledInDefer(class UClass *(*InRegister)(), const TCHAR* Name, bool bDynamic);
+COREUOBJECT_API void UObjectCompiledInDefer(class UClass *(*InRegister)(), class UClass *(*InStaticClass)(), const TCHAR* Name, bool bDynamic);
 
 struct FCompiledInDefer
 {
-	FCompiledInDefer(class UClass *(*InRegister)(), const TCHAR* Name, bool bDynamic)
+	FCompiledInDefer(class UClass *(*InRegister)(), class UClass *(*InStaticClass)(), const TCHAR* Name, bool bDynamic)
 	{
-		UObjectCompiledInDefer(InRegister, Name, bDynamic);
+		UObjectCompiledInDefer(InRegister, InStaticClass, Name, bDynamic);
 	}
 };
 
