Missing file from previous checkin

diff --git a/Engine/Source/Runtime/CoreUObject/Public/UObject/ObjectBase.h b/Engine/Source/Runtime/CoreUObject/Public/UObject/ObjectBase.h
--- a/Engine/Source/Runtime/CoreUObject/Public/UObject/ObjectBase.h
+++ b/Engine/Source/Runtime/CoreUObject/Public/UObject/ObjectBase.h
@@ -1487,7 +1487,7 @@ public: \
 		check(Class->GetClass()); \
 		return Class; \
 	} \
-	static FCompiledInDefer Z_CompiledInDefer_UClass_##TClass(Z_Construct_UClass_##TClass, TEXT(#TClass), false);
+	static FCompiledInDefer Z_CompiledInDefer_UClass_##TClass(Z_Construct_UClass_##TClass, &TClass::StaticClass, TEXT(#TClass), false);
 
 #define IMPLEMENT_CORE_INTRINSIC_CLASS(TClass, TSuperClass, InitCode) \
 	IMPLEMENT_INTRINSIC_CLASS(TClass, COREUOBJECT_API, TSuperClass, COREUOBJECT_API, InitCode)
