GetDynamicClassConstructFn fix

diff --git a/Engine/Source/Runtime/CoreUObject/Private/UObject/UObjectBase.cpp b/Engine/Source/Runtime/CoreUObject/Private/UObject/UObjectBase.cpp
--- a/Engine/Source/Runtime/CoreUObject/Private/UObject/UObjectBase.cpp
+++ b/Engine/Source/Runtime/CoreUObject/Private/UObject/UObjectBase.cpp
@@ -1053,7 +1053,7 @@ UClass::StaticClassFunctionType GetDynamicClassConstructFn(FName ClassName)
 	UE_CLOG(!ClassConstructFn, LogUObjectBase, Fatal, TEXT("Unable to find construct function pointer for dynamic class %s. Make sure dynamic class exists."), *ClassName.ToString());
 	if (ClassConstructFn)
 	{
-		return ClassConstructFn->StaticClassFn;
+		return ClassConstructFn->ZConstructFn;
 	}
 	// We should never get here. We either find the class or assert.
 	return nullptr;
