Add default parameters consistent with LoadObject for LoadClass

diff --git a/Engine/Source/Runtime/CoreUObject/Public/UObject/UObjectGlobals.h b/Engine/Source/Runtime/CoreUObject/Public/UObject/UObjectGlobals.h
--- a/Engine/Source/Runtime/CoreUObject/Public/UObject/UObjectGlobals.h
+++ b/Engine/Source/Runtime/CoreUObject/Public/UObject/UObjectGlobals.h
@@ -1245,14 +1245,14 @@ inline T* FindObjectSafe( UObject* Outer, const TCHAR* Name, bool ExactClass=fal
 
 // Load an object.
 template< class T > 
-inline T* LoadObject( UObject* Outer, const TCHAR* Name, const TCHAR* Filename=NULL, uint32 LoadFlags=LOAD_None, UPackageMap* Sandbox=NULL )
+inline T* LoadObject( UObject* Outer, const TCHAR* Name, const TCHAR* Filename=nullptr, uint32 LoadFlags=LOAD_None, UPackageMap* Sandbox=nullptr )
 {
 	return (T*)StaticLoadObject( T::StaticClass(), Outer, Name, Filename, LoadFlags, Sandbox );
 }
 
 // Load a class object.
 template< class T > 
-inline UClass* LoadClass( UObject* Outer, const TCHAR* Name, const TCHAR* Filename, uint32 LoadFlags, UPackageMap* Sandbox )
+inline UClass* LoadClass( UObject* Outer, const TCHAR* Name, const TCHAR* Filename=nullptr, uint32 LoadFlags=LOAD_None, UPackageMap* Sandbox=nullptr )
 {
 	return StaticLoadClass( T::StaticClass(), Outer, Name, Filename, LoadFlags, Sandbox );
 }
