Back out revision 222 from //depot/UE4/Engine/Source/Runtime/CoreUObject/Private/UObject/Obj.cpp

diff --git a/Engine/Source/Runtime/CoreUObject/Private/UObject/Obj.cpp b/Engine/Source/Runtime/CoreUObject/Private/UObject/Obj.cpp
--- a/Engine/Source/Runtime/CoreUObject/Private/UObject/Obj.cpp
+++ b/Engine/Source/Runtime/CoreUObject/Private/UObject/Obj.cpp
@@ -867,11 +867,6 @@ void UObject::Serialize( FArchive& Ar )
 	UObject* LoadOuter = GetOuter();
 	FName LoadName = GetFName();
 
-	if (Ar.IsLoading() && Cast<UDynamicClass>(ObjClass) && ObjClass->GetName().Contains(TEXT("RRRActor_C")))
-	{
-		int k = 0; k++;
-	}
-
 	// Make sure this object's class's data is loaded.
 	if(ObjClass->HasAnyFlags(RF_NeedLoad) )
 	{
