Core: Allowed for specifying custom data in FStructOnScope

diff --git a/Engine/Source/Runtime/CoreUObject/Public/UObject/Class.h b/Engine/Source/Runtime/CoreUObject/Public/UObject/Class.h
--- a/Engine/Source/Runtime/CoreUObject/Public/UObject/Class.h
+++ b/Engine/Source/Runtime/CoreUObject/Public/UObject/Class.h
@@ -1122,10 +1122,12 @@ protected:
 		{
 			SampleStructMemory = (uint8*)FMemory::Malloc(ScriptStruct->GetStructureSize());
 			ScriptStruct.Get()->InitializeStruct(SampleStructMemory);
+			OwnsMemory = true;
 		}
 	}
 
 public:
+
 	FStructOnScope(const UStruct* InScriptStruct)
 		: ScriptStruct(InScriptStruct)
 		, SampleStructMemory(NULL)
@@ -1133,16 +1135,39 @@ public:
 		Initialize();
 	}
 
-	virtual uint8* GetStructMemory() { return SampleStructMemory; }
+	FStructOnScope(const UStruct* InScriptStruct, uint8* InData)
+		: ScriptStruct(InScriptStruct)
+		, SampleStructMemory(InData)
+		, OwnsMemory(false)
+	{ }
 
-	virtual const uint8* GetStructMemory() const { return SampleStructMemory; }
+	virtual uint8* GetStructMemory()
+	{
+		return SampleStructMemory;
+	}
 
-	virtual const UStruct* GetStruct() const { return ScriptStruct.Get(); }
+	virtual const uint8* GetStructMemory() const
+	{
+		return SampleStructMemory;
+	}
+
+	virtual const UStruct* GetStruct() const
+	{
+		return ScriptStruct.Get();
+	}
 
-	virtual bool IsValid() const { return ScriptStruct.IsValid() && SampleStructMemory; }
+	virtual bool IsValid() const
+	{
+		return ScriptStruct.IsValid() && SampleStructMemory;
+	}
 
 	virtual void Destroy()
 	{
+		if (!OwnsMemory)
+		{
+			return;
+		}
+
 		if (ScriptStruct.IsValid() && SampleStructMemory)
 		{
 			ScriptStruct.Get()->DestroyStruct(SampleStructMemory);
@@ -1169,8 +1194,14 @@ public:
 	}
 
 private:
+
 	FStructOnScope(const FStructOnScope&);
 	FStructOnScope& operator=(const FStructOnScope&);
+
+private:
+
+	/** Whether the struct memory is owned by this instance. */
+	bool OwnsMemory;
 };
 
 /*-----------------------------------------------------------------------------
