Merging CL 2687751 to main

Async loading priority typedef'd and made into a signed type

diff --git a/Engine/Source/Runtime/CoreUObject/Private/Serialization/AsyncLoading.cpp b/Engine/Source/Runtime/CoreUObject/Private/Serialization/AsyncLoading.cpp
--- a/Engine/Source/Runtime/CoreUObject/Private/Serialization/AsyncLoading.cpp
+++ b/Engine/Source/Runtime/CoreUObject/Private/Serialization/AsyncLoading.cpp
@@ -1992,7 +1992,7 @@ void FAsyncPackage::UpdateLoadPercentage()
 	LoadPercentage = FMath::Max(NewLoadPercentage, LoadPercentage);
 }
 
-int32 LoadPackageAsync(const FString& InName, const FGuid* InGuid /*= nullptr*/, const TCHAR* InPackageToLoadFrom /*= nullptr*/, FLoadPackageAsyncDelegate InCompletionDelegate /*= FLoadPackageAsyncDelegate()*/, EPackageFlags InPackageFlags /*= PKG_None*/, int32 InPIEInstanceID /*= INDEX_NONE*/, uint32 InPackagePriority /*= 0*/)
+int32 LoadPackageAsync(const FString& InName, const FGuid* InGuid /*= nullptr*/, const TCHAR* InPackageToLoadFrom /*= nullptr*/, FLoadPackageAsyncDelegate InCompletionDelegate /*= FLoadPackageAsyncDelegate()*/, EPackageFlags InPackageFlags /*= PKG_None*/, int32 InPIEInstanceID /*= INDEX_NONE*/, int32 InPackagePriority /*= 0*/)
 {
 #if !WITH_EDITOR
 	if (GPreloadPackageDependencies)
@@ -2046,14 +2046,14 @@ int32 LoadPackageAsync(const FString& InName, const FGuid* InGuid /*= nullptr*/,
 	return RequestID;
 }
 
-int32 LoadPackageAsync(const FString& PackageName, FLoadPackageAsyncDelegate CompletionDelegate, uint32 InPackagePriority /*= 0*/, EPackageFlags InPackageFlags /*= PKG_None*/)
+int32 LoadPackageAsync(const FString& PackageName, FLoadPackageAsyncDelegate CompletionDelegate, int32 InPackagePriority /*= 0*/, EPackageFlags InPackageFlags /*= PKG_None*/)
 {
 	const FGuid* Guid = nullptr;
 	const TCHAR* PackageToLoadFrom = nullptr;
 	return LoadPackageAsync(PackageName, Guid, PackageToLoadFrom, CompletionDelegate, InPackageFlags, -1, InPackagePriority );
 }
 
-int32 LoadPackageAsync(const FString& InName, const FGuid* InGuid, FName InType /* Unused, deprecated */, const TCHAR* InPackageToLoadFrom /*= nullptr*/, FLoadPackageAsyncDelegate InCompletionDelegate /*= FLoadPackageAsyncDelegate()*/, EPackageFlags InPackageFlags /*= PKG_None*/, int32 InPIEInstanceID /*= INDEX_NONE*/, uint32 InPackagePriority /*= 0*/)
+int32 LoadPackageAsync(const FString& InName, const FGuid* InGuid, FName InType /* Unused, deprecated */, const TCHAR* InPackageToLoadFrom /*= nullptr*/, FLoadPackageAsyncDelegate InCompletionDelegate /*= FLoadPackageAsyncDelegate()*/, EPackageFlags InPackageFlags /*= PKG_None*/, int32 InPIEInstanceID /*= INDEX_NONE*/, int32 InPackagePriority /*= 0*/)
 {
 	return LoadPackageAsync(InName, InGuid, InPackageToLoadFrom, InCompletionDelegate, InPackageFlags, InPIEInstanceID, InPackagePriority);
 }
@@ -2535,7 +2535,7 @@ bool FArchiveAsync::Precache( int64 RequestOffset, int64 RequestSize )
 									PrecacheEndPos[CURRENT] - PrecacheStartPos[CURRENT], 
 									PrecacheBuffer[CURRENT], 
 									&PrecacheReadStatus[CURRENT],
-									AIOP_Normal );
+									AIOP_Normal);
 			check(RequestId);
 		}
 
diff --git a/Engine/Source/Runtime/CoreUObject/Public/Serialization/AsyncLoading.h b/Engine/Source/Runtime/CoreUObject/Public/Serialization/AsyncLoading.h
--- a/Engine/Source/Runtime/CoreUObject/Public/Serialization/AsyncLoading.h
+++ b/Engine/Source/Runtime/CoreUObject/Public/Serialization/AsyncLoading.h
@@ -78,7 +78,7 @@ struct FAsyncPackage
 	}
 
 	/** Returns package loading priority. */
-	FORCEINLINE uint32 GetPriority() const
+	FORCEINLINE TAsyncLoadPriority GetPriority() const
 	{
 		return Desc.Priority;
 	}
diff --git a/Engine/Source/Runtime/CoreUObject/Public/Serialization/AsyncPackage.h b/Engine/Source/Runtime/CoreUObject/Public/Serialization/AsyncPackage.h
--- a/Engine/Source/Runtime/CoreUObject/Public/Serialization/AsyncPackage.h
+++ b/Engine/Source/Runtime/CoreUObject/Public/Serialization/AsyncPackage.h
@@ -21,14 +21,14 @@ struct FAsyncPackageDesc
 	/** The flags that should be applied to the package */
 	EPackageFlags PackageFlags;
 	/** Package loading priority. Higher number is higher priority. */
-	uint32 Priority;
+	TAsyncLoadPriority Priority;
 #if WITH_EDITOR
 	/** Editor only: PIE instance ID this package belongs to, INDEX_NONE otherwise */
 	int32 PIEInstanceID;
 #endif
 
 
-	FAsyncPackageDesc(int32 InRequestID, const FName& InName, FName InPackageToLoadFrom = NAME_None, const FGuid& InGuid = FGuid(), FLoadPackageAsyncDelegate InCompletionDelegate = FLoadPackageAsyncDelegate(), EPackageFlags InPackageFlags = PKG_None, int32 InPIEInstanceID = INDEX_NONE, uint32 InPriority = 0)
+	FAsyncPackageDesc(int32 InRequestID, const FName& InName, FName InPackageToLoadFrom = NAME_None, const FGuid& InGuid = FGuid(), FLoadPackageAsyncDelegate InCompletionDelegate = FLoadPackageAsyncDelegate(), EPackageFlags InPackageFlags = PKG_None, int32 InPIEInstanceID = INDEX_NONE, TAsyncLoadPriority InPriority = 0)
 		: RequestID(InRequestID)
 		, Name(InName)
 		, NameToLoad(InPackageToLoadFrom)
diff --git a/Engine/Source/Runtime/CoreUObject/Public/UObject/UObjectGlobals.h b/Engine/Source/Runtime/CoreUObject/Public/UObject/UObjectGlobals.h
--- a/Engine/Source/Runtime/CoreUObject/Public/UObject/UObjectGlobals.h
+++ b/Engine/Source/Runtime/CoreUObject/Public/UObject/UObjectGlobals.h
@@ -291,6 +291,9 @@ namespace EAsyncLoadingResult
 	};
 }
 
+/** The type that represents an async loading priority */
+typedef int32 TAsyncLoadPriority;
+
 /**
  * Delegate called on completion of async package loading
  * @param	PackageName			Package name we were trying to load
@@ -313,7 +316,7 @@ DECLARE_DELEGATE_ThreeParams(FLoadPackageAsyncDelegate, const FName& /*PackageNa
  * @return Unique ID associated with this load request (the same package can be associated with multiple IDs).
  */
 DEPRECATED(4.9, "LoadPackageAsync override that takes package type parameter FName InType is deprecated.")
-COREUOBJECT_API int32 LoadPackageAsync(const FString& InName, const FGuid* InGuid, FName InType, const TCHAR* InPackageToLoadFrom = nullptr, FLoadPackageAsyncDelegate InCompletionDelegate = FLoadPackageAsyncDelegate(), EPackageFlags InPackageFlags = PKG_None, int32 InPIEInstanceID = INDEX_NONE, uint32 InPackagePriority = 0);
+COREUOBJECT_API int32 LoadPackageAsync(const FString& InName, const FGuid* InGuid, FName InType, const TCHAR* InPackageToLoadFrom = nullptr, FLoadPackageAsyncDelegate InCompletionDelegate = FLoadPackageAsyncDelegate(), EPackageFlags InPackageFlags = PKG_None, int32 InPIEInstanceID = INDEX_NONE, TAsyncLoadPriority InPackagePriority = 0);
 
 /**
 * Asynchronously load a package and all contained objects that match context flags. Non- blocking.
@@ -327,7 +330,7 @@ COREUOBJECT_API int32 LoadPackageAsync(const FString& InName, const FGuid* InGui
 * @param	InPackagePriority		Loading priority
 * @return Unique ID associated with this load request (the same package can be associated with multiple IDs).
 */
-COREUOBJECT_API int32 LoadPackageAsync(const FString& InName, const FGuid* InGuid = nullptr, const TCHAR* InPackageToLoadFrom = nullptr, FLoadPackageAsyncDelegate InCompletionDelegate = FLoadPackageAsyncDelegate(), EPackageFlags InPackageFlags = PKG_None, int32 InPIEInstanceID = INDEX_NONE, uint32 InPackagePriority = 0);
+COREUOBJECT_API int32 LoadPackageAsync(const FString& InName, const FGuid* InGuid = nullptr, const TCHAR* InPackageToLoadFrom = nullptr, FLoadPackageAsyncDelegate InCompletionDelegate = FLoadPackageAsyncDelegate(), EPackageFlags InPackageFlags = PKG_None, int32 InPIEInstanceID = INDEX_NONE, TAsyncLoadPriority InPackagePriority = 0);
 
 /**
 * Asynchronously load a package and all contained objects that match context flags. Non- blocking.
@@ -337,7 +340,7 @@ COREUOBJECT_API int32 LoadPackageAsync(const FString& InName, const FGuid* InGui
 * @param	InPackagePriority		Loading priority
 * @return Unique ID associated with this load request (the same package can be associated with multiple IDs).
 */
-COREUOBJECT_API int32 LoadPackageAsync(const FString& InName, FLoadPackageAsyncDelegate InCompletionDelegate, uint32 InPackagePriority = 0, EPackageFlags InPackageFlags = PKG_None);
+COREUOBJECT_API int32 LoadPackageAsync(const FString& InName, FLoadPackageAsyncDelegate InCompletionDelegate, TAsyncLoadPriority InPackagePriority = 0, EPackageFlags InPackageFlags = PKG_None);
 
 /**
 * Cancels all async package loading requests.
