ConstructDynamicType works also when a type was saved as dynamic class.

diff --git a/Engine/Source/Runtime/CoreUObject/Private/UObject/UObjectBase.cpp b/Engine/Source/Runtime/CoreUObject/Private/UObject/UObjectBase.cpp
--- a/Engine/Source/Runtime/CoreUObject/Private/UObject/UObjectBase.cpp
+++ b/Engine/Source/Runtime/CoreUObject/Private/UObject/UObjectBase.cpp
@@ -1085,7 +1085,7 @@ UEnum* ConstructDynamicEnum(FName EnumName)
 UObject* ConstructDynamicType(FName TypeName, FName TypeClass)
 {
 	UObject* Result = nullptr;
-	if (TypeClass == UClass::StaticClass()->GetFName())
+	if ((TypeClass == UClass::StaticClass()->GetFName()) || (TypeClass == UDynamicClass::StaticClass()->GetFName()))
 	{
 		Result = ConstructDynamicClass(TypeName);
 	}
