Removing unused SerializeRootSet

diff --git a/Engine/Source/Runtime/CoreUObject/Private/UObject/GarbageCollection.cpp b/Engine/Source/Runtime/CoreUObject/Private/UObject/GarbageCollection.cpp
--- a/Engine/Source/Runtime/CoreUObject/Private/UObject/GarbageCollection.cpp
+++ b/Engine/Source/Runtime/CoreUObject/Private/UObject/GarbageCollection.cpp
@@ -378,25 +378,6 @@ static void LogClassCountInfo( const TCHAR* LogText, TMap<const FName,uint32>& C
 };
 #endif
 
-/**
- * Serializes the global root set and objects that have any of the passed in KeepFlags to passed in archive.
- *
- * @param Ar			Archive to serialize with
- * @param KeepFlags		Objects with any of those flags will be serialized regardless of whether they are part of the root
- *						set or not.
- */
-void SerializeRootSet( FArchive& Ar, EObjectFlags KeepFlags )
-{
-	for( FObjectIterator It; It; ++It )
-	{
-		UObject* Obj = *It;
-		if(	Obj->HasAnyFlags(KeepFlags|RF_RootSet) )
-		{
-			Ar << Obj;
-		}
-	}
-}
-
 /**
  * Handles object reference, potentially NULL'ing
  *
diff --git a/Engine/Source/Runtime/CoreUObject/Public/UObject/UObjectGlobals.h b/Engine/Source/Runtime/CoreUObject/Public/UObject/UObjectGlobals.h
--- a/Engine/Source/Runtime/CoreUObject/Public/UObject/UObjectGlobals.h
+++ b/Engine/Source/Runtime/CoreUObject/Public/UObject/UObjectGlobals.h
@@ -375,7 +375,6 @@ COREUOBJECT_API void CollectGarbage(EObjectFlags KeepFlags, bool bPerformFullPur
 * @param	bPerformFullPurge	if true, perform a full purge after the mark pass
 */
 COREUOBJECT_API bool TryCollectGarbage(EObjectFlags KeepFlags, bool bPerformFullPurge = true);
-COREUOBJECT_API void SerializeRootSet(FArchive& Ar, EObjectFlags KeepFlags);
 
 /**
  * Returns whether an incremental purge is still pending/ in progress.
