Add support for more map types to AddReferencedObjects<> map specialization.

diff --git a/Engine/Source/Runtime/CoreUObject/Public/UObject/UObjectGlobals.h b/Engine/Source/Runtime/CoreUObject/Public/UObject/UObjectGlobals.h
--- a/Engine/Source/Runtime/CoreUObject/Public/UObject/UObjectGlobals.h
+++ b/Engine/Source/Runtime/CoreUObject/Public/UObject/UObjectGlobals.h
@@ -1577,8 +1577,8 @@ public:
 	* @param ReferencingObject Referencing object (if available).
 	* @param ReferencingProperty Referencing property (if available).
 	*/
-	template <typename TKeyType, typename TValueType>
-	void AddReferencedObjects(TMap<TKeyType, TValueType>& Map, const UObject* ReferencingObject = NULL, const UProperty* ReferencingProperty = NULL)
+	template <typename TKeyType, typename TValueType, typename TAllocator, typename TKeyFuncs >
+	void AddReferencedObjects(TMapBase<TKeyType, TValueType, TAllocator, TKeyFuncs>& Map, const UObject* ReferencingObject = NULL, const UProperty* ReferencingProperty = NULL)
 	{
 		static_assert(CanConverFromTo<TKeyType, UObjectBase*>::Result || CanConverFromTo<TValueType, UObjectBase*>::Result, "At least one of TMap template types must be derived from UObject");
 		for (auto& It : Map)
