New way of replacing old BPGC references with new (converted and native) DynamicClass references in Package::Save

diff --git a/Engine/Source/Developer/BlueprintCompilerCppBackend/Private/BlueprintCompilerCppBackendGatherDependencies.cpp b/Engine/Source/Developer/BlueprintCompilerCppBackend/Private/BlueprintCompilerCppBackendGatherDependencies.cpp
--- a/Engine/Source/Developer/BlueprintCompilerCppBackend/Private/BlueprintCompilerCppBackendGatherDependencies.cpp
+++ b/Engine/Source/Developer/BlueprintCompilerCppBackend/Private/BlueprintCompilerCppBackendGatherDependencies.cpp
@@ -2,6 +2,7 @@
 
 #include "BlueprintCompilerCppBackendModulePrivatePCH.h"
 #include "BlueprintCompilerCppBackendGatherDependencies.h"
+#include "BlueprintSupport.h"
 
 FGatherConvertedClassDependencies::FGatherConvertedClassDependencies(UStruct* InStruct) : OriginalStruct(InStruct)
 {
@@ -12,9 +13,7 @@ FGatherConvertedClassDependencies::FGatherConvertedClassDependencies(UStruct* In
 
 bool FGatherConvertedClassDependencies::WillClassBeConverted(const UBlueprintGeneratedClass* InClass) const
 {
-	// TODO:
-
-	return InClass && !InClass->HasAnyFlags(RF_ClassDefaultObject);
+	return FReplaceCookedBPGC::Get().CouldBeConverted(InClass);
 }
 
 void FGatherConvertedClassDependencies::DependenciesForHeader()
diff --git a/Engine/Source/Runtime/CoreUObject/Private/Blueprint/BlueprintSupport.cpp b/Engine/Source/Runtime/CoreUObject/Private/Blueprint/BlueprintSupport.cpp
--- a/Engine/Source/Runtime/CoreUObject/Private/Blueprint/BlueprintSupport.cpp
+++ b/Engine/Source/Runtime/CoreUObject/Private/Blueprint/BlueprintSupport.cpp
@@ -1767,4 +1767,136 @@ void FConvertedBlueprintsDependencies::GetAssets(FName ClassName, TArray<FName>&
 	{
 		Func(OutPackagePaths);
 	}
-}
\ No newline at end of file
+}
+
+#if WITH_EDITOR
+
+struct FConvertibleTypeHelper
+{
+	static FName NameBPGC;
+	static FName NameWidgetBPGC;
+	static FName NameUDS;
+	static FName NameUDE;
+
+	static bool IsConvertible(const UObject* Obj)
+	{
+		const FName ClassName = Obj ? Obj->GetClass()->GetFName() : NAME_None;
+		return (ClassName != NAME_None)
+			&& ((ClassName == NameBPGC)
+			|| (ClassName == NameUDS)
+			|| (ClassName == NameUDE)
+			|| (ClassName == NameWidgetBPGC));
+	}
+};
+
+FName FConvertibleTypeHelper::NameBPGC(TEXT("BlueprintGeneratedClass"));
+FName FConvertibleTypeHelper::NameWidgetBPGC(TEXT("WidgetBlueprintGeneratedClass"));
+FName FConvertibleTypeHelper::NameUDS(TEXT("UserDefinedStruct"));
+FName FConvertibleTypeHelper::NameUDE(TEXT("UserDefinedEnum"));
+
+FReplaceCookedBPGC::FReplaceCookedBPGC()
+	: NativeScriptPackage(nullptr)
+	, bEnabled(false)
+{
+}
+
+FReplaceCookedBPGC& FReplaceCookedBPGC::Get()
+{
+	static FReplaceCookedBPGC ReplaceCookedBPGC;
+	return ReplaceCookedBPGC;
+}
+
+void FReplaceCookedBPGC::AddReferencedObjects(FReferenceCollector& Collector)
+{
+	Collector.AddReferencedObjects(ReplaceMap);
+	Collector.AddReferencedObject(NativeScriptPackage);
+}
+
+void FReplaceCookedBPGC::AddConvertedPackageName(FName PackageName)
+{
+	ConvertedPackagesNames.Add(PackageName);
+}
+
+void FReplaceCookedBPGC::AddConvertedFieldStub(UObject* Object)
+{
+	if (Object && ensure(CouldBeConverted(Object)))
+	{
+		const FName ClassName = Object->GetClass()->GetFName();
+		if ((ClassName == FConvertibleTypeHelper::NameBPGC) || (ClassName == FConvertibleTypeHelper::NameWidgetBPGC))
+		{
+			AddConvertedFieldStub<UDynamicClass>(Object->GetPathName(), Object->GetFName());
+		}
+		else if (ClassName == FConvertibleTypeHelper::NameUDS)
+		{
+			AddConvertedFieldStub<UScriptStruct>(Object->GetPathName(), Object->GetFName());
+		}
+		else if (ClassName == FConvertibleTypeHelper::NameUDE)
+		{
+			AddConvertedFieldStub<UEnum>(Object->GetPathName(), Object->GetFName());
+		}
+		else
+		{
+			check(false);
+		}
+	}
+}
+
+UObject* FReplaceCookedBPGC::FindReplacementStub(UObject* Object)
+{
+	if (bEnabled && Object && CouldBeConverted(Object))
+	{
+		const FString OriginalPathName = Object->GetPathName();
+		UObject** StubObjPtr = ReplaceMap.Find(OriginalPathName);
+		UObject* StubObj = StubObjPtr ? *StubObjPtr : nullptr;
+		check(!StubObj || StubObj->IsValidLowLevelFast());
+		if (StubObj)
+		{
+			int32 k = 0; k++;
+		}
+		return StubObj;
+	}
+	return nullptr;
+}
+
+bool FReplaceCookedBPGC::CouldBeConverted(const UObject* Object) const
+{
+	if (Object && !Object->HasAnyFlags(RF_ClassDefaultObject) && FConvertibleTypeHelper::IsConvertible(Object))
+	{
+		// TODO: check excluded assets list.
+		return true;
+	}
+	return false;
+	
+}
+
+bool FReplaceCookedBPGC::PackageShouldNotBeSaved(const UPackage* InOuter) const
+{
+	return InOuter && bEnabled && ConvertedPackagesNames.Contains(InOuter->GetFName());
+}
+
+void FReplaceCookedBPGC::Initialize(const FString& NativePackageName)
+{
+	if (ensure(!NativeScriptPackage))
+	{
+		NativeScriptPackage = CreatePackage(nullptr, *NativePackageName);
+	}
+	bEnabled = true;
+}
+
+void DummyClassConstructor(const FObjectInitializer&)
+{
+	check(false);
+}
+
+void FReplaceCookedBPGC::AdditionalStubFieldInitialization(UField* Stub)
+{
+	check(Stub && Stub->IsValidLowLevelFast());
+	if (auto Class = Cast<UClass>(Stub))
+	{
+		Class->ClassConstructor = &DummyClassConstructor;
+		Class->ClassAddReferencedObjects = &UObject::AddReferencedObjects;
+	}
+	Stub->SetFlags(RF_Native);
+}
+
+#endif //WITH_EDITOR
\ No newline at end of file
diff --git a/Engine/Source/Runtime/CoreUObject/Private/UObject/LinkerSave.cpp b/Engine/Source/Runtime/CoreUObject/Private/UObject/LinkerSave.cpp
--- a/Engine/Source/Runtime/CoreUObject/Private/UObject/LinkerSave.cpp
+++ b/Engine/Source/Runtime/CoreUObject/Private/UObject/LinkerSave.cpp
@@ -136,6 +136,7 @@ void FLinkerSave::Serialize( void* V, int64 Length )
 FArchive& FLinkerSave::operator<<( FName& InName )
 {
 	int32 Save = MapName(InName);
+	ensure(Save != INDEX_NONE);
 	int32 Number = InName.GetNumber();
 	FArchive& Ar = *this;
 	return Ar << Save << Number;
diff --git a/Engine/Source/Runtime/CoreUObject/Private/UObject/SavePackage.cpp b/Engine/Source/Runtime/CoreUObject/Private/UObject/SavePackage.cpp
--- a/Engine/Source/Runtime/CoreUObject/Private/UObject/SavePackage.cpp
+++ b/Engine/Source/Runtime/CoreUObject/Private/UObject/SavePackage.cpp
@@ -5,6 +5,7 @@
 #include "Interface.h"
 #include "TargetPlatform.h"
 #include "UObject/UObjectThreadContext.h"
+#include "BlueprintSupport.h"
 
 DEFINE_LOG_CATEGORY_STATIC(LogSavePackage, Log, All);
 
@@ -305,6 +306,21 @@ public:
 		ReferencedNames.Add(NameNoNumber);
 	}
 
+#if WITH_EDITOR
+	void AddReplacementsNames(UObject* Obj)
+	{
+		if (Obj && FReplaceCookedBPGC::Get().IsEnabled())
+		{
+			if (UObject* ReplObj = FReplaceCookedBPGC::Get().FindReplacementStub(Obj))
+			{
+				MarkNameAsReferenced(ReplObj->GetFName());
+				MarkNameAsReferenced(ReplObj->GetClass()->GetFName());
+				MarkNameAsReferenced(ReplObj->GetOuter()->GetFName());
+			}
+		}
+	}
+#endif //WITH_EDITOR
+
 	/** 
 	 * Add the marked names to a linker
 	 * @param Linker linker to save the marked names to
@@ -823,6 +839,9 @@ FArchive& FArchiveSaveTagImports::operator<<( UObject*& Obj )
 				if (!bIsEditorOnly)
 				{
 					Obj->Mark(OBJECTMARK_TagImp);
+#if WITH_EDITOR
+					SavePackageState->AddReplacementsNames(Obj);
+#endif //WITH_EDITOR
 				}
 
 				if ( Obj->HasAnyFlags(RF_ClassDefaultObject) )
@@ -3084,6 +3103,52 @@ TMap<UObject*, UObject*> UnmarkExportTagFromDuplicates()
 
 #if WITH_EDITOR
 COREUOBJECT_API extern bool GOutputCookingWarnings;
+
+struct FImportReplacementsHelper
+{
+	TSet<UObject*> UniqueTagImpObjects;
+
+	FImportReplacementsHelper(const TArray<UObject*>& InTagImpObjects)
+	{
+		if (FReplaceCookedBPGC::Get().IsEnabled())
+		{
+			UniqueTagImpObjects = TSet<UObject*>(InTagImpObjects);
+		}
+	}
+
+	UObject* Handle(UObject* Obj, TArray<FObjectImport>& ImportMap)
+	{
+		if (FReplaceCookedBPGC::Get().IsEnabled())
+		{
+			UObject* ReplacedObjectToImport = FReplaceCookedBPGC::Get().FindReplacementStub(Obj);
+			Obj = ReplacedObjectToImport ? ReplacedObjectToImport : Obj;
+			if (ReplacedObjectToImport)
+			{
+				{
+					ensure(ReplacedObjectToImport->GetOuter() == ReplacedObjectToImport->GetOutermost());
+					bool bAlreadyAdded = false;
+					UniqueTagImpObjects.Add(ReplacedObjectToImport->GetOuter(), &bAlreadyAdded);
+					if (!bAlreadyAdded)
+					{
+						new(ImportMap)FObjectImport(ReplacedObjectToImport->GetOuter());
+					}
+				}
+
+									{
+										bool bAlreadyAdded = false;
+										UniqueTagImpObjects.Add(ReplacedObjectToImport->GetClass(), &bAlreadyAdded);
+										if (!bAlreadyAdded)
+										{
+											new(ImportMap)FObjectImport(ReplacedObjectToImport->GetClass());
+										}
+									}
+			}
+		}
+		return Obj;
+	}
+
+};
+
 #endif
 
 ESavePackageResult UPackage::Save(UPackage* InOuter, UObject* Base, EObjectFlags TopLevelFlags, const TCHAR* Filename,
@@ -3119,6 +3184,14 @@ ESavePackageResult UPackage::Save(UPackage* InOuter, UObject* Base, EObjectFlags
 				UE_CLOG(!(SaveFlags & SAVE_NoError), LogSavePackage, Display, TEXT("Package marked as editor-only: %s. Package will not be saved."), *InOuter->GetName());
 				return ESavePackageResult::ReferencedOnlyByEditorOnlyData;
 			}
+#if WITH_EDITOR
+			else if (FReplaceCookedBPGC::Get().PackageShouldNotBeSaved(InOuter))
+			{
+				UE_LOG(LogSavePackage, Display, TEXT("Package %s contains assets, that were converted into native code. Package will not be saved."), *InOuter->GetName());
+				// Or replace with the NativeScriptPackage ?
+				return ESavePackageResult::ContainsConvertedAssets;
+			}
+#endif	// #if WITH_EDITOR
 		}
 #endif
 		// if we are cooking we should be doing it in the editor
@@ -3468,6 +3541,18 @@ ESavePackageResult UPackage::Save(UPackage* InOuter, UObject* Base, EObjectFlags
 						{
 							UE_LOG(LogSavePackage, Fatal, TEXT("%s"), *FString::Printf( TEXT("Transient object imported: %s"), *Obj->GetFullName() ) );
 						}
+#if WITH_EDITOR
+						if (FReplaceCookedBPGC::Get().IsEnabled())
+						{
+							for (UObject*& Dependency : ImportTagger.Dependencies)
+							{
+								if (UObject* NewDependency = FReplaceCookedBPGC::Get().FindReplacementStub(Dependency))
+								{
+									Dependency = NewDependency;
+								}
+							}
+						}
+#endif //WITH_EDITOR
 
 						// add the list of dependencies to the dependency map
 						ObjectDependencies.Add(Obj, ImportTagger.Dependencies);
@@ -3541,6 +3626,9 @@ ESavePackageResult UPackage::Save(UPackage* InOuter, UObject* Base, EObjectFlags
 						check(Obj->HasAnyMarks(EObjectMark(OBJECTMARK_TagExp|OBJECTMARK_TagImp)));
 
 						SavePackageState->MarkNameAsReferenced(Obj->GetFName());
+#if WITH_EDITOR
+						SavePackageState->AddReplacementsNames(Obj);
+#endif //WITH_EDITOR
 						if( Obj->GetOuter() )
 						{
 							SavePackageState->MarkNameAsReferenced(Obj->GetOuter()->GetFName());
@@ -3853,11 +3941,17 @@ ESavePackageResult UPackage::Save(UPackage* InOuter, UObject* Base, EObjectFlags
 				{
 					TArray<UObject*> TagImpObjects;
 					GetObjectsWithAnyMarks(TagImpObjects, OBJECTMARK_TagImp);
+#if WITH_EDITOR
+					FImportReplacementsHelper ImportReplacementsHelper(TagImpObjects);
+#endif //WITH_EDITOR
 					for(int32 Index = 0; Index < TagImpObjects.Num(); Index++)
 					{
 						UObject* Obj = TagImpObjects[Index];
 						check(Obj->HasAnyMarks(OBJECTMARK_TagImp));
-						new( Linker->ImportMap )FObjectImport( Obj );
+#if WITH_EDITOR
+						Obj = ImportReplacementsHelper.Handle(Obj, Linker->ImportMap);
+#endif //WITH_EDITOR
+						new(Linker->ImportMap)FObjectImport(Obj);
 					}
 				}
 
@@ -4236,9 +4330,20 @@ ESavePackageResult UPackage::Save(UPackage* InOuter, UObject* Base, EObjectFlags
 						{
 							// Set class index.
 							// If this is *exactly* a UClass, store null instead; for anything else, including UClass-derived classes, map it
-							if( Export.Object->GetClass() != UClass::StaticClass() )
+							UClass* ObjClass = Export.Object->GetClass();
+#if WITH_EDITOR
+							if (FReplaceCookedBPGC::Get().IsEnabled())
+							{
+								if (UObject* NewObjClass = FReplaceCookedBPGC::Get().FindReplacementStub(ObjClass))
+								{
+									ObjClass = CastChecked<UClass>(NewObjClass);
+								}
+							}
+#endif	// #if WITH_EDITOR
+
+							if (ObjClass != UClass::StaticClass())
 							{
-								Export.ClassIndex = Linker->MapObject(Export.Object->GetClass());
+								Export.ClassIndex = Linker->MapObject(ObjClass);
 								check(!Export.ClassIndex.IsNull());
 							}
 							else
diff --git a/Engine/Source/Runtime/CoreUObject/Public/Blueprint/BlueprintSupport.h b/Engine/Source/Runtime/CoreUObject/Public/Blueprint/BlueprintSupport.h
--- a/Engine/Source/Runtime/CoreUObject/Public/Blueprint/BlueprintSupport.h
+++ b/Engine/Source/Runtime/CoreUObject/Public/Blueprint/BlueprintSupport.h
@@ -65,10 +65,14 @@ private:
 };
 
 /**
-* Stores info about BPCG, UDS, UDE, etc,, converted into native code.
-* Native entities (generated from BP items) have "ReplaceConverted" with original object path.
-* This path is used to fix linker.
-*/
+ * Stores info about BPCG, UDS, UDE, etc.. converted into native code.
+ *
+ * Native entities (generated from BP items) have "ReplaceConverted" with original object path.
+ * This path is used to fix linker while loading.
+ * 
+ * This struct is used only when the code converted from BP is compiled in an editor build. 
+ * It's usually only for a test/development purpose.
+ */
 struct COREUOBJECT_API FReplaceConvertedAssetManager : public FGCObject
 {
 private:
@@ -106,6 +110,55 @@ public:
 	static void OnModulesChanged(FName ModuleThatChanged, EModuleChangeReason ReasonForChange);
 };
 
+/**
+ * The struct is used while saving cooked package to find replacements of converted BP assets.
+ */
+struct COREUOBJECT_API FReplaceCookedBPGC : public FGCObject
+{
+private:
+	UPackage* NativeScriptPackage;
+	TMap<FString, UObject*> ReplaceMap;
+	TSet<FName> ConvertedPackagesNames;
+	bool bEnabled;
+
+	FReplaceCookedBPGC();
+
+	void AdditionalStubFieldInitialization(UField* Stub);
+
+public:
+	static FReplaceCookedBPGC& Get();
+	virtual void AddReferencedObjects(FReferenceCollector& Collector) override;
+
+	// INITIALIZATION
+	void Initialize(const FString& NativePackageName);
+	void AddConvertedPackageName(FName PackageName);
+	template<typename FieldType>
+	void AddConvertedFieldStub(const FString& OriginalPathName, const FName Name)
+	{
+		UObject*& NativeObj = ReplaceMap.FindOrAdd(OriginalPathName);
+		if (!NativeObj)
+		{
+			check(NativeScriptPackage);
+			FieldType* NewField = NewObject<FieldType>(NativeScriptPackage, Name);
+			AdditionalStubFieldInitialization(NewField);
+			NewField->Bind();
+			NativeObj = NewField;
+		}
+	}
+	void AddConvertedFieldStub(UObject* Object);
+
+	// USAGE
+	bool IsEnabled() const
+	{
+		return bEnabled;
+	}
+	UObject* FindReplacementStub(UObject* Object);
+
+	// The object is of a convertible type, and the object is not listed in excluded assets.
+	bool CouldBeConverted(const UObject* Object) const;
+	bool PackageShouldNotBeSaved(const UPackage* InOuter) const;
+};
+
 #endif //WITH_EDITOR
 
 /** 
diff --git a/Engine/Source/Runtime/CoreUObject/Public/UObject/CoreObject.h b/Engine/Source/Runtime/CoreUObject/Public/UObject/CoreObject.h
--- a/Engine/Source/Runtime/CoreUObject/Public/UObject/CoreObject.h
+++ b/Engine/Source/Runtime/CoreUObject/Public/UObject/CoreObject.h
@@ -33,7 +33,9 @@ enum class ESavePackageResult
 	/** [When cooking] Package was not saved because it contained editor-only data */
 	ContainsEditorOnlyData,
 	/** [When cooking] Package was not saved because it was referenced by editor-only properties */
-	ReferencedOnlyByEditorOnlyData
+	ReferencedOnlyByEditorOnlyData, 
+	/** [When cooking] Package was not saved because it contains assets that were converted into native code */
+	ContainsConvertedAssets
 };
 
 /**
