Merging CL# 2696877

Merging FTlsObjectInitializers with FUObjectThreadcontext to avoid excessive TLS lookups.

diff --git a/Engine/Source/Runtime/CoreUObject/Private/UObject/Obj.cpp b/Engine/Source/Runtime/CoreUObject/Private/UObject/Obj.cpp
--- a/Engine/Source/Runtime/CoreUObject/Private/UObject/Obj.cpp
+++ b/Engine/Source/Runtime/CoreUObject/Private/UObject/Obj.cpp
@@ -13,7 +13,6 @@
 #include "Serialization/ArchiveDescribeReference.h"
 #include "FindStronglyConnected.h"
 #include "HotReloadInterface.h"
-#include "UObject/TlsObjectInitializers.h"
 #include "UObject/UObjectThreadContext.h"
 #include "ExclusiveLoadPackageTimeTracker.h"
 #include "Serialization/DeferredMessageLog.h"
@@ -81,7 +80,7 @@ void UObject::EnsureNotRetrievingVTablePtr() const
 
 UObject* UObject::CreateDefaultSubobject(FName SubobjectFName, UClass* ReturnType, UClass* ClassToCreateByDefault, bool bIsRequired, bool bAbstract, bool bIsTransient)
 {
-	auto CurrentInitializer = FTlsObjectInitializers::Top();
+	FObjectInitializer* CurrentInitializer = FUObjectThreadContext::Get().TopInitializer();
 	UE_CLOG(!CurrentInitializer, LogObj, Fatal, TEXT("No object initializer found during construction."));
 	UE_CLOG(CurrentInitializer->Obj != this, LogObj, Fatal, TEXT("Using incorrect object initializer."));
 	return CurrentInitializer->CreateDefaultSubobject(this, SubobjectFName, ReturnType, ClassToCreateByDefault, bIsRequired, bAbstract, bIsTransient);
@@ -89,7 +88,7 @@ UObject* UObject::CreateDefaultSubobject(FName SubobjectFName, UClass* ReturnTyp
 
 UObject* UObject::CreateEditorOnlyDefaultSubobjectImpl(FName SubobjectName, UClass* ReturnType, bool bTransient)
 {
-	auto CurrentInitializer = FTlsObjectInitializers::Top();
+	FObjectInitializer* CurrentInitializer = FUObjectThreadContext::Get().TopInitializer();
 	return CurrentInitializer->CreateEditorOnlyDefaultSubobject(this, SubobjectName, ReturnType, bTransient);
 }
 
diff --git a/Engine/Source/Runtime/CoreUObject/Private/UObject/TlsObjectInitializers.cpp b/Engine/Source/Runtime/CoreUObject/Private/UObject/TlsObjectInitializers.cpp
deleted file mode 100644
diff --git a/Engine/Source/Runtime/CoreUObject/Private/UObject/UObjectGlobals.cpp b/Engine/Source/Runtime/CoreUObject/Private/UObject/UObjectGlobals.cpp
--- a/Engine/Source/Runtime/CoreUObject/Private/UObject/UObjectGlobals.cpp
+++ b/Engine/Source/Runtime/CoreUObject/Private/UObject/UObjectGlobals.cpp
@@ -8,7 +8,6 @@
 #include "Misc/SecureHash.h"
 #include "UObject/ConstructorHelpers.h"
 #include "UObject/UObjectAnnotation.h"
-#include "UObject/TlsObjectInitializers.h"
 #include "UObject/UObjectThreadContext.h"
 #include "BlueprintSupport.h" // for FDeferredObjInitializerTracker
 #include "ExclusiveLoadPackageTimeTracker.h"
@@ -2090,7 +2089,7 @@ UObject::UObject()
 	EnsureNotRetrievingVTablePtr();
 #endif // WITH_HOT_RELOAD_CTORS
 
-	FObjectInitializer* ObjectInitializerPtr = FTlsObjectInitializers::Top();
+	FObjectInitializer* ObjectInitializerPtr = FUObjectThreadContext::Get().TopInitializer();
 	UE_CLOG(!ObjectInitializerPtr, LogUObjectGlobals, Fatal, TEXT("%s is not being constructed with either NewObject, NewNamedObject or ConstructObject."), *GetName());
 	FObjectInitializer& ObjectInitializer = *ObjectInitializerPtr;
 	UE_CLOG(ObjectInitializer.Obj != nullptr && ObjectInitializer.Obj != this, LogUObjectGlobals, Fatal, TEXT("UObject() constructor called but it's not the object that's currently being constructed with NewObject. Maybe you trying to construct it on the stack which is not supported."));
@@ -2123,12 +2122,12 @@ FObjectInitializer::FObjectInitializer()
 	, bIsDeferredInitializer(false)
 #endif // USE_CIRCULAR_DEPENDENCY_LOAD_DEFERRING
 {
-	auto& ThreadContext = FUObjectThreadContext::Get();
+	FUObjectThreadContext& ThreadContext = FUObjectThreadContext::Get();
 	// Mark we're in the constructor now.	
 	ThreadContext.IsInConstructor++;
 	LastConstructedObject = ThreadContext.ConstructedObject;
 	ThreadContext.ConstructedObject = Obj;
-	FTlsObjectInitializers::Push(this);
+	ThreadContext.PushInitializer(this);
 }	
 
 FObjectInitializer::FObjectInitializer(UObject* InObj, UObject* InObjectArchetype, bool bInCopyTransientsFromClassDefaults, bool bInShouldIntializeProps, struct FObjectInstancingGraph* InInstanceGraph)
@@ -2144,12 +2143,12 @@ FObjectInitializer::FObjectInitializer(UObject* InObj, UObject* InObjectArchetyp
 	, bIsDeferredInitializer(false)
 #endif // USE_CIRCULAR_DEPENDENCY_LOAD_DEFERRING
 {
-	auto& ThreadContext = FUObjectThreadContext::Get();
+	FUObjectThreadContext& ThreadContext = FUObjectThreadContext::Get();
 	// Mark we're in the constructor now.
 	ThreadContext.IsInConstructor++;
 	LastConstructedObject = ThreadContext.ConstructedObject;
 	ThreadContext.ConstructedObject = Obj;
-	FTlsObjectInitializers::Push(this);
+	ThreadContext.PushInitializer(this);
 }
 
 /**
@@ -2159,16 +2158,16 @@ FObjectInitializer::~FObjectInitializer()
 {
 
 #if USE_CIRCULAR_DEPENDENCY_LOAD_DEFERRING
-	// if we're not at the top of FTlsObjectInitializers, then this is most 
+	// if we're not at the top of ObjectInitializers, then this is most 
 	// likely a deferred FObjectInitializer that's a copy of one that was used 
 	// in a constructor (that has already been popped)
 	if (!bIsDeferredInitializer)
 	{
 #endif // USE_CIRCULAR_DEPENDENCY_LOAD_DEFERRING
-		check(FTlsObjectInitializers::Top() == this);
-
-		FTlsObjectInitializers::Pop();
-		auto& ThreadContext = FUObjectThreadContext::Get();
+		FUObjectThreadContext& ThreadContext = FUObjectThreadContext::Get();
+		check(ThreadContext.TopInitializer() == this);
+		ThreadContext.PopInitializer();
+		
 		// Let the FObjectFinders know we left the constructor.
 		ThreadContext.IsInConstructor--;
 		check(ThreadContext.IsInConstructor >= 0);
@@ -2231,10 +2230,10 @@ FObjectInitializer::~FObjectInitializer()
 					bIsPostConstructInitDeferred = true;
 					DeferredCopy->bIsDeferredInitializer = true;
 
-					// make sure this wasn't mistakenly pushed into FTlsObjectInitializers
+					// make sure this wasn't mistakenly pushed into ObjectInitializers
 					// (the copy constructor should have been what was invoked, 
-					// which doesn't push to FTlsObjectInitializers)
-					check(FTlsObjectInitializers::Top() != DeferredCopy);
+					// which doesn't push to ObjectInitializers)
+					check(FUObjectThreadContext::Get().TopInitializer() != DeferredCopy);
 				}
 			}
 		}
@@ -2666,15 +2665,15 @@ UObject* StaticConstructObject
 
 void FObjectInitializer::AssertIfInConstructor(UObject* Outer, const TCHAR* ErrorMessage)
 {
-	auto& ThreadContext = FUObjectThreadContext::Get();
+	FUObjectThreadContext& ThreadContext = FUObjectThreadContext::Get();
 	UE_CLOG(ThreadContext.IsInConstructor && Outer == ThreadContext.ConstructedObject, LogUObjectGlobals, Fatal, TEXT("%s"), ErrorMessage);
 }
 
 FObjectInitializer& FObjectInitializer::Get()
 {
-	auto& ThreadContext = FUObjectThreadContext::Get();
+	FUObjectThreadContext& ThreadContext = FUObjectThreadContext::Get();
 	UE_CLOG(!ThreadContext.IsInConstructor, LogUObjectGlobals, Fatal, TEXT("FObjectInitializer::Get() can only be used inside of UObject-derived class constructor."));
-	return FTlsObjectInitializers::TopChecked();
+	return ThreadContext.TopInitializerChecked();
 }
 
 /**
@@ -2710,7 +2709,7 @@ void FScopedObjectFlagMarker::RestoreObjectFlags()
 
 void ConstructorHelpers::FailedToFind(const TCHAR* ObjectToFind)
 {
-	auto CurrentInitializer = FTlsObjectInitializers::Top();
+	FObjectInitializer* CurrentInitializer = FUObjectThreadContext::Get().TopInitializer();
 	const FString Message = FString::Printf(TEXT("CDO Constructor (%s): Failed to find %s\n"),
 		(CurrentInitializer && CurrentInitializer->GetClass()) ? *CurrentInitializer->GetClass()->GetName() : TEXT("Unknown"),
 		ObjectToFind);
@@ -2727,7 +2726,7 @@ void ConstructorHelpers::CheckFoundViaRedirect(UObject *Object, const FString& P
 		NewString.ReplaceInline(TEXT(" "), TEXT("'"), ESearchCase::CaseSensitive);
 		NewString += TEXT("'");
 
-		auto CurrentInitializer = FTlsObjectInitializers::Top();
+		FObjectInitializer* CurrentInitializer = FUObjectThreadContext::Get().TopInitializer();
 		const FString Message = FString::Printf(TEXT("CDO Constructor (%s): Followed redirector (%s), change code to new path (%s)\n"),
 			(CurrentInitializer && CurrentInitializer->GetClass()) ? *CurrentInitializer->GetClass()->GetName() : TEXT("Unknown"),
 			ObjectToFind, *NewString);
diff --git a/Engine/Source/Runtime/CoreUObject/Private/UObject/UObjectThreadContext.cpp b/Engine/Source/Runtime/CoreUObject/Private/UObject/UObjectThreadContext.cpp
--- a/Engine/Source/Runtime/CoreUObject/Private/UObject/UObjectThreadContext.cpp
+++ b/Engine/Source/Runtime/CoreUObject/Private/UObject/UObjectThreadContext.cpp
@@ -7,6 +7,8 @@
 #include "CoreUObjectPrivate.h"
 #include "UObject/UObjectThreadContext.h"
 
+DEFINE_LOG_CATEGORY(LogUObjectThreadContext);
+
 FUObjectThreadContext::FUObjectThreadContext()
 : ImportCount(0)
 , ForcedExportCount(0)
diff --git a/Engine/Source/Runtime/CoreUObject/Public/UObject/TlsObjectInitializers.h b/Engine/Source/Runtime/CoreUObject/Public/UObject/TlsObjectInitializers.h
deleted file mode 100644
diff --git a/Engine/Source/Runtime/CoreUObject/Public/UObject/UObjectThreadContext.h b/Engine/Source/Runtime/CoreUObject/Public/UObject/UObjectThreadContext.h
--- a/Engine/Source/Runtime/CoreUObject/Public/UObject/UObjectThreadContext.h
+++ b/Engine/Source/Runtime/CoreUObject/Public/UObject/UObjectThreadContext.h
@@ -6,14 +6,57 @@
 
 #pragma once
 
+DECLARE_LOG_CATEGORY_EXTERN(LogUObjectThreadContext, Log, All);
+
 class COREUOBJECT_API FUObjectThreadContext : public TThreadSingleton<FUObjectThreadContext>
 {
 	friend TThreadSingleton<FUObjectThreadContext>;
 
 	FUObjectThreadContext();
 
+	/** Stack of currently used FObjectInitializers for this thread */
+	TArray<FObjectInitializer*> InitializerStack;
+
 public:
 	
+	/**
+	* Remove top element from the stack.
+	*/
+	void PopInitializer()
+	{
+		InitializerStack.Pop(/*bAllowShrinking=*/ false);
+	}
+
+	/**
+	* Push new FObjectInitializer on stack.
+	* @param	Initializer			Object initializer to push.
+	*/
+	void PushInitializer(FObjectInitializer* Initializer)
+	{
+		InitializerStack.Push(Initializer);
+	}
+
+	/**
+	* Retrieve current FObjectInitializer for current thread.
+	* @return Current FObjectInitializer.
+	*/
+	FObjectInitializer* TopInitializer()
+	{
+		return InitializerStack.Num() ? InitializerStack.Last() : nullptr;
+	}
+
+	/**
+	* Retrieves current FObjectInitializer for current thread. Will assert of no ObjectInitializer is currently set.
+	* @return Current FObjectInitializer reference.
+	*/
+	FObjectInitializer& TopInitializerChecked()
+	{
+		FObjectInitializer* ObjectInitializerPtr = TopInitializer();
+		UE_CLOG(!ObjectInitializerPtr, LogUObjectThreadContext, Fatal, TEXT("Tried to get the current ObjectInitializer, but none is set. Please use NewObject or NewNamedObject to construct new UObject-derived classes."));
+		return *ObjectInitializerPtr;
+	}
+
+
 	/** Imports for EndLoad optimization.	*/
 	int32 ImportCount;
 	/** Forced exports for EndLoad optimization. */
