Added comments in SavePackage.cpp

diff --git a/Engine/Source/Runtime/CoreUObject/Private/UObject/SavePackage.cpp b/Engine/Source/Runtime/CoreUObject/Private/UObject/SavePackage.cpp
--- a/Engine/Source/Runtime/CoreUObject/Private/UObject/SavePackage.cpp
+++ b/Engine/Source/Runtime/CoreUObject/Private/UObject/SavePackage.cpp
@@ -307,6 +307,9 @@ public:
 	}
 
 #if WITH_EDITOR
+	/**
+	 * Names are cached before we replace objects for imports. So the names of replacements must be ealier.
+	 */
 	void AddReplacementsNames(UObject* Obj)
 	{
 		if (Obj && FReplaceCookedBPGC::Get().IsEnabled())
@@ -3104,6 +3107,10 @@ TMap<UObject*, UObject*> UnmarkExportTagFromDuplicates()
 #if WITH_EDITOR
 COREUOBJECT_API extern bool GOutputCookingWarnings;
 
+/**
+ * Find imports to converted BPGC, that should replaced by a generated native class. 
+ * For each converted class a stub is generated. The import object is replaced with the stub.
+ */
 struct FImportReplacementsHelper
 {
 	TSet<UObject*> UniqueTagImpObjects;
