Cis fix.

diff --git a/Engine/Source/Runtime/CoreUObject/Public/UObject/UnrealType.h b/Engine/Source/Runtime/CoreUObject/Public/UObject/UnrealType.h
--- a/Engine/Source/Runtime/CoreUObject/Public/UObject/UnrealType.h
+++ b/Engine/Source/Runtime/CoreUObject/Public/UObject/UnrealType.h
@@ -1860,7 +1860,7 @@ public:
 
 	virtual FString GetCPPType(FString* ExtendedTypeText, uint32 CPPExportFlags) const override
 	{
-		return GetCPPTypeCustom(ExtendedTypeText, CPPExportFlags, PropertyClass);
+		return GetCPPTypeCustom(ExtendedTypeText, CPPExportFlags, UObjectPropertyBase::PropertyClass);
 	}
 };
 
