Merging CL 2707469.

Blueprint multithreading improvements

Moved various exception-handling-related variables from static/global to TLS.
Fixed instance of static contention in execLet.

diff --git a/Engine/Source/Runtime/Core/Private/Misc/Core.cpp b/Engine/Source/Runtime/Core/Private/Misc/Core.cpp
--- a/Engine/Source/Runtime/Core/Private/Misc/Core.cpp
+++ b/Engine/Source/Runtime/Core/Private/Misc/Core.cpp
@@ -252,9 +252,6 @@ bool					GPumpingMessagesOutsideOfMainLoop = false;
 /** Total blueprint compile time.																			*/
 double GBlueprintCompileTime = 0.0;
 
-/** Stack names from the VM to be unrolled when we assert */
-TArray<FScriptTraceStackNode> GScriptStack;
-
 DEFINE_STAT(STAT_AudioMemory);
 DEFINE_STAT(STAT_TextureMemory);
 DEFINE_STAT(STAT_MemoryPhysXTotalAllocationSize);
diff --git a/Engine/Source/Runtime/Core/Private/Misc/CoreMisc.cpp b/Engine/Source/Runtime/Core/Private/Misc/CoreMisc.cpp
--- a/Engine/Source/Runtime/Core/Private/Misc/CoreMisc.cpp
+++ b/Engine/Source/Runtime/Core/Private/Misc/CoreMisc.cpp
@@ -988,6 +988,18 @@ FBoolConfigValueHelper::FBoolConfigValueHelper(const TCHAR* Section, const TCHAR
 	GConfig->GetBool(Section, Key, bValue, Filename);
 }
 
+/*----------------------------------------------------------------------------
+FBlueprintExceptionTracker
+----------------------------------------------------------------------------*/
+#if DO_BLUEPRINT_GUARD
+void FBlueprintExceptionTracker::ResetRunaway()
+{
+	Runaway = 0;
+	Recurse = 0;
+	bRanaway = false;
+}
+#endif // DO_BLUEPRINT_GUARD
+
 #if WITH_HOT_RELOAD_CTORS
 bool GIsRetrievingVTablePtr = false;
 
diff --git a/Engine/Source/Runtime/Core/Private/Misc/OutputDevice.cpp b/Engine/Source/Runtime/Core/Private/Misc/OutputDevice.cpp
--- a/Engine/Source/Runtime/Core/Private/Misc/OutputDevice.cpp
+++ b/Engine/Source/Runtime/Core/Private/Misc/OutputDevice.cpp
@@ -186,14 +186,15 @@ void StaticFailDebug( const TCHAR* Error, const ANSICHAR* File, int32 Line, cons
 //
 void VARARGS FDebug::LogAssertFailedMessage(const ANSICHAR* Expr, const ANSICHAR* File, int32 Line, const TCHAR* Format/*=TEXT("")*/, ...)
 {
-#if !(UE_BUILD_SHIPPING || UE_BUILD_TEST)
+#if DO_BLUEPRINT_GUARD
 	// Walk the script stack, if any
-	if( IsInGameThread() && GScriptStack.Num() > 0 )
+	FBlueprintExceptionTracker& BlueprintExceptionTracker = FBlueprintExceptionTracker::Get();
+	if( BlueprintExceptionTracker.ScriptStack.Num() > 0 )
 	{
 		FString ScriptStack = TEXT( "\n\nScript Stack:\n" );
-		while( GScriptStack.Num() )
+		while( BlueprintExceptionTracker.ScriptStack.Num() )
 		{
-			ScriptStack += GScriptStack.Pop().GetStackDescription() + TEXT( "\n" );
+			ScriptStack += BlueprintExceptionTracker.ScriptStack.Pop().GetStackDescription() + TEXT( "\n" );
 		}
 
 		UE_LOG( LogOutputDevice, Warning, TEXT( "%s" ), *ScriptStack );
diff --git a/Engine/Source/Runtime/Core/Public/CoreGlobals.h b/Engine/Source/Runtime/Core/Public/CoreGlobals.h
--- a/Engine/Source/Runtime/Core/Public/CoreGlobals.h
+++ b/Engine/Source/Runtime/Core/Public/CoreGlobals.h
@@ -309,9 +309,6 @@ extern CORE_API bool GPumpingMessagesOutsideOfMainLoop;
 /** Total blueprint compile time. */
 extern CORE_API double GBlueprintCompileTime;
 
-/** Stack names from the VM to be unrolled when we assert */
-extern CORE_API TArray<FScriptTraceStackNode> GScriptStack;
-
 #if WITH_HOT_RELOAD_CTORS
 /**
  * Ensures that current thread is during retrieval of vtable ptr of some
diff --git a/Engine/Source/Runtime/Core/Public/Misc/CoreMisc.h b/Engine/Source/Runtime/Core/Public/Misc/CoreMisc.h
--- a/Engine/Source/Runtime/Core/Public/Misc/CoreMisc.h
+++ b/Engine/Source/Runtime/Core/Public/Misc/CoreMisc.h
@@ -3,7 +3,8 @@
 #pragma once
 
 #include "IntPoint.h"
-
+#include "Map.h"
+#include "ThreadingBase.h"
 
 /**
  * Exec handler that registers itself and is being routed via StaticExec.
@@ -386,4 +387,41 @@ public:
 	{
 		return bValue;
 	}
-};
\ No newline at end of file
+};
+
+#if !(UE_BUILD_SHIPPING || UE_BUILD_TEST)
+	#define DO_BLUEPRINT_GUARD 1
+#endif
+
+#if DO_BLUEPRINT_GUARD
+/** 
+ * Helper struct for dealing with Blueprint exceptions 
+ */
+struct CORE_API FBlueprintExceptionTracker : TThreadSingleton<FBlueprintExceptionTracker>
+{
+	FBlueprintExceptionTracker()
+		: Runaway(0)
+		, Recurse(0)
+		, bRanaway(false)
+		, ScriptEntryTag(0)
+	{}
+
+	void ResetRunaway();
+
+public:
+	// map of currently displayed warnings in exception handler
+	TMap<FName, int32> DisplayedWarningsMap;
+
+	// runaway tracking
+	int32 Runaway;
+	int32 Recurse;
+	bool bRanaway;
+
+	// Script entry point tracking
+	int32 ScriptEntryTag;
+
+	// Stack names from the VM to be unrolled when we assert
+	TArray<FScriptTraceStackNode> ScriptStack;
+};
+
+#endif // DO_BLUEPRINT_GUARD
\ No newline at end of file
diff --git a/Engine/Source/Runtime/CoreUObject/Private/UObject/ScriptCore.cpp b/Engine/Source/Runtime/CoreUObject/Private/UObject/ScriptCore.cpp
--- a/Engine/Source/Runtime/CoreUObject/Private/UObject/ScriptCore.cpp
+++ b/Engine/Source/Runtime/CoreUObject/Private/UObject/ScriptCore.cpp
@@ -36,20 +36,15 @@ COREUOBJECT_API int32 GMaximumScriptLoopIterations = 1000000;
 	#define RECURSE_LIMIT 250
 #endif
 
-#if !(UE_BUILD_SHIPPING || UE_BUILD_TEST)
-	#define DO_GUARD 1
-#endif
-
-//@TODO: ScriptParallel: Contended static usage
-#if DO_GUARD
-	static int32 Runaway=0;
-	static int32 Recurse=0;
-	static bool Ranaway = false;
-	#define CHECK_RUNAWAY {++Runaway;}
-	COREUOBJECT_API void GInitRunaway() {Recurse=Runaway=0; Ranaway = false;}
+#if DO_BLUEPRINT_GUARD
+#define CHECK_RUNAWAY { ++FBlueprintExceptionTracker::Get().Runaway; }
+COREUOBJECT_API void GInitRunaway() 
+{
+	FBlueprintExceptionTracker::Get().ResetRunaway();
+}
 #else
-	#define CHECK_RUNAWAY
-	COREUOBJECT_API void GInitRunaway() {}
+#define CHECK_RUNAWAY
+COREUOBJECT_API void GInitRunaway() {}
 #endif
 
 #define IMPLEMENT_FUNCTION(cls,func) \
@@ -76,7 +71,7 @@ void FBlueprintCoreDelegates::ThrowScriptException(const UObject* ActiveObject,
 	case EBlueprintExceptionType::Tracepoint:
 	case EBlueprintExceptionType::WireTracepoint:
 		break;
-#if WITH_EDITOR
+#if WITH_EDITOR && DO_BLUEPRINT_GUARD
 	case EBlueprintExceptionType::AccessViolation:
 		{
 			struct FIntConfigValueHelper
@@ -92,9 +87,8 @@ void FBlueprintCoreDelegates::ThrowScriptException(const UObject* ActiveObject,
 			static const FIntConfigValueHelper MaxNumOfAccessViolation;
 			if (MaxNumOfAccessViolation.Value > 0)
 			{
-				static TMap<FName, int32> DisplayedWarningsMap;
 				const FName ActiveObjectName = ActiveObject ? ActiveObject->GetFName() : FName();
-				int32& Num = DisplayedWarningsMap.FindOrAdd(ActiveObjectName);
+				int32& Num = FBlueprintExceptionTracker::Get().DisplayedWarningsMap.FindOrAdd(ActiveObjectName);
 				if (Num > MaxNumOfAccessViolation.Value)
 				{
 					break;
@@ -108,7 +102,11 @@ void FBlueprintCoreDelegates::ThrowScriptException(const UObject* ActiveObject,
 		break;
 	}
 
-	OnScriptException.Broadcast(ActiveObject, StackFrame, Info);
+	// cant fire arbitrary delegates here off the game thead
+	if(IsInGameThread())
+	{
+		OnScriptException.Broadcast(ActiveObject, StackFrame, Info);
+	}
 
 	if (Info.GetType() == EBlueprintExceptionType::FatalError)
 	{
@@ -264,15 +262,17 @@ void FFrame::KismetExecutionMessage(const TCHAR* Message, ELogVerbosity::Type Ve
 #endif
 	}
 
-#if !(UE_BUILD_SHIPPING || UE_BUILD_TEST)
+#if DO_BLUEPRINT_GUARD
 	// Walk the script stack, if any
 	FString ScriptStack;
-	if( GScriptStack.Num() > 0 )
+
+	FBlueprintExceptionTracker& BlueprintExceptionTracker = FBlueprintExceptionTracker::Get();
+	if( BlueprintExceptionTracker.ScriptStack.Num() > 0 )
 	{
 		ScriptStack = TEXT( "Script call stack:\n" );
-		for( int32 i = GScriptStack.Num() - 1; i >= 0; --i )
+		for( int32 i = BlueprintExceptionTracker.ScriptStack.Num() - 1; i >= 0; --i )
 		{
-			ScriptStack += TEXT( "\t" ) + GScriptStack[i].GetStackDescription() + TEXT( "\n" );
+			ScriptStack += TEXT( "\t" ) + BlueprintExceptionTracker.ScriptStack[i].GetStackDescription() + TEXT( "\n" );
 		}
 	}
 #endif
@@ -732,15 +732,15 @@ void UObject::ProcessInternal( FFrame& Stack, RESULT_DECL )
 		// No POD struct can ever be stored in this buffer. 
 		MS_ALIGN(16) uint8 Buffer[MAX_SIMPLE_RETURN_VALUE_SIZE] GCC_ALIGN(16);
 
-#if DO_GUARD
-		if(Ranaway)
+#if DO_BLUEPRINT_GUARD
+		if(FBlueprintExceptionTracker::Get().bRanaway)
 		{
 			// If we have a return property, return a zeroed value in it, to try and save execution as much as possible
 			UProperty* ReturnProp = ((UFunction*)Stack.Node)->GetReturnProperty();
 			ClearReturnValue(ReturnProp, RESULT_PARAM);
 			return;
 		}
-		else if (++Recurse == RECURSE_LIMIT)
+		else if (++FBlueprintExceptionTracker::Get().Recurse == RECURSE_LIMIT)
 		{
 			// We've hit the recursion limit, so print out the stack, warn, and then continue with a zeroed return value.
 			UE_LOG(LogScriptCore, Log, TEXT("%s"), *Stack.GetStackTrace());
@@ -756,7 +756,7 @@ void UObject::ProcessInternal( FFrame& Stack, RESULT_DECL )
 
 			// This flag prevents repeated warnings of infinite loop, script exception handler 
 			// is expected to have terminated execution appropriately:
-			Ranaway = true;
+			FBlueprintExceptionTracker::Get().bRanaway = true;
 
 			return;
 		}
@@ -767,8 +767,8 @@ void UObject::ProcessInternal( FFrame& Stack, RESULT_DECL )
 		// Execute the bytecode
 		while (*Stack.Code != EX_Return)
 		{
-#if DO_GUARD
-			if( Runaway > GMaximumScriptLoopIterations )
+#if DO_BLUEPRINT_GUARD
+			if( FBlueprintExceptionTracker::Get().Runaway > GMaximumScriptLoopIterations )
 			{
 				// We've hit the recursion limit, so print out the stack, warn, and then continue with a zeroed return value.
 				UE_LOG(LogScriptCore, Log, TEXT("%s"), *Stack.GetStackTrace());
@@ -783,7 +783,7 @@ void UObject::ProcessInternal( FFrame& Stack, RESULT_DECL )
 
 				// Need to reset Runaway counter BEFORE throwing script exception, because the exception causes a modal dialog,
 				// and other scripts running will then erroneously think they are also "runaway".
-				Runaway = 0;
+				FBlueprintExceptionTracker::Get().Runaway = 0;
 
 				FBlueprintCoreDelegates::ThrowScriptException(this, Stack, RunawayLoopExceptionInfo);
 				return;
@@ -805,8 +805,8 @@ void UObject::ProcessInternal( FFrame& Stack, RESULT_DECL )
 			Stack.Code++;
 		}
 
-#if DO_GUARD
-		--Recurse;
+#if DO_BLUEPRINT_GUARD
+		--FBlueprintExceptionTracker::Get().Recurse;
 #endif
 	}
 	else
@@ -970,8 +970,6 @@ UFunction* UObject::FindFunctionChecked( FName InName ) const
 
 void UObject::ProcessEvent( UFunction* Function, void* Parms )
 {
-	static int32 ScriptEntryTag = 0;
-
 	checkf(!HasAnyFlags(RF_Unreachable),TEXT("%s  Function: '%s'"), *GetFullName(), *Function->GetPathName());
 	checkf(!FUObjectThreadContext::Get().IsRoutingPostLoad, TEXT("Cannot call UnrealScript (%s - %s) while PostLoading objects"), *GetFullName(), *Function->GetFullName());
 
@@ -1013,9 +1011,12 @@ void UObject::ProcessEvent( UFunction* Function, void* Parms )
 	}
 	checkSlow((Function->ParmsSize == 0) || (Parms != NULL));
 
-	ScriptEntryTag++;
+#if DO_BLUEPRINT_GUARD
+	FBlueprintExceptionTracker& BlueprintExceptionTracker = FBlueprintExceptionTracker::Get();
+	BlueprintExceptionTracker.ScriptEntryTag++;
 
-	CONDITIONAL_SCOPE_CYCLE_COUNTER(STAT_BlueprintTime, ScriptEntryTag == 1);
+	CONDITIONAL_SCOPE_CYCLE_COUNTER(STAT_BlueprintTime, BlueprintExceptionTracker.ScriptEntryTag == 1);
+#endif
 
 #if UE_BLUEPRINT_EVENTGRAPH_FASTCALLS
 	// Fast path for ubergraph calls
@@ -1138,7 +1139,9 @@ void UObject::ProcessEvent( UFunction* Function, void* Parms )
 		}
 	}
 
-	--ScriptEntryTag;
+#if DO_BLUEPRINT_GUARD
+	--BlueprintExceptionTracker.ScriptEntryTag;
+#endif
 }
 
 #ifdef _MSC_VER
@@ -1615,9 +1618,7 @@ void UObject::execLet(FFrame& Stack, RESULT_DECL)
 		}
 		else
 		{
-			//@TODO: ScriptParallel: Contended static usage
-			static uint8 Crud[1024];//@temp
-			Stack.MostRecentPropertyAddress = Crud;
+			Stack.MostRecentPropertyAddress = (uint8*)FMemory_Alloca(1024);
 			FMemory::Memzero(Stack.MostRecentPropertyAddress, sizeof(FString));
 		}
 	}
diff --git a/Engine/Source/Runtime/CoreUObject/Public/UObject/Stack.h b/Engine/Source/Runtime/CoreUObject/Public/UObject/Stack.h
--- a/Engine/Source/Runtime/CoreUObject/Public/UObject/Stack.h
+++ b/Engine/Source/Runtime/CoreUObject/Public/UObject/Stack.h
@@ -109,10 +109,11 @@ public:
 
 	virtual ~FFrame()
 	{
-#if !(UE_BUILD_SHIPPING || UE_BUILD_TEST)
-		if (GScriptStack.Num())
+#if DO_BLUEPRINT_GUARD
+		FBlueprintExceptionTracker& BlueprintExceptionTracker = FBlueprintExceptionTracker::Get();
+		if (BlueprintExceptionTracker.ScriptStack.Num())
 		{
-			GScriptStack.Pop(false);
+			BlueprintExceptionTracker.ScriptStack.Pop(false);
 		}
 #endif
 	}
@@ -187,9 +188,9 @@ inline FFrame::FFrame( UObject* InObject, UFunction* InNode, void* InLocals, FFr
 	, CurrentNativeFunction(NULL)
 	, bArrayContextFailed(false)
 {
-#if !(UE_BUILD_SHIPPING || UE_BUILD_TEST)
+#if DO_BLUEPRINT_GUARD
 	FScriptTraceStackNode StackNode(InNode->GetOuter()->GetFName(), InNode->GetFName());
- 	GScriptStack.Push(StackNode);
+	FBlueprintExceptionTracker::Get().ScriptStack.Push(StackNode);
 #endif
 }
 
