Bumping the max object count up to 16M

diff --git a/Engine/Source/Runtime/CoreUObject/Public/UObject/UObjectArray.h b/Engine/Source/Runtime/CoreUObject/Public/UObject/UObjectArray.h
--- a/Engine/Source/Runtime/CoreUObject/Public/UObject/UObjectArray.h
+++ b/Engine/Source/Runtime/CoreUObject/Public/UObject/UObjectArray.h
@@ -397,7 +397,7 @@ private:
 		return false;
 	}
 
-	typedef TStaticIndirectArrayThreadSafeRead<UObjectBase, 8 * 1024 * 1024 /* Max 8M UObjects */, 16384 /* allocated in 64K/128K chunks */ > TUObjectArray;
+	typedef TStaticIndirectArrayThreadSafeRead<UObjectBase, 16 * 1024 * 1024 /* Max 16M UObjects */, 16384 /* allocated in 64K/128K chunks */ > TUObjectArray;
 
 	/**
 	 * return the object array for use by debug visualizers
