Major HotReload fix.

diff --git a/Engine/Source/Developer/HotReload/Private/HotReload.cpp b/Engine/Source/Developer/HotReload/Private/HotReload.cpp
--- a/Engine/Source/Developer/HotReload/Private/HotReload.cpp
+++ b/Engine/Source/Developer/HotReload/Private/HotReload.cpp
@@ -11,6 +11,8 @@
 #include "KismetEditorUtilities.h"
 #endif
 
+#include "KismetEditorUtilities.h"
+
 DEFINE_LOG_CATEGORY(LogHotReload);
 
 #define LOCTEXT_NAMESPACE "HotReload"
@@ -183,10 +185,13 @@ private:
 	void GetPackagesToRebindAndDependentModules(const TArray<FString>& InGameModuleNames, TArray<UPackage*>& OutPackagesToRebind, TArray<FName>& OutDependentModules);
 
 #if WITH_ENGINE
+	void RegisterForReinstancing(UClass* OldClass, UClass* NewClass);
+	void ReinstanceClasses();
+
 	/**
 	 * Called from CoreUObject to re-instance hot-reloaded classes
 	 */
-	void ReinstanceClass(UClass* OldClass, UClass* NewClass);
+	void ReinstanceClass(UClass* OldClass, UClass* NewClass, const TMap<UClass*, UClass*>& OldToNewClassesMap);
 #endif
 
 	/**
@@ -393,7 +398,8 @@ void FHotReloadModule::StartupModule()
 
 #if WITH_ENGINE
 	// Register re-instancing delegate (Core)
-	FCoreUObjectDelegates::ReplaceHotReloadClassDelegate.BindRaw(this, &FHotReloadModule::ReinstanceClass);
+	FCoreUObjectDelegates::RegisterClassForHotReloadReinstancingDelegate.BindRaw(this, &FHotReloadModule::RegisterForReinstancing);
+	FCoreUObjectDelegates::ReinstanceHotReloadedClassesDelegate.BindRaw(this, &FHotReloadModule::ReinstanceClasses);
 #endif
 
 	// Register directory watcher delegate
@@ -646,11 +652,11 @@ void FHotReloadModule::DoHotReloadCallback(bool bRecompileFinished, ECompilation
 /**
  * Gets duplicated CDO from the cache, renames it and returns.
  */
-UObject* GetCachedCDODuplicate(UClass* Class, FName Name)
+UObject* GetCachedCDODuplicate(UObject* CDO, FName Name)
 {
 	UObject* DupCDO = nullptr;
 
-	UObject** DupCDOPtr = GetDuplicatedCDOMap().Find(Class);
+	UObject** DupCDOPtr = GetDuplicatedCDOMap().Find(CDO);
 	if (DupCDOPtr != nullptr)
 	{
 		DupCDO = *DupCDOPtr;
@@ -987,9 +993,46 @@ ECompilationResult::Type FHotReloadModule::RebindPackagesInternal(TArray<UPackag
 }
 
 #if WITH_ENGINE
-void FHotReloadModule::ReinstanceClass(UClass* OldClass, UClass* NewClass)
+namespace {
+	static TArray<TPair<UClass*, UClass*> >& GetClassesToReinstance()
+	{
+		static TArray<TPair<UClass*, UClass*> > Data;
+		return Data;
+	}
+}
+
+void FHotReloadModule::RegisterForReinstancing(UClass* OldClass, UClass* NewClass)
+{
+	TPair<UClass*, UClass*> Pair;
+	
+	Pair.Key = OldClass;
+	Pair.Value = NewClass;
+
+	GetClassesToReinstance().Add(MoveTemp(Pair));
+}
+
+void FHotReloadModule::ReinstanceClasses()
+{
+	TMap<UClass*, UClass*> OldToNewClassesMap;
+	for (TPair<UClass*, UClass*>& Pair : GetClassesToReinstance())
+	{
+		if (Pair.Value != nullptr)
+		{
+			OldToNewClassesMap.Add(Pair.Key, Pair.Value);
+		}
+	}
+
+	for (TPair<UClass*, UClass*>& Pair : GetClassesToReinstance())
+	{
+		ReinstanceClass(Pair.Key, Pair.Value, OldToNewClassesMap);
+	}
+
+	GetClassesToReinstance().Empty();
+}
+
+void FHotReloadModule::ReinstanceClass(UClass* OldClass, UClass* NewClass, const TMap<UClass*, UClass*>& OldToNewClassesMap)
 {	
-	auto ReinstanceHelper = FHotReloadClassReinstancer::Create(NewClass, OldClass, ReconstructedCDOsMap, HotReloadBPSetToRecompile, HotReloadBPSetToRecompileBytecodeOnly);
+	TSharedPtr<FHotReloadClassReinstancer> ReinstanceHelper = FHotReloadClassReinstancer::Create(NewClass, OldClass, OldToNewClassesMap, ReconstructedCDOsMap, HotReloadBPSetToRecompile, HotReloadBPSetToRecompileBytecodeOnly);
 	if (ReinstanceHelper->ClassNeedsReinstancing())
 	{
 		UE_LOG(LogHotReload, Log, TEXT("Re-instancing %s after hot-reload."), NewClass ? *NewClass->GetName() : *OldClass->GetName());
diff --git a/Engine/Source/Developer/HotReload/Private/HotReloadClassReinstancer.cpp b/Engine/Source/Developer/HotReload/Private/HotReloadClassReinstancer.cpp
--- a/Engine/Source/Developer/HotReload/Private/HotReloadClassReinstancer.cpp
+++ b/Engine/Source/Developer/HotReload/Private/HotReloadClassReinstancer.cpp
@@ -5,6 +5,9 @@
 
 #if WITH_ENGINE
 #include "Engine/BlueprintGeneratedClass.h"
+#include "Layers/ILayers.h"
+#include "BlueprintEditor.h"
+#include "Kismet2/CompilerResultsLog.h"
 
 void FHotReloadClassReinstancer::SetupNewClassReinstancing(UClass* InNewClass, UClass* InOldClass)
 {
@@ -271,19 +274,25 @@ void FHotReloadClassReinstancer::RecreateCDOAndSetupOldClassReinstancing(UClass*
 	}
 }
 
-FHotReloadClassReinstancer::FHotReloadClassReinstancer(UClass* InNewClass, UClass* InOldClass, TMap<UObject*, UObject*>& OutReconstructedCDOsMap, TSet<UBlueprint*>& InBPSetToRecompile, TSet<UBlueprint*>& InBPSetToRecompileBytecodeOnly)
+FHotReloadClassReinstancer::FHotReloadClassReinstancer(UClass* InNewClass, UClass* InOldClass, const TMap<UClass*, UClass*>& InOldToNewClassesMap, TMap<UObject*, UObject*>& OutReconstructedCDOsMap, TSet<UBlueprint*>& InBPSetToRecompile, TSet<UBlueprint*>& InBPSetToRecompileBytecodeOnly)
 	: NewClass(nullptr)
 	, bNeedsReinstancing(false)
 	, CopyOfPreviousCDO(nullptr)
 	, ReconstructedCDOsMap(OutReconstructedCDOsMap)
 	, BPSetToRecompile(InBPSetToRecompile)
 	, BPSetToRecompileBytecodeOnly(InBPSetToRecompileBytecodeOnly)
+	, OldToNewClassesMap(InOldToNewClassesMap)
 {
 	ensure(InOldClass);
 	ensure(!HotReloadedOldClass && !HotReloadedNewClass);
 	HotReloadedOldClass = InOldClass;
 	HotReloadedNewClass = InNewClass ? InNewClass : InOldClass;
 
+	for (const TPair<UClass*, UClass*>& OldToNewClass : OldToNewClassesMap)
+	{
+		ObjectsThatShouldUseOldStuff.Add(OldToNewClass.Key);
+	}
+
 	// If InNewClass is NULL, then the old class has not changed after hot-reload.
 	// However, we still need to check for changes to its constructor code (CDO values).
 	if (InNewClass)
@@ -298,7 +307,7 @@ FHotReloadClassReinstancer::FHotReloadClassReinstancer(UClass* InNewClass, UClas
 			FArchiveReplaceObjectRef<UObject> ReplaceObjectArch(*BlueprintIt, ClassRedirects, false, true, true);
 			if (ReplaceObjectArch.GetCount())
 			{
-				EnlistDependentBlueprintToRecompile(*BlueprintIt, true);
+				EnlistDependentBlueprintToRecompile(*BlueprintIt, false);
 			}
 		}
 	}
diff --git a/Engine/Source/Developer/HotReload/Private/HotReloadClassReinstancer.h b/Engine/Source/Developer/HotReload/Private/HotReloadClassReinstancer.h
--- a/Engine/Source/Developer/HotReload/Private/HotReloadClassReinstancer.h
+++ b/Engine/Source/Developer/HotReload/Private/HotReloadClassReinstancer.h
@@ -94,7 +94,7 @@ class FHotReloadClassReinstancer : public FBlueprintCompileReinstancer
 public:
 
 	/** Sets the re-instancer up to re-instance native classes */
-	FHotReloadClassReinstancer(UClass* InNewClass, UClass* InOldClass, TMap<UObject*, UObject*>& OutReconstructedCDOsMap, TSet<UBlueprint*>& InBPSetToRecompile, TSet<UBlueprint*>& InBPSetToRecompileBytecodeOnly);
+	FHotReloadClassReinstancer(UClass* InNewClass, UClass* InOldClass, const TMap<UClass*, UClass*>& OldToNewClassesMap, TMap<UObject*, UObject*>& OutReconstructedCDOsMap, TSet<UBlueprint*>& InBPSetToRecompile, TSet<UBlueprint*>& InBPSetToRecompileBytecodeOnly);
 	
 	/** Destructor */
 	virtual ~FHotReloadClassReinstancer();
@@ -109,9 +109,9 @@ public:
 	void ReinstanceObjectsAndUpdateDefaults();
 
 	/** Creates the reinstancer as a sharable object */
-	static TSharedPtr<FHotReloadClassReinstancer> Create(UClass* InNewClass, UClass* InOldClass, TMap<UObject*, UObject*>& OutReconstructedCDOsMap, TSet<UBlueprint*>& InBPSetToRecompile, TSet<UBlueprint*>& InBPSetToRecompileBytecodeOnly)
+	static TSharedPtr<FHotReloadClassReinstancer> Create(UClass* InNewClass, UClass* InOldClass, const TMap<UClass*, UClass*>& OldToNewClassesMap, TMap<UObject*, UObject*>& OutReconstructedCDOsMap, TSet<UBlueprint*>& InBPSetToRecompile, TSet<UBlueprint*>& InBPSetToRecompileBytecodeOnly)
 	{
-		return MakeShareable(new FHotReloadClassReinstancer(InNewClass, InOldClass, OutReconstructedCDOsMap, InBPSetToRecompile, InBPSetToRecompileBytecodeOnly));
+		return MakeShareable(new FHotReloadClassReinstancer(InNewClass, InOldClass, OldToNewClassesMap, OutReconstructedCDOsMap, InBPSetToRecompile, InBPSetToRecompileBytecodeOnly));
 	}
 
 	// FSerializableObject interface
@@ -134,7 +134,7 @@ private:
 	TMap<UObject*, UObject*>& ReconstructedCDOsMap;
 	TSet<UBlueprint*>& BPSetToRecompile;
 	TSet<UBlueprint*>& BPSetToRecompileBytecodeOnly;
-
+	const TMap<UClass*, UClass*>& OldToNewClassesMap;
 };
 
 #endif // WITH_ENGINE
\ No newline at end of file
diff --git a/Engine/Source/Editor/BlueprintGraph/Private/K2Node_DynamicCast.cpp b/Engine/Source/Editor/BlueprintGraph/Private/K2Node_DynamicCast.cpp
--- a/Engine/Source/Editor/BlueprintGraph/Private/K2Node_DynamicCast.cpp
+++ b/Engine/Source/Editor/BlueprintGraph/Private/K2Node_DynamicCast.cpp
@@ -186,16 +186,22 @@ UEdGraphPin* UK2Node_DynamicCast::GetInvalidCastPin() const
 
 UEdGraphPin* UK2Node_DynamicCast::GetCastResultPin() const
 {
-	UEdGraphPin* Pin = NULL;
-
-	if(TargetType != NULL)
+	if(TargetType != nullptr)
 	{
 		const UEdGraphSchema_K2* K2Schema = GetDefault<UEdGraphSchema_K2>();
-		FString PinName = K2Schema->PN_CastedValuePrefix + TargetType->GetDisplayNameText().ToString();
-		Pin = FindPin(PinName);
+
+		for (int32 PinIdx = 0; PinIdx < Pins.Num(); PinIdx++)
+		{
+			if (Pins[PinIdx]->PinType.PinSubCategoryObject == *TargetType
+				&& Pins[PinIdx]->Direction == EGPD_Output
+				&& Pins[PinIdx]->PinName.StartsWith(K2Schema->PN_CastedValuePrefix))
+			{
+				return Pins[PinIdx];
+			}
+		}
 	}
 		
-	return Pin;
+	return nullptr;
 }
 
 UEdGraphPin* UK2Node_DynamicCast::GetCastSourcePin() const
diff --git a/Engine/Source/Editor/BlueprintGraph/Private/K2Node_EditablePinBase.cpp b/Engine/Source/Editor/BlueprintGraph/Private/K2Node_EditablePinBase.cpp
--- a/Engine/Source/Editor/BlueprintGraph/Private/K2Node_EditablePinBase.cpp
+++ b/Engine/Source/Editor/BlueprintGraph/Private/K2Node_EditablePinBase.cpp
@@ -182,11 +182,22 @@ void UK2Node_EditablePinBase::Serialize(FArchive& Ar)
 	else
 	{
 		SerializedItems.Empty(UserDefinedPins.Num());
+
 		for (int32 Index = 0; Index < UserDefinedPins.Num(); ++Index)
 		{
 			SerializedItems.Add(*(UserDefinedPins[Index].Get()));
 		}
+
 		Ar << SerializedItems;
+
+		if (Ar.IsModifyingWeakAndStrongReferences())
+		{
+			UserDefinedPins.Empty(SerializedItems.Num());
+			for (int32 Index = 0; Index < SerializedItems.Num(); ++Index)
+			{
+				UserDefinedPins.Add(MakeShareable(new FUserPinInfo(SerializedItems[Index])));
+			}
+		}
 	}
 }
 
diff --git a/Engine/Source/Editor/UnrealEd/Private/Kismet2/KismetReinstanceUtilities.cpp b/Engine/Source/Editor/UnrealEd/Private/Kismet2/KismetReinstanceUtilities.cpp
--- a/Engine/Source/Editor/UnrealEd/Private/Kismet2/KismetReinstanceUtilities.cpp
+++ b/Engine/Source/Editor/UnrealEd/Private/Kismet2/KismetReinstanceUtilities.cpp
@@ -183,7 +183,7 @@ FBlueprintCompileReinstancer::FBlueprintCompileReinstancer(UClass* InClassToRein
 		// Temporarily suspend the undo buffer; we don't need to record the duplicated CDO until it is fully resolved
  		ITransaction* CurrentTransaction = GUndo;
  		GUndo = NULL;
-		DuplicatedClass->ClassDefaultObject = GetClassCDODuplicate(ClassToReinstance, DuplicatedClass->GetDefaultObjectName());
+		DuplicatedClass->ClassDefaultObject = GetClassCDODuplicate(ClassToReinstance->GetDefaultObject(), DuplicatedClass->GetDefaultObjectName());
 
 		// Restore the undo buffer
 		GUndo = CurrentTransaction;
@@ -1492,16 +1492,16 @@ void FBlueprintCompileReinstancer::ReparentChild(UClass* ChildClass)
 	ChildClass->StaticLink(true);
 }
 
-UObject* FBlueprintCompileReinstancer::GetClassCDODuplicate(UClass* Class, FName Name)
+UObject* FBlueprintCompileReinstancer::GetClassCDODuplicate(UObject* CDO, FName Name)
 {
 	UObject* DupCDO = nullptr;
 
 	FCDODuplicatesProvider& CDODupProvider = GetCDODuplicatesProviderDelegate();
 
-	if (!CDODupProvider.IsBound() || (DupCDO = CDODupProvider.Execute(Class, Name)) == nullptr)
+	if (!CDODupProvider.IsBound() || (DupCDO = CDODupProvider.Execute(CDO, Name)) == nullptr)
 	{
 		GIsDuplicatingClassForReinstancing = true;
-		DupCDO = (UObject*)StaticDuplicateObject(Class->GetDefaultObject(), GetTransientPackage(), *Name.ToString());
+		DupCDO = (UObject*)StaticDuplicateObject(CDO, GetTransientPackage(), *Name.ToString());
 		GIsDuplicatingClassForReinstancing = false;
 	}
 
diff --git a/Engine/Source/Editor/UnrealEd/Public/Kismet2/KismetReinstanceUtilities.h b/Engine/Source/Editor/UnrealEd/Public/Kismet2/KismetReinstanceUtilities.h
--- a/Engine/Source/Editor/UnrealEd/Public/Kismet2/KismetReinstanceUtilities.h
+++ b/Engine/Source/Editor/UnrealEd/Public/Kismet2/KismetReinstanceUtilities.h
@@ -28,7 +28,7 @@ public:
 	/**
 	 * CDO duplicates provider delegate type.
 	 */
-	DECLARE_DELEGATE_RetVal_TwoParams(UObject*, FCDODuplicatesProvider, UClass*, FName);
+	DECLARE_DELEGATE_RetVal_TwoParams(UObject*, FCDODuplicatesProvider, UObject*, FName);
 
 	/**
 	 * Gets CDO duplicates provider delegate.
@@ -158,10 +158,10 @@ private:
 	 * Gets class CDO duplicate from cache (or creates it if not available or not
 	 * during hot-reload) for given class.
 	 *
-	 * @param Class Class to get (or create) CDO for.
+	 * @param CDO Object to get duplicate of.
 	 * @param Name The name that CDO has to be renamed with (or created with).
 	 */
-	static UObject* GetClassCDODuplicate(UClass* Class, FName Name);
+	static UObject* GetClassCDODuplicate(UObject* CDO, FName Name);
 };
 
 
diff --git a/Engine/Source/Programs/UnrealHeaderTool/Private/CodeGenerator.cpp b/Engine/Source/Programs/UnrealHeaderTool/Private/CodeGenerator.cpp
--- a/Engine/Source/Programs/UnrealHeaderTool/Private/CodeGenerator.cpp
+++ b/Engine/Source/Programs/UnrealHeaderTool/Private/CodeGenerator.cpp
@@ -2791,6 +2791,8 @@ void FNativeClassHeaderGenerator::ExportGeneratedStructBodyMacros(FUnrealSourceF
 				ExportMirrorsForNoexportStructs(Structs, /*Indent=*/ 2, GeneratedStructRegisterFunctionText);
 			}
 
+			FString CRCFuncName = FString::Printf(TEXT("Get_%s_CRC"), *SingletonName.Replace(TEXT("()"), TEXT(""), ESearchCase::CaseSensitive));
+
 			// Structs can either have a UClass or UPackage as outer (if delcared in non-UClass header).
 			if (ScriptStruct->GetOuter()->IsA(UStruct::StaticClass()))
 			{
@@ -2800,15 +2802,8 @@ void FNativeClassHeaderGenerator::ExportGeneratedStructBodyMacros(FUnrealSourceF
 			{
 				GeneratedStructRegisterFunctionText.Logf(TEXT("\t\tUPackage* Outer=%s;\r\n"), *GetPackageSingletonName(CastChecked<UPackage>(ScriptStruct->GetOuter())));
 			}
-			if (!bIsDynamic)
-			{
-				GeneratedStructRegisterFunctionText.Logf(TEXT("\t\tstatic UScriptStruct* ReturnStruct = NULL;\r\n"));
-			}
-			else
-			{
-				GeneratedStructRegisterFunctionText.Logf(TEXT("\t\tUScriptStruct* ReturnStruct = Cast<UScriptStruct>(StaticFindObjectFast(UScriptStruct::StaticClass(), Outer, TEXT(\"%s\")));\r\n"),
-					*ScriptStruct->GetName());
-			}
+			GeneratedStructRegisterFunctionText.Logf(TEXT("\t\textern uint32 %s();\r\n"), *CRCFuncName);
+			GeneratedStructRegisterFunctionText.Logf(TEXT("\t\tstatic UScriptStruct* ReturnStruct = FindExistingStructIfHotReloadOrDynamic(Outer, TEXT(\"%s\"), sizeof(%s), %s(), %s);\r\n"), *ScriptStruct->GetName(), NameLookupCPP.GetNameCPP(Struct), *CRCFuncName, bIsDynamic ? TEXT("true") : TEXT("false"));
 			GeneratedStructRegisterFunctionText.Logf(TEXT("\t\tif (!ReturnStruct)\r\n"));
 			GeneratedStructRegisterFunctionText.Logf(TEXT("\t\t{\r\n"));
 			FString BaseStructString(TEXT("NULL"));
@@ -2866,7 +2861,7 @@ void FNativeClassHeaderGenerator::ExportGeneratedStructBodyMacros(FUnrealSourceF
 
 			auto& GeneratedFunctionText = GetGeneratedFunctionTextDevice();
 			GeneratedFunctionText += GeneratedStructRegisterFunctionText;
-			GeneratedFunctionText.Logf(TEXT("\tuint32 Get_%s_CRC() { return %uU; }\r\n"), *SingletonName.Replace(TEXT("()"), TEXT(""), ESearchCase::CaseSensitive), StructCrc);
+			GeneratedFunctionText.Logf(TEXT("\tuint32 %s() { return %uU; }\r\n"), *CRCFuncName, StructCrc);
 
 			//CallSingletons.Logf(TEXT("\t\t\t\tOuterClass->LinkChild(%s); // %u\r\n"), *SingletonName, StructCrc);
 		}
@@ -2925,6 +2920,8 @@ void FNativeClassHeaderGenerator::ExportGeneratedEnumsInitCode(const TArray<UEnu
 
 			FUHTStringBuilder GeneratedEnumRegisterFunctionText;
 
+			FString CRCFuncName = FString::Printf(TEXT("Get_%s_CRC"), *SingletonName.Replace(TEXT("()"), TEXT(""), ESearchCase::CaseSensitive));
+
 			GeneratedEnumRegisterFunctionText.Logf(TEXT("\tUEnum* %s\r\n"), *EnumSingletonName);
 			GeneratedEnumRegisterFunctionText.Logf(TEXT("\t{\r\n"));
 			// Enums can either have a UClass or UPackage as outer (if declared in non-UClass header).
@@ -2936,7 +2933,8 @@ void FNativeClassHeaderGenerator::ExportGeneratedEnumsInitCode(const TArray<UEnu
 			{
 				GeneratedEnumRegisterFunctionText.Logf(TEXT("\t\tUPackage* Outer=%s;\r\n"), *GetPackageSingletonName(CastChecked<UPackage>(Enum->GetOuter())));
 			}
-			GeneratedEnumRegisterFunctionText.Logf(TEXT("\t\tstatic UEnum* ReturnEnum = NULL;\r\n"));
+			GeneratedEnumRegisterFunctionText.Logf(TEXT("\t\textern uint32 %s();\r\n"), *CRCFuncName);
+			GeneratedEnumRegisterFunctionText.Logf(TEXT("\t\tstatic UEnum* ReturnEnum = FindExistingEnumIfHotReload(Outer, TEXT(\"%s\"), 0, %s());\r\n"), *Enum->GetName(), *CRCFuncName);
 			GeneratedEnumRegisterFunctionText.Logf(TEXT("\t\tif (!ReturnEnum)\r\n"));
 			GeneratedEnumRegisterFunctionText.Logf(TEXT("\t\t{\r\n"));
 
@@ -2975,6 +2973,7 @@ void FNativeClassHeaderGenerator::ExportGeneratedEnumsInitCode(const TArray<UEnu
 
 			uint32 EnumCrc = GenerateTextCRC(*GeneratedEnumRegisterFunctionText);
 			GGeneratedCodeCRCs.Add(Enum, EnumCrc);
+			GeneratedFunctionText.Logf(TEXT("\tuint32 %s() { return %uU; }\r\n"), *CRCFuncName, EnumCrc);
 			// CallSingletons.Logf(TEXT("\t\t\t\tOuterClass->LinkChild(%s); // %u\r\n"), *EnumSingletonName, EnumCrc);
 		}
 	}
diff --git a/Engine/Source/Runtime/CoreUObject/Private/Serialization/FindReferencersArchive.cpp b/Engine/Source/Runtime/CoreUObject/Private/Serialization/FindReferencersArchive.cpp
--- a/Engine/Source/Runtime/CoreUObject/Private/Serialization/FindReferencersArchive.cpp
+++ b/Engine/Source/Runtime/CoreUObject/Private/Serialization/FindReferencersArchive.cpp
@@ -10,20 +10,24 @@
  *
  * @param	PotentialReferencer		the object to serialize which may contain references to our target objects
  * @param	InTargetObjects			array of objects to search for references to
+ * @param	bFindAlsoWeakReferences should we also look into weak references?
  */
-FFindReferencersArchive::FFindReferencersArchive( UObject* InPotentialReferencer, TArray<UObject*> InTargetObjects )
+FFindReferencersArchive::FFindReferencersArchive(UObject* InPotentialReferencer, TArray<UObject*> InTargetObjects, bool bFindAlsoWeakReferences)
 {
 	// use the optimized RefLink to skip over properties which don't contain object references
 	ArIsObjectReferenceCollector = true;
 
+	// look also for weak references if user asked for it (we won't modify them, but there is no archive option for that)
+	ArIsModifyingWeakAndStrongReferences = bFindAlsoWeakReferences;
+
 	// ALL objects reference their outers...it's just log spam here
 	ArIgnoreOuterRef = true;
 
 	// initialize the map
-	for ( int32 ObjIndex = 0; ObjIndex < InTargetObjects.Num(); ObjIndex++ )
+	for (int32 ObjIndex = 0; ObjIndex < InTargetObjects.Num(); ObjIndex++)
 	{
 		UObject* TargetObject = InTargetObjects[ObjIndex];
-		if ( TargetObject != nullptr )
+		if (TargetObject != nullptr)
 		{
 			TargetObjects.Add(TargetObject, 0);
 		}
@@ -32,7 +36,7 @@ FFindReferencersArchive::FFindReferencersArchive( UObject* InPotentialReferencer
 	PotentialReferencer = InPotentialReferencer;
 
 	// now start the search
-	if ( PotentialReferencer )
+	if (PotentialReferencer)
 	{
 		PotentialReferencer->Serialize(*this);
 	}
diff --git a/Engine/Source/Runtime/CoreUObject/Private/UObject/UObjectBase.cpp b/Engine/Source/Runtime/CoreUObject/Private/UObject/UObjectBase.cpp
--- a/Engine/Source/Runtime/CoreUObject/Private/UObject/UObjectBase.cpp
+++ b/Engine/Source/Runtime/CoreUObject/Private/UObject/UObjectBase.cpp
@@ -396,78 +396,8 @@ void UObjectCompiledInDeferStruct(class UScriptStruct *(*InRegister)(), const TC
 	}
 }
 
-#if WITH_HOT_RELOAD
-struct FStructOrEnumCompiledInfo : FFieldCompiledInInfo
-{
-	FStructOrEnumCompiledInfo(SIZE_T InClassSize, uint32 InCrc)
-		: FFieldCompiledInInfo(InClassSize, InCrc)
-	{
-	}
-
-	virtual UClass* Register() const
-	{
-		return nullptr;
-	}
-};
-
-/** Registered struct info (including size and reflection info) */
-static TMap<FName, FStructOrEnumCompiledInfo*>& GetStructOrEnumGeneratedCodeInfo()
-{
-	static TMap<FName, FStructOrEnumCompiledInfo*> StructOrEnumCompiledInfoMap;
-	return StructOrEnumCompiledInfoMap;
-}
-#endif
-
 class UScriptStruct *GetStaticStruct(class UScriptStruct *(*InRegister)(), UObject* StructOuter, const TCHAR* StructName, SIZE_T Size, uint32 Crc)
 {
-#if WITH_HOT_RELOAD
-	// Try to get the existing info for this struct
-	const FName StructFName(StructName);
-	FStructOrEnumCompiledInfo* Info = nullptr;
-	auto ExistingInfo = GetStructOrEnumGeneratedCodeInfo().Find(StructFName);
-	if (!ExistingInfo)
-	{
-		// New struct
-		Info = new FStructOrEnumCompiledInfo(Size, Crc);
-		Info->bHasChanged = true;
-		GetStructOrEnumGeneratedCodeInfo().Add(StructFName, Info);
-	}
-	else
-	{
-		// Hot-relaoded struct, check if it has changes
-		Info = *ExistingInfo;
-		Info->bHasChanged = (*ExistingInfo)->Size != Size || (*ExistingInfo)->Crc != Crc;
-		Info->Size = Size;
-		Info->Crc = Crc;
-	}
-
-	if (GIsHotReload)
-	{		
-		if (!Info->bHasChanged)
-		{
-			// New struct added during hot-reload
-			UScriptStruct* ReturnStruct = FindObject<UScriptStruct>(StructOuter, StructName);
-			if (ReturnStruct)
-			{
-				UE_LOG(LogClass, Log, TEXT( "%s HotReload."), StructName);
-				return ReturnStruct;
-			}
-			UE_LOG(LogClass, Log, TEXT("Could not find existing script struct %s for HotReload. Assuming new"), StructName);
-		}
-		else
-		{
-			// Existing struct, make sure we destroy the old one
-			UScriptStruct* ExistingStruct = FindObject<UScriptStruct>(StructOuter, StructName);
-			if (ExistingStruct)
-			{
-				// Make sure the old struct is not used by anything
-				ExistingStruct->ClearFlags(RF_RootSet | RF_Standalone | RF_Public);
-				const FName OldStructRename = MakeUniqueObjectName(GetTransientPackage(), ExistingStruct->GetClass(), *FString::Printf(TEXT("HOTRELOADED_%s"), StructName));
-				ExistingStruct->Rename(*OldStructRename.ToString(), GetTransientPackage());
-			}
-		}
-	}
-#endif
 	return (*InRegister)();
 }
 
@@ -550,7 +480,7 @@ static TArray<FFieldCompiledInInfo*>& GetDeferredClassRegistration()
 }
 
 #if WITH_HOT_RELOAD
-TMap<UClass*, UObject*>& GetDuplicatedCDOMap()
+TMap<UObject*, UObject*>& GetDuplicatedCDOMap()
 {
 	/**
 	 * GC referencer object, so duplicated CDOs aren't destroyed if still
@@ -571,7 +501,7 @@ TMap<UClass*, UObject*>& GetDuplicatedCDOMap()
 		}
 
 		/** Duplicated CDOs map. */
-		TMap<UClass*, UObject*> Map;
+		TMap<UObject*, UObject*> Map;
 	};
 
 	static FDuplicatedCDOMap Cache;
@@ -707,7 +637,7 @@ void UClassRegisterAllCompiledInClasses()
 /** Re-instance all existing classes that have changed during hot-reload */
 void UClassReplaceHotReloadClasses()
 {
-	if (FCoreUObjectDelegates::ReplaceHotReloadClassDelegate.IsBound())
+	if (FCoreUObjectDelegates::RegisterClassForHotReloadReinstancingDelegate.IsBound())
 	{
 		for (auto Class : GetHotReloadClasses())
 		{
@@ -719,9 +649,11 @@ void UClassReplaceHotReloadClasses()
 				RegisteredClass = Class->Register();
 			}
 
-			FCoreUObjectDelegates::ReplaceHotReloadClassDelegate.Execute(Class->OldClass, RegisteredClass);
+			FCoreUObjectDelegates::RegisterClassForHotReloadReinstancingDelegate.Execute(Class->OldClass, RegisteredClass);
 		}
 	}
+
+	FCoreUObjectDelegates::ReinstanceHotReloadedClassesDelegate.ExecuteIfBound();
 	GetHotReloadClasses().Empty();
 }
 
@@ -750,7 +682,7 @@ static void UClassGenerateCDODuplicatesForHotReload()
 					*MakeUniqueObjectName(GetTransientPackage(), Class, TEXT("HOTRELOAD_CDO_DUPLICATE")).ToString()
 					);
 				GIsDuplicatingClassForReinstancing = false;
-				GetDuplicatedCDOMap().Add(Class, DupCDO);
+				GetDuplicatedCDOMap().Add(Class->GetDefaultObject(), DupCDO);
 			}
 		}
 	}
@@ -1034,6 +966,129 @@ const TCHAR* DebugFullName(UObject* Object)
 	}
 }
 
+namespace
+{
+#if WITH_HOT_RELOAD
+	struct FStructOrEnumCompiledInfo : FFieldCompiledInInfo
+	{
+		struct FKey
+		{
+			UObject* Outer;
+			FName Name;
+
+			friend bool operator==(const FKey& A, const FKey& B)
+			{
+				return A.Outer == B.Outer && A.Name == B.Name;
+			}
+
+			friend uint32 GetTypeHash(const FKey& Key)
+			{
+				return HashCombine(GetTypeHash(Key.Outer), GetTypeHash(Key.Name));
+			}
+		};
+
+		/** Registered struct info (including size and reflection info) */
+		static TMap<FStructOrEnumCompiledInfo::FKey, FStructOrEnumCompiledInfo*>& GetRegisteredInfo()
+		{
+			static TMap<FStructOrEnumCompiledInfo::FKey, FStructOrEnumCompiledInfo*> StructOrEnumCompiledInfoMap;
+			return StructOrEnumCompiledInfoMap;
+		}
+
+		FStructOrEnumCompiledInfo(SIZE_T InClassSize, uint32 InCrc)
+			: FFieldCompiledInInfo(InClassSize, InCrc)
+		{
+		}
+
+		virtual UClass* Register() const
+		{
+			return nullptr;
+		}
+	};
+
+
+#endif // WITH_HOT_RELOAD
+
+	template <typename TType>
+	TType* FindExistingStructOrEnumIfHotReload(UObject* Outer, const TCHAR* Name, SIZE_T Size, uint32 Crc)
+	{
+#if WITH_HOT_RELOAD
+		FStructOrEnumCompiledInfo::FKey Key;
+
+		Key.Outer = Outer;
+		Key.Name = Name;
+
+		FStructOrEnumCompiledInfo* Info = nullptr;
+		FStructOrEnumCompiledInfo** ExistingInfo = FStructOrEnumCompiledInfo::GetRegisteredInfo().Find(Key);
+		if (!ExistingInfo)
+		{
+			// New struct
+			Info = new FStructOrEnumCompiledInfo(Size, Crc);
+			Info->bHasChanged = true;
+			FStructOrEnumCompiledInfo::GetRegisteredInfo().Add(Key, Info);
+		}
+		else
+		{
+			// Hot-relaoded struct, check if it has changes
+			Info = *ExistingInfo;
+			Info->bHasChanged = (*ExistingInfo)->Size != Size || (*ExistingInfo)->Crc != Crc;
+			Info->Size = Size;
+			Info->Crc = Crc;
+		}
+
+		if (!GIsHotReload)
+		{
+			return nullptr;
+		}
+
+		if (!Info->bHasChanged)
+		{
+			// New type added during hot-reload
+			TType* Return = FindObject<TType>(Outer, Name);
+			if (Return)
+			{
+				UE_LOG(LogClass, Log, TEXT("%s HotReload."), Name);
+				return Return;
+			}
+			UE_LOG(LogClass, Log, TEXT("Could not find existing type %s for HotReload. Assuming new"), Name);
+		}
+		else
+		{
+			// Existing type, make sure we destroy the old one
+			TType* Existing = FindObject<TType>(Outer, Name);
+			if (Existing)
+			{
+				// Make sure the old struct is not used by anything
+				Existing->ClearFlags(RF_RootSet | RF_Standalone | RF_Public);
+				const FName OldRename = MakeUniqueObjectName(GetTransientPackage(), Existing->GetClass(), *FString::Printf(TEXT("HOTRELOADED_%s"), Name));
+				Existing->Rename(*OldRename.ToString(), GetTransientPackage());
+			}
+		}
+#endif
+
+		return nullptr;
+	}
+}
+
+UScriptStruct* FindExistingStructIfHotReloadOrDynamic(UObject* Outer, const TCHAR* StructName, SIZE_T Size, uint32 Crc, bool bIsDynamic)
+{
+	if (UScriptStruct* Output = FindExistingStructOrEnumIfHotReload<UScriptStruct>(Outer, StructName, Size, Crc))
+	{
+		return Output;
+	}
+
+	if (bIsDynamic)
+	{
+		return Cast<UScriptStruct>(StaticFindObjectFast(UScriptStruct::StaticClass(), Outer, StructName));
+	}
+
+	return nullptr;
+}
+
+UEnum* FindExistingEnumIfHotReload(UObject* Outer, const TCHAR* EnumName, SIZE_T Size, uint32 Crc)
+{
+	return FindExistingStructOrEnumIfHotReload<UEnum>(Outer, EnumName, Size, Crc);
+}
+
 UClass* ConstructDynamicClass(FName ClassName)
 {
 	UClass* Result = nullptr;
diff --git a/Engine/Source/Runtime/CoreUObject/Private/UObject/UObjectGlobals.cpp b/Engine/Source/Runtime/CoreUObject/Private/UObject/UObjectGlobals.cpp
--- a/Engine/Source/Runtime/CoreUObject/Private/UObject/UObjectGlobals.cpp
+++ b/Engine/Source/Runtime/CoreUObject/Private/UObject/UObjectGlobals.cpp
@@ -40,7 +40,8 @@ DEFINE_STAT(STAT_NameTableMemorySize);
 DEFINE_STAT(STAT_DestroyObject);
 
 /** CoreUObject delegates */
-FCoreUObjectDelegates::FReplaceHotReloadClassDelegate FCoreUObjectDelegates::ReplaceHotReloadClassDelegate;
+FCoreUObjectDelegates::FRegisterClassForHotReloadReinstancingDelegate FCoreUObjectDelegates::RegisterClassForHotReloadReinstancingDelegate;
+FCoreUObjectDelegates::FReinstanceHotReloadedClassesDelegate FCoreUObjectDelegates::ReinstanceHotReloadedClassesDelegate;
 // Delegates used by SavePackage()
 FCoreUObjectDelegates::FIsPackageOKToSaveDelegate FCoreUObjectDelegates::IsPackageOKToSaveDelegate;
 FCoreUObjectDelegates::FAutoPackageBackupDelegate FCoreUObjectDelegates::AutoPackageBackupDelegate;
diff --git a/Engine/Source/Runtime/CoreUObject/Public/Serialization/ArchiveUObject.h b/Engine/Source/Runtime/CoreUObject/Public/Serialization/ArchiveUObject.h
--- a/Engine/Source/Runtime/CoreUObject/Public/Serialization/ArchiveUObject.h
+++ b/Engine/Source/Runtime/CoreUObject/Public/Serialization/ArchiveUObject.h
@@ -407,8 +407,9 @@ public:
 	 *
 	 * @param	PotentialReferencer		the object to serialize which may contain references to our target objects
 	 * @param	InTargetObjects			array of objects to search for references to
+	 * @param	bFindAlsoWeakReferences should we also look into weak references?
 	 */
-	COREUOBJECT_API FFindReferencersArchive( class UObject* PotentialReferencer, TArray<class UObject*> InTargetObjects );
+	COREUOBJECT_API FFindReferencersArchive(class UObject* PotentialReferencer, TArray<class UObject*> InTargetObjects, bool bFindAlsoWeakReferences = false);
 
 	/**
 	 * Retrieves the number of references from PotentialReferencer to the object specified.
@@ -491,14 +492,15 @@ public:
 	 * @param	PackageToCheck	if specified, only objects contained in this package will be searched
 	 *							for references to 
 	 * @param	bIgnoreTemplates If true, do not record template objects
+	 * @param	bFindAlsoWeakReferences If true, also look into weak references
 	 */
-	TFindObjectReferencers( TArray< T* > TargetObjects, UPackage* PackageToCheck=NULL, bool bIgnoreTemplates = true )
+	TFindObjectReferencers( TArray< T* > TargetObjects, UPackage* PackageToCheck=NULL, bool bIgnoreTemplates = true, bool bFindAlsoWeakReferences = false)
 	: TMultiMap< T*, UObject* >()
 	{
 		TArray<UObject*> ReferencedObjects;
 		TMap<UObject*, int32> ReferenceCounts;
 
-		FFindReferencersArchive FindReferencerAr(nullptr, ( TArray<UObject*>& )TargetObjects);
+		FFindReferencersArchive FindReferencerAr(nullptr, (TArray<UObject*>&)TargetObjects, bFindAlsoWeakReferences);
 
 		// Loop over every object to find any reference that may exist for the target objects
 		for (FObjectIterator It; It; ++It)
diff --git a/Engine/Source/Runtime/CoreUObject/Public/UObject/UObjectBase.h b/Engine/Source/Runtime/CoreUObject/Public/UObject/UObjectBase.h
--- a/Engine/Source/Runtime/CoreUObject/Public/UObject/UObjectBase.h
+++ b/Engine/Source/Runtime/CoreUObject/Public/UObject/UObjectBase.h
@@ -372,6 +372,9 @@ struct FCompiledInDeferEnum
  */
 COREUOBJECT_API class UEnum *GetStaticEnum(class UEnum *(*InRegister)(), UObject* EnumOuter, const TCHAR* EnumName);
 
+COREUOBJECT_API class UScriptStruct* FindExistingStructIfHotReloadOrDynamic(UObject* Outer, const TCHAR* StructName, SIZE_T Size, uint32 Crc, bool bIsDynamic);
+COREUOBJECT_API class UEnum* FindExistingEnumIfHotReload(UObject* Outer, const TCHAR* EnumName, SIZE_T Size, uint32 Crc);
+
 /** @return	True if there are any newly-loaded UObjects that are waiting to be registered by calling ProcessNewlyLoadedUObjects() */
 COREUOBJECT_API bool AnyNewlyLoadedUObjects();
 
@@ -380,7 +383,7 @@ COREUOBJECT_API void ProcessNewlyLoadedUObjects();
 
 #if WITH_HOT_RELOAD
 /** Map of duplicated CDOs for reinstancing during hot-reload purposes. */
-COREUOBJECT_API TMap<UClass*, UObject*>& GetDuplicatedCDOMap();
+COREUOBJECT_API TMap<UObject*, UObject*>& GetDuplicatedCDOMap();
 #endif // WITH_HOT_RELOAD
 
 /**
diff --git a/Engine/Source/Runtime/CoreUObject/Public/UObject/UObjectGlobals.h b/Engine/Source/Runtime/CoreUObject/Public/UObject/UObjectGlobals.h
--- a/Engine/Source/Runtime/CoreUObject/Public/UObject/UObjectGlobals.h
+++ b/Engine/Source/Runtime/CoreUObject/Public/UObject/UObjectGlobals.h
@@ -1764,9 +1764,13 @@ struct COREUOBJECT_API FCoreUObjectDelegates
 	/** Delegate used by SavePackage() to check whether a package should be saved */
 	static FIsPackageOKToSaveDelegate IsPackageOKToSaveDelegate;
 
-	/** Delegate for replacing hot-reloaded classes that changed after hot-reload */
-	DECLARE_DELEGATE_TwoParams(FReplaceHotReloadClassDelegate, UClass*, UClass*);
-	static FReplaceHotReloadClassDelegate ReplaceHotReloadClassDelegate;
+	/** Delegate for registering hot-reloaded classes that changed after hot-reload for reinstancing */
+	DECLARE_DELEGATE_TwoParams(FRegisterClassForHotReloadReinstancingDelegate, UClass*, UClass*);
+	static FRegisterClassForHotReloadReinstancingDelegate RegisterClassForHotReloadReinstancingDelegate;
+
+	/** Delegate for reinstancing hot-reloaded classes */
+	DECLARE_DELEGATE(FReinstanceHotReloadedClassesDelegate);
+	static FReinstanceHotReloadedClassesDelegate ReinstanceHotReloadedClassesDelegate;
 
 	// Sent at the very beginning of LoadMap
 	static FSimpleMulticastDelegate PreLoadMap;
