Added support for diffing cooked packages on save against supplied cooked package directory.

diff --git a/Engine/Source/Runtime/Core/Private/Internationalization/Text.cpp b/Engine/Source/Runtime/Core/Private/Internationalization/Text.cpp
--- a/Engine/Source/Runtime/Core/Private/Internationalization/Text.cpp
+++ b/Engine/Source/Runtime/Core/Private/Internationalization/Text.cpp
@@ -949,7 +949,7 @@ CORE_API FArchive& operator<<( FArchive& Ar, FText& Value )
 	}
 
 #if WITH_EDITOR
-	/*if (Ar.IsCooking() && Ar.IsSaving() && Ar.IsPersistent() && !(Ar.GetDebugSerializationFlags()&DSF_IsTagExportsArchive))
+	if (Ar.IsCooking() && Ar.IsSaving() && Ar.IsPersistent() && (Ar.GetDebugSerializationFlags()&DSF_EnableCookerWarnings))
 	{
 		if (!!(Value.Flags & ETextFlag::ConvertedProperty))
 		{
@@ -959,7 +959,7 @@ CORE_API FArchive& operator<<( FArchive& Ar, FText& Value )
 		{
 			UE_LOG(LogText, Warning, TEXT("Saving FText \"%s\" which has been initialized from FString at cook time resave of source package %s may fix issue."), *Value.ToString(), *Ar.GetArchiveName())
 		}
-	}*/
+	}
 #endif
 
 	const int32 OriginalFlags = Value.Flags;
diff --git a/Engine/Source/Runtime/Core/Private/Serialization/Archive.cpp b/Engine/Source/Runtime/Core/Private/Serialization/Archive.cpp
--- a/Engine/Source/Runtime/Core/Private/Serialization/Archive.cpp
+++ b/Engine/Source/Runtime/Core/Private/Serialization/Archive.cpp
@@ -94,6 +94,9 @@ void FArchive::Reset()
 	SerializedProperty = nullptr;
 #if WITH_EDITORONLY_DATA
 	EditorOnlyPropertyStack = 0;
+#endif
+#if WITH_EDITOR
+	ArDebugSerializationFlags			= 0;
 #endif
 	// Reset all custom versions to the current registered versions.
 	ResetCustomVersions();
@@ -149,6 +152,17 @@ FString FArchive::GetArchiveName() const
 	return TEXT("FArchive");
 }
 
+#if WITH_EDITOR
+FArchive::FScopeAddDebugData::FScopeAddDebugData(FArchive& InAr, const FName& DebugData) : Ar(InAr)
+{
+	Ar.PushDebugDataString(DebugData);
+}
+
+void FArchive::PushDebugDataString(const FName& DebugData)
+{
+}
+#endif
+
 FArchive& FArchive::operator<<( class FLazyObjectPtr& LazyObjectPtr )
 {
 	// The base FArchive does not implement this method. Use FArchiveUOBject instead.
diff --git a/Engine/Source/Runtime/Core/Public/Serialization/ArchiveBase.h b/Engine/Source/Runtime/Core/Public/Serialization/ArchiveBase.h
--- a/Engine/Source/Runtime/Core/Public/Serialization/ArchiveBase.h
+++ b/Engine/Source/Runtime/Core/Public/Serialization/ArchiveBase.h
@@ -700,6 +700,15 @@ public:
 		return ((ArPortFlags & Flags) == Flags);
 	}
 
+	FORCEINLINE uint32 GetDebugSerializationFlags() const
+	{
+#if WITH_EDITOR
+		return ArDebugSerializationFlags;
+#else
+		return 0;
+#endif
+	}
+
 	FORCEINLINE bool ShouldSkipBulkData() const
 	{
 		return ArShouldSkipBulkData;
@@ -800,6 +809,18 @@ public:
 		ArPortFlags = InPortFlags;
 	}
 
+	/**
+	 * Sets the archives custom serialization modifier flags (nothing to do with PortFlags or Custom versions)
+	 * 
+	 * @param InCustomFlags the new flags to use for custom serialization
+	 */
+	void SetDebugSerializationFlags(uint32 InCustomFlags)
+	{
+#if WITH_EDITOR
+		ArDebugSerializationFlags = InCustomFlags;
+#endif
+	}
+
 	/**
 	 * Returns if an async close operation has finished or not, as well as if there was an error.
 	 *
@@ -1044,9 +1065,72 @@ public:
 			
 	/** Max size of data that this archive is allowed to serialize. */
 	int64 ArMaxSerializeSize;
+
+	class FScopeSetDebugSerializationFlags
+	{
+	private:
+#if WITH_EDITOR
+		uint32 PreviousFlags;
+		FArchive& Ar;
+#endif
+	public:
+		/**
+		 * Initializes an object which will set flags for the scope of this code
+		 * 
+		 * @param NewFlags new flags to set 
+		 * @param Remove should we add these flags or remove them default is to add
+		 */
+#if WITH_EDITOR
+		FScopeSetDebugSerializationFlags(FArchive& InAr, uint32 NewFlags, bool Remove = false)
+			: Ar(InAr)
+		{
+
+			PreviousFlags = Ar.GetDebugSerializationFlags();
+			if (Remove)
+			{
+				Ar.SetDebugSerializationFlags( PreviousFlags & ~NewFlags);
+			}
+			else
+			{
+				Ar.SetDebugSerializationFlags( PreviousFlags | NewFlags);
+			}
+
+		}
+		~FScopeSetDebugSerializationFlags()
+		{
+
+			Ar.SetDebugSerializationFlags( PreviousFlags);
+		}
+#else
+		FScopeSetDebugSerializationFlags(FArchive& InAr, uint32 NewFlags, bool Remove = false)
+		{}
+		~FScopeSetDebugSerializationFlags()
+		{}
+#endif
+	};
+
+#if WITH_EDITOR
+	/** Custom serialization modifier flags can be used for anything */
+	uint32 ArDebugSerializationFlags;
+	/** Debug stack storage if you want to add data to the archive for usage further down the serialization stack this should be used in conjunction with the FScopeAddDebugData struct */
 	
-private:
+	virtual void PushDebugDataString(const FName& DebugData);
+	virtual void PopDebugDataString() { }
+
+	class FScopeAddDebugData
+	{
+	private:
+		FArchive& Ar;
+	public:
+		CORE_API FScopeAddDebugData(FArchive& InAr, const FName& DebugData);
 
+		~FScopeAddDebugData()
+		{
+			Ar.PopDebugDataString();
+		}
+	};
+#endif	
+private:
 	/** Holds the cooking target platform. */
 	const ITargetPlatform* CookingTargetPlatform;
 
@@ -1209,7 +1293,16 @@ public:
 	{
 		return InnerArchive.IsCloseComplete(bHasError);
 	}
-
+#if WITH_EDITOR
+	virtual void PushDebugDataString(const FName& DebugData) override
+	{
+		InnerArchive.PushDebugDataString(DebugData);
+	}
+	virtual void PopDebugDataString() override
+	{
+		InnerArchive.PopDebugDataString();
+	}
+#endif
 protected:
 
 	/** Holds the archive that this archive is a proxy to. */
diff --git a/Engine/Source/Runtime/CoreUObject/Private/Serialization/BulkData.cpp b/Engine/Source/Runtime/CoreUObject/Private/Serialization/BulkData.cpp
--- a/Engine/Source/Runtime/CoreUObject/Private/Serialization/BulkData.cpp
+++ b/Engine/Source/Runtime/CoreUObject/Private/Serialization/BulkData.cpp
@@ -3,6 +3,9 @@
 
 #include "CoreUObjectPrivate.h"
 #include "TargetPlatform.h"
+#if WITH_EDITOR
+#include "DebugSerializationFlags.h"
+#endif
 
 /*-----------------------------------------------------------------------------
 	Constructors and operators
@@ -778,8 +781,12 @@ void FUntypedBulkData::Serialize( FArchive& Ar, UObject* Owner, int32 Idx )
 		FThreadSafeBulkDataToObjectMap::Get().Add( this, Owner );
 #endif
 		// Offset where the bulkdata flags are stored
-		int64 SavedBulkDataFlagsPos	= Ar.Tell();
-		Ar << BulkDataFlags;
+		int64 SavedBulkDataFlagsPos = Ar.Tell();
+		{
+			FArchive::FScopeSetDebugSerializationFlags S(Ar, DSF_IgnoreDiff);
+			Ar << BulkDataFlags;
+		}
+
 
 		// Number of elements in array.
 		Ar << ElementCount;
@@ -919,14 +926,19 @@ void FUntypedBulkData::Serialize( FArchive& Ar, UObject* Owner, int32 Idx )
 			// Keep track of position we are going to serialize placeholder BulkDataSizeOnDisk.
 			SavedBulkDataSizeOnDiskPos = Ar.Tell();
 			BulkDataSizeOnDisk = INDEX_NONE;
-			// And serialize the placeholder which is going to be overwritten later.
-			Ar << BulkDataSizeOnDisk;
 
-			// Keep track of position we are going to serialize placeholder BulkDataOffsetInFile.
-			SavedBulkDataOffsetInFilePos = Ar.Tell();
-			BulkDataOffsetInFile = INDEX_NONE;
-			// And serialize the placeholder which is going to be overwritten later.
-			Ar << BulkDataOffsetInFile;
+			{
+				FArchive::FScopeSetDebugSerializationFlags S(Ar, DSF_IgnoreDiff);
+
+				// And serialize the placeholder which is going to be overwritten later.
+				Ar << BulkDataSizeOnDisk;
+				// Keep track of position we are going to serialize placeholder BulkDataOffsetInFile.
+				SavedBulkDataOffsetInFilePos = Ar.Tell();
+				BulkDataOffsetInFile = INDEX_NONE;
+				// And serialize the placeholder which is going to be overwritten later.
+				Ar << BulkDataOffsetInFile;
+
+			}
 
 				// try to get the linkersave object
 			FLinkerSave* LinkerSave = Cast<FLinkerSave>(Ar.GetLinker());
@@ -977,17 +989,21 @@ void FUntypedBulkData::Serialize( FArchive& Ar, UObject* Owner, int32 Idx )
 			// store current file offset before seeking back
 			int64 CurrentFileOffset = Ar.Tell();
 
-			// Seek back and overwrite the flags 
-			Ar.Seek(SavedBulkDataFlagsPos);
-			Ar << BulkDataFlags;
+			{
+				FArchive::FScopeSetDebugSerializationFlags S(Ar, DSF_IgnoreDiff);
 
-			// Seek back and overwrite placeholder for BulkDataSizeOnDisk
-			Ar.Seek( SavedBulkDataSizeOnDiskPos );
-			Ar << BulkDataSizeOnDisk;
+				// Seek back and overwrite the flags 
+				Ar.Seek(SavedBulkDataFlagsPos);
+				Ar << BulkDataFlags;
 
-			// Seek back and overwrite placeholder for BulkDataOffsetInFile
-			Ar.Seek( SavedBulkDataOffsetInFilePos );
-			Ar << BulkDataOffsetInFile;
+				// Seek back and overwrite placeholder for BulkDataSizeOnDisk
+				Ar.Seek(SavedBulkDataSizeOnDiskPos);
+				Ar << BulkDataSizeOnDisk;
+
+				// Seek back and overwrite placeholder for BulkDataOffsetInFile
+				Ar.Seek(SavedBulkDataOffsetInFilePos);
+				Ar << BulkDataOffsetInFile;
+			}
 
 			// Seek to the end of written data so we don't clobber any data in subsequent write 
 			// operations
diff --git a/Engine/Source/Runtime/CoreUObject/Private/UObject/Class.cpp b/Engine/Source/Runtime/CoreUObject/Private/UObject/Class.cpp
--- a/Engine/Source/Runtime/CoreUObject/Private/UObject/Class.cpp
+++ b/Engine/Source/Runtime/CoreUObject/Private/UObject/Class.cpp
@@ -1030,6 +1030,7 @@ void UStruct::SerializeTaggedProperties(FArchive& Ar, uint8* Data, UStruct* Defa
 				FName Name;  
 				Ar << Name;
 				FText Text = FText::FromName(Name);
+				Text.Flags |= ETextFlag::ConvertedProperty;
 				CastChecked<UTextProperty>(Property)->SetPropertyValue_InContainer(Data, Text, Tag.ArrayIndex);
 				AdvanceProperty = true;
 				continue; 
@@ -1375,6 +1376,11 @@ void UStruct::SerializeTaggedProperties(FArchive& Ar, uint8* Data, UStruct* Defa
 						{
 							DefaultValue = NULL;
 						}
+#if WITH_EDITOR
+						static const FName NAME_PropertySerialize = FName(TEXT("PropertySerialize"));
+						FArchive::FScopeAddDebugData P(Ar, NAME_PropertySerialize);
+						FArchive::FScopeAddDebugData S(Ar, Property->GetFName());
+#endif
 						FPropertyTag Tag( Ar, Property, Idx, DataPtr, DefaultValue );
 						Ar << Tag;
 
@@ -1445,6 +1451,8 @@ void UStruct::Serialize( FArchive& Ar )
 
 		if (Ar.IsSaving())
 		{
+			FArchive::FScopeSetDebugSerializationFlags S(Ar, DSF_IgnoreDiff);
+
 			Ar << ScriptBytecodeSize;
 
 			int32 ScriptStorageSize = 0;
@@ -1471,8 +1479,8 @@ void UStruct::Serialize( FArchive& Ar )
 
 				// force writing to a buffer
 				TArray<uint8> TempScript;
-				FMemoryWriter MemWriter(TempScript, Ar.IsPersistent());
-				LinkerSave->Saver = &MemWriter;
+					FMemoryWriter MemWriter(TempScript, Ar.IsPersistent());
+					LinkerSave->Saver = &MemWriter;
 
 				// now, use the linker to save the byte code, but writing to memory
 				while (iCode < ScriptBytecodeSize)
@@ -1480,11 +1488,11 @@ void UStruct::Serialize( FArchive& Ar )
 					SerializeExpr(iCode, Ar);
 				}
 
-				// restore the saver
-				LinkerSave->Saver = SavedSaver;
+					// restore the saver
+					LinkerSave->Saver = SavedSaver;
 
-				// now write out the memory bytes
-				Ar.Serialize(TempScript.GetData(), TempScript.Num());
+					// now write out the memory bytes
+					Ar.Serialize(TempScript.GetData(), TempScript.Num());
 
 				// and update the SHA (does nothing if not currently calculating SHA)
 				LinkerSave->UpdateScriptSHAKey(TempScript);
@@ -1504,6 +1512,8 @@ void UStruct::Serialize( FArchive& Ar )
 
 			if (Ar.IsSaving())
 			{
+				FArchive::FScopeSetDebugSerializationFlags S(Ar, DSF_IgnoreDiff);
+
 				int32 const BytecodeEndOffset = Ar.Tell();
 
 				// go back and write on-disk size
diff --git a/Engine/Source/Runtime/CoreUObject/Private/UObject/LinkerSave.cpp b/Engine/Source/Runtime/CoreUObject/Private/UObject/LinkerSave.cpp
--- a/Engine/Source/Runtime/CoreUObject/Private/UObject/LinkerSave.cpp
+++ b/Engine/Source/Runtime/CoreUObject/Private/UObject/LinkerSave.cpp
@@ -44,6 +44,43 @@ FLinkerSave::FLinkerSave(UPackage* InParent, const TCHAR* InFilename, bool bForc
 	}
 }
 
+
+FLinkerSave::FLinkerSave(UPackage* InParent, FArchive *InSaver, bool bForceByteSwapping, bool bInSaveUnversioned)
+: FLinker(ELinkerType::Save, InParent, TEXT("$$Memory$$"))
+, Saver(nullptr)
+{
+	if (FPlatformProperties::HasEditorOnlyData())
+	{
+		// Create file saver.
+		Saver = InSaver;
+		check(Saver);
+#if WITH_EDITOR
+		ArDebugSerializationFlags = Saver->ArDebugSerializationFlags;
+#endif
+		
+
+		UPackage* Package = dynamic_cast<UPackage*>(LinkerRoot);
+
+		// Set main summary info.
+		Summary.Tag = PACKAGE_FILE_TAG;
+		Summary.SetFileVersions(GPackageFileUE4Version, GPackageFileLicenseeUE4Version, bInSaveUnversioned);
+		Summary.SavedByEngineVersion = FEngineVersion::Current();
+		Summary.CompatibleWithEngineVersion = FEngineVersion::CompatibleWith();
+		Summary.PackageFlags = Package ? (Package->PackageFlags & ~PKG_NewlyCreated) : 0;
+
+		if (Package)
+		{
+			Summary.FolderName = Package->GetFolderName().ToString();
+			Summary.ChunkIDs = Package->GetChunkIDs();
+		}
+
+		// Set status info.
+		ArIsSaving = 1;
+		ArIsPersistent = 1;
+		ArForceByteSwapping = bForceByteSwapping;
+	}
+}
+
 FLinkerSave::FLinkerSave(UPackage* InParent, bool bForceByteSwapping, bool bInSaveUnversioned )
 :	FLinker(ELinkerType::Save, InParent,TEXT("$$Memory$$"))
 ,	Saver(nullptr)
@@ -130,9 +167,18 @@ int64 FLinkerSave::Tell()
 
 void FLinkerSave::Serialize( void* V, int64 Length )
 {
+#if WITH_EDITOR
+	Saver->ArDebugSerializationFlags = ArDebugSerializationFlags;
+#endif
 	Saver->Serialize( V, Length );
 }
 	
+
+FString FLinkerSave::GetArchiveName() const
+{
+	return Saver->GetArchiveName();
+}
+
 FArchive& FLinkerSave::operator<<( FName& InName )
 {
 	int32 Save = MapName(InName);
diff --git a/Engine/Source/Runtime/CoreUObject/Private/UObject/Obj.cpp b/Engine/Source/Runtime/CoreUObject/Private/UObject/Obj.cpp
--- a/Engine/Source/Runtime/CoreUObject/Private/UObject/Obj.cpp
+++ b/Engine/Source/Runtime/CoreUObject/Private/UObject/Obj.cpp
@@ -983,6 +983,11 @@ void UObject::SerializeScriptProperties( FArchive& Ar ) const
 		const bool bBreakSerializationRecursion = BreakSerializationRecursion && Ar.IsLoading() && Ar.GetLinker();
 #else 
 		const bool bBreakSerializationRecursion = false;
+#endif
+#if WITH_EDITOR
+		static const FName NAME_SerializeScriptProperties = FName(TEXT("SerializeScriptProperties"));
+		FArchive::FScopeAddDebugData P(Ar, NAME_SerializeScriptProperties);
+		FArchive::FScopeAddDebugData S(Ar, ObjClass->GetFName());
 #endif
 		ObjClass->SerializeTaggedProperties(Ar, (uint8*)this, HasAnyFlags(RF_ClassDefaultObject) ? ObjClass->GetSuperClass() : ObjClass, (uint8*)DiffObject, bBreakSerializationRecursion ? this : NULL);
 	}
diff --git a/Engine/Source/Runtime/CoreUObject/Private/UObject/SavePackage.cpp b/Engine/Source/Runtime/CoreUObject/Private/UObject/SavePackage.cpp
--- a/Engine/Source/Runtime/CoreUObject/Private/UObject/SavePackage.cpp
+++ b/Engine/Source/Runtime/CoreUObject/Private/UObject/SavePackage.cpp
@@ -6,6 +6,7 @@
 #include "TargetPlatform.h"
 #include "UObject/UObjectThreadContext.h"
 #include "BlueprintSupport.h"
+#include "DebugSerializationFlags.h"
 
 DEFINE_LOG_CATEGORY_STATIC(LogSavePackage, Log, All);
 
@@ -3158,6 +3159,78 @@ struct FImportReplacementsHelper
 
 #endif
 
+#if WITH_EDITOR
+
+class FDiffSerializeArchive : public FBufferArchive
+{
+private:
+	FArchive *TestArchive;
+	TArray<FName> DebugDataStack;
+	bool bDisable;
+public:
+
+
+	FDiffSerializeArchive(const FName& InFilename, FArchive *InTestArchive) : FBufferArchive(true, InFilename), TestArchive(InTestArchive)
+	{ 
+		ArDebugSerializationFlags = DSF_IgnoreDiff;
+		bDisable = false;
+	}
+
+	virtual void Serialize(void* Data, int64 Num) override
+	{
+		TArray<int8> TestMemory;
+
+		if (TestArchive)
+		{
+			int64 Pos = FMath::Min(FBufferArchive::Tell(), TestArchive->TotalSize());
+			TestArchive->Seek(Pos);
+			TestMemory.AddZeroed(Num);
+			int64 ReadSize = FMath::Min(Num, TestArchive->TotalSize() - Pos);
+			TestArchive->Serialize((void*)TestMemory.GetData(), ReadSize);
+
+			if (!(ArDebugSerializationFlags&DSF_IgnoreDiff) && (!bDisable))
+			{
+				if (FMemory::Memcmp((void*)TestMemory.GetData(), Data, Num) != 0)
+				{
+					// get the calls debug callstack and 
+					FString DebugStackString;
+					for (const auto& DebugData : DebugDataStack)
+					{
+						DebugStackString += DebugData.ToString();
+						DebugStackString += TEXT("->");
+					}
+
+					UE_LOG(LogSavePackage, Warning, TEXT("Diff cooked package archive recognized a difference %lld Filename %s, stack %s "), Pos, *ArchiveName.ToString(), *DebugStackString);
+
+					// only log one message per archive, from this point the entire package is probably messed up
+					bDisable = true;
+
+					static int i = 0;
+					i++;
+				}
+			}
+		}
+		FBufferArchive::Serialize(Data, Num);
+	}
+
+	virtual void PushDebugDataString(const FName& DebugData) override
+	{
+		DebugDataStack.Add(DebugData);
+	}
+	virtual void PopDebugDataString() override
+	{
+		DebugDataStack.Pop();
+	}
+
+	virtual FString GetArchiveName() const override
+	{
+		return TestArchive->GetArchiveName();
+	}
+};
+
+#endif
+
+
 ESavePackageResult UPackage::Save(UPackage* InOuter, UObject* Base, EObjectFlags TopLevelFlags, const TCHAR* Filename,
 	FOutputDevice* Error, FLinkerLoad* Conform, bool bForceByteSwapping, bool bWarnOfLongFilename, uint32 SaveFlags, 
 	const class ITargetPlatform* TargetPlatform, const FDateTime&  FinalTimeStamp, bool bSlowTask)
@@ -3448,6 +3521,25 @@ ESavePackageResult UPackage::Save(UPackage* InOuter, UObject* Base, EObjectFlags
 
 				FLinkerSave* Linker = nullptr;
 				
+#if WITH_EDITOR
+				FString DiffCookedPackagesPath;
+				// if we are cooking and we have diff cooked packages on the commandline then do some special stuff
+
+				if ((!!TargetPlatform) && FParse::Value(FCommandLine::Get(), TEXT("DiffCookedPackages="), DiffCookedPackagesPath))
+				{
+					FString TestArchiveFilename = Filename;
+					// TestArchiveFilename.ReplaceInline(TEXT("Cooked"), TEXT("CookedDiff"));
+					DiffCookedPackagesPath.ReplaceInline(TEXT("\\"), TEXT("/"));
+					FString CookedPath = FPaths::ConvertRelativePathToFull(FPaths::GameSavedDir() + TEXT("Cooked/"));
+					CookedPath.ReplaceInline(TEXT("\\"), TEXT("/"));
+					TestArchiveFilename.ReplaceInline(*CookedPath, *DiffCookedPackagesPath);
+					
+					FArchive* TestArchive = IFileManager::Get().CreateFileReader(*TestArchiveFilename); 
+					FArchive* Saver = new FDiffSerializeArchive(InOuter->FileName, TestArchive);
+					Linker = new FLinkerSave(InOuter, Saver, bForceByteSwapping);
+				}
+				else 
+#endif
 				if (bCompressFromMemory || bSaveAsync)
 				{
 					// Allocate the linker with a memory writer, forcing byte swapping if wanted.
@@ -3459,6 +3551,14 @@ ESavePackageResult UPackage::Save(UPackage* InOuter, UObject* Base, EObjectFlags
 					Linker = new FLinkerSave(InOuter, *TempFilename, bForceByteSwapping, bSaveUnversioned);
 				}
 
+
+#if WITH_EDITOR
+				if (!!TargetPlatform)
+				{
+					Linker->SetDebugSerializationFlags(DSF_EnableCookerWarnings | Linker->GetDebugSerializationFlags());
+				}
+#endif
+
 				// Use the custom versions we had previously gleaned from the export tag pass
 				Linker->Summary.SetCustomVersionContainer(ExportTaggerArchive.GetCustomVersions());
 
@@ -4320,6 +4420,9 @@ ESavePackageResult UPackage::Save(UPackage* InOuter, UObject* Base, EObjectFlags
 				UE_LOG_COOK_TIME(TEXT("SaveThumbNailsAssetRegistryDataWorldLevelInfo"));
 				
 				{
+#if WITH_EDITOR
+					FArchive::FScopeSetDebugSerializationFlags S(*Linker, DSF_IgnoreDiff, true);
+#endif
 					FScopedSlowTask ExportScope(Linker->ExportMap.Num());
 
 					// Save exports.
diff --git a/Engine/Source/Runtime/CoreUObject/Public/UObject/Linker.h b/Engine/Source/Runtime/CoreUObject/Public/UObject/Linker.h
--- a/Engine/Source/Runtime/CoreUObject/Public/UObject/Linker.h
+++ b/Engine/Source/Runtime/CoreUObject/Public/UObject/Linker.h
@@ -2183,6 +2183,8 @@ public:
 	FLinkerSave(UPackage* InParent, const TCHAR* InFilename, bool bForceByteSwapping, bool bInSaveUnversioned = false );
 	/** Constructor for memory writer */
 	FLinkerSave(UPackage* InParent, bool bForceByteSwapping, bool bInSaveUnversioned = false );
+	/** Constructor for custom savers */
+	FLinkerSave(UPackage* InParent, FArchive *InSaver, bool bForceByteSwapping, bool bInSaveUnversioned = false);
 
 	/** Returns the appropriate name index for the source name, or 0 if not found in NameIndices */
 	int32 MapName(const FName& Name) const;
@@ -2196,6 +2198,15 @@ public:
 	FArchive& operator<<( FLazyObjectPtr& LazyObjectPtr );
 	FArchive& operator<<( FAssetPtr& AssetPtr );
 
+#if WITH_EDITOR
+	// proxy for debugdata
+	virtual void PushDebugDataString(const FName& DebugData) override { Saver->PushDebugDataString(DebugData); }
+	virtual void PopDebugDataString() override { Saver->PopDebugDataString(); }
+#endif
+
+
+	virtual FString GetArchiveName() const override;
+
 	/**
 	 * If this archive is a FLinkerLoad or FLinkerSave, returns a pointer to the FLinker portion.
 	 */
diff --git a/Engine/Source/Runtime/CoreUObject/Public/UObject/PropertyTag.h b/Engine/Source/Runtime/CoreUObject/Public/UObject/PropertyTag.h
--- a/Engine/Source/Runtime/CoreUObject/Public/UObject/PropertyTag.h
+++ b/Engine/Source/Runtime/CoreUObject/Public/UObject/PropertyTag.h
@@ -7,6 +7,8 @@
 #ifndef __UNPROPERTYTAG_H__
 #define __UNPROPERTYTAG_H__
 
+#include "DebugSerializationFlags.h"
+
 /**
  *  A tag describing a class property, to aid in serialization.
  */
@@ -91,8 +93,10 @@ struct FPropertyTag
 			// property has been serialized.
 			Tag.SizeOffset = Ar.Tell();
 		}
-		Ar << Tag.Size << Tag.ArrayIndex;
-
+		{
+			FArchive::FScopeSetDebugSerializationFlags S(Ar, DSF_IgnoreDiff);
+			Ar << Tag.Size << Tag.ArrayIndex;
+		}
 		// only need to serialize this for structs
 		if (Tag.Type == NAME_StructProperty)
 		{
diff --git a/Engine/Source/Runtime/Engine/Private/TextureDerivedData.cpp b/Engine/Source/Runtime/Engine/Private/TextureDerivedData.cpp
--- a/Engine/Source/Runtime/Engine/Private/TextureDerivedData.cpp
+++ b/Engine/Source/Runtime/Engine/Private/TextureDerivedData.cpp
@@ -24,6 +24,8 @@ enum
 #include "ImageCore.h"
 #include "Engine/TextureCube.h"
 
+#include "DebugSerializationFlags.h"
+
 /*------------------------------------------------------------------------------
 	Versioning for texture derived data.
 ------------------------------------------------------------------------------*/
@@ -1901,7 +1903,12 @@ void UTexture::SerializeCookedPlatformData(FArchive& Ar)
 				Ar << PixelFormatName;
 				int32 SkipOffsetLoc = Ar.Tell();
 				int32 SkipOffset = 0;
-				Ar << SkipOffset;
+
+				{
+					FArchive::FScopeSetDebugSerializationFlags S(Ar,DSF_IgnoreDiff);
+					Ar << SkipOffset;
+				}
+
 				// Pass streamable flag for inlining mips
 				PlatformDataToSave->SerializeCooked(Ar, this, BuildSettings.bStreamable);
 				SkipOffset = Ar.Tell();
diff --git a/Engine/Source/Runtime/ShaderCore/Private/VertexFactory.cpp b/Engine/Source/Runtime/ShaderCore/Private/VertexFactory.cpp
--- a/Engine/Source/Runtime/ShaderCore/Private/VertexFactory.cpp
+++ b/Engine/Source/Runtime/ShaderCore/Private/VertexFactory.cpp
@@ -9,6 +9,8 @@
 #include "VertexFactory.h"
 #include "RHICommandList.h"
 
+#include "DebugSerializationFlags.h"
+
 uint32 FVertexFactoryType::NextHashIndex = 0;
 bool FVertexFactoryType::bInitializedSerializationHistory = false;
 
@@ -347,8 +349,12 @@ bool operator<<(FArchive& Ar,FVertexFactoryParameterRef& Ref)
 
 	// Need to be able to skip over parameters for no longer existing vertex factories.
 	int32 SkipOffset = Ar.Tell();
-	// Write placeholder.
-	Ar << SkipOffset;
+	{
+		FArchive::FScopeSetDebugSerializationFlags S(Ar, DSF_IgnoreDiff);
+		// Write placeholder.
+		Ar << SkipOffset;
+	}
+
 
 	if(Ref.Parameters)
 	{
diff --git a/Engine/Source/Runtime/ShaderCore/Public/Shader.h b/Engine/Source/Runtime/ShaderCore/Public/Shader.h
--- a/Engine/Source/Runtime/ShaderCore/Public/Shader.h
+++ b/Engine/Source/Runtime/ShaderCore/Public/Shader.h
@@ -13,6 +13,10 @@
 // For FShaderUniformBufferParameter
 #include "ShaderParameters.h"
 
+#if WITH_EDITOR
+#include "DebugSerializationFlags.h"
+#endif
+
 class FShaderType;
 class FVertexFactoryParameterRef;
 
@@ -1133,10 +1137,15 @@ public:
 				checkSlow(FName(Type->GetName()) != NAME_None);
 
 				Ar << Type;
-
 				int32 SkipOffset = Ar.Tell();
-				// Serialize a placeholder value, we will overwrite this with an offset to the end of the shader
-				Ar << SkipOffset;
+
+				{
+#if WITH_EDITOR
+					FArchive::FScopeSetDebugSerializationFlags S(Ar,DSF_IgnoreDiff);
+#endif
+					// Serialize a placeholder value, we will overwrite this with an offset to the end of the shader
+					Ar << SkipOffset;
+				}
 
 				if (bHandleShaderKeyChanges)
 				{
