Fixed compilation error.

diff --git a/Engine/Source/Runtime/CoreUObject/Private/Serialization/BulkData.cpp b/Engine/Source/Runtime/CoreUObject/Private/Serialization/BulkData.cpp
--- a/Engine/Source/Runtime/CoreUObject/Private/Serialization/BulkData.cpp
+++ b/Engine/Source/Runtime/CoreUObject/Private/Serialization/BulkData.cpp
@@ -783,7 +783,9 @@ void FUntypedBulkData::Serialize( FArchive& Ar, UObject* Owner, int32 Idx )
 		// Offset where the bulkdata flags are stored
 		int64 SavedBulkDataFlagsPos = Ar.Tell();
 		{
+#if WITH_EDITOR
 			FArchive::FScopeSetDebugSerializationFlags S(Ar, DSF_IgnoreDiff);
+#endif
 			Ar << BulkDataFlags;
 		}
 
