Fix for skipped elements/infinite loop in UMapProperty::Identical.

diff --git a/Engine/Source/Runtime/CoreUObject/Private/UObject/PropertyMap.cpp b/Engine/Source/Runtime/CoreUObject/Private/UObject/PropertyMap.cpp
--- a/Engine/Source/Runtime/CoreUObject/Private/UObject/PropertyMap.cpp
+++ b/Engine/Source/Runtime/CoreUObject/Private/UObject/PropertyMap.cpp
@@ -33,7 +33,6 @@ namespace UE4MapProperty_Private
 			}
 
 			++Index;
-			--Num;
 		}
 
 		return false;
