Improved archive debug stack for UFunction.

diff --git a/Engine/Source/Runtime/CoreUObject/Private/UObject/Class.cpp b/Engine/Source/Runtime/CoreUObject/Private/UObject/Class.cpp
--- a/Engine/Source/Runtime/CoreUObject/Private/UObject/Class.cpp
+++ b/Engine/Source/Runtime/CoreUObject/Private/UObject/Class.cpp
@@ -4458,6 +4458,12 @@ void UFunction::Invoke(UObject* Obj, FFrame& Stack, RESULT_DECL)
 
 void UFunction::Serialize( FArchive& Ar )
 {
+#if WITH_EDITOR
+	const static FName NAME_UFunction(TEXT("UFunction"));
+	FArchive::FScopeAddDebugData S(Ar, NAME_UFunction);
+	FArchive::FScopeAddDebugData Q(Ar, GetFName());
+#endif
+
 	Super::Serialize( Ar );
 
 	Ar.ThisContainsCode();
