Fix for stale ParentFuncMap during blueprint reinstancing after child blueprints are reparented

diff --git a/Engine/Source/Runtime/CoreUObject/Private/UObject/Class.cpp b/Engine/Source/Runtime/CoreUObject/Private/UObject/Class.cpp
--- a/Engine/Source/Runtime/CoreUObject/Private/UObject/Class.cpp
+++ b/Engine/Source/Runtime/CoreUObject/Private/UObject/Class.cpp
@@ -3356,6 +3356,7 @@ void UClass::SetSuperStruct(UStruct* NewSuperStruct)
 #if UCLASS_FAST_ISA_IMPL & 2
 	FFastIndexingClassTree::Unregister(this);
 #endif
+	ClearFunctionMapsCaches();
 	Super::SetSuperStruct(NewSuperStruct);
 #if UCLASS_FAST_ISA_IMPL & 2
 	FFastIndexingClassTree::Register(this);
