Mergine engine changes to //UE4/Main @ 2748007

diff --git a/Engine/Source/Editor/Sequencer/Private/DisplayNodes/SequencerTrackNode.h b/Engine/Source/Editor/Sequencer/Private/DisplayNodes/SequencerTrackNode.h
--- a/Engine/Source/Editor/Sequencer/Private/DisplayNodes/SequencerTrackNode.h
+++ b/Engine/Source/Editor/Sequencer/Private/DisplayNodes/SequencerTrackNode.h
@@ -89,6 +89,7 @@ public:
 	virtual float GetNodeHeight() const override;
 	virtual FNodePadding GetNodePadding() const override;
 	virtual ESequencerNode::Type GetType() const override;
+	virtual void BuildContextMenu( FMenuBuilder& MenuBuilder );
 
 private:
 
diff --git a/Engine/Source/Editor/Sequencer/Public/ISequencerTrackEditor.h b/Engine/Source/Editor/Sequencer/Public/ISequencerTrackEditor.h
--- a/Engine/Source/Editor/Sequencer/Public/ISequencerTrackEditor.h
+++ b/Engine/Source/Editor/Sequencer/Public/ISequencerTrackEditor.h
@@ -68,6 +68,12 @@ public:
 	 */
 	virtual TSharedPtr<SWidget> BuildOutlinerEditWidget(const FGuid& ObjectBinding, UMovieSceneTrack* Track) = 0;
 
+	/**
+	 * Builds the context menu for the track.
+	 * @param MenuBuilder The menu builder to use to build the track menu. 
+	 */
+	virtual void BuildTrackContextMenu( FMenuBuilder& MenuBuilder, UMovieSceneTrack* Track ) = 0;
+
 	/**
 	 * Called when an asset is dropped into Sequencer. Can potentially consume the asset
 	 * so it doesn't get added as a spawnable.
diff --git a/Engine/Source/Editor/Sequencer/Public/MovieSceneTrackEditor.h b/Engine/Source/Editor/Sequencer/Public/MovieSceneTrackEditor.h
--- a/Engine/Source/Editor/Sequencer/Public/MovieSceneTrackEditor.h
+++ b/Engine/Source/Editor/Sequencer/Public/MovieSceneTrackEditor.h
@@ -200,6 +200,7 @@ public:
 	virtual void BuildObjectBindingEditButtons(TSharedPtr<SHorizontalBox> EditBox, const FGuid& ObjectBinding, const UClass* ObjectClass) override { }
 	virtual void BuildObjectBindingTrackMenu(FMenuBuilder& MenuBuilder, const FGuid& ObjectBinding, const UClass* ObjectClass) override { }
 	virtual TSharedPtr<SWidget> BuildOutlinerEditWidget(const FGuid& ObjectBinding, UMovieSceneTrack* Track) override  { return TSharedPtr<SWidget>(); }
+	virtual void BuildTrackContextMenu( FMenuBuilder& MenuBuilder, UMovieSceneTrack* Track ) override { }
 	virtual bool HandleAssetAdded(UObject* Asset, const FGuid& TargetObjectGuid) override { return false; }
 
 	virtual bool IsAllowedKeyAll() const
diff --git a/Engine/Source/Editor/WorldBrowser/Private/Tiles/WorldTileCollectionModel.cpp b/Engine/Source/Editor/WorldBrowser/Private/Tiles/WorldTileCollectionModel.cpp
--- a/Engine/Source/Editor/WorldBrowser/Private/Tiles/WorldTileCollectionModel.cpp
+++ b/Engine/Source/Editor/WorldBrowser/Private/Tiles/WorldTileCollectionModel.cpp
@@ -1517,10 +1517,10 @@ static void SetupLandscapeImportLayers(const FTiledLandscapeImportSettings& InIm
 				ReadWeightmapFile(LayerImportInfo.LayerData, LayerImportInfo.SourceFilePath, FILEREAD_Silent);
 			}
 		}
-		
+
 		LayerImportInfo.LayerInfo = GetLandscapeLayerInfoObject(LayerImportInfo.LayerName, ContentPath);
 		LayerImportInfo.LayerInfo->bNoWeightBlend = LayerSettings.bNoBlendWeight;
-						
+
 		OutLayerInfo.Add(LayerImportInfo);
 	}
 }
@@ -1606,7 +1606,8 @@ void FWorldTileCollectionModel::ImportTiledLandscape_Executed()
 
 			// Setup layers list for importing
 			SetupLandscapeImportLayers(ImportSettings, GetWorld()->GetOutermost()->GetName(), TileIndex, TileImportSettings.ImportLayers);
-						
+			TileImportSettings.ImportLayerType = ELandscapeImportAlphamapType::Additive;
+
 			if (ReadRawFile(TileImportSettings.HeightData, Filename, FILEREAD_Silent))
 			{
 				FString MapFileName = WorldRootPath + TileName + FPackageName::GetMapPackageExtension();
diff --git a/Engine/Source/Editor/WorldBrowser/Private/Tiles/WorldTileModel.cpp b/Engine/Source/Editor/WorldBrowser/Private/Tiles/WorldTileModel.cpp
--- a/Engine/Source/Editor/WorldBrowser/Private/Tiles/WorldTileModel.cpp
+++ b/Engine/Source/Editor/WorldBrowser/Private/Tiles/WorldTileModel.cpp
@@ -930,17 +930,18 @@ ALandscapeProxy* FWorldTileModel::ImportLandscapeTile(const FLandscapeImportSett
 	// Cache pointer to landscape in the level model
 	Landscape = LandscapeProxy;
 
-	// Create landscape components	
+	// Create landscape components
 	LandscapeProxy->Import(
-		Settings.LandscapeGuid, 
-		Settings.SizeX, 
-		Settings.SizeY, 
-		Settings.ComponentSizeQuads, 
-		Settings.SectionsPerComponent, 
-		Settings.QuadsPerSection, 
-		Settings.HeightData.GetData(), 
-		*Settings.HeightmapFilename, 
-		Settings.ImportLayers);
+		Settings.LandscapeGuid,
+		0, 0,
+		Settings.SizeX - 1,
+		Settings.SizeY - 1,
+		Settings.SectionsPerComponent,
+		Settings.QuadsPerSection,
+		Settings.HeightData.GetData(),
+		*Settings.HeightmapFilename,
+		Settings.ImportLayers,
+		Settings.ImportLayerType);
 
 	return LandscapeProxy;
 }
diff --git a/Engine/Source/Editor/WorldBrowser/Private/Tiles/WorldTileModel.h b/Engine/Source/Editor/WorldBrowser/Private/Tiles/WorldTileModel.h
--- a/Engine/Source/Editor/WorldBrowser/Private/Tiles/WorldTileModel.h
+++ b/Engine/Source/Editor/WorldBrowser/Private/Tiles/WorldTileModel.h
@@ -53,6 +53,7 @@ public:
 		TArray<uint16>						HeightData;
 		TArray<FLandscapeImportLayerInfo>	ImportLayers;
 		FString								HeightmapFilename;
+		ELandscapeImportAlphamapType		ImportLayerType;
 	};
 	
 	/**
diff --git a/Engine/Source/Programs/ShaderCompileWorker/Private/ShaderCompileWorker.cpp b/Engine/Source/Programs/ShaderCompileWorker/Private/ShaderCompileWorker.cpp
--- a/Engine/Source/Programs/ShaderCompileWorker/Private/ShaderCompileWorker.cpp
+++ b/Engine/Source/Programs/ShaderCompileWorker/Private/ShaderCompileWorker.cpp
@@ -14,9 +14,13 @@
 #define DEBUG_USING_CONSOLE	0
 
 // this is for the protocol, not the data, bump if FShaderCompilerInput or ProcessInputFromArchive changes (also search for the second one with the same name, todo: put into one header file)
-const int32 ShaderCompileWorkerInputVersion = 4;
+const int32 ShaderCompileWorkerInputVersion = 5;
 // this is for the protocol, not the data, bump if FShaderCompilerOutput or WriteToOutputArchive changes (also search for the second one with the same name, todo: put into one header file)
-const int32 ShaderCompileWorkerOutputVersion = 2;
+const int32 ShaderCompileWorkerOutputVersion = 3;
+// this is for the protocol, not the data, bump if FShaderCompilerOutput or WriteToOutputArchive changes (also search for the second one with the same name, todo: put into one header file)
+const int32 ShaderCompileWorkerSingleJobHeader = 'S';
+// this is for the protocol, not the data, bump if FShaderCompilerOutput or WriteToOutputArchive changes (also search for the second one with the same name, todo: put into one header file)
+const int32 ShaderCompileWorkerPipelineJobHeader = 'P';
 
 double LastCompileTime = 0.0;
 
@@ -223,8 +227,6 @@ private:
 
 	void ProcessInputFromArchive(FArchive* InputFilePtr, TArray<FJobResult>& OutJobResults)
 	{
-		int32 NumBatches = 0;
-
 		FArchive& InputFile = *InputFilePtr;
 		int32 InputVersion;
 		InputFile << InputVersion;
@@ -235,31 +237,49 @@ private:
 
 		VerifyFormatVersions(ReceivedFormatVersionMap);
 
-		InputFile << NumBatches;
+		// Individual jobs
+		{
+			int32 SingleJobHeader = ShaderCompileWorkerSingleJobHeader;
+			InputFile << SingleJobHeader;
+			checkf(ShaderCompileWorkerSingleJobHeader == SingleJobHeader, TEXT("Exiting due to ShaderCompilerWorker expecting job header %d, got %d instead! Did you forget to build ShaderCompilerWorker?"), ShaderCompileWorkerSingleJobHeader, SingleJobHeader);
 
-		// Flush cache, to make sure we load the latest version of the input file.
-		// (Otherwise quick changes to a shader file can result in the wrong output.)
-		FlushShaderFileCache();
+			int32 NumBatches = 0;
+			InputFile << NumBatches;
 
-		for (int32 BatchIndex = 0; BatchIndex < NumBatches; BatchIndex++)
-		{
-			// Deserialize the job's inputs.
-			FShaderCompilerInput CompilerInput;
-			InputFile << CompilerInput;
+			// Flush cache, to make sure we load the latest version of the input file.
+			// (Otherwise quick changes to a shader file can result in the wrong output.)
+			FlushShaderFileCache();
 
-			if (IsValidRef(CompilerInput.SharedEnvironment))
+			for (int32 BatchIndex = 0; BatchIndex < NumBatches; BatchIndex++)
 			{
-				// Merge the shared environment into the per-shader environment before calling into the compile function
-				CompilerInput.Environment.Merge(*CompilerInput.SharedEnvironment);
-			}
+				// Deserialize the job's inputs.
+				FShaderCompilerInput CompilerInput;
+				InputFile << CompilerInput;
+
+				if (IsValidRef(CompilerInput.SharedEnvironment))
+				{
+					// Merge the shared environment into the per-shader environment before calling into the compile function
+					CompilerInput.Environment.Merge(*CompilerInput.SharedEnvironment);
+				}
 
-			// Process the job.
-			FShaderCompilerOutput CompilerOutput;
-			ProcessCompilationJob(CompilerInput,CompilerOutput,WorkingDirectory);
+				// Process the job.
+				FShaderCompilerOutput CompilerOutput;
+				ProcessCompilationJob(CompilerInput, CompilerOutput, WorkingDirectory);
 
-			// Serialize the job's output.
+				// Serialize the job's output.
 			FJobResult& JobResult = *new(OutJobResults) FJobResult;
-			JobResult.CompilerOutput = CompilerOutput;
+				JobResult.CompilerOutput = CompilerOutput;
+			}
+		}
+
+		// Shader pipeline jobs
+		{
+			int32 PipelineJobHeader = ShaderCompileWorkerPipelineJobHeader;
+			InputFile << PipelineJobHeader;
+			checkf(ShaderCompileWorkerPipelineJobHeader == PipelineJobHeader, TEXT("Exiting due to ShaderCompilerWorker expecting pipeline job header %d, got %d instead! Did you forget to build ShaderCompilerWorker?"), ShaderCompileWorkerSingleJobHeader, PipelineJobHeader);
+
+			int32 NumPipelines = 0;
+			InputFile << NumPipelines;
 		}
 	}
 
@@ -330,13 +350,26 @@ private:
 		OutputFile << ErrorStringLength;
 		OutputFile << ErrorStringLength;
 
-		int32 NumBatches = JobResults.Num();
-		OutputFile << NumBatches;
+		{
+			int32 SingleJobHeader = ShaderCompileWorkerSingleJobHeader;
+			OutputFile << SingleJobHeader;
+
+			int32 NumBatches = JobResults.Num();
+			OutputFile << NumBatches;
+
+			for (int32 ResultIndex = 0; ResultIndex < JobResults.Num(); ResultIndex++)
+			{
+				FJobResult& JobResult = JobResults[ResultIndex];
+				OutputFile << JobResult.CompilerOutput;
+			}
+		}
 
-		for (int32 ResultIndex = 0; ResultIndex < JobResults.Num(); ResultIndex++)
 		{
-			FJobResult& JobResult = JobResults[ResultIndex];
-			OutputFile << JobResult.CompilerOutput;
+			int32 PipelineJobHeader = ShaderCompileWorkerPipelineJobHeader;
+			OutputFile << PipelineJobHeader;
+
+			int32 NumPipelines = 0;
+			OutputFile << NumPipelines;
 		}
 	}
 
diff --git a/Engine/Source/Programs/UnrealBuildTool/System/UnrealBuildTool.cs b/Engine/Source/Programs/UnrealBuildTool/System/UnrealBuildTool.cs
--- a/Engine/Source/Programs/UnrealBuildTool/System/UnrealBuildTool.cs
+++ b/Engine/Source/Programs/UnrealBuildTool/System/UnrealBuildTool.cs
@@ -2143,31 +2143,38 @@ namespace UnrealBuildTool
 					throw new BuildException("ShouldDoHotReload cannot handle multiple binaries returning from UEBuildTarget.MakeExecutablePaths");
 				}
 
-				var EditorProcessFilename = EditorProcessFilenames[0].FullName;
-				var EditorProcessName = Path.GetFileNameWithoutExtension(EditorProcessFilename);
-				var EditorProcesses = BuildHostPlatform.Current.GetProcessesByName(EditorProcessName);
-				var BinariesPath = Path.GetFullPath(Path.GetDirectoryName(EditorProcessFilename));
-				var PerProjectBinariesPath = (TargetDesc.ProjectFile != null) ? Path.Combine(TargetDesc.ProjectFile.Directory.FullName, BinariesPath.Substring(BinariesPath.LastIndexOf("Binaries"))) : "";
-				bIsRunning = EditorProcesses.FirstOrDefault(EditorProc =>
-					{
-						if (!Path.GetFullPath(EditorProc.Filename).StartsWith(BinariesPath, StringComparison.InvariantCultureIgnoreCase))
-						{
-							return false;
-						}
+				var Processes = BuildHostPlatform.Current.GetProcesses();
+                var RootDir = Path.Combine(Path.GetDirectoryName(Assembly.GetExecutingAssembly().Modules.First().FullyQualifiedName), "..", "..", "..");
+                var EditorRunsDir = Path.Combine(RootDir, "Engine", "Intermediate", "EditorRuns");
 
-						if (PerProjectBinariesPath.Equals(BinariesPath, StringComparison.InvariantCultureIgnoreCase))
-						{
-							return false;
-						}
+                if (!Directory.Exists(EditorRunsDir))
+                {
+                    return false;
+                }
 
-						var Modules = BuildHostPlatform.Current.GetProcessModules(EditorProc.PID, EditorProc.Filename);
-						if (!Modules.Any(Module => Module.StartsWith(PerProjectBinariesPath)))
-						{
-							return false;
-						}
+                var EditorRunsFiles = new DirectoryInfo(EditorRunsDir).GetFiles();
+
+                var NormalizedProcFileName = Path.GetFullPath(EditorProcessFilenames[0].ToString()).TrimEnd(Path.DirectorySeparatorChar).TrimEnd(Path.AltDirectorySeparatorChar).ToLowerInvariant();
 
-						return true;
-					}) != default(BuildHostPlatform.ProcessInfo);
+				foreach(var File in EditorRunsFiles)
+				{
+					int PID;
+					BuildHostPlatform.ProcessInfo Proc = null;
+					if (!Int32.TryParse(File.Name, out PID) || (Proc = Processes.FirstOrDefault(P => P.PID == PID)) == default(BuildHostPlatform.ProcessInfo))
+					{
+						// Delete stale files (it may happen if editor crashes).
+						File.Delete();
+						continue;
+					}
+
+                    // Don't break here to allow clean-up of other stale files.
+                    if (!bIsRunning)
+                    {
+                        // Otherwise check if the path matches.
+                        var FileProcName = Path.GetFullPath(Proc.Filename).TrimEnd(Path.DirectorySeparatorChar).TrimEnd(Path.AltDirectorySeparatorChar).ToLowerInvariant();
+                        bIsRunning = NormalizedProcFileName == FileProcName;
+                    }
+				}
 			}
 			return bIsRunning;
 		}
diff --git a/Engine/Source/Runtime/Core/Private/Async/TaskGraph.cpp b/Engine/Source/Runtime/Core/Private/Async/TaskGraph.cpp
--- a/Engine/Source/Runtime/Core/Private/Async/TaskGraph.cpp
+++ b/Engine/Source/Runtime/Core/Private/Async/TaskGraph.cpp
@@ -242,7 +242,7 @@ static FAutoConsoleVariableRef CVarMaxTasksToStartOnDequeue(
 	ECVF_Cheat
 	);
 
-static int32 GAllAnyThreadTasksFromGameHiPri = 0;
+int32 GAllAnyThreadTasksFromGameHiPri = 1;
 static FAutoConsoleVariableRef CVarAllAnyThreadTasksFromGameHiPri(
 	TEXT("TaskGraph.AllAnyThreadTasksFromGameHiPri"),
 	GAllAnyThreadTasksFromGameHiPri,
@@ -1354,22 +1354,9 @@ public:
 			TASKGRAPH_SCOPE_CYCLE_COUNTER(3, STAT_TaskGraph_QueueTask_AnyThread);
 			if (FPlatformProcess::SupportsMultithreading())
 			{
-				ENamedThreads::Type CurrentThreadIfKnown = ENamedThreads::AnyThread;
-				if (GAllAnyThreadTasksFromGameHiPri || !GFastSchedulerLatched)
-				{
-				    if (ENamedThreads::GetThreadIndex(InCurrentThreadIfKnown) == ENamedThreads::AnyThread)
-				    {
-					    CurrentThreadIfKnown = GetCurrentThread();
-				    }
-				    else
-				    {
-					    CurrentThreadIfKnown = ENamedThreads::GetThreadIndex(InCurrentThreadIfKnown);
-					    checkThreadGraph(CurrentThreadIfKnown == GetCurrentThread());
-				    }
-				}
 				{
 					TASKGRAPH_SCOPE_CYCLE_COUNTER(4, STAT_TaskGraph_QueueTask_IncomingAnyThreadTasks_Push);
-					if ((GAllAnyThreadTasksFromGameHiPri && CurrentThreadIfKnown == ENamedThreads::GameThread) ||  ENamedThreads::GetPriority(Task->ThreadToExecuteOn))
+					if (ENamedThreads::GetPriority(Task->ThreadToExecuteOn))
 					{
 						IncomingAnyThreadTasksHiPri.Push(Task);
 					}
@@ -1387,6 +1374,16 @@ public:
 				}
 				else
 				{
+					ENamedThreads::Type CurrentThreadIfKnown = ENamedThreads::AnyThread;
+					if (ENamedThreads::GetThreadIndex(InCurrentThreadIfKnown) == ENamedThreads::AnyThread)
+					{
+						CurrentThreadIfKnown = GetCurrentThread();
+					}
+					else
+					{
+						CurrentThreadIfKnown = ENamedThreads::GetThreadIndex(InCurrentThreadIfKnown);
+						checkThreadGraph(CurrentThreadIfKnown == GetCurrentThread());
+					}
 					FTaskThreadBase* TempTarget;
 					{
 						TASKGRAPH_SCOPE_CYCLE_COUNTER(5, STAT_TaskGraph_QueueTask_StalledUnnamedThreads_Pop);
@@ -2403,7 +2400,7 @@ void FGraphEvent::DispatchSubsequents(TArray<FBaseGraphTask*>& NewTasks, ENamedT
 			STAT_FNullGraphTask_DontCompleteUntil,
 			STATGROUP_TaskGraphTasks);
 
-		TGraphTask<FNullGraphTask>::CreateTask(FGraphEventRef(this), &TempEventsToWaitFor, CurrentThreadIfKnown).ConstructAndDispatchWhenReady(GET_STATID(STAT_FNullGraphTask_DontCompleteUntil));
+		TGraphTask<FNullGraphTask>::CreateTask(FGraphEventRef(this), &TempEventsToWaitFor, CurrentThreadIfKnown).ConstructAndDispatchWhenReady(GET_STATID(STAT_FNullGraphTask_DontCompleteUntil), ENamedThreads::HiPri(ENamedThreads::AnyThread));
 		return;
 	}
 
@@ -2588,10 +2585,10 @@ static void TaskGraphBenchmark(const TArray<FString>& Args)
 		Tasks.Reserve(1000);
 		for (int32 Index = 0; Index < 1000; Index++)
 		{
-			Tasks.Emplace(TGraphTask<FNullGraphTask>::CreateTask(nullptr, ENamedThreads::GameThread).ConstructAndDispatchWhenReady(TStatId()));
+			Tasks.Emplace(TGraphTask<FNullGraphTask>::CreateTask(nullptr, ENamedThreads::GameThread).ConstructAndDispatchWhenReady(TStatId(), ENamedThreads::AnyThread));
 		}
 		QueueTime = FPlatformTime::Seconds();
-		FGraphEventRef Join = TGraphTask<FNullGraphTask>::CreateTask(&Tasks, ENamedThreads::GameThread).ConstructAndDispatchWhenReady(TStatId());
+		FGraphEventRef Join = TGraphTask<FNullGraphTask>::CreateTask(&Tasks, ENamedThreads::GameThread).ConstructAndDispatchWhenReady(TStatId(), ENamedThreads::AnyThread);
 		JoinTime = FPlatformTime::Seconds();
 		FTaskGraphInterface::Get().WaitUntilTaskCompletes(Join, ENamedThreads::GameThread_Local);
 		EndTime = FPlatformTime::Seconds();
@@ -2605,11 +2602,11 @@ static void TaskGraphBenchmark(const TArray<FString>& Args)
 		ParallelFor(1000, 
 			[&Tasks](int32 Index)
 			{
-				Tasks[Index] = TGraphTask<FNullGraphTask>::CreateTask().ConstructAndDispatchWhenReady(TStatId());
+				Tasks[Index] = TGraphTask<FNullGraphTask>::CreateTask().ConstructAndDispatchWhenReady(TStatId(), ENamedThreads::AnyThread);
 			}
 		);
 		QueueTime = FPlatformTime::Seconds();
-		FGraphEventRef Join = TGraphTask<FNullGraphTask>::CreateTask(&Tasks, ENamedThreads::GameThread).ConstructAndDispatchWhenReady(TStatId());
+		FGraphEventRef Join = TGraphTask<FNullGraphTask>::CreateTask(&Tasks, ENamedThreads::GameThread).ConstructAndDispatchWhenReady(TStatId(), ENamedThreads::AnyThread);
 		JoinTime = FPlatformTime::Seconds();
 		FTaskGraphInterface::Get().WaitUntilTaskCompletes(Join, ENamedThreads::GameThread_Local);
 		EndTime = FPlatformTime::Seconds();
@@ -2628,14 +2625,14 @@ static void TaskGraphBenchmark(const TArray<FString>& Args)
 			ENamedThreads::Type CurrentThread = FTaskGraphInterface::Get().GetCurrentThreadIfKnown();
 			for (int32 InnerIndex = 0; InnerIndex < 100; InnerIndex++)
 			{
-				InnerTasks[InnerIndex] = TGraphTask<FNullGraphTask>::CreateTask(nullptr, CurrentThread).ConstructAndDispatchWhenReady(TStatId());
+				InnerTasks[InnerIndex] = TGraphTask<FNullGraphTask>::CreateTask(nullptr, CurrentThread).ConstructAndDispatchWhenReady(TStatId(), ENamedThreads::AnyThread);
 			}
 			// join the above tasks
-			Tasks[Index] = TGraphTask<FNullGraphTask>::CreateTask(&InnerTasks, CurrentThread).ConstructAndDispatchWhenReady(TStatId());
+			Tasks[Index] = TGraphTask<FNullGraphTask>::CreateTask(&InnerTasks, CurrentThread).ConstructAndDispatchWhenReady(TStatId(), ENamedThreads::AnyThread);
 		}
 		);
 		QueueTime = FPlatformTime::Seconds();
-		FGraphEventRef Join = TGraphTask<FNullGraphTask>::CreateTask(&Tasks, ENamedThreads::GameThread).ConstructAndDispatchWhenReady(TStatId());
+		FGraphEventRef Join = TGraphTask<FNullGraphTask>::CreateTask(&Tasks, ENamedThreads::GameThread).ConstructAndDispatchWhenReady(TStatId(), ENamedThreads::AnyThread);
 		JoinTime = FPlatformTime::Seconds();
 		FTaskGraphInterface::Get().WaitUntilTaskCompletes(Join, ENamedThreads::GameThread_Local);
 		EndTime = FPlatformTime::Seconds();
@@ -2654,14 +2651,14 @@ static void TaskGraphBenchmark(const TArray<FString>& Args)
 			ENamedThreads::Type CurrentThread = FTaskGraphInterface::Get().GetCurrentThreadIfKnown();
 			for (int32 InnerIndex = 0; InnerIndex < 10; InnerIndex++)
 			{
-				InnerTasks[InnerIndex] = TGraphTask<FNullGraphTask>::CreateTask(nullptr, CurrentThread).ConstructAndDispatchWhenReady(TStatId());
+				InnerTasks[InnerIndex] = TGraphTask<FNullGraphTask>::CreateTask(nullptr, CurrentThread).ConstructAndDispatchWhenReady(TStatId(), ENamedThreads::AnyThread);
 			}
 			// join the above tasks
-			Tasks[Index] = TGraphTask<FNullGraphTask>::CreateTask(&InnerTasks, CurrentThread).ConstructAndDispatchWhenReady(TStatId());
+			Tasks[Index] = TGraphTask<FNullGraphTask>::CreateTask(&InnerTasks, CurrentThread).ConstructAndDispatchWhenReady(TStatId(), ENamedThreads::AnyThread);
 		}
 		);
 		QueueTime = FPlatformTime::Seconds();
-		FGraphEventRef Join = TGraphTask<FNullGraphTask>::CreateTask(&Tasks, ENamedThreads::GameThread).ConstructAndDispatchWhenReady(TStatId());
+		FGraphEventRef Join = TGraphTask<FNullGraphTask>::CreateTask(&Tasks, ENamedThreads::GameThread).ConstructAndDispatchWhenReady(TStatId(), ENamedThreads::AnyThread);
 		JoinTime = FPlatformTime::Seconds();
 		FTaskGraphInterface::Get().WaitUntilTaskCompletes(Join, ENamedThreads::GameThread_Local);
 		EndTime = FPlatformTime::Seconds();
diff --git a/Engine/Source/Runtime/Core/Public/Async/TaskGraphInterfaces.h b/Engine/Source/Runtime/Core/Public/Async/TaskGraphInterfaces.h
--- a/Engine/Source/Runtime/Core/Public/Async/TaskGraphInterfaces.h
+++ b/Engine/Source/Runtime/Core/Public/Async/TaskGraphInterfaces.h
@@ -26,6 +26,8 @@ DECLARE_CYCLE_STAT_EXTERN(TEXT("FTriggerEventGraphTask"), STAT_FTriggerEventGrap
 DECLARE_CYCLE_STAT_EXTERN(TEXT("ParallelFor"), STAT_ParallelFor, STATGROUP_TaskGraphTasks, CORE_API);
 DECLARE_CYCLE_STAT_EXTERN(TEXT("ParallelForTask"), STAT_ParallelForTask, STATGROUP_TaskGraphTasks, CORE_API);
 
+extern CORE_API int32 GAllAnyThreadTasksFromGameHiPri;
+
 namespace ENamedThreads
 {
 	enum Type
@@ -93,6 +95,11 @@ namespace ENamedThreads
 	{
 		return Type(ThreadAndIndex | HighPriority);
 	}
+
+	FORCEINLINE Type AnyThreadGame()
+	{
+		return Type(AnyThread | (GAllAnyThreadTasksFromGameHiPri ? HighPriority : 0));
+	}
 }
 
 
@@ -1139,7 +1146,7 @@ public:
 	 *	@param StatId The stat id for this task.
 	 *	@param InDesiredThread; Thread to run on, can be ENamedThreads::AnyThread
 	**/
-	FNullGraphTask(const TStatId& StatId, ENamedThreads::Type InDesiredThread = ENamedThreads::AnyThread)
+	FNullGraphTask(const TStatId& StatId, ENamedThreads::Type InDesiredThread)
 		: FCustomStatIDGraphTaskBase(StatId)
 		, DesiredThread(InDesiredThread)
 	{
@@ -1193,7 +1200,7 @@ public:
 
 	ENamedThreads::Type GetDesiredThread()
 	{
-		return ENamedThreads::AnyThread;
+		return ENamedThreads::HiPri(ENamedThreads::AnyThread);
 	}
 
 	static ESubsequentsMode::Type GetSubsequentsMode() { return ESubsequentsMode::TrackSubsequents; }
diff --git a/Engine/Source/Runtime/Core/Public/Math/Quat.h b/Engine/Source/Runtime/Core/Public/Math/Quat.h
--- a/Engine/Source/Runtime/Core/Public/Math/Quat.h
+++ b/Engine/Source/Runtime/Core/Public/Math/Quat.h
@@ -135,6 +135,14 @@ public:
 	 */
 	FORCEINLINE bool Equals(const FQuat& Q, float Tolerance=KINDA_SMALL_NUMBER) const;
 
+	/**
+	 * Checks whether this Quaternion is an Identity Quaternion.
+	 * Assumes Quaternion tested is normalized.
+	 *
+	 * @return true if Quaternion is a normalized Identity Quaternion.
+	 */
+	FORCEINLINE bool IsIdentity() const;
+
 	/**
 	 * Subtracts another quaternion from this.
 	 * This is a component-wise subtraction; composing quaternions should be done via multiplication.
@@ -726,6 +734,10 @@ FORCEINLINE bool FQuat::Equals(const FQuat& Q, float Tolerance) const
 #endif // PLATFORM_ENABLE_VECTORINTRINSICS
 }
 
+FORCEINLINE bool FQuat::IsIdentity() const
+{
+	return ((W * W) > ((1.f - DELTA) * (1.f - DELTA)));
+}
 
 FORCEINLINE FQuat FQuat::operator-=(const FQuat& Q)
 {
diff --git a/Engine/Source/Runtime/Core/Public/Math/UnrealMath.h b/Engine/Source/Runtime/Core/Public/Math/UnrealMath.h
--- a/Engine/Source/Runtime/Core/Public/Math/UnrealMath.h
+++ b/Engine/Source/Runtime/Core/Public/Math/UnrealMath.h
@@ -470,7 +470,7 @@ inline FString FMath::FormatIntToHumanReadable(int32 Val)
 template<class U>
 FORCEINLINE_DEBUGGABLE FRotator FMath::Lerp( const FRotator& A, const FRotator& B, const U& Alpha)
 {
-	return A + (B - A).GetNormalized() * Alpha;
+	return (A * (1 - Alpha) + B * Alpha).GetNormalized();
 }
 
 
diff --git a/Engine/Source/Runtime/CoreUObject/Private/Serialization/AsyncLoading.cpp b/Engine/Source/Runtime/CoreUObject/Private/Serialization/AsyncLoading.cpp
--- a/Engine/Source/Runtime/CoreUObject/Private/Serialization/AsyncLoading.cpp
+++ b/Engine/Source/Runtime/CoreUObject/Private/Serialization/AsyncLoading.cpp
@@ -417,12 +417,57 @@ FAsyncPackage* FAsyncLoadingThread::FindExistingPackageAndAddCompletionCallback(
 	return Result;
 }
 
-void FAsyncLoadingThread::ProcessAsyncPackageRequest(FAsyncPackageDesc* InRequest, FAsyncPackage* InRootPackage, TMap<FName, int32>& InDependencyTracker)
+void FAsyncLoadingThread::UpdateExistingPackagePriorities(FAsyncPackage* InPackage, TAsyncLoadPriority InNewPriority, TSet<FName>& InDependencyTracker, IAssetRegistryInterface* InAssetRegistry)
+{
+	InDependencyTracker.Add(InPackage->GetPackageName());
+
+	if (InNewPriority > InPackage->GetPriority())
+	{
+		AsyncPackages.Remove(InPackage);
+		InPackage->SetPriority(InNewPriority);
+
+		// Reduce loading counters ready for InsertPackage to increment them again
+		AsyncLoadingCounter.Decrement();
+		AsyncPackagesCounter.Decrement();
+
+		InsertPackage(InPackage, InAssetRegistry != nullptr ? EAsyncPackageInsertMode::InsertAfterMatchingPriorities : EAsyncPackageInsertMode::InsertBeforeMatchingPriorities);
+	}
+
+	if (InAssetRegistry)
+	{
+		TArray<FName> Dependencies;
+		InAssetRegistry->GetDependencies(InPackage->GetPackageName(), Dependencies, EAssetRegistryDependencyType::Hard);
+
+		for (FName DependencyName : Dependencies)
+		{
+			if (!InDependencyTracker.Contains(DependencyName))
+			{
+				int32 PackageIndex = FindPackageByName(AsyncPackages, DependencyName);
+
+				if (PackageIndex >= 0)
+				{
+					FAsyncPackage* DependencyPackage = AsyncPackages[PackageIndex];
+					UpdateExistingPackagePriorities(DependencyPackage, InNewPriority, InDependencyTracker, InAssetRegistry);
+				}
+			}
+		}
+	}
+}
+
+void FAsyncLoadingThread::ProcessAsyncPackageRequest(FAsyncPackageDesc* InRequest, FAsyncPackage* InRootPackage, TSet<FName>& InDependencyTracker, IAssetRegistryInterface* InAssetRegistry)
 {
 	check(!InDependencyTracker.Contains(InRequest->Name));
+	InDependencyTracker.Add(InRequest->Name);
+
 	FAsyncPackage* Package = FindExistingPackageAndAddCompletionCallback(InRequest, AsyncPackages);
 
-	if (!Package)
+	if (Package)
+	{
+		// The package is already sitting in the queue. Make sure the its priority, and the priority of all its
+		// dependencies is at least as high as the priority of this request
+		UpdateExistingPackagePriorities(Package, InRequest->Priority, InDependencyTracker, InAssetRegistry);
+	}
+	else
 	{
 		// [BLOCKING] LoadedPackages are accessed on the main thread too, so lock to be able to add a completion callback
 #if THREADSAFE_UOBJECTS
@@ -452,13 +497,10 @@ void FAsyncLoadingThread::ProcessAsyncPackageRequest(FAsyncPackageDesc* InReques
 		Package->SetDependencyRootPackage(InRootPackage);
 
 #if !WITH_EDITOR
-		if (GPreloadPackageDependencies != 0)
-		{
-			auto AssetRegistry = IAssetRegistryInterface::GetPtr();
-			if (AssetRegistry)
+		if (InAssetRegistry)
 			{
 				TArray<FName> Dependencies;
-				AssetRegistry->GetDependencies(Package->GetPackageName(), Dependencies, EAssetRegistryDependencyType::Hard);
+			InAssetRegistry->GetDependencies(Package->GetPackageName(), Dependencies, EAssetRegistryDependencyType::Hard);
 
 				if (InRootPackage == nullptr)
 				{
@@ -467,27 +509,19 @@ void FAsyncLoadingThread::ProcessAsyncPackageRequest(FAsyncPackageDesc* InReques
 
 				for (auto DependencyName : Dependencies)
 				{
-					if (!InDependencyTracker.Contains(DependencyName))
-					{
-						// Try and find the package in memory first
-						UPackage* ClassPackage = FindObjectFast<UPackage>( NULL, DependencyName, false, false );
-
-						if (ClassPackage == nullptr)
+				if (!InDependencyTracker.Contains(DependencyName) && FindObjectFast<UPackage>(nullptr, DependencyName, false, false ) == nullptr)
 						{
 							QueuedPackagesCounter.Increment();
 							const int32 RequestID = GPackageRequestID.Increment();
 							FAsyncLoadingThread::Get().AddPendingRequest(RequestID);
 							FAsyncPackageDesc DependencyPackageRequest(RequestID, DependencyName, NAME_None, FGuid(), FLoadPackageAsyncDelegate(), InRequest->PackageFlags, INDEX_NONE, InRequest->Priority);
-							InDependencyTracker.Add(InRequest->Name);
-							ProcessAsyncPackageRequest(&DependencyPackageRequest, InRootPackage, InDependencyTracker);
-						}
-					}
+					ProcessAsyncPackageRequest(&DependencyPackageRequest, InRootPackage, InDependencyTracker, InAssetRegistry);
 				}
 			}
 		}
 #endif
 		// Add to queue according to priority.
-		InsertPackage(Package);
+		InsertPackage(Package, InAssetRegistry != nullptr ? EAsyncPackageInsertMode::InsertAfterMatchingPriorities : EAsyncPackageInsertMode::InsertBeforeMatchingPriorities);
 
 		// For all other cases this is handled in FindExistingPackageAndAddCompletionCallback
 		const int32 QueuedPackagesCount = QueuedPackagesCounter.Decrement();
@@ -515,13 +549,20 @@ int32 FAsyncLoadingThread::CreateAsyncPackagesFromQueue()
 
 	if (QueueCopy.Num() > 0)
 	{
+		IAssetRegistryInterface* AssetRegistry = nullptr;
+		
+		if (GPreloadPackageDependencies && IsPlatformFileCompatibleWithDependencyPreloading())
+		{
+			AssetRegistry = IAssetRegistryInterface::GetPtr();
+		}
+		
 		double Timer = 0;
 		{
 			SCOPE_SECONDS_COUNTER(Timer);
 			for (FAsyncPackageDesc* PackageRequest : QueueCopy)
 			{
 				DependencyTracker.Reset();
-				ProcessAsyncPackageRequest(PackageRequest, nullptr, DependencyTracker);
+				ProcessAsyncPackageRequest(PackageRequest, nullptr, DependencyTracker, AssetRegistry);
 				delete PackageRequest;
 			}
 		}
@@ -533,7 +574,7 @@ int32 FAsyncLoadingThread::CreateAsyncPackagesFromQueue()
 	return NumCreated;
 }
 
-void FAsyncLoadingThread::InsertPackage(FAsyncPackage* Package)
+void FAsyncLoadingThread::InsertPackage(FAsyncPackage* Package, EAsyncPackageInsertMode InsertMode)
 {
 	checkSlow(IsInAsyncLoadThread());
 
@@ -547,12 +588,32 @@ void FAsyncLoadingThread::InsertPackage(FAsyncPackage* Package)
 #if THREADSAFE_UOBJECTS
 		FScopeLock LockAsyncPackages(&AsyncPackagesCritical);
 #endif
+		int32 InsertIndex = -1;
+
+		switch (InsertMode)
+		{
+		case EAsyncPackageInsertMode::InsertAfterMatchingPriorities:
+			{
+				InsertIndex = AsyncPackages.IndexOfByPredicate([Package](const FAsyncPackage* Element)
+				{
+					return Element->GetPriority() < Package->GetPriority();
+				});
+
+				break;
+			}
+
+		case EAsyncPackageInsertMode::InsertBeforeMatchingPriorities:
+			{
 		// Insert new package keeping descending priority order in AsyncPackages
-		int32 InsertIndex = AsyncPackages.IndexOfByPredicate([Package](const FAsyncPackage* Element)
+				InsertIndex = AsyncPackages.IndexOfByPredicate([Package](const FAsyncPackage* Element)
 		{
 			return Element->GetPriority() <= Package->GetPriority();
 		});
 
+				break;
+			}
+		};
+
 		InsertIndex = InsertIndex == INDEX_NONE ? AsyncPackages.Num() : InsertIndex;
 
 		AsyncPackages.InsertUninitialized(InsertIndex);
@@ -1500,80 +1561,80 @@ EAsyncPackageState::Type FAsyncPackage::LoadImports()
 		}
 		if (ImportDependencyPackages.Num() == 0)
 		{
-			while (Import->OuterIndex.IsImport())
-			{
-				Import = &Linker->Imp(Import->OuterIndex);
-			}
-			check(Import->OuterIndex.IsNull());
+		while (Import->OuterIndex.IsImport())
+		{
+			Import = &Linker->Imp(Import->OuterIndex);
+		}
+		check(Import->OuterIndex.IsNull());
 
-			// @todo: why do we need this? some UFunctions have null outer in the linker.
-			if (Import->ClassName != NAME_Package)
-			{
-				continue;
-			}
-			
+		// @todo: why do we need this? some UFunctions have null outer in the linker.
+		if (Import->ClassName != NAME_Package)
+		{
+			continue;
+		}
 
-			// Don't try to import a package that is in an import table that we know is an invalid entry
-			if (FLinkerLoad::KnownMissingPackages.Contains(Import->ObjectName))
-			{
-				continue;
-			}
 
-			// Our import package name is the import name
+		// Don't try to import a package that is in an import table that we know is an invalid entry
+		if (FLinkerLoad::KnownMissingPackages.Contains(Import->ObjectName))
+		{
+			continue;
+		}
+
+		// Our import package name is the import name
 			ImportDependencyPackages.Add(Import->ObjectName);
 		}
 
 		for (const FName& ImportPackageFName : ImportDependencyPackages)
 		{
-			// Handle circular dependencies - try to find existing packages.
-			UPackage* ExistingPackage = dynamic_cast<UPackage*>(StaticFindObjectFast(UPackage::StaticClass(), nullptr, ImportPackageFName, true));
-			if (ExistingPackage && !(ExistingPackage->PackageFlags & PKG_CompiledIn) && !ExistingPackage->bHasBeenFullyLoaded)//!ExistingPackage->HasAnyFlags(RF_WasLoaded))
-			{
-				// The import package already exists. Check if it's currently being streamed as well. If so, make sure
-				// we add all dependencies that don't yet have linkers created otherwise we risk that if the current package
-				// doesn't depend on any other packages that have not yet started streaming, creating imports is going
-				// to load packages blocking the main thread.
-				int32 PendingAsyncPackageIndex = FAsyncLoadingThread::Get().FindAsyncPackage(ImportPackageFName);
-				if (PendingAsyncPackageIndex != INDEX_NONE)
-				{
-					FAsyncPackage& PendingPackage = *FAsyncLoadingThread::Get().GetPackage(PendingAsyncPackageIndex);
-					FLinkerLoad* PendingPackageLinker = PendingPackage.Linker;
-					if (PendingPackageLinker == nullptr || !PendingPackageLinker->HasFinishedInitialization())
-					{
-						// Add this import to the dependency list.
-						AddUniqueLinkerDependencyPackage(AsyncQueueIndex, PendingPackage);
-					}
-					else
-					{
-						UE_LOG(LogStreaming, Verbose, TEXT("FAsyncPackage::LoadImports for %s: Linker exists for %s"), *Desc.NameToLoad.ToString(), *ImportPackageFName.ToString());
-						// Only keep a reference to this package so that its linker doesn't go away too soon
-						PendingPackage.DependencyRefCount.Increment();
-						ReferencedImports.Add(&PendingPackage);
-						// Check if we need to add its dependencies too.
-						TSet<FAsyncPackage*> SearchedPackages;
-						AddDependencyTree(AsyncQueueIndex, PendingPackage, SearchedPackages);
-					}
-				}
-			}
-
-			if (!ExistingPackage && ContainsDependencyPackage(PendingImportedPackages, ImportPackageFName) == INDEX_NONE)
+		// Handle circular dependencies - try to find existing packages.
+		UPackage* ExistingPackage = dynamic_cast<UPackage*>(StaticFindObjectFast(UPackage::StaticClass(), nullptr, ImportPackageFName, true));
+		if (ExistingPackage && !(ExistingPackage->PackageFlags & PKG_CompiledIn) && !ExistingPackage->bHasBeenFullyLoaded)//!ExistingPackage->HasAnyFlags(RF_WasLoaded))
+		{
+			// The import package already exists. Check if it's currently being streamed as well. If so, make sure
+			// we add all dependencies that don't yet have linkers created otherwise we risk that if the current package
+			// doesn't depend on any other packages that have not yet started streaming, creating imports is going
+			// to load packages blocking the main thread.
+			int32 PendingAsyncPackageIndex = FAsyncLoadingThread::Get().FindAsyncPackage(ImportPackageFName);
+			if (PendingAsyncPackageIndex != INDEX_NONE)
 			{
-				const FString ImportPackageName(Import->ObjectName.ToString());
-				// The package doesn't exist and this import is not in the dependency list so add it now.
-				if (!FPackageName::IsShortPackageName(ImportPackageName))
+				FAsyncPackage& PendingPackage = *FAsyncLoadingThread::Get().GetPackage(PendingAsyncPackageIndex);
+				FLinkerLoad* PendingPackageLinker = PendingPackage.Linker;
+				if (PendingPackageLinker == nullptr || !PendingPackageLinker->HasFinishedInitialization())
 				{
-					UE_LOG(LogStreaming, Verbose, TEXT("FAsyncPackage::LoadImports for %s: Loading %s"), *Desc.NameToLoad.ToString(), *ImportPackageName);
-					AddImportDependency(AsyncQueueIndex, ImportPackageFName);
+					// Add this import to the dependency list.
+					AddUniqueLinkerDependencyPackage(AsyncQueueIndex, PendingPackage);
 				}
 				else
 				{
-					// This usually means there's a reference to a script package from another project
-					UE_LOG(LogStreaming, Warning, TEXT("FAsyncPackage::LoadImports for %s: Short package name in imports list: %s"), *Desc.NameToLoad.ToString(), *ImportPackageName);
+					UE_LOG(LogStreaming, Verbose, TEXT("FAsyncPackage::LoadImports for %s: Linker exists for %s"), *Desc.NameToLoad.ToString(), *ImportPackageFName.ToString());
+					// Only keep a reference to this package so that its linker doesn't go away too soon
+					PendingPackage.DependencyRefCount.Increment();
+					ReferencedImports.Add(&PendingPackage);
+					// Check if we need to add its dependencies too.
+					TSet<FAsyncPackage*> SearchedPackages;
+					AddDependencyTree(AsyncQueueIndex, PendingPackage, SearchedPackages);
 				}
 			}
+		}
 
-			UpdateLoadPercentage();
+		if (!ExistingPackage && ContainsDependencyPackage(PendingImportedPackages, ImportPackageFName) == INDEX_NONE)
+			{
+			const FString ImportPackageName(Import->ObjectName.ToString());
+			// The package doesn't exist and this import is not in the dependency list so add it now.
+			if (!FPackageName::IsShortPackageName(ImportPackageName))
+			{
+				UE_LOG(LogStreaming, Verbose, TEXT("FAsyncPackage::LoadImports for %s: Loading %s"), *Desc.NameToLoad.ToString(), *ImportPackageName);
+				AddImportDependency(AsyncQueueIndex, ImportPackageFName);
+			}
+			else
+			{
+				// This usually means there's a reference to a script package from another project
+				UE_LOG(LogStreaming, Warning, TEXT("FAsyncPackage::LoadImports for %s: Short package name in imports list: %s"), *Desc.NameToLoad.ToString(), *ImportPackageName);
+			}
 		}
+
+		UpdateLoadPercentage();
+	}
 	}
 
 	if (PendingImportedPackages.Num())
diff --git a/Engine/Source/Runtime/CoreUObject/Private/Serialization/AsyncLoadingThread.h b/Engine/Source/Runtime/CoreUObject/Private/Serialization/AsyncLoadingThread.h
--- a/Engine/Source/Runtime/CoreUObject/Private/Serialization/AsyncLoadingThread.h
+++ b/Engine/Source/Runtime/CoreUObject/Private/Serialization/AsyncLoadingThread.h
@@ -5,6 +5,8 @@
 =============================================================================*/
 #pragma once
 
+class IAssetRegistryInterface;
+
 /**
  * Async loading thread. Preloads/serializes packages on async loading thread. Postloads objects on the game thread.
  */
@@ -74,8 +76,15 @@ class FAsyncLoadingThread : public FRunnable
 	/** Async loading thread ID */
 	static uint32 AsyncLoadingThreadID;
 
-	/** Helper map for tracking the dependency packages that have been requested while loading a package */
-	TMap<FName, int32> DependencyTracker;
+	/** Helper for tracking the dependency packages that have been requested while loading a package */
+	TSet<FName> DependencyTracker;
+
+	/** Enum describing async package request insertion mode */
+	enum class EAsyncPackageInsertMode
+	{
+		InsertBeforeMatchingPriorities,	// Insert this package before all other packages of the same priority
+		InsertAfterMatchingPriorities	// Insert this package after all other packages of the same priority
+	};
 
 #if LOOKING_FOR_PERF_ISSUES
 	/** Thread safe counter used to accumulate cycles spent on blocking. Using stats may generate to many stats messages. */
@@ -206,9 +215,10 @@ public:
 	/**
 	* [ASYNC THREAD] Inserts package to queue according to priority.
 	*
-	* @param PackageName async package name.
+	* @param PackageName - async package name.
+	* @param InsertMode - Insert mode, describing how we insert this package into the request list
 	*/
-	void InsertPackage(FAsyncPackage* Package);
+	void InsertPackage(FAsyncPackage* Package, EAsyncPackageInsertMode InsertMode = EAsyncPackageInsertMode::InsertBeforeMatchingPriorities);
 
 	/**
 	* [ASYNC THREAD] Finds an existing async package in the LoadedPackages by its name.
@@ -356,7 +366,12 @@ private:
 	* [ASYNC THREAD] Internal helper function for processing a package load request. If dependency preloading is enabled, 
 	* it will call itself recursively for all the package dependencies
 	*/
-	void ProcessAsyncPackageRequest(FAsyncPackageDesc* InRequest, FAsyncPackage* InRootPackage, TMap<FName, int32>& InDependencyTracker);
+	void ProcessAsyncPackageRequest(FAsyncPackageDesc* InRequest, FAsyncPackage* InRootPackage, TSet<FName>& InDependencyTracker, IAssetRegistryInterface* InAssetRegistry);
+
+	/**
+	* [ASYNC THREAD] Internal helper function for updating the priorities of an existing package and all its dependencies
+	*/
+	void UpdateExistingPackagePriorities(FAsyncPackage* InPackage, TAsyncLoadPriority InNewPriority, TSet<FName>& InDependencyTracker, IAssetRegistryInterface* InAssetRegistry);
 
 	/**
 	* [ASYNC THREAD] Finds existing async package and adds the new request's completion callback to it.
diff --git a/Engine/Source/Runtime/CoreUObject/Private/UObject/UObjectGlobals.cpp b/Engine/Source/Runtime/CoreUObject/Private/UObject/UObjectGlobals.cpp
--- a/Engine/Source/Runtime/CoreUObject/Private/UObject/UObjectGlobals.cpp
+++ b/Engine/Source/Runtime/CoreUObject/Private/UObject/UObjectGlobals.cpp
@@ -944,12 +944,21 @@ UPackage* LoadPackageInternal(UPackage* InOuter, const TCHAR* InLongPackageName,
 
 		for (auto Dependency : PackageDependencies)
 		{
-			if (!InDependencyTracker.Contains(Dependency) && FindObjectFast<UPackage>(NULL, Dependency, false, false) == nullptr)
+			if (!InDependencyTracker.Contains(Dependency) && FindObjectFast<UPackage>(nullptr, Dependency, false, false) == nullptr)
 			{
 				InDependencyTracker.Add(Dependency);
 				LoadPackageInternal(InOuter, *Dependency.ToString(), LoadFlags, nullptr, InDependencyTracker, InAssetRegistry);
 			}
 		}
+
+		// Check that the package we are going to load hasn't somehow been loaded behind our backs
+		// while loading a dependency. In some cases, circular dependencies can mean this happens
+		// and so there is no need to carry on any further.
+		Result = FindObjectFast<UPackage>(nullptr, InLongPackageName, false, false);
+		if (Result != nullptr)
+		{
+			return Result;
+		}
 	}
 
 	// Try to load.
@@ -1131,6 +1140,20 @@ UPackage* LoadPackageInternal(UPackage* InOuter, const TCHAR* InLongPackageName,
 	return Result;
 }
 
+bool IsPlatformFileCompatibleWithDependencyPreloading()
+{
+	static bool bResultCached = false;
+	static bool bResult = true;
+	if (!bResultCached)
+	{
+		IPlatformFile& PlatformFile = FPlatformFileManager::Get().GetPlatformFile();
+		const TCHAR* PlatformFileType = PlatformFile.GetName();
+		bResultCached = true;
+		bResult = (FCString::Stricmp(TEXT("StreamingFile"), PlatformFileType) != 0) && (FCString::Stricmp(TEXT("NetworkFile"), PlatformFileType) != 0);
+	}
+	return bResult;
+}
+
 UPackage* LoadPackageInternal(UPackage* InOuter, const TCHAR* InLongPackageName, uint32 LoadFlags, FLinkerLoad* ImportLinker)
 {
 	IAssetRegistryInterface* AssetRegistry = nullptr;
@@ -1139,8 +1162,11 @@ UPackage* LoadPackageInternal(UPackage* InOuter, const TCHAR* InLongPackageName,
 	static auto CVarPreloadDependencies = IConsoleManager::Get().FindConsoleVariable(TEXT("s.PreloadPackageDependencies"));
 	
 	if (bAllowDependencyPreloading && CVarPreloadDependencies && CVarPreloadDependencies->GetInt() != 0)
-	{
-		AssetRegistry = IAssetRegistryInterface::GetPtr();
+	{		
+		if (IsPlatformFileCompatibleWithDependencyPreloading())
+		{
+			AssetRegistry = IAssetRegistryInterface::GetPtr();
+		}
 	}
 #endif
 
diff --git a/Engine/Source/Runtime/CoreUObject/Public/Serialization/AsyncLoading.h b/Engine/Source/Runtime/CoreUObject/Public/Serialization/AsyncLoading.h
--- a/Engine/Source/Runtime/CoreUObject/Public/Serialization/AsyncLoading.h
+++ b/Engine/Source/Runtime/CoreUObject/Public/Serialization/AsyncLoading.h
@@ -83,6 +83,12 @@ struct FAsyncPackage
 		return Desc.Priority;
 	}
 
+	/** Returns package loading priority. */
+	FORCEINLINE void SetPriority(TAsyncLoadPriority InPriority)
+	{
+		Desc.Priority = InPriority;
+	}
+
 	/** Returns true if loading has failed */
 	FORCEINLINE bool HasLoadFailed() const
 	{
diff --git a/Engine/Source/Runtime/CoreUObject/Public/UObject/UObjectGlobals.h b/Engine/Source/Runtime/CoreUObject/Public/UObject/UObjectGlobals.h
--- a/Engine/Source/Runtime/CoreUObject/Public/UObject/UObjectGlobals.h
+++ b/Engine/Source/Runtime/CoreUObject/Public/UObject/UObjectGlobals.h
@@ -459,6 +459,13 @@ COREUOBJECT_API int32 GetNumAsyncPackages();
  */
 COREUOBJECT_API bool IsLoading();
 
+/**
+* Determines whether the current platform file is compatible with dependency preloading
+*
+* @return true if we can use dependency preloading
+*/
+COREUOBJECT_API bool IsPlatformFileCompatibleWithDependencyPreloading();
+
 /**
  * State of the async package after the last tick.
  */
diff --git a/Engine/Source/Runtime/Engine/Classes/Animation/AnimInstance.h b/Engine/Source/Runtime/Engine/Classes/Animation/AnimInstance.h
--- a/Engine/Source/Runtime/Engine/Classes/Animation/AnimInstance.h
+++ b/Engine/Source/Runtime/Engine/Classes/Animation/AnimInstance.h
@@ -825,6 +825,9 @@ public:
 
 protected:
 
+	/** Gets the index of the state machine matching MachineName */
+	int32 GetStateMachineIndex(FName MachineName);
+
 	/** Gets the runtime instance of the specified state machine */
 	FAnimNode_StateMachine* GetStateMachineInstance(int32 MachineIndex);
 
@@ -930,6 +933,9 @@ public:
 	static void AddReferencedObjects(UObject* InThis, FReferenceCollector& Collector);
 	//~ End UObject Interface
 
+	virtual void OnUROSkipTickAnimation() {}
+	virtual void OnUROPreInterpolation() {}
+
 	// Animation phase trigger
 	// start with initialize
 	// update happens in every tick. Can happen in parallel with others if conditions are right.
diff --git a/Engine/Source/Runtime/Engine/Classes/Animation/AnimStateMachineTypes.h b/Engine/Source/Runtime/Engine/Classes/Animation/AnimStateMachineTypes.h
--- a/Engine/Source/Runtime/Engine/Classes/Animation/AnimStateMachineTypes.h
+++ b/Engine/Source/Runtime/Engine/Classes/Animation/AnimStateMachineTypes.h
@@ -276,7 +276,11 @@ public:
 	{}
 
 	// Finds a state by name or INDEX_NONE if no such state exists
-	ENGINE_API int32 FindStateIndex(FName StateName) const;
+	ENGINE_API int32 FindStateIndex(const FName& StateName) const;
+
+	// Find the index of a transition from StateNameFrom to StateNameTo
+	ENGINE_API int32 FindTransitionIndex(const FName& InStateNameFrom, const FName& InStateNameTo) const;
+	ENGINE_API int32 FindTransitionIndex(const int32 InStateIdxFrom, const int32 InStateIdxTo) const;
 };
 
 UCLASS()
diff --git a/Engine/Source/Runtime/Engine/Classes/Components/SkeletalMeshComponent.h b/Engine/Source/Runtime/Engine/Classes/Components/SkeletalMeshComponent.h
--- a/Engine/Source/Runtime/Engine/Classes/Components/SkeletalMeshComponent.h
+++ b/Engine/Source/Runtime/Engine/Classes/Components/SkeletalMeshComponent.h
@@ -243,6 +243,13 @@ namespace EAnimationMode
 	};
 }
 
+/** Enum for indicating whether kinematic updates can be deferred */
+enum class EAllowKinematicDeferral
+{
+	AllowDeferral,
+	DisallowDeferral
+};
+
 /**
 * Tick function that prepares for cloth tick
 **/
@@ -1075,9 +1082,9 @@ public:
 	 *	transform based on the animated transform.
 	 *	@param	Teleport		Whether movement is a 'teleport' (ie infers no physics velocity, but moves simulating bodies) or not
 	 *	@param	bNeedsSkinning	Whether we may need  to send new triangle data for per-poly skeletal mesh collision
-	 *	@perem	bNoDefer		Whether we can defer actual update of bodies (if 'physics only' collision)
+	 *	@perem	AllowDeferral	Whether we can defer actual update of bodies (if 'physics only' collision)
 	 */
-	void UpdateKinematicBonesToAnim(const TArray<FTransform>& InSpaceBases, ETeleportType Teleport, bool bNeedsSkinning, bool bNoDefer = false);
+	void UpdateKinematicBonesToAnim(const TArray<FTransform>& InSpaceBases, ETeleportType Teleport, bool bNeedsSkinning, EAllowKinematicDeferral DeferralAllowed = EAllowKinematicDeferral::AllowDeferral);
 
 	DEPRECATED(4.9, "bForceUpdate is no longer used. Please use Teleport")
 	void UpdateKinematicBonesToAnim(const TArray<FTransform>& InSpaceBases, bool bTeleport, bool bNeedsSkinning, bool bForceUpdate)
diff --git a/Engine/Source/Runtime/Engine/Classes/Engine/DemoNetDriver.h b/Engine/Source/Runtime/Engine/Classes/Engine/DemoNetDriver.h
--- a/Engine/Source/Runtime/Engine/Classes/Engine/DemoNetDriver.h
+++ b/Engine/Source/Runtime/Engine/Classes/Engine/DemoNetDriver.h
@@ -8,7 +8,8 @@
 
 DECLARE_LOG_CATEGORY_EXTERN( LogDemo, Log, All );
 
-DECLARE_DELEGATE_OneParam(FOnGotoTimeDelegate, const bool);
+DECLARE_MULTICAST_DELEGATE(FOnGotoTimeMCDelegate);
+DECLARE_DELEGATE_OneParam(FOnGotoTimeDelegate, const bool /* bWasSuccessful */);
 
 class UDemoNetDriver;
 
@@ -89,7 +90,8 @@ class ENGINE_API UDemoNetDriver : public UNetDriver
 
 	void		LoadCheckpoint( FArchive* GotoCheckpointArchive, int64 GotoCheckpointSkipExtraTimeInMS );
 
-	FOnGotoTimeDelegate OnGotoTimeDelegate;
+	/** Public delegate for external systems to be notified when scrubbing is complete. Only called for successful scrub. */
+	FOnGotoTimeMCDelegate OnGotoTimeDelegate;
 
 private:
 	bool		bIsFastForwarding;
@@ -106,6 +108,9 @@ private:
 	TSharedPtr< FQueuedReplayTask >				ActiveReplayTask;
 	TSharedPtr< FQueuedReplayTask >				ActiveScrubReplayTask;
 
+	/** Set via GotoTimeInSeconds, only fired once (at most). Called for successful or failed scrub. */
+	FOnGotoTimeDelegate OnGotoTimeDelegate_Transient;
+	
 public:
 
 	// UNetDriver interface.
@@ -120,13 +125,19 @@ public:
 	virtual bool IsAvailable() const override { return true; }
 	void SkipTime(const float InTimeToSkip);
 	void SkipTimeInternal( const float SecondsToSkip, const bool InFastForward, const bool InIsForCheckpoint );
-	bool InitConnectInternal( FString& Error );
+	bool InitConnectInternal(FString& Error);
 	virtual bool ShouldClientDestroyTearOffActors() const override;
 	virtual bool ShouldSkipRepNotifies() const override;
 	virtual bool ShouldQueueBunchesForActorGUID(FNetworkGUID InGUID) const override;
 	virtual FNetworkGUID GetGUIDForActor(const AActor* InActor) const override;
 	virtual AActor* GetActorForGUID(FNetworkGUID InGUID) const override;
 
+	/** 
+	 * Scrubs playback to the given time. 
+	 * 
+	 * @param TimeInSeconds
+	 * @param InOnGotoTimeDelegate		Delegate to call when finished. Will be called only once at most.
+	*/
 	void GotoTimeInSeconds(const float TimeInSeconds, const FOnGotoTimeDelegate& InOnGotoTimeDelegate = FOnGotoTimeDelegate());
 
 	bool IsRecording();
@@ -183,4 +194,7 @@ public:
 	void AddNonQueuedGUIDForScrubbing(FNetworkGUID InGUID);
 
 	virtual bool IsLevelInitializedForActor( const AActor* InActor, const UNetConnection* InConnection ) const override;
+
+	/** Called when a "go to time" operation is completed. */
+	void NotifyGotoTimeFinished(bool bWasSuccessful);
 };
diff --git a/Engine/Source/Runtime/Engine/Classes/Engine/StaticMesh.h b/Engine/Source/Runtime/Engine/Classes/Engine/StaticMesh.h
--- a/Engine/Source/Runtime/Engine/Classes/Engine/StaticMesh.h
+++ b/Engine/Source/Runtime/Engine/Classes/Engine/StaticMesh.h
@@ -278,6 +278,13 @@ struct FAssetEditorOrbitCameraPosition
 	FRotator CamOrbitRotation;
 };
 
+#if WITH_EDITOR
+/** delegate type for pre mesh build events */
+DECLARE_MULTICAST_DELEGATE_OneParam(FOnPreMeshBuild, class UStaticMesh*);
+/** delegate type for pre mesh build events */
+DECLARE_MULTICAST_DELEGATE_OneParam(FOnPostMeshBuild, class UStaticMesh*);
+#endif
+
 /**
  * A StaticMesh is a piece of geometry that consists of a static set of polygons.
  * Static Meshes can be translated, rotated, and scaled, but they cannot have their vertices animated in any way. As such, they are more efficient
@@ -628,6 +635,12 @@ public:
 
 	ENGINE_API void GenerateLodsInPackage();
 
+	/** Get multicast delegate broadcast prior to mesh building */
+	FOnPreMeshBuild& OnPreMeshBuild() { return PreMeshBuild; }
+
+	/** Get multicast delegate broadcast after mesh building */
+	FOnPostMeshBuild& OnPostMeshBuild() { return PostMeshBuild; }
+
 private:
 	/**
 	 * Converts legacy LODDistance in the source models to Display Factor
@@ -646,5 +659,9 @@ private:
 
 	/** Calculates the extended bounds */
 	void CalculateExtendedBounds();
+
+	FOnPreMeshBuild PreMeshBuild;
+	FOnPostMeshBuild PostMeshBuild;
+
 #endif // #if WITH_EDITOR
 };
diff --git a/Engine/Source/Runtime/Engine/Classes/GameFramework/Character.h b/Engine/Source/Runtime/Engine/Classes/GameFramework/Character.h
--- a/Engine/Source/Runtime/Engine/Classes/GameFramework/Character.h
+++ b/Engine/Source/Runtime/Engine/Classes/GameFramework/Character.h
@@ -319,7 +319,7 @@ public:
 	const FVector& GetBaseTranslationOffset() const { return BaseTranslationOffset; }
 
 	/** @return Saved rotation offset of mesh. */
-	const FQuat& GetBaseRotationOffset() const { return BaseRotationOffset; }
+	const virtual FQuat GetBaseRotationOffset() const { return BaseRotationOffset; }
 
 	//~ Begin INavAgentInterface Interface
 	virtual FVector GetNavAgentLocation() const override;
diff --git a/Engine/Source/Runtime/Engine/Private/Animation/AnimInstance.cpp b/Engine/Source/Runtime/Engine/Private/Animation/AnimInstance.cpp
--- a/Engine/Source/Runtime/Engine/Private/Animation/AnimInstance.cpp
+++ b/Engine/Source/Runtime/Engine/Private/Animation/AnimInstance.cpp
@@ -865,8 +865,12 @@ void UAnimInstance::BindNativeDelegates()
 					{
 						if(TransitionExit.CanTakeDelegateIndex != INDEX_NONE)
 						{
-							const FAnimationTransitionBetweenStates& Transition = StateMachine->GetTransitionInfo(TransitionExit.TransitionIndex);
-							if(StateMachine->GetStateInfo(Transition.NextState).StateName == Binding.NextStateName)
+							// In case the state machine hasn't been initilized, we need to re-get the desc
+							const FBakedAnimationStateMachine* MachineDesc = GetMachineDescription(AnimBlueprintGeneratedClass, StateMachine);
+							const FAnimationTransitionBetweenStates& Transition = MachineDesc->Transitions[TransitionExit.TransitionIndex];
+							const FBakedAnimationState& BakedState = MachineDesc->States[Transition.NextState];
+							
+							if(BakedState.StateName == Binding.NextStateName)
 							{
 								FAnimNode_TransitionResult* ResultNode = GetNodeFromPropertyIndex<FAnimNode_TransitionResult>(this, AnimBlueprintGeneratedClass, TransitionExit.CanTakeDelegateIndex);
 								if(ResultNode)
@@ -3235,6 +3239,33 @@ const FBakedAnimationStateMachine* UAnimInstance::GetStateMachineInstanceDesc(FN
 	return nullptr;
 }
 
+int32 UAnimInstance::GetStateMachineIndex(FName MachineName)
+{
+	if(UAnimBlueprintGeneratedClass* AnimBlueprintClass = Cast<UAnimBlueprintGeneratedClass>((UObject*)GetClass()))
+	{
+		for(int32 MachineIndex = 0; MachineIndex < AnimBlueprintClass->AnimNodeProperties.Num(); MachineIndex++)
+		{
+			UStructProperty* Property = AnimBlueprintClass->AnimNodeProperties[AnimBlueprintClass->AnimNodeProperties.Num() - 1 - MachineIndex];
+			if(Property && Property->Struct == FAnimNode_StateMachine::StaticStruct())
+			{
+				FAnimNode_StateMachine* StateMachine = Property->ContainerPtrToValuePtr<FAnimNode_StateMachine>(this);
+				if(StateMachine)
+				{
+					if(const FBakedAnimationStateMachine* MachineDescription = GetMachineDescription(AnimBlueprintClass, StateMachine))
+					{
+						if(MachineDescription->MachineName == MachineName)
+						{
+							return MachineIndex;
+						}
+					}
+				}
+			}
+		}
+	}
+
+	return INDEX_NONE;
+}
+
 FAnimNode_Base* UAnimInstance::GetCheckedNodeFromIndexUntyped(int32 NodeIdx, UScriptStruct* RequiredStructType)
 {
 	FAnimNode_Base* NodePtr = nullptr;
diff --git a/Engine/Source/Runtime/Engine/Private/Animation/AnimStateMachineTypes.cpp b/Engine/Source/Runtime/Engine/Private/Animation/AnimStateMachineTypes.cpp
--- a/Engine/Source/Runtime/Engine/Private/Animation/AnimStateMachineTypes.cpp
+++ b/Engine/Source/Runtime/Engine/Private/Animation/AnimStateMachineTypes.cpp
@@ -13,15 +13,42 @@ UAnimStateMachineTypes::UAnimStateMachineTypes(const FObjectInitializer& ObjectI
 {
 }
 
-int32 FBakedAnimationStateMachine::FindStateIndex(FName StateName) const
+int32 FBakedAnimationStateMachine::FindStateIndex(const FName& InStateName) const
 {
 	for (int32 Index = 0; Index < States.Num(); Index++)
 	{
-		if (States[Index].StateName == StateName)
+		if (States[Index].StateName == InStateName)
 		{
 			return Index;
 		}
 	}
 
 	return INDEX_NONE;
-}
\ No newline at end of file
+}
+
+int32 FBakedAnimationStateMachine::FindTransitionIndex(const FName& InStateNameFrom, const FName& InStateNameTo) const
+{
+	return FindTransitionIndex(FindStateIndex(InStateNameFrom), FindStateIndex(InStateNameTo));
+}
+
+int32 FBakedAnimationStateMachine::FindTransitionIndex(const int32 InStateIdxFrom, const int32 InStateIdxTo) const
+{
+	// Early out if any request is invalid
+	if(InStateIdxFrom == INDEX_NONE || InStateIdxTo == INDEX_NONE)
+	{
+		return INDEX_NONE;
+	}
+
+	const int32 Count = Transitions.Num();
+	for(int32 Index = 0 ; Index < Count ; ++Index)
+	{
+		const FAnimationTransitionBetweenStates& Transition = Transitions[Index];
+		if(Transition.PreviousState == InStateIdxFrom && Transition.NextState == InStateIdxTo)
+		{
+			// If we match the request, output the index.
+			return Index;
+		}
+	}
+
+	return INDEX_NONE;
+}
diff --git a/Engine/Source/Runtime/Engine/Private/Components/CharacterMovementComponent.cpp b/Engine/Source/Runtime/Engine/Private/Components/CharacterMovementComponent.cpp
--- a/Engine/Source/Runtime/Engine/Private/Components/CharacterMovementComponent.cpp
+++ b/Engine/Source/Runtime/Engine/Private/Components/CharacterMovementComponent.cpp
@@ -1266,7 +1266,7 @@ void UCharacterMovementComponent::SimulateRootMotion(float DeltaSeconds, const F
 
 		// Apply Root Motion rotation after movement is complete.
 		const FQuat RootMotionRotationQuat = WorldSpaceRootMotionTransform.GetRotation();
-		if( !RootMotionRotationQuat.Equals(FQuat::Identity) )
+		if( !RootMotionRotationQuat.IsIdentity() )
 		{
 			const FQuat NewActorRotationQuat = RootMotionRotationQuat * UpdatedComponent->GetComponentQuat();
 			MoveUpdatedComponent(FVector::ZeroVector, NewActorRotationQuat, true);
@@ -1817,7 +1817,7 @@ void UCharacterMovementComponent::PerformMovement(float DeltaSeconds)
 		{
 			const FQuat OldActorRotationQuat = UpdatedComponent->GetComponentQuat();
 			const FQuat RootMotionRotationQuat = RootMotionParams.RootMotionTransform.GetRotation();
-			if( !RootMotionRotationQuat.Equals(FQuat::Identity) )
+			if( !RootMotionRotationQuat.IsIdentity() )
 			{
 				const FQuat NewActorRotationQuat = RootMotionRotationQuat * OldActorRotationQuat;
 				MoveUpdatedComponent(FVector::ZeroVector, NewActorRotationQuat, true);
@@ -6294,13 +6294,22 @@ void UCharacterMovementComponent::SmoothClientPosition_UpdateVisuals()
 	{
 		if (NetworkSmoothingMode == ENetworkSmoothingMode::Linear)
 		{
-			UpdatedComponent->SetRelativeRotation(ClientData->MeshRotationOffset);
-			// TODO: maybe we can just adjust the mesh RelativeLocation before moving the capsule above, so it's a single update?
-			const FVector NewRelTranslation = UpdatedComponent->GetComponentToWorld().InverseTransformVectorNoScale(ClientData->MeshTranslationOffset) + CharacterOwner->GetBaseTranslationOffset();
-			CharacterOwner->GetMesh()->SetRelativeLocation(NewRelTranslation);
+			// Adjust capsule rotation and mesh location. Optimized to trigger only one transform chain update.
+			// If we know the rotation is changing that will update children, so it's sufficient to set RelativeLocation directly on the mesh.
+			const FVector NewRelLocation = ClientData->MeshRotationOffset.UnrotateVector(ClientData->MeshTranslationOffset) + CharacterOwner->GetBaseTranslationOffset();
+			if (!UpdatedComponent->GetComponentQuat().Equals(ClientData->MeshRotationOffset, SCENECOMPONENT_QUAT_TOLERANCE))
+			{
+				CharacterOwner->GetMesh()->RelativeLocation = NewRelLocation;
+				UpdatedComponent->SetWorldRotation(ClientData->MeshRotationOffset);
+			}
+			else
+			{
+				CharacterOwner->GetMesh()->SetRelativeLocation(NewRelLocation);
+			}
 		}
 		else if (NetworkSmoothingMode == ENetworkSmoothingMode::Exponential)
 		{
+			// Adjust mesh location and rotation
 			const FVector NewRelTranslation = UpdatedComponent->GetComponentToWorld().InverseTransformVectorNoScale(ClientData->MeshTranslationOffset) + CharacterOwner->GetBaseTranslationOffset();
 			const FQuat NewRelRotation = ClientData->MeshRotationOffset * CharacterOwner->GetBaseRotationOffset();
 			CharacterOwner->GetMesh()->SetRelativeLocationAndRotation(NewRelTranslation, NewRelRotation);
diff --git a/Engine/Source/Runtime/Engine/Private/Components/SkeletalMeshComponent.cpp b/Engine/Source/Runtime/Engine/Private/Components/SkeletalMeshComponent.cpp
--- a/Engine/Source/Runtime/Engine/Private/Components/SkeletalMeshComponent.cpp
+++ b/Engine/Source/Runtime/Engine/Private/Components/SkeletalMeshComponent.cpp
@@ -61,7 +61,7 @@ public:
 	}
 	static ENamedThreads::Type GetDesiredThread()
 	{
-		return ENamedThreads::AnyThread;
+		return ENamedThreads::AnyThreadGame();
 	}
 	static ESubsequentsMode::Type GetSubsequentsMode()
 	{
@@ -640,9 +640,17 @@ void USkeletalMeshComponent::TickPose(float DeltaTime, bool bNeedsValidRootMotio
 			NotTicked.Reset();
 		}
 	}
-	else if (CVarSpewAnimRateOptimization.GetValueOnGameThread())
+	else
 	{
-		NotTicked.Increment();
+		if (AnimScriptInstance)
+		{
+			AnimScriptInstance->OnUROSkipTickAnimation();
+		}
+
+		if (CVarSpewAnimRateOptimization.GetValueOnGameThread())
+		{
+			NotTicked.Increment();
+		}
 	}
 }
 
@@ -745,18 +753,11 @@ void USkeletalMeshComponent::FillSpaceBases(const USkeletalMesh* InSkeletalMesh,
 #endif
 		FTransform::Multiply(SpaceBase, LocalTransformsData + BoneIndex, ParentSpaceBase);
 
+		SpaceBase->NormalizeRotation();
+
 		checkSlow(SpaceBase->IsRotationNormalized());
 		checkSlow(!SpaceBase->ContainsNaN());
 	}
-
-	/**
-	 * Normalize rotations.
-	 * We want to remove any loss of precision due to accumulation of error.
-	 * i.e. A componentSpace transform is the accumulation of all of its local space parents. The further down the chain, the greater the error.
-	 * SpaceBases are used by external systems, we feed this to PhysX, send this to gameplay through bone and socket queries, etc.
-	 * So this is a good place to make sure all transforms are normalized.
-	 */
-	FAnimationRuntime::NormalizeRotations(DestSpaceBases);
 }
 
 /** Takes sorted array Base and then adds any elements from sorted array Insert which is missing from it, preserving order.
@@ -1309,17 +1310,14 @@ void USkeletalMeshComponent::PostAnimEvaluation(FAnimationEvaluationContext& Eva
 	{
 		SCOPE_CYCLE_COUNTER(STAT_InterpolateSkippedFrames);
 
-		const float Alpha = AnimUpdateRateParams->GetInterpolationAlpha();
-		FAnimationRuntime::LerpBoneTransforms(LocalAtoms, CachedLocalAtoms, Alpha, RequiredBones);
-		if (bDoubleBufferedBlendSpaces)
+		if (AnimScriptInstance)
 		{
-			// We need to prep our space bases for interp (TODO: Would a new Lerp function that took
-			// separate input and output be quicker than current copy + lerp in place?)
-			TArray<FTransform>& LocalEditableSpaceBases = GetEditableSpaceBases();
-			LocalEditableSpaceBases.Reset();
-			LocalEditableSpaceBases.Append(GetSpaceBases());
+			AnimScriptInstance->OnUROPreInterpolation();
 		}
-		FAnimationRuntime::LerpBoneTransforms(GetEditableSpaceBases(), CachedSpaceBases, Alpha, RequiredBones);
+
+		const float Alpha = AnimUpdateRateParams->GetInterpolationAlpha();
+		FAnimationRuntime::LerpBoneTransforms(LocalAtoms, CachedLocalAtoms, Alpha, RequiredBones);
+		FillSpaceBases(SkeletalMesh, LocalAtoms, GetEditableSpaceBases());
 
 		// interpolate curve
 		EvaluationContext.Curve.BlendWith(CachedCurve, Alpha);
diff --git a/Engine/Source/Runtime/Engine/Private/Components/SkinnedMeshComponent.cpp b/Engine/Source/Runtime/Engine/Private/Components/SkinnedMeshComponent.cpp
--- a/Engine/Source/Runtime/Engine/Private/Components/SkinnedMeshComponent.cpp
+++ b/Engine/Source/Runtime/Engine/Private/Components/SkinnedMeshComponent.cpp
@@ -133,10 +133,15 @@ namespace FAnimUpdateRateManager
 	}
 
 	static TAutoConsoleVariable<int32> CVarForceAnimRate(
-		TEXT("ForceAnimRate"),
+		TEXT("a.URO.ForceAnimRate"),
 		0,
 		TEXT("Non-zero to force anim rate. 10 = eval anim every ten frames for those meshes that can do it. In some cases a frame is considered to be 30fps."));
 
+	static TAutoConsoleVariable<int32> CVarForceInterpolation(
+		TEXT("a.URO.ForceInterpolation"),
+		0,
+		TEXT("Set to 1 to force interpolation"));
+
 	void AnimUpdateRateSetParams(FAnimUpdateRateParametersTracker* Tracker, float DeltaTime, bool bRecentlyRendered, float MaxDistanceFactor, bool bNeedsValidRootMotion, bool bUsingRootMotionFromEverything)
 	{
 		// default rules for setting update rates
@@ -1990,10 +1995,17 @@ void FAnimUpdateRateParameters::SetTrailMode(float DeltaTime, uint8 UpdateRateSh
 	OptimizeMode = TrailMode;
 	ThisTickDelta = DeltaTime;
 
+	const int32 ForceAnimRate = FAnimUpdateRateManager::CVarForceAnimRate.GetValueOnGameThread();
+	if (ForceAnimRate > 0)
+	{
+		NewUpdateRate = ForceAnimRate;
+		NewEvaluationRate = ForceAnimRate;
+	}
+
 	UpdateRate = FMath::Max(NewUpdateRate, 1);
 	// Make sure EvaluationRate is a multiple of UpdateRate.
 	EvaluationRate = FMath::Max((NewEvaluationRate / UpdateRate) * UpdateRate, 1);
-	bInterpolateSkippedFrames = bNewInterpSkippedFrames;
+	bInterpolateSkippedFrames = bNewInterpSkippedFrames || (FAnimUpdateRateManager::CVarForceInterpolation.GetValueOnAnyThread() == 1);
 
 	// Make sure we don't overflow. we don't need very large numbers.
 	const uint32 Counter = (GFrameCounter + UpdateRateShift)% MAX_uint8;
diff --git a/Engine/Source/Runtime/Engine/Private/DemoNetDriver.cpp b/Engine/Source/Runtime/Engine/Private/DemoNetDriver.cpp
--- a/Engine/Source/Runtime/Engine/Private/DemoNetDriver.cpp
+++ b/Engine/Source/Runtime/Engine/Private/DemoNetDriver.cpp
@@ -160,8 +160,7 @@ public:
 			Driver->DemoCurrentTime = OldTimeInSeconds;
 
 			// Call delegate if any
-			Driver->OnGotoTimeDelegate.ExecuteIfBound( false );
-			Driver->OnGotoTimeDelegate.Unbind();
+			Driver->NotifyGotoTimeFinished(false);
 
 			GotoCheckpointSkipExtraTimeInMS = -2;	// So tick can detect failure case
 			return;
@@ -1593,14 +1592,15 @@ void UDemoNetDriver::SkipTimeInternal( const float SecondsToSkip, const bool InF
 
 void UDemoNetDriver::GotoTimeInSeconds(const float TimeInSeconds, const FOnGotoTimeDelegate& InOnGotoTimeDelegate)
 {
+	OnGotoTimeDelegate_Transient = InOnGotoTimeDelegate;
+
 	if ( IsNamedTaskInQueue( TEXT( "FGotoTimeInSecondsTask" ) ) )
 	{
-		InOnGotoTimeDelegate.ExecuteIfBound( false );
+		NotifyGotoTimeFinished(false);
 		return;		// Don't allow scrubbing if we already are
 	}
 
 	AddReplayTask( new FGotoTimeInSecondsTask( this, TimeInSeconds ) );
-	OnGotoTimeDelegate = InOnGotoTimeDelegate;
 }
 
 void UDemoNetDriver::JumpToEndOfLiveReplay()
@@ -1723,10 +1723,9 @@ void UDemoNetDriver::FinalizeFastForward( const float StartTime )
 	// Reset the never-queue GUID list, we'll rebuild it
 	NonQueuedGUIDsForScrubbing.Reset();
 
-	const auto FastForwardTotalSeconds = FPlatformTime::Seconds() - StartTime;
+	const float FastForwardTotalSeconds = FPlatformTime::Seconds() - StartTime;
 
-	OnGotoTimeDelegate.ExecuteIfBound( true );
-	OnGotoTimeDelegate.Unbind();
+	NotifyGotoTimeFinished(true);
 
 	UE_LOG( LogDemo, Log, TEXT( "Fast forward took %.2f seconds." ), FastForwardTotalSeconds );
 }
@@ -2279,6 +2278,20 @@ bool UDemoNetDriver::IsLevelInitializedForActor( const AActor* InActor, const UN
 	return ( DemoFrameNum > 2 || Super::IsLevelInitializedForActor( InActor, InConnection ) );
 }
 
+void UDemoNetDriver::NotifyGotoTimeFinished(bool bWasSuccessful)
+{
+	// execute and clear the transient delegate
+	OnGotoTimeDelegate_Transient.ExecuteIfBound(bWasSuccessful);
+	OnGotoTimeDelegate_Transient.Unbind();
+
+	// execute and keep the permanent delegate
+	// call only when successful
+	if (bWasSuccessful)
+	{
+		OnGotoTimeDelegate.Broadcast();
+	}
+}
+
 
 /*-----------------------------------------------------------------------------
 	UDemoPendingNetGame.
diff --git a/Engine/Source/Runtime/Engine/Private/Particles/ParticleComponents.cpp b/Engine/Source/Runtime/Engine/Private/Particles/ParticleComponents.cpp
--- a/Engine/Source/Runtime/Engine/Private/Particles/ParticleComponents.cpp
+++ b/Engine/Source/Runtime/Engine/Private/Particles/ParticleComponents.cpp
@@ -3698,7 +3698,7 @@ public:
 
 	ENamedThreads::Type GetDesiredThread()
 	{
-		return ENamedThreads::AnyThread;
+		return ENamedThreads::AnyThreadGame();
 	}
 
 	static ESubsequentsMode::Type GetSubsequentsMode() { return ESubsequentsMode::TrackSubsequents; }
diff --git a/Engine/Source/Runtime/Engine/Private/PhysicsEngine/ClothManager.cpp b/Engine/Source/Runtime/Engine/Private/PhysicsEngine/ClothManager.cpp
--- a/Engine/Source/Runtime/Engine/Private/PhysicsEngine/ClothManager.cpp
+++ b/Engine/Source/Runtime/Engine/Private/PhysicsEngine/ClothManager.cpp
@@ -15,6 +15,7 @@ TAutoConsoleVariable<int32> CVarParallelCloth(TEXT("p.ParallelCloth"), 0, TEXT("
 void FClothManagerData::PrepareCloth(float DeltaTime, FTickFunction& TickFunction)
 {
 #if WITH_APEX_CLOTHING
+	const bool bParallelCloth = !!CVarParallelCloth.GetValueOnAnyThread() && !PhysSingleThreadedMode();
 	QUICK_SCOPE_CYCLE_COUNTER(STAT_FClothManager_PrepareCloth);
 	if (SkeletalMeshComponents.Num())
 	{
@@ -23,7 +24,7 @@ void FClothManagerData::PrepareCloth(float DeltaTime, FTickFunction& TickFunctio
 		for (USkeletalMeshComponent* SkeletalMeshComponent : SkeletalMeshComponents)
 		{
 			FClothSimulationContext& ClothSimulationContext = SkeletalMeshComponent->InternalClothSimulationContext;
-			if (bParallelCloth)
+			if(bParallelCloth)
 			{
 				SkeletalMeshComponent->ParallelTickClothing(DeltaTime, ClothSimulationContext);
 			}else
@@ -76,7 +77,7 @@ void FClothManager::SetupClothTickFunction(bool bTickClothSim)
 	check(IsInGameThread());
 	check(!IsPreparingClothAsync());
 
-	bool bTickOnAny = !!CVarParallelCloth.GetValueOnGameThread();
+	bool bTickOnAny = !!CVarParallelCloth.GetValueOnGameThread() && !PhysSingleThreadedMode();
 	StartIgnorePhysicsClothTickFunction.bRunOnAnyThread = bTickOnAny;
 	StartClothTickFunction.bRunOnAnyThread = bTickOnAny;
 
diff --git a/Engine/Source/Runtime/Engine/Private/PhysicsEngine/PhysAnim.cpp b/Engine/Source/Runtime/Engine/Private/PhysicsEngine/PhysAnim.cpp
--- a/Engine/Source/Runtime/Engine/Private/PhysicsEngine/PhysAnim.cpp
+++ b/Engine/Source/Runtime/Engine/Private/PhysicsEngine/PhysAnim.cpp
@@ -44,7 +44,7 @@ public:
 	}
 	static ENamedThreads::Type GetDesiredThread()
 	{
-		return ENamedThreads::AnyThread;
+		return ENamedThreads::AnyThreadGame();
 	}
 	static ESubsequentsMode::Type GetSubsequentsMode()
 	{
@@ -452,7 +452,7 @@ void USkeletalMeshComponent::CompleteParallelBlendPhysics()
 	ParallelBlendPhysicsCompletionTask.SafeRelease();
 }
 
-void USkeletalMeshComponent::UpdateKinematicBonesToAnim(const TArray<FTransform>& InSpaceBases, ETeleportType Teleport, bool bNeedsSkinning, bool bNoDefer)
+void USkeletalMeshComponent::UpdateKinematicBonesToAnim(const TArray<FTransform>& InSpaceBases, ETeleportType Teleport, bool bNeedsSkinning, EAllowKinematicDeferral DeferralAllowed)
 {
 	SCOPE_CYCLE_COUNTER(STAT_UpdateRBBones);
 
@@ -491,7 +491,7 @@ void USkeletalMeshComponent::UpdateKinematicBonesToAnim(const TArray<FTransform>
 	}
 
 	// If we are only using bodies for physics, don't need to move them right away, can defer until simulation (unless told not to)
-	if(BodyInstance.GetCollisionEnabled() == ECollisionEnabled::PhysicsOnly && !bNoDefer)
+	if(BodyInstance.GetCollisionEnabled() == ECollisionEnabled::PhysicsOnly && DeferralAllowed == EAllowKinematicDeferral::AllowDeferral)
 	{
 		PhysScene->MarkForPreSimKinematicUpdate(this, Teleport, bNeedsSkinning);
 		return;
diff --git a/Engine/Source/Runtime/Engine/Private/PhysicsEngine/PhysScene.cpp b/Engine/Source/Runtime/Engine/Private/PhysicsEngine/PhysScene.cpp
--- a/Engine/Source/Runtime/Engine/Private/PhysicsEngine/PhysScene.cpp
+++ b/Engine/Source/Runtime/Engine/Private/PhysicsEngine/PhysScene.cpp
@@ -156,7 +156,7 @@ public:
 	}
 	static ENamedThreads::Type GetDesiredThread()
 	{
-		return ENamedThreads::AnyThread;
+		return ENamedThreads::AnyThreadGame();
 	}
 	static ESubsequentsMode::Type GetSubsequentsMode()
 	{
@@ -563,7 +563,7 @@ bool FPhysScene::SubstepSimulation(uint32 SceneType, FGraphEventRef &InOutComple
 	{
 		//we have valid scene and subtime so enqueue task
 		PhysXCompletionTask* Task = new PhysXCompletionTask(InOutCompletionEvent, SceneType, PScene->getTaskManager());
-		ENamedThreads::Type NamedThread = PhysSingleThreadedMode() ? ENamedThreads::GameThread : ENamedThreads::AnyThread;
+		ENamedThreads::Type NamedThread = PhysSingleThreadedMode() ? ENamedThreads::GameThread : ENamedThreads::AnyThreadGame();
 
 		DECLARE_CYCLE_STAT(TEXT("FSimpleDelegateGraphTask.SubstepSimulationImp"),
 			STAT_FSimpleDelegateGraphTask_SubstepSimulationImp,
@@ -801,7 +801,7 @@ void FPhysScene::UpdateKinematicsOnDeferredSkelMeshes()
 		check(SkelComp->bDeferredKinematicUpdate); // Should be true if in map!
 
 		// Perform kinematic updates
-		SkelComp->UpdateKinematicBonesToAnim(SkelComp->GetSpaceBases(), Info.TeleportType, Info.bNeedsSkinning, true);
+		SkelComp->UpdateKinematicBonesToAnim(SkelComp->GetSpaceBases(), Info.TeleportType, Info.bNeedsSkinning, EAllowKinematicDeferral::DisallowDeferral);
 
 		// Clear deferred flag
 		SkelComp->bDeferredKinematicUpdate = false; 
@@ -943,9 +943,8 @@ void FPhysScene::TickPhysScene(uint32 SceneType, FGraphEventRef& InOutCompletion
 
 	if (!bTaskOutstanding)
 	{
-		check(IsInGameThread());
-		static TArray<FBaseGraphTask*> NewTasks;
-		InOutCompletionEvent->DispatchSubsequents(NewTasks, ENamedThreads::GameThread); // nothing to do, so nothing to wait for
+		TArray<FBaseGraphTask*> NewTasks;
+		InOutCompletionEvent->DispatchSubsequents(NewTasks, ENamedThreads::AnyThread); // nothing to do, so nothing to wait for
 	}
 #if WITH_SUBSTEPPING
 	bSubstepping = UPhysicsSettings::Get()->bSubstepping;
@@ -966,6 +965,8 @@ void FPhysScene::KillVisualDebugger()
 
 void FPhysScene::WaitPhysScenes()
 {
+	check(IsInGameThread());
+
 	FGraphEventArray ThingsToComplete;
 	if (PhysicsSceneCompletion.GetReference())
 	{
@@ -1004,6 +1005,9 @@ void FPhysScene::WaitClothScene()
 		QUICK_SCOPE_CYCLE_COUNTER(STAT_FPhysScene_WaitClothScene);
 		FTaskGraphInterface::Get().WaitUntilTasksComplete(ThingsToComplete, ENamedThreads::GameThread);
 	}
+	
+	PhysicsSubsceneCompletion[PST_Cloth] = nullptr;
+	bPhysXSceneExecuting[PST_Cloth] = false;
 }
 
 void FPhysScene::SceneCompletionTask(ENamedThreads::Type CurrentThread, const FGraphEventRef& MyCompletionGraphEvent, EPhysicsSceneType SceneType)
@@ -1067,8 +1071,12 @@ void FPhysScene::ProcessPhysScene(uint32 SceneType)
 	}
 #endif // WITH_PHYSX
 
+    if(SceneType != PST_Cloth)	//Cloth potentially runs this code off the game thread where setting to null is not thread safe. For cloth we do it in a later game thread sync point
+	{
 	PhysicsSubsceneCompletion[SceneType] = NULL;
 	bPhysXSceneExecuting[SceneType] = false;
+	}
+	
 }
 #if WITH_PHYSX
 void FPhysScene::UpdateActiveTransforms(uint32 SceneType)
@@ -1318,7 +1326,7 @@ void FPhysScene::StartFrame()
 				STATGROUP_TaskGraphTasks);
 
 			PhysicsSceneCompletion = TGraphTask<FNullGraphTask>::CreateTask(&FinishPrerequisites, ENamedThreads::GameThread).ConstructAndDispatchWhenReady(
-				GET_STATID(STAT_FNullGraphTask_ProcessPhysScene_Join), PhysSingleThreadedMode() ? ENamedThreads::GameThread : ENamedThreads::AnyThread);
+				GET_STATID(STAT_FNullGraphTask_ProcessPhysScene_Join), PhysSingleThreadedMode() ? ENamedThreads::GameThread : ENamedThreads::AnyThreadGame());
 		}
 		else
 		{
@@ -1343,7 +1351,7 @@ void FPhysScene::StartCloth()
 		{
 			if (PhysicsSubsceneCompletion[PST_Cloth].GetReference())
 			{
-				ENamedThreads::Type CurrentThread = CVarParallelCloth.GetValueOnAnyThread() ? ENamedThreads::AnyThread : ENamedThreads::GameThread;
+				ENamedThreads::Type CurrentThread = (CVarParallelCloth.GetValueOnAnyThread() && !PhysSingleThreadedMode()) ? ENamedThreads::AnyThreadGame() : ENamedThreads::GameThread;
 				DECLARE_CYCLE_STAT(TEXT("FDelegateGraphTask.ProcessPhysScene_Cloth"),
 				STAT_FDelegateGraphTask_ProcessPhysScene_Cloth,
 					STATGROUP_TaskGraphTasks);
@@ -1352,7 +1360,7 @@ void FPhysScene::StartCloth()
 					FDelegateGraphTask::CreateAndDispatchWhenReady(
 					FDelegateGraphTask::FDelegate::CreateRaw(this, &FPhysScene::SceneCompletionTask, PST_Cloth),
 					GET_STATID(STAT_FDelegateGraphTask_ProcessPhysScene_Cloth), PhysicsSubsceneCompletion[PST_Cloth],
-					CurrentThread, ENamedThreads::GameThread
+					CurrentThread, CurrentThread
 					)
 					);
 			}
diff --git a/Engine/Source/Runtime/Engine/Private/PhysicsEngine/PhysSubstepTasks.cpp b/Engine/Source/Runtime/Engine/Private/PhysicsEngine/PhysSubstepTasks.cpp
--- a/Engine/Source/Runtime/Engine/Private/PhysicsEngine/PhysSubstepTasks.cpp
+++ b/Engine/Source/Runtime/Engine/Private/PhysicsEngine/PhysSubstepTasks.cpp
@@ -361,7 +361,7 @@ void FPhysSubstepTask::SubstepSimulationStart()
 	PhysXCompletionTask* SubstepTask = new PhysXCompletionTask(CompletionEvent,
 		 PST_MAX //we don't care about sub-step time. The full time is recorded by FullSimulationTask
 		,PAScene->getTaskManager());
-	ENamedThreads::Type NamedThread = PhysSingleThreadedMode() ? ENamedThreads::GameThread : ENamedThreads::AnyThread;
+	ENamedThreads::Type NamedThread = PhysSingleThreadedMode() ? ENamedThreads::GameThread : ENamedThreads::AnyThreadGame();
 
 	DECLARE_CYCLE_STAT(TEXT("FDelegateGraphTask.ProcessPhysSubstepSimulation"),
 		STAT_FDelegateGraphTask_ProcessPhysSubstepSimulation,
diff --git a/Engine/Source/Runtime/Engine/Private/PlayerState.cpp b/Engine/Source/Runtime/Engine/Private/PlayerState.cpp
--- a/Engine/Source/Runtime/Engine/Private/PlayerState.cpp
+++ b/Engine/Source/Runtime/Engine/Private/PlayerState.cpp
@@ -190,12 +190,15 @@ void APlayerState::OnRep_PlayerName()
 
 void APlayerState::OnRep_bIsInactive()
 {
-	// remove and re-add from the GameState so it's in the right list
-	GetWorld()->GameState->RemovePlayerState(this);
-	GetWorld()->GameState->AddPlayerState(this);
+	// remove and re-add from the GameState so it's in the right list  
+	UWorld* World = GetWorld();
+	if (ensure(World && World->GameState))
+	{
+		World->GameState->RemovePlayerState(this);
+		World->GameState->AddPlayerState(this);
+	}
 }
 
-
 bool APlayerState::ShouldBroadCastWelcomeMessage(bool bExiting)
 {
 	return (!bIsInactive && GetNetMode() != NM_Standalone);
@@ -203,14 +206,15 @@ bool APlayerState::ShouldBroadCastWelcomeMessage(bool bExiting)
 
 void APlayerState::Destroyed()
 {
-	if ( GetWorld()->GameState != NULL )
+	UWorld* World = GetWorld();
+	if (World->GameState != NULL)
 	{
-		GetWorld()->GameState->RemovePlayerState(this);
+		World->GameState->RemovePlayerState(this);
 	}
 
 	if( ShouldBroadCastWelcomeMessage(true) )
 	{
-		for( FConstPlayerControllerIterator Iterator = GetWorld()->GetPlayerControllerIterator(); Iterator; ++Iterator )
+		for (FConstPlayerControllerIterator Iterator = World->GetPlayerControllerIterator(); Iterator; ++Iterator)
 		{
 			APlayerController* PlayerController = *Iterator;
 			if( PlayerController )
diff --git a/Engine/Source/Runtime/Engine/Private/SceneView.cpp b/Engine/Source/Runtime/Engine/Private/SceneView.cpp
--- a/Engine/Source/Runtime/Engine/Private/SceneView.cpp
+++ b/Engine/Source/Runtime/Engine/Private/SceneView.cpp
@@ -608,6 +608,15 @@ uint32 FSceneView::GetViewKey() const
 	return 0;
 }
 
+uint32 FSceneView::GetOcclusionFrameCounter() const
+{
+	if (State)
+	{
+		return State->GetOcclusionFrameCounter();
+	}
+	return MAX_uint32;
+}
+
 
 void FSceneView::UpdateViewMatrix()
 {
diff --git a/Engine/Source/Runtime/Engine/Private/ShaderCompiler/ShaderCompiler.cpp b/Engine/Source/Runtime/Engine/Private/ShaderCompiler/ShaderCompiler.cpp
--- a/Engine/Source/Runtime/Engine/Private/ShaderCompiler/ShaderCompiler.cpp
+++ b/Engine/Source/Runtime/Engine/Private/ShaderCompiler/ShaderCompiler.cpp
@@ -16,6 +16,15 @@
 
 DEFINE_LOG_CATEGORY(LogShaderCompilers);
 
+// this is for the protocol, not the data, bump if FShaderCompilerInput or ProcessInputFromArchive changes (also search for the second one with the same name, todo: put into one header file)
+const int32 ShaderCompileWorkerInputVersion = 5;
+// this is for the protocol, not the data, bump if FShaderCompilerOutput or WriteToOutputArchive changes (also search for the second one with the same name, todo: put into one header file)
+const int32 ShaderCompileWorkerOutputVersion = 3;
+// this is for the protocol, not the data, bump if FShaderCompilerOutput or WriteToOutputArchive changes (also search for the second one with the same name, todo: put into one header file)
+const int32 ShaderCompileWorkerSingleJobHeader = 'S';
+// this is for the protocol, not the data, bump if FShaderCompilerOutput or WriteToOutputArchive changes (also search for the second one with the same name, todo: put into one header file)
+const int32 ShaderCompileWorkerPipelineJobHeader = 'P';
+
 // Set to 1 to debug ShaderCompilerWorker.exe. Set a breakpoint in LaunchWorker() to get the cmd-line.
 #define DEBUG_SHADERCOMPILEWORKER 0
 
@@ -169,21 +178,36 @@ static inline void GetFormatVersionMap(TMap<FString, uint16>& OutFormatVersionMa
 // Serialize Queued Job information
 static bool DoWriteTasks(const TArray<FShaderCompileJob*>& QueuedJobs, FArchive& TransferFile)
 {
-	int32 ShaderCompileWorkerInputVersion = 4;
-	TransferFile << ShaderCompileWorkerInputVersion;
+	int32 InputVersion = ShaderCompileWorkerInputVersion;
+	TransferFile << InputVersion;
 
 	static TMap<FString, uint16> FormatVersionMap;
 	GetFormatVersionMap(FormatVersionMap);
 
 	TransferFile << FormatVersionMap;
 
-	int32 NumBatches = QueuedJobs.Num();
-	TransferFile << NumBatches;
+	// Write individual shader jobs
+	{
+		int32 SingleJobHeader = ShaderCompileWorkerSingleJobHeader;
+		TransferFile << SingleJobHeader;
+
+		int32 NumBatches = QueuedJobs.Num();
+		TransferFile << NumBatches;
+
+		// Serialize all the batched jobs
+		for (int32 JobIndex = 0; JobIndex < QueuedJobs.Num(); JobIndex++)
+		{
+			TransferFile << QueuedJobs[JobIndex]->Input;
+		}
+	}
 
-	// Serialize all the batched jobs
-	for (int32 JobIndex = 0; JobIndex < QueuedJobs.Num(); JobIndex++)
+	// Write shader pipeline jobs
 	{
-		TransferFile << QueuedJobs[JobIndex]->Input;
+		int32 PipelineJobHeader = ShaderCompileWorkerPipelineJobHeader;
+		TransferFile << PipelineJobHeader;
+
+		int32 NumBatches = 0;
+		TransferFile << NumBatches;
 	}
 
 	return TransferFile.Close();
@@ -192,24 +216,28 @@ static bool DoWriteTasks(const TArray<FShaderCompileJob*>& QueuedJobs, FArchive&
 // Process results from Worker Process
 static void DoReadTaskResults(const TArray<FShaderCompileJob*>& QueuedJobs, FArchive& OutputFile)
 {
-	const int32 OutputVersion = 2;
-	int32 ShaderCompileWorkerOutputVersion;
-	OutputFile << ShaderCompileWorkerOutputVersion;
+	int32 OutputVersion = ShaderCompileWorkerOutputVersion;
+	OutputFile << OutputVersion;
+
+	auto ModalErrorOrLog = [&](const FString& Text)
+		{
+			if (FPlatformProperties::SupportsWindowedMode())
+			{
+				UE_LOG(LogShaderCompilers, Error, TEXT("%s"), *Text);
+				FMessageDialog::Open(EAppMsgType::Ok, FText::FromString(Text));
+				FPlatformMisc::RequestExit(false);
+				return;
+			}
+			else
+			{
+				UE_LOG(LogShaderCompilers, Fatal, TEXT("%s"), *Text);
+			}
+		};
 
 	if (ShaderCompileWorkerOutputVersion != OutputVersion)
 	{
 		FString Text = FString::Printf(TEXT("Expecting ShaderCompilerWorker output version %d, got %d instead! Forgot to build ShaderCompilerWorker?"), ShaderCompileWorkerOutputVersion, OutputVersion);
-		if (FPlatformProperties::SupportsWindowedMode())
-		{
-			UE_LOG(LogShaderCompilers, Error, TEXT("%s"), *Text);
-			FMessageDialog::Open(EAppMsgType::Ok, FText::FromString(Text));
-			FPlatformMisc::RequestExit(false);
-			return;
-		}
-		else
-		{
-			UE_LOG(LogShaderCompilers, Fatal, TEXT("%s"), *Text);
-		}
+		ModalErrorOrLog(Text);
 	}
 
 	int32 ErrorCode;
@@ -251,23 +279,52 @@ static void DoReadTaskResults(const TArray<FShaderCompileJob*>& QueuedJobs, FArc
 		UE_LOG(LogShaderCompilers, Fatal, TEXT("ShaderCompileWorker crashed! \n %s \n %s"), ExceptionInfo.GetData(), Callstack.GetData());
 	}
 
-	int32 NumJobs;
-	OutputFile << NumJobs;
-	checkf(NumJobs == QueuedJobs.Num(), TEXT("Worker returned %u jobs, %u expected"), NumJobs, QueuedJobs.Num());
+	auto ReadSingleJob = [&](FShaderCompileJob* CurrentJob)
+		{
+			check(!CurrentJob->bFinalized);
+			CurrentJob->bFinalized = true;
+
+			// Deserialize the shader compilation output.
+			OutputFile << CurrentJob->Output;
+
+			// Generate a hash of the output and cache it
+			// The shader processing this output will use it to search for existing FShaderResources
+			CurrentJob->Output.GenerateOutputHash();
+			CurrentJob->bSucceeded = CurrentJob->Output.bSucceeded;
+	};
 
-	for (int32 JobIndex = 0; JobIndex < NumJobs; JobIndex++)
+	// Read single jobs
 	{
-		FShaderCompileJob* CurrentJob = QueuedJobs[JobIndex];
-		check(!CurrentJob->bFinalized);
-		CurrentJob->bFinalized = true;
+		int32 SingleJobHeader = -1;
+		OutputFile << SingleJobHeader;
+		if (SingleJobHeader != ShaderCompileWorkerSingleJobHeader)
+		{
+			FString Text = FString::Printf(TEXT("Expecting ShaderCompilerWorker Single Jobs %d, got %d instead! Forgot to build ShaderCompilerWorker?"), ShaderCompileWorkerSingleJobHeader, SingleJobHeader);
+			ModalErrorOrLog(Text);
+		}
 
-		// Deserialize the shader compilation output.
-		OutputFile << CurrentJob->Output;
+		int32 NumJobs;
+		OutputFile << NumJobs;
+		checkf(NumJobs == QueuedJobs.Num(), TEXT("Worker returned %u single jobs, %u expected"), NumJobs, QueuedJobs.Num());
+		for (int32 JobIndex = 0; JobIndex < NumJobs; JobIndex++)
+		{
+			auto* CurrentJob = QueuedJobs[JobIndex];
+			ReadSingleJob(CurrentJob);
+		}
+	}
 
-		// Generate a hash of the output and cache it
-		// The shader processing this output will use it to search for existing FShaderResources
-		CurrentJob->Output.GenerateOutputHash();
-		CurrentJob->bSucceeded = CurrentJob->Output.bSucceeded;
+	// Pipeline jobs
+	{
+		int32 PipelineJobHeader = -1;
+		OutputFile << PipelineJobHeader;
+		if (PipelineJobHeader != ShaderCompileWorkerPipelineJobHeader)
+		{
+			FString Text = FString::Printf(TEXT("Expecting ShaderCompilerWorker Pipeline Jobs %d, got %d instead! Forgot to build ShaderCompilerWorker?"), ShaderCompileWorkerPipelineJobHeader, PipelineJobHeader);
+			ModalErrorOrLog(Text);
+		}
+
+		int32 NumJobs;
+		OutputFile << NumJobs;
 	}
 }
 
@@ -1161,22 +1218,28 @@ void FShaderCompilingManager::ProcessCompiledShaderMaps(
 			TArray<FMaterial*> MaterialsArray = *Materials;
 			bool bSuccess = true;
 
+			auto CheckSingleJob = [&](FShaderCompileJob* SingleJob)
+				{
+					if (SingleJob->bSucceeded)
+					{
+						check(SingleJob->Output.ShaderCode.GetShaderCodeSize() > 0);
+					}
+
+					if (GShowShaderWarnings || !SingleJob->bSucceeded)
+					{
+						for (int32 ErrorIndex = 0; ErrorIndex < SingleJob->Output.Errors.Num(); ErrorIndex++)
+						{
+							Errors.AddUnique(SingleJob->Output.Errors[ErrorIndex].StrippedErrorMessage);
+						}
+					}
+				};
+
 			for (int32 JobIndex = 0; JobIndex < ResultArray.Num(); JobIndex++)
 			{
 				FShaderCompileJob& CurrentJob = *ResultArray[JobIndex];
 				bSuccess = bSuccess && CurrentJob.bSucceeded;
 
-				if (CurrentJob.bSucceeded)
-				{
-					check(CurrentJob.Output.ShaderCode.GetShaderCodeSize() > 0);
-				}
-				else 
-				{
-					for (int32 ErrorIndex = 0; ErrorIndex < CurrentJob.Output.Errors.Num(); ErrorIndex++)
-					{
-						Errors.AddUnique(CurrentJob.Output.Errors[ErrorIndex].StrippedErrorMessage);
-					}
-				}
+				CheckSingleJob(&CurrentJob);
 			}
 
 			bool bShaderMapComplete = true;
@@ -1811,6 +1874,8 @@ void GlobalBeginCompileShader(
 	Input.ShaderFormat = LegacyShaderPlatformToShaderFormat(EShaderPlatform(Target.Platform));
 	Input.SourceFilename = SourceFilename;
 	Input.EntryPointName = FunctionName;
+	Input.bCompilingForShaderPipeline = false;
+	Input.bIncludeUsedOutputs = false;
 	Input.DumpDebugInfoRootPath = GShaderCompilingManager->GetAbsoluteShaderDebugInfoDirectory() / Input.ShaderFormat.ToString();
 	// asset material name or "Global"
 	Input.DebugGroupName = DebugGroupName;
diff --git a/Engine/Source/Runtime/Engine/Private/ShaderDerivedDataVersion.h b/Engine/Source/Runtime/Engine/Private/ShaderDerivedDataVersion.h
--- a/Engine/Source/Runtime/Engine/Private/ShaderDerivedDataVersion.h
+++ b/Engine/Source/Runtime/Engine/Private/ShaderDerivedDataVersion.h
@@ -9,5 +9,5 @@
 // In case of merge conflicts with DDC versions, you MUST generate a new GUID and set this new
 // guid as version
 
-#define GLOBALSHADERMAP_DERIVEDDATA_VER			TEXT("af1606fcec214373844d2b6cf9df71c8")
-#define MATERIALSHADERMAP_DERIVEDDATA_VER		TEXT("B2870EB526BD4F359760A52F6ADD1B95")
+#define GLOBALSHADERMAP_DERIVEDDATA_VER			TEXT("653D597996674C67855E3A31377FCDB5")
+#define MATERIALSHADERMAP_DERIVEDDATA_VER		TEXT("3F0644619F9C48FF91E0D494C95EB038")
diff --git a/Engine/Source/Runtime/Engine/Private/StaticMeshBuild.cpp b/Engine/Source/Runtime/Engine/Private/StaticMeshBuild.cpp
--- a/Engine/Source/Runtime/Engine/Private/StaticMeshBuild.cpp
+++ b/Engine/Source/Runtime/Engine/Private/StaticMeshBuild.cpp
@@ -68,6 +68,8 @@ void UStaticMesh::Build(bool bSilent, TArray<FText>* OutErrors)
 		GWarn->BeginSlowTask( StatusUpdate, true );	
 	}
 
+	PreMeshBuild.Broadcast(this);
+
 	// Detach all instances of this static mesh from the scene.
 	FStaticMeshComponentRecreateRenderStateContext RecreateRenderStateContext(this,false);
 
@@ -91,7 +93,6 @@ void UStaticMesh::Build(bool bSilent, TArray<FText>* OutErrors)
 	CreateBodySetup();
 	check(BodySetup != NULL);
 
-#if WITH_EDITOR
 	if( SourceModels.Num() )
 	{
 		// Rescale simple collision if the user changed the mesh build scale
@@ -102,7 +103,6 @@ void UStaticMesh::Build(bool bSilent, TArray<FText>* OutErrors)
 	// TODO_STATICMESH: Not necessary any longer?
 	BodySetup->InvalidatePhysicsData();
 	BodySetup->CreatePhysicsMeshes();
-#endif
 
 	// Compare the derived data keys to see if renderable mesh data has actually changed.
 	check(RenderData);
@@ -144,6 +144,8 @@ void UStaticMesh::Build(bool bSilent, TArray<FText>* OutErrors)
 	// Calculate extended bounds
 	CalculateExtendedBounds();
 
+	PostMeshBuild.Broadcast(this);
+
 	if(!bSilent)
 	{
 		GWarn->EndSlowTask();
diff --git a/Engine/Source/Runtime/Engine/Private/TickTaskManager.cpp b/Engine/Source/Runtime/Engine/Private/TickTaskManager.cpp
--- a/Engine/Source/Runtime/Engine/Private/TickTaskManager.cpp
+++ b/Engine/Source/Runtime/Engine/Private/TickTaskManager.cpp
@@ -291,7 +291,7 @@ class FTickTaskSequencer
 		}
 		FORCEINLINE ENamedThreads::Type GetDesiredThread()
 		{
-			return ENamedThreads::HiPri(ENamedThreads::AnyThread);
+			return ENamedThreads::AnyThread;
 		}
 		FORCEINLINE static ESubsequentsMode::Type GetSubsequentsMode() 
 		{ 
@@ -366,7 +366,7 @@ public:
 		UseContext.Thread = ENamedThreads::GameThread;
 		if (TickFunction->bRunOnAnyThread && bAllowConcurrentTicks && bIsOriginalTickGroup)
 		{
-			UseContext.Thread = ENamedThreads::AnyThread;
+			UseContext.Thread = ENamedThreads::AnyThreadGame();
 		}
 
 		TickFunction->TaskPointer = TGraphTask<FTickFunctionTask>::CreateTask(Prerequisites, TickContext.Thread).ConstructAndHold(TickFunction, &UseContext, bLogTicks, bLogTicksShowPrerequistes);
diff --git a/Engine/Source/Runtime/Engine/Private/UnrealEngine.cpp b/Engine/Source/Runtime/Engine/Private/UnrealEngine.cpp
--- a/Engine/Source/Runtime/Engine/Private/UnrealEngine.cpp
+++ b/Engine/Source/Runtime/Engine/Private/UnrealEngine.cpp
@@ -7479,6 +7479,7 @@ void DrawStatsHUD( UWorld* World, FViewport* Viewport, FCanvas* Canvas, UCanvas*
 
 		if (GEngine->bEnableOnScreenDebugMessagesDisplay && GEngine->bEnableOnScreenDebugMessages)
 		{
+			const int32 MaxYPos = CanvasObject ? CanvasObject->SizeY : 700;
 			if (GEngine->PriorityScreenMessages.Num() > 0)
 			{
 				FCanvasTextItem MessageTextItem( FVector2D( 0, 0 ), FText::GetEmpty(), GEngine->GetSmallFont(), FLinearColor::White );
@@ -7486,7 +7487,7 @@ void DrawStatsHUD( UWorld* World, FViewport* Viewport, FCanvas* Canvas, UCanvas*
 				for (int32 PrioIndex = GEngine->PriorityScreenMessages.Num() - 1; PrioIndex >= 0; PrioIndex--)
 				{
 					FScreenMessageString& Message = GEngine->PriorityScreenMessages[PrioIndex];
-					if (YPos < 700)
+					if (YPos < MaxYPos)
 					{
 						MessageTextItem.Text =  FText::FromString( Message.ScreenMessage );
 						MessageTextItem.SetColor( Message.DisplayColor );
@@ -7509,7 +7510,7 @@ void DrawStatsHUD( UWorld* World, FViewport* Viewport, FCanvas* Canvas, UCanvas*
 				for (TMap<int32, FScreenMessageString>::TIterator MsgIt(GEngine->ScreenMessages); MsgIt; ++MsgIt)
 				{
 					FScreenMessageString& Message = MsgIt.Value();
-					if (YPos < 700)
+					if (YPos < MaxYPos)
 					{
 						MessageTextItem.Text =  FText::FromString( Message.ScreenMessage );
 						MessageTextItem.SetColor( Message.DisplayColor );
diff --git a/Engine/Source/Runtime/Engine/Public/MaterialShared.h b/Engine/Source/Runtime/Engine/Public/MaterialShared.h
--- a/Engine/Source/Runtime/Engine/Public/MaterialShared.h
+++ b/Engine/Source/Runtime/Engine/Public/MaterialShared.h
@@ -353,7 +353,8 @@ namespace EMaterialShaderMapUsage
 		MaterialExportNormal,
 		MaterialExportMetallic,
 		MaterialExportRoughness,
-		MaterialExportAO
+		MaterialExportAO,
+		MaterialExportEmissive,
 	};
 }
 
diff --git a/Engine/Source/Runtime/Engine/Public/SceneManagement.h b/Engine/Source/Runtime/Engine/Public/SceneManagement.h
--- a/Engine/Source/Runtime/Engine/Public/SceneManagement.h
+++ b/Engine/Source/Runtime/Engine/Public/SceneManagement.h
@@ -155,6 +155,10 @@ public:
 	 * returns a unique key for the view state, non-zero
 	 */
 	virtual uint32 GetViewKey() const = 0;
+	/** 
+	 * returns the occlusion frame counter 
+	 */
+	virtual uint32 GetOcclusionFrameCounter() const = 0;
 protected:
 	// Don't allow direct deletion of the view state, Destroy should be called instead.
 	virtual ~FSceneViewStateInterface() {}
diff --git a/Engine/Source/Runtime/Engine/Public/SceneView.h b/Engine/Source/Runtime/Engine/Public/SceneView.h
--- a/Engine/Source/Runtime/Engine/Public/SceneView.h
+++ b/Engine/Source/Runtime/Engine/Public/SceneView.h
@@ -771,6 +771,11 @@ public:
 	 */
 	uint32 GetViewKey() const;
 
+	/** 
+	 * returns a the occlusion frame counter or MAX_uint32 if there is no view state
+	 */
+	uint32 GetOcclusionFrameCounter() const;
+
 	/** Allow things like HMD displays to update the view matrix at the last minute, to minimize perceived latency */
 	void UpdateViewMatrix();
 
diff --git a/Engine/Source/Runtime/GameLiveStreaming/GameLiveStreaming.cpp b/Engine/Source/Runtime/GameLiveStreaming/GameLiveStreaming.cpp
--- a/Engine/Source/Runtime/GameLiveStreaming/GameLiveStreaming.cpp
+++ b/Engine/Source/Runtime/GameLiveStreaming/GameLiveStreaming.cpp
@@ -200,16 +200,32 @@ void FGameLiveStreaming::StartBroadcastingGame( const FGameBroadcastConfig& Game
 			{
 				for( int32 TextureIndex = 0; TextureIndex < 2; ++TextureIndex )
 				{
-					FRHIResourceCreateInfo CreateInfo;
-					ReadbackTextures[ TextureIndex ] = RHICreateTexture2D(
-						BroadcastConfig.VideoBufferWidth,
-						BroadcastConfig.VideoBufferHeight,
-						PF_B8G8R8A8,
-						1,
-						1,
-						TexCreate_CPUReadback,
-						CreateInfo
-						);
+					ReadbackTextures[ TextureIndex ] = nullptr;
+				}
+				ENQUEUE_UNIQUE_RENDER_COMMAND_THREEPARAMETER(
+					FWebMRecordCreateBufers,
+					int32,InVideoWidth,BroadcastConfig.VideoBufferWidth,
+					int32,InVideoHeight,BroadcastConfig.VideoBufferHeight,
+					FTexture2DRHIRef*,InReadbackTextures,ReadbackTextures,
+				{
+					for (int32 TextureIndex = 0; TextureIndex < 2; ++TextureIndex)
+					{
+						FRHIResourceCreateInfo CreateInfo;
+						InReadbackTextures[ TextureIndex ] = RHICreateTexture2D(
+							InVideoWidth,
+							InVideoHeight,
+							PF_B8G8R8A8,
+							1,
+							1,
+							TexCreate_CPUReadback,
+							CreateInfo
+							);
+					}
+				});
+				FlushRenderingCommands();
+				for (int32 TextureIndex = 0; TextureIndex < 2; ++TextureIndex)
+				{
+					check(ReadbackTextures[TextureIndex].GetReference());
 				}
 	
 				ReadbackTextureIndex = 0;
