[UE-20194] Don't warn about a missing parent function if the function itself is no longer valid as a result of compile-on-load.

change summary:
- modified ULinkerLoad::CreateExport() to also check the validity of the current export map entry after loading a function's Outer in the case of a missing/NULL SuperStruct (since loading the Outer may invalidate the entry - DOBPs)

diff --git a/Engine/Source/Runtime/CoreUObject/Private/UObject/LinkerLoad.cpp b/Engine/Source/Runtime/CoreUObject/Private/UObject/LinkerLoad.cpp
--- a/Engine/Source/Runtime/CoreUObject/Private/UObject/LinkerLoad.cpp
+++ b/Engine/Source/Runtime/CoreUObject/Private/UObject/LinkerLoad.cpp
@@ -3493,19 +3493,16 @@ UObject* FLinkerLoad::CreateExport( int32 Index )
 			{
 				if( LoadClass->IsChildOf(UFunction::StaticClass()) )
 				{
+					// In the case of a function object, the outer should be the function's class. For Blueprints, loading
+					// the outer class may also invalidate this entry in the export map. In that case, we won't actually be
+					// keeping the function object around, so there's no need to warn here about the missing parent object.
 					UObject* ObjOuter = IndexToObject(Export.OuterIndex);
-					if (ObjOuter)
+					if (ObjOuter && !Export.bExportLoadFailed)
 					{
 						UClass* FuncClass = Cast<UClass>(ObjOuter);
 						if (FuncClass && FuncClass->ClassGeneratedBy && !FuncClass->ClassGeneratedBy->HasAnyFlags(RF_BeingRegenerated))
 						{
-							// If this is a function (NOT being regenerated) whose parent has been removed, give it a 
-							// NULL parent, as we would have in the script compiler
-							//
-							// @TODO: Figure out why this occurs, as it happens frequently, and we can only suggest 
-							//        users resave their assets (downgrading from a warning to keep from spamming logs 
-							//        and since it seems to be fixed up with class regeneration... could be related to 
-							//        data-only construction scripts
+							// If this is a function (NOT being regenerated) whose parent has been removed, give it a NULL parent, as we would have in the script compiler.
 							UE_LOG(LogLinker, Display, TEXT("CreateExport: Failed to load Parent for %s; removing parent information, but keeping function"), *GetExportFullName(Index));
 						}
 					}
