Proper fix for not adding imports for Outers of objects that have been skipped for the current cook target platform. Object marks should only be checked when actually cooking.

diff --git a/Engine/Source/Runtime/CoreUObject/Private/UObject/SavePackage.cpp b/Engine/Source/Runtime/CoreUObject/Private/UObject/SavePackage.cpp
--- a/Engine/Source/Runtime/CoreUObject/Private/UObject/SavePackage.cpp
+++ b/Engine/Source/Runtime/CoreUObject/Private/UObject/SavePackage.cpp
@@ -864,8 +864,7 @@ FArchive& FArchiveSaveTagImports::operator<<( UObject*& Obj )
 				}
 
 				// If the object has been excluded, don't add its outer
-				// TODO: RobM, please investigate as this change is causing problems with BP macros
-				//if (!Obj->HasAllMarks(ObjectMarks))
+				if (!IsCooking() || !Obj->HasAllMarks(ObjectMarks))
 				{
 					UObject* Parent = Obj->GetOuter();
 					if( Parent )
