Copying //UE4/Dev-Blueprints to Dev-Main (//UE4/Dev-Main) @ 2781164

diff --git a/Engine/Source/Runtime/CoreUObject/Private/Blueprint/BlueprintSupport.cpp b/Engine/Source/Runtime/CoreUObject/Private/Blueprint/BlueprintSupport.cpp
--- a/Engine/Source/Runtime/CoreUObject/Private/Blueprint/BlueprintSupport.cpp
+++ b/Engine/Source/Runtime/CoreUObject/Private/Blueprint/BlueprintSupport.cpp
@@ -1412,6 +1412,74 @@ void FLinkerLoad::ResetDeferredLoadingState()
 #endif // USE_CIRCULAR_DEPENDENCY_LOAD_DEFERRING
 }
 
+#if WITH_EDITORONLY_DATA
+extern int32 GLinkerAllowDynamicClasses;
+#endif
+
+UObject* FLinkerLoad::FindImport(UClass* ImportClass, UObject* ImportOuter, const TCHAR* Name)
+{	
+	UObject* Result = StaticFindObject(ImportClass, ImportOuter, Name);
+#if WITH_EDITORONLY_DATA
+	static FName NAME_BlueprintGeneratedClass(TEXT("BlueprintGeneratedClass"));
+	if (GLinkerAllowDynamicClasses && !Result && ImportClass->GetFName() == NAME_BlueprintGeneratedClass)
+	{
+		Result = StaticFindObject(UDynamicClass::StaticClass(), ImportOuter, Name);
+	}
+#endif
+	return Result;
+}
+
+UObject* FLinkerLoad::FindImportFast(UClass* ImportClass, UObject* ImportOuter, FName Name)
+{
+	UObject* Result = StaticFindObjectFast(ImportClass, ImportOuter, Name);
+#if WITH_EDITORONLY_DATA
+	static FName NAME_BlueprintGeneratedClass(TEXT("BlueprintGeneratedClass"));
+	if (GLinkerAllowDynamicClasses && !Result && ImportClass->GetFName() == NAME_BlueprintGeneratedClass)
+	{
+		Result = StaticFindObjectFast(UDynamicClass::StaticClass(), ImportOuter, Name);
+	}
+#endif
+	return Result;
+}
+
+void FLinkerLoad::CreateDynamicTypeLoader()
+{
+	// In this case we can skip serializing PackageFileSummary and fill all the required info here
+	bHasSerializedPackageFileSummary = true;
+
+	// Try to get dependencies for dynamic classes
+	TArray<FBlueprintDependencyData> ImportDependencyPackages;
+	FConvertedBlueprintsDependencies::Get().GetAssets(LinkerRoot->GetFName(), ImportDependencyPackages);
+
+	// Create Imports
+	for (FBlueprintDependencyData& Import : ImportDependencyPackages)
+	{
+		FObjectImport* ObjectImport = new(ImportMap)FObjectImport(nullptr);
+		ObjectImport->ClassName = Import.ClassName;
+		ObjectImport->ClassPackage = Import.ClassPackageName;
+		ObjectImport->ObjectName = Import.ObjectName;
+		ObjectImport->OuterIndex = FPackageIndex::FromImport(ImportMap.Num());
+
+		FObjectImport* OuterImport = new(ImportMap)FObjectImport(nullptr);
+		OuterImport->ClassName = NAME_Package;
+		OuterImport->ClassPackage = GLongCoreUObjectPackageName;
+		OuterImport->ObjectName = Import.PackageName;
+	}
+
+	// Create Exports
+	{
+		FObjectExport* Export = new (ExportMap)FObjectExport();
+		const FName* TypeNamePtr = GetConvertedDynamicPackageNameToTypeName().Find(LinkerRoot->GetFName());
+		Export->ObjectName = TypeNamePtr ? *TypeNamePtr : NAME_None;
+		Export->ThisIndex = FPackageIndex::FromExport(ExportMap.Num() - 1);
+		// This allows us to skip creating two additional imports for UDynamicClass and its package
+		Export->bDynamicClass = true;
+		Export->ObjectFlags |= RF_Public;
+	}
+
+	LinkerRoot->SetPackageFlags(LinkerRoot->GetPackageFlags() | PKG_CompiledIn);
+}
+
 /*******************************************************************************
  * UObject
  ******************************************************************************/
@@ -1655,251 +1723,150 @@ void FDeferredObjInitializerTracker::ResolveDeferredSubClassObjects(UClass* Supe
 // don't want other files ending up with this internal define
 #undef DEFERRED_DEPENDENCY_CHECK
 
-#if WITH_EDITOR
-
-FReplaceConvertedAssetManager::FReplaceConvertedAssetManager() 
-	: bIsEnabled(false)
-{
-	FModuleManager::Get().OnModulesChanged().AddStatic(&FReplaceConvertedAssetManager::OnModulesChanged);
-
-	// FOR DEVELOPMENT/TEST ONLY:
-	const FBoolConfigValueHelper LoadNativeConvertedBPClass(TEXT("Kismet"), TEXT("bLoadNativeConvertedBPClassInEditor"), GEngineIni);
-	SetEnabled(LoadNativeConvertedBPClass);
-}
-
-void FReplaceConvertedAssetManager::OnModulesChanged(FName ModuleThatChanged, EModuleChangeReason ReasonForChange)
-{
-	Get().GatherOriginalPathsOfConvertedAssets();
-}
-
-FReplaceConvertedAssetManager& FReplaceConvertedAssetManager::Get()
-{
-	static FReplaceConvertedAssetManager ReplaceConvertedAssetManager;
-	return ReplaceConvertedAssetManager;
-}
-
-UObject* FReplaceConvertedAssetManager::FindReplacement(const FString& OriginalPathName) const
-{
-	UObject* const* ObjPtr = ReplaceMap.Find(OriginalPathName);
-	return ObjPtr ? *ObjPtr : nullptr;
-}
-
-UPackage* FReplaceConvertedAssetManager::FindPackageReplacement(const FString& OriginalPathName) const
-{
-	for (auto It = ReplaceMap.CreateConstIterator(); It; ++It)
-	{
-		if (It.Key().StartsWith(OriginalPathName))
-		{
-			auto Val = It.Value();
-			return ensure(Val) ? Val->GetOutermost() : nullptr;
-		}
-	}
-	return nullptr;
-}
-
-void FReplaceConvertedAssetManager::AddReferencedObjects(FReferenceCollector& Collector)
-{
-	Collector.AddReferencedObjects(ReplaceMap);
-}
-
-void FReplaceConvertedAssetManager::GatherOriginalPathsOfConvertedAssets()
-{
-	if (!IsEnabled())
-	{
-		return;
-	}
-
-	ReplaceMap.Reset();
-	const FName ReplaceConverted(TEXT("ReplaceConverted"));
-
-	auto FillMap = [&](UField* Field)
-	{
-		if (Field && Field->HasMetaData(ReplaceConverted))
-		{
-			auto CombinedPaths = Field->GetMetaData(ReplaceConverted);
-			TArray<FString> Paths;
-			CombinedPaths.ParseIntoArray(Paths, TEXT(","));
-			for (auto& Path : Paths)
-			{
-				ReplaceMap.Add(Path, Field);
-			}
-		}
-	};
-
-	for (auto FuncIter : FConvertedBlueprintsDependencies::Get().CreateClassFunctions)
-	{
-		auto DynamicClass = (*FuncIter)();
-		FillMap(DynamicClass);
-	}
-
-	for (UScriptStruct* LocalScriptStruct : TObjectRange<UScriptStruct>())
-	{
-		FillMap(LocalScriptStruct);
-	}
-
-	for (UEnum* LocalEnum : TObjectRange<UEnum>())
-	{
-		FillMap(LocalEnum);
-	}
-}
-
-#endif //WITH_EDITOR
-
 FConvertedBlueprintsDependencies& FConvertedBlueprintsDependencies::Get()
 {
 	static FConvertedBlueprintsDependencies ConvertedBlueprintsDependencies;
 	return ConvertedBlueprintsDependencies;
 }
 
-void FConvertedBlueprintsDependencies::RegisterClass(FName ClassName, GetDependenciesNamesFunc GetAssets)
+void FConvertedBlueprintsDependencies::RegisterClass(FName PackageName, GetDependenciesNamesFunc GetAssets)
 {
-	check(!ClassNameToGetter.Contains(ClassName));
+	check(!PackageNameToGetter.Contains(PackageName));
 	ensure(GetAssets);
-	ClassNameToGetter.Add(ClassName, GetAssets);
+	PackageNameToGetter.Add(PackageName, GetAssets);
 }
 
-void FConvertedBlueprintsDependencies::GetAssets(FName ClassName, TArray<FName>& OutPackagePaths) const
+void FConvertedBlueprintsDependencies::GetAssets(FName PackageName, TArray<FBlueprintDependencyData>& OutDependencies) const
 {
-	auto FuncPtr = ClassNameToGetter.Find(ClassName);
+	auto FuncPtr = PackageNameToGetter.Find(PackageName);
 	auto Func = (FuncPtr) ? (*FuncPtr) : nullptr;
 	ensure(Func || !FuncPtr);
 	if (Func)
 	{
-		Func(OutPackagePaths);
+		Func(OutDependencies);
 	}
 }
 
+/*******************************************************************************
+ * FScriptCookReplacementCoordinator
+ ******************************************************************************/
 #if WITH_EDITOR
 
-struct FConvertibleTypeHelper
-{
-	static FName NameBPGC;
-	static FName NameWidgetBPGC;
-	static FName NameUDS;
-	static FName NameUDE;
-
-	static bool IsConvertible(const UObject* Obj)
-	{
-		const FName ClassName = Obj ? Obj->GetClass()->GetFName() : NAME_None;
-		return (ClassName != NAME_None)
-			&& ((ClassName == NameBPGC)
-			|| (ClassName == NameUDS)
-			|| (ClassName == NameUDE)
-			|| (ClassName == NameWidgetBPGC));
-	}
-};
-
-FName FConvertibleTypeHelper::NameBPGC(TEXT("BlueprintGeneratedClass"));
-FName FConvertibleTypeHelper::NameWidgetBPGC(TEXT("WidgetBlueprintGeneratedClass"));
-FName FConvertibleTypeHelper::NameUDS(TEXT("UserDefinedStruct"));
-FName FConvertibleTypeHelper::NameUDE(TEXT("UserDefinedEnum"));
+FScriptCookReplacementCoordinator* FScriptCookReplacementCoordinator::CoordinatorInstance = nullptr;
 
-FReplaceCookedBPGC::FReplaceCookedBPGC()
-	: NativeScriptPackage(nullptr)
-	, bEnabled(false)
+FScriptCookReplacementCoordinator* FScriptCookReplacementCoordinator::Get()
 {
+	return CoordinatorInstance;
 }
 
-FReplaceCookedBPGC& FReplaceCookedBPGC::Get()
+void FScriptCookReplacementCoordinator::Create(bool bEnabled, const TArray<FString>& ExcludedAssetTypes, const TArray<FString>& ExcludedBlueprintTypes, const TMap<UObject*, UClass*>& ReplacementMap)
 {
-	static FReplaceCookedBPGC ReplaceCookedBPGC;
-	return ReplaceCookedBPGC;
+	CoordinatorInstance = new FScriptCookReplacementCoordinator(bEnabled, ExcludedAssetTypes, ExcludedBlueprintTypes, ReplacementMap);
 }
 
-void FReplaceCookedBPGC::AddReferencedObjects(FReferenceCollector& Collector)
+FScriptCookReplacementCoordinator::FScriptCookReplacementCoordinator(bool bInEnabled, const TArray<FString>& InExcludedAssetTypes, const TArray<FString>& InExcludedBlueprintTypes, const TMap<UObject*, UClass*>& InReplacementMap)
+	: bEnabled(bInEnabled)
+	, ExcludedAssetTypes(InExcludedAssetTypes)
+	, ExcludedBlueprintTypes(InExcludedBlueprintTypes)
+	, ReplacementMap(InReplacementMap)
 {
-	Collector.AddReferencedObjects(ReplaceMap);
-	Collector.AddReferencedObject(NativeScriptPackage);
 }
 
-void FReplaceCookedBPGC::AddConvertedPackageName(FName PackageName)
+bool FScriptCookReplacementCoordinator::Initialize()
 {
-	ConvertedPackagesNames.Add(PackageName);
+	return true;
 }
 
-void FReplaceCookedBPGC::AddConvertedFieldStub(UObject* Object)
+DEFINE_LOG_CATEGORY_STATIC(LogDan, Log, All);
+
+UClass* FScriptCookReplacementCoordinator::FindReplacedClass(const UObject* Obj) const
 {
-	if (Object && ensure(CouldBeConverted(Object)))
+	if (!bEnabled)
 	{
-		const FName ClassName = Object->GetClass()->GetFName();
-		if ((ClassName == FConvertibleTypeHelper::NameBPGC) || (ClassName == FConvertibleTypeHelper::NameWidgetBPGC))
-		{
-			AddConvertedFieldStub<UDynamicClass>(Object->GetPathName(), Object->GetFName());
-		}
-		else if (ClassName == FConvertibleTypeHelper::NameUDS)
-		{
-			AddConvertedFieldStub<UScriptStruct>(Object->GetPathName(), Object->GetFName());
-		}
-		else if (ClassName == FConvertibleTypeHelper::NameUDE)
-		{
-			AddConvertedFieldStub<UEnum>(Object->GetPathName(), Object->GetFName());
-		}
-		else
-		{
-			check(false);
-		}
+		return nullptr;
 	}
-}
 
-UObject* FReplaceCookedBPGC::FindReplacementStub(UObject* Object)
-{
-	if (bEnabled && Object && CouldBeConverted(Object))
+	// we're only looking to replace class types:
+	const UClass* AsClass = Obj->GetClass();
+	while (AsClass)
 	{
-		const FString OriginalPathName = Object->GetPathName();
-		UObject** StubObjPtr = ReplaceMap.Find(OriginalPathName);
-		UObject* StubObj = StubObjPtr ? *StubObjPtr : nullptr;
-		check(!StubObj || StubObj->IsValidLowLevelFast());
-		if (StubObj)
+		if (UClass* const* ReplacementPtr = ReplacementMap.Find(AsClass))
 		{
-			int32 k = 0; k++;
+			return *ReplacementPtr;
 		}
-		return StubObj;
+
+		AsClass = AsClass->GetSuperClass();
 	}
+
 	return nullptr;
 }
 
-bool FReplaceCookedBPGC::CouldBeConverted(const UObject* Object) const
+EReplacementResult FScriptCookReplacementCoordinator::IsTargetedForReplacement(const UPackage* Package) const
 {
-	if (Object && !Object->HasAnyFlags(RF_ClassDefaultObject) && FConvertibleTypeHelper::IsConvertible(Object))
+	if (!bEnabled)
 	{
-		if (IsTargetedForConversionDelegate.IsBound())
+		return EReplacementResult::DontReplace;
+	}
+
+	// non-native packages with enums and structs should be converted, unless they are blacklisted:
+	UStruct* Struct = nullptr;
+	UEnum* Enum = nullptr;
+	TArray<UObject*> Objects;
+	GetObjectsWithOuter(Package, Objects, false);
+	for (auto Entry : Objects)
+	{
+		Struct = Cast<UStruct>(Entry);
+		Enum = Cast<UEnum>(Entry);
+		if (Struct || Enum)
 		{
-			return IsTargetedForConversionDelegate.Execute(Object);
+			break;
 		}
-		return true;
 	}
-	return false;
-	
-}
 
-bool FReplaceCookedBPGC::PackageShouldNotBeSaved(const UPackage* InOuter) const
-{
-	return InOuter && bEnabled && ConvertedPackagesNames.Contains(InOuter->GetFName());
+	UObject* Target = Struct;
+	if (Target == nullptr)
+	{
+		Target = Enum;
+	}
+	return IsTargetedForReplacement(Target);
 }
 
-void FReplaceCookedBPGC::Initialize(const FString& NativePackageName)
+EReplacementResult FScriptCookReplacementCoordinator::IsTargetedForReplacement(const UObject* Object) const
 {
-	if (ensure(!NativeScriptPackage))
+	const UStruct* Struct = Cast<UStruct>(Object);
+	const UEnum* Enum = Cast<UEnum>(Object);
+
+	if (Struct == nullptr && Enum == nullptr)
 	{
-		NativeScriptPackage = CreatePackage(nullptr, *NativePackageName);
+		return EReplacementResult::DontReplace;
 	}
-	bEnabled = true;
-}
 
-void DummyClassConstructor(const FObjectInitializer&)
-{
-	check(false);
-}
+	// check blacklists:
+	// we can't use FindObject, because we may be converting a type while saving
+	if ((Struct && ExcludedAssetTypes.Find(Struct->GetPathName()) != INDEX_NONE) ||
+		(Enum && ExcludedAssetTypes.Find(Enum->GetPathName()) != INDEX_NONE))
+	{
+		return EReplacementResult::GenerateStub;
+	}
 
-void FReplaceCookedBPGC::AdditionalStubFieldInitialization(UField* Stub)
-{
-	check(Stub && Stub->IsValidLowLevelFast());
-	if (auto Class = Cast<UClass>(Stub))
+	EReplacementResult Result = EReplacementResult::ReplaceCompletely;
+	while (Struct)
 	{
-		Class->ClassConstructor = &DummyClassConstructor;
-		Class->ClassAddReferencedObjects = &UObject::AddReferencedObjects;
+		// This happens because the cooker incorrectly cooks editor only packages. Specifically happens for the blackjack sample
+		// project due to a FStringAssetReference in BaseEditor.ini:
+		if (Struct->RootPackageHasAnyFlags(PKG_EditorOnly))
+		{
+			return EReplacementResult::DontReplace;
+		}
+
+		if (ExcludedBlueprintTypes.Find(Struct->GetPathName()) != INDEX_NONE)
+		{
+			Result = EReplacementResult::GenerateStub;
+		}
+		Struct = Struct->GetSuperStruct();
 	}
-	Stub->SetInternalFlags(EInternalObjectFlags::Native);
+
+	return Result;
 }
 
-#endif //WITH_EDITOR
\ No newline at end of file
+#endif // WITH_EDITOR
diff --git a/Engine/Source/Runtime/CoreUObject/Private/UObject/Class.cpp b/Engine/Source/Runtime/CoreUObject/Private/UObject/Class.cpp
--- a/Engine/Source/Runtime/CoreUObject/Private/UObject/Class.cpp
+++ b/Engine/Source/Runtime/CoreUObject/Private/UObject/Class.cpp
@@ -4291,7 +4291,7 @@ IMPLEMENT_CORE_INTRINSIC_CLASS(UClass, UStruct,
 	}
 );
 
-UClass::StaticClassFunctionType GetDynamicClassConstructFn(FName ClassName);
+UClass::StaticClassFunctionType GetDynamicClassConstructFn(FName ClassPathName);
 
 void GetPrivateStaticClassBody(
 	const TCHAR* PackageName,
@@ -4382,7 +4382,7 @@ void GetPrivateStaticClassBody(
 			InClassFlags,
 			InClassCastFlags,
 			InConfigName,
-			EObjectFlags(RF_Public | RF_Standalone | RF_Transient | RF_MarkAsNative | RF_Dynamic),
+			EObjectFlags(RF_Public | RF_Standalone | RF_Transient | RF_Dynamic | (GIsInitialLoad ? RF_MarkAsRootSet : RF_NoFlags)),
 			InClassConstructor,
 #if WITH_HOT_RELOAD_CTORS
 			InClassVTableHelperCtorCaller,
@@ -4405,7 +4405,7 @@ void GetPrivateStaticClassBody(
 	if (bIsDynamic)
 	{
 		// Now call the UHT-generated Z_Construct* function for the dynamic class
-		UClass::StaticClassFunctionType ZConstructDynamicClassFn = GetDynamicClassConstructFn(Name);
+		UClass::StaticClassFunctionType ZConstructDynamicClassFn = GetDynamicClassConstructFn(*ReturnClass->GetPathName());
 		check(ZConstructDynamicClassFn);
 		ZConstructDynamicClassFn();
 	}
@@ -4782,6 +4782,7 @@ UDynamicClass constructors.
 */
 UDynamicClass::UDynamicClass(const FObjectInitializer& ObjectInitializer)
 : UClass(ObjectInitializer)
+, AnimClassImplementation(nullptr)
 {
 	// If you add properties here, please update the other constructors and PurgeClass()
 }
@@ -4791,6 +4792,7 @@ UDynamicClass::UDynamicClass(const FObjectInitializer& ObjectInitializer)
 */
 UDynamicClass::UDynamicClass(const FObjectInitializer& ObjectInitializer, UClass* InBaseClass)
 : UClass(ObjectInitializer, InBaseClass)
+, AnimClassImplementation(nullptr)
 {
 }
 
@@ -4823,6 +4825,7 @@ UDynamicClass::UDynamicClass(
 , InClassVTableHelperCtorCaller
 #endif // WITH_HOT_RELOAD_CTORS
 , InClassAddReferencedObjects)
+, AnimClassImplementation(nullptr)
 {
 }
 
@@ -4837,6 +4840,8 @@ void UDynamicClass::AddReferencedObjects(UObject* InThis, FReferenceCollector& C
 	Collector.AddReferencedObjects(This->ComponentTemplates, This);
 	Collector.AddReferencedObjects(This->Timelines, This);
 
+	Collector.AddReferencedObject(This->AnimClassImplementation, This);
+
 	Super::AddReferencedObjects(This, Collector);
 }
 
@@ -4851,6 +4856,8 @@ void UDynamicClass::PurgeClass(bool bRecompilingOnLoad)
 	DynamicBindingObjects.Empty();
 	ComponentTemplates.Empty();
 	Timelines.Empty();
+
+	AnimClassImplementation = nullptr;
 }
 
 IMPLEMENT_CORE_INTRINSIC_CLASS(UDynamicClass, UClass,
diff --git a/Engine/Source/Runtime/CoreUObject/Private/UObject/Linker.cpp b/Engine/Source/Runtime/CoreUObject/Private/UObject/Linker.cpp
--- a/Engine/Source/Runtime/CoreUObject/Private/UObject/Linker.cpp
+++ b/Engine/Source/Runtime/CoreUObject/Private/UObject/Linker.cpp
@@ -76,15 +76,31 @@ void FGenerationInfo::Serialize(FArchive& Ar, const struct FPackageFileSummary&
 	Ar << ExportCount << NameCount;
 }
 
-FName FLinkerTables::GetExportClassName( int32 i )
+#if WITH_EDITORONLY_DATA
+extern int32 GLinkerAllowDynamicClasses;
+#endif
+
+FName FLinker::GetExportClassName( int32 i )
 {
 	if (ExportMap.IsValidIndex(i))
 	{
 		FObjectExport& Export = ExportMap[i];
-		if( !Export.ClassIndex.IsNull() )
+		if (!Export.ClassIndex.IsNull())
 		{
 			return ImpExp(Export.ClassIndex).ObjectName;
 		}
+#if WITH_EDITORONLY_DATA
+		else if (GLinkerAllowDynamicClasses && Export.bDynamicClass)
+		{
+			static FName NAME_BlueprintGeneratedClass(TEXT("BlueprintGeneratedClass"));
+			return NAME_BlueprintGeneratedClass;
+		}
+#else
+		else if (Export.bDynamicClass)
+		{
+			return GetDynamicTypeClassName(*GetExportPathName(i));
+		}
+#endif
 	}
 	return NAME_Class;
 }
@@ -434,6 +450,24 @@ static void LogGetPackageLinkerError(FArchiveUObject* LinkerArchive, const TCHAR
 	}
 }
 
+/** Customized version of FPackageName::DoesPackageExist that takes dynamic native class packages into account */
+static bool DoesPackageExistForGetPackageLinker(const FString& LongPackageName, const FGuid* Guid, FString& OutFilename, const bool ShouldGetLocalizedPackage)
+{
+	if (!ShouldGetLocalizedPackage && 
+#if WITH_EDITORONLY_DATA
+		GLinkerAllowDynamicClasses && 
+#endif
+		GetConvertedDynamicPackageNameToTypeName().Contains(*LongPackageName))
+	{
+		OutFilename = FPackageName::LongPackageNameToFilename(LongPackageName);
+		return true;
+	}
+	else
+	{
+		return FPackageName::DoesPackageExist(LongPackageName, Guid, &OutFilename, ShouldGetLocalizedPackage);
+	}
+}
+
 //
 // Find or create the linker for a package.
 //
@@ -470,8 +504,8 @@ FLinkerLoad* GetPackageLinker
 	
 		FString NativeFilename;
 		FString LocalizedFilename;
-		const bool DoesNativePackageExist = FPackageName::DoesPackageExist(InOuter->GetName(), CompatibleGuid, &NativeFilename, false);
-		const bool DoesLocalizedPackageExist = FPackageName::DoesPackageExist(InOuter->GetName(), CompatibleGuid, &LocalizedFilename, true);
+		const bool DoesNativePackageExist = DoesPackageExistForGetPackageLinker(InOuter->GetName(), CompatibleGuid, NativeFilename, false);
+		const bool DoesLocalizedPackageExist = DoesPackageExistForGetPackageLinker(InOuter->GetName(), CompatibleGuid, LocalizedFilename, true);
 
 		// If we are the editor, we must have a native package. If we are the game, we must have a localized package or a native package.
 		if ( (GIsEditor && !DoesNativePackageExist) || (!GIsEditor && !DoesLocalizedPackageExist && !DoesNativePackageExist) )
@@ -549,8 +583,8 @@ FLinkerLoad* GetPackageLinker
 		// Verify that the file exists.
 		FString NativeFilename;
 		FString LocalizedFilename;
-		const bool DoesNativePackageExist = FPackageName::DoesPackageExist(PackageName, CompatibleGuid, &NativeFilename, false);
-		const bool DoesLocalizedPackageExist = FPackageName::DoesPackageExist(PackageName, CompatibleGuid, &LocalizedFilename, true);
+		const bool DoesNativePackageExist = DoesPackageExistForGetPackageLinker(PackageName, CompatibleGuid, NativeFilename, false);
+		const bool DoesLocalizedPackageExist = DoesPackageExistForGetPackageLinker(PackageName, CompatibleGuid, LocalizedFilename, true);
 
 		if( (GIsEditor && !DoesNativePackageExist) || (!GIsEditor && !DoesLocalizedPackageExist && !DoesNativePackageExist) )
 		{
diff --git a/Engine/Source/Runtime/CoreUObject/Private/UObject/LinkerLoad.cpp b/Engine/Source/Runtime/CoreUObject/Private/UObject/LinkerLoad.cpp
--- a/Engine/Source/Runtime/CoreUObject/Private/UObject/LinkerLoad.cpp
+++ b/Engine/Source/Runtime/CoreUObject/Private/UObject/LinkerLoad.cpp
@@ -29,6 +29,16 @@ DECLARE_DWORD_ACCUMULATOR_STAT(TEXT("Linker Count"), STAT_LinkerCount, STATGROUP
 DECLARE_DWORD_ACCUMULATOR_STAT(TEXT("Live Linker Count"), STAT_LiveLinkerCount, STATGROUP_LinkerCount);
 DECLARE_FLOAT_ACCUMULATOR_STAT(TEXT("Fixup editor-only flags time"), STAT_EditorOnlyFixupTime, STATGROUP_LinkerCount);
 
+#if WITH_EDITORONLY_DATA
+int32 GLinkerAllowDynamicClasses = 0;
+static FAutoConsoleVariableRef CVarLinkerAllowDynamicClasses(
+	TEXT("linker.AllowDynamicClasses"),
+	GLinkerAllowDynamicClasses,
+	TEXT("If true, linkers will attempt to use dynamic classes instead of class assets."),
+	ECVF_Default
+	);
+#endif
+
 /** Map that keeps track of any precached full package reads															*/
 TMap<FString, FLinkerLoad::FPackagePrecacheInfo> FLinkerLoad::PackagePrecacheMap;
 
@@ -342,9 +352,9 @@ static bool IgnoreMissingReferencedClass( FName ClassName )
 	return MissingClassesToIgnore.Find( ClassName ) != INDEX_NONE;
 }
 
-static inline int32 HashNames( FName A, FName B, FName C )
+static inline int32 HashNames(FName Object, FName Class, FName Package)
 {
-	return A.GetComparisonIndex() + 7 * B.GetComparisonIndex() + 31 * FPackageName::GetShortFName(C).GetComparisonIndex();
+	return Object.GetComparisonIndex() + 7 * Class.GetComparisonIndex() + 31 * FPackageName::GetShortFName(Package).GetComparisonIndex();
 }
 
 static FORCEINLINE bool IsCoreUObjectPackage(const FName& PackageName)
@@ -709,6 +719,7 @@ FLinkerLoad::FLinkerLoad(UPackage* InParent, const TCHAR* InFilename, uint32 InL
 : FLinker(ELinkerType::Load, InParent, InFilename)
 , LoadFlags(InLoadFlags)
 , bHaveImportsBeenVerified(false)
+, bDynamicClassLinker(false)
 , Loader(nullptr)
 , AsyncRoot(nullptr)
 , NameMapIndex(0)
@@ -817,7 +828,7 @@ FLinkerLoad::ELinkerStatus FLinkerLoad::CreateLoader()
 	// This should have been initialized in InitUObject
 	check(bActiveRedirectsMapInitialized);
 
-	if( !Loader )
+	if( !Loader && !bDynamicClassLinker )
 	{
 		bool bIsSeekFree = LoadFlags & LOAD_SeekFree;
 
@@ -828,13 +839,22 @@ FLinkerLoad::ELinkerStatus FLinkerLoad::CreateLoader()
 		LoadProgressScope->EnterProgressFrame();
 #endif
 
+		// Check if this linker was created for dynamic class package
+		bDynamicClassLinker = GetConvertedDynamicPackageNameToTypeName().Contains(LinkerRoot->GetFName());
+		if (bDynamicClassLinker
+#if WITH_EDITORONLY_DATA
+			&& GLinkerAllowDynamicClasses
+#endif
+			)
+		{
+			// In this case we can skip serializing PackageFileSummary and fill all the required info here
+			CreateDynamicTypeLoader();
+		}
 		// NOTE: Precached memory read gets highest priority, then memory reader, then seek free, then normal
-
 		// check to see if there is was an async preload request for this file
-		FPackagePrecacheInfo* PrecacheInfo = PackagePrecacheMap.Find(*Filename);
-		// if so, serialize from memory (note this will have uncompressed a fully compressed package)
-		if (PrecacheInfo)
+		else if (FPackagePrecacheInfo* PrecacheInfo = PackagePrecacheMap.Find(*Filename))
 		{
+			// if so, serialize from memory (note this will have uncompressed a fully compressed package)
 			// block until the async read is complete
 			if( PrecacheInfo->SynchronizationObject->GetValue() != 0 )
 			{
@@ -908,8 +928,8 @@ FLinkerLoad::ELinkerStatus FLinkerLoad::CreateLoader()
 				return LINKER_Failed;
 			}
 		}
-		check( Loader );
-		check( !Loader->IsError() );
+		check(bDynamicClassLinker || Loader);
+		check(bDynamicClassLinker || !Loader->IsError());
 
 		//if( FLinkerLoad::FindExistingLinkerForPackage(LinkerRoot) )
 		//{
@@ -1326,36 +1346,6 @@ FLinkerLoad::ELinkerStatus FLinkerLoad::FixupImportMap()
 				for( int32 i=0; i<ImportMap.Num(); i++ )
 				{
 					FObjectImport& Import = ImportMap[i];
-
-#if WITH_EDITOR
-					// Fix references to assets that were converted into native code
-					
-					if (FReplaceConvertedAssetManager::Get().IsEnabled() && !Import.XObject)
-					{
-						const FString ImportPathName = GetImportPathName(i);
-						const FName PackageClassName(TEXT("Package"));
-						const bool bIsPackage = Import.ClassName == PackageClassName; // there must be a better way
-						if (bIsPackage)
-						{
-							if (auto NewNativePackage = FReplaceConvertedAssetManager::Get().FindPackageReplacement(ImportPathName))
-							{
-								Import.XObject = NewNativePackage;
-								Import.ObjectName = NewNativePackage->GetFName();
-								continue;
-							}
-						}
-						else if (auto ConvertedObject = FReplaceConvertedAssetManager::Get().FindReplacement(ImportPathName))
-						{
-							Import.ObjectName = ConvertedObject->GetFName();
-							Import.XObject = ConvertedObject;
-							Import.ClassName = ConvertedObject->GetClass()->GetFName();
-							auto ClassPackage = ConvertedObject->GetClass()->GetOuterUPackage();
-							Import.ClassPackage = ClassPackage ? ClassPackage->GetFName() : FName();
-							continue;
-						}
-					}
-#endif //WITH_EDITOR
-
 					{
 						FSubobjectRedirect* Redirect = SubobjectNameRedirects.Find(Import.ObjectName);
 						if (Redirect)
@@ -2017,6 +2007,13 @@ FName FLinkerLoad::GetExportClassPackage( int32 i )
 		// the export's class is contained within the same package
 		return LinkerRoot->GetFName();
 	}
+#if WITH_EDITORONLY_DATA
+	else if (GLinkerAllowDynamicClasses && Export.bDynamicClass)
+	{
+		static FName NAME_EnginePackage(TEXT("/Script/Engine"));
+		return NAME_EnginePackage;
+	}
+#endif
 	else
 	{
 		return GLongCoreUObjectPackageName;
@@ -2720,7 +2717,7 @@ bool FLinkerLoad::VerifyImportInner(const int32 ImportIndex, FString& WarningSuf
 					}
 				}
 
-				UObject* FindObject = StaticFindObject(FindClass, FindOuter, *Import.ObjectName.ToString());
+				UObject* FindObject = FindImport(FindClass, FindOuter, *Import.ObjectName.ToString());
 				// Reference to in memory-only package's object, native transient class or CDO of such a class.
 				bool bIsInMemoryOnlyOrNativeTransient = bCameFromMemoryOnlyPackage || (FindObject != NULL && (FindObject->HasAllFlags(EObjectFlags(RF_Public | RF_MarkAsNative | RF_Transient)) || (FindObject->HasAnyFlags(RF_ClassDefaultObject) && FindObject->GetClass()->HasAllFlags(EObjectFlags(RF_Public | RF_MarkAsNative | RF_Transient)))));
 				// Check for structs which have been moved to another header (within the same class package).
@@ -3454,6 +3451,17 @@ UObject* FLinkerLoad::CreateExport( int32 Index )
 		check(Export.ObjectName!=NAME_None || !(Export.ObjectFlags&RF_Public));
 		check(IsLoading());
 
+		if (Export.bDynamicClass)
+		{
+			// Export is a dynamic type, construct it using registered native functions
+			Export.Object = ConstructDynamicType(*GetExportPathName(Index));
+			if (Export.Object)
+			{
+				Export.Object->SetLinker(this, Index);
+			}
+			return Export.Object;
+		}
+
 		UClass* LoadClass = GetExportLoadClass(Index);
 		if( !LoadClass && !Export.ClassIndex.IsNull() ) // Hack to load packages with classes which do not exist.
 		{
@@ -3952,8 +3960,6 @@ bool FLinkerLoad::IsImportNative(const int32 Index) const
 	return bIsImportNative;
 }
 
-UObject* ConstructDynamicType(FName TypeName, FName TypeClass);
-
 // Return the loaded object corresponding to an import index; any errors are fatal.
 UObject* FLinkerLoad::CreateImport( int32 Index )
 {
@@ -3975,7 +3981,7 @@ UObject* FLinkerLoad::CreateImport( int32 Index )
 			// Try to find existing version in memory first.
 			if( UPackage* ClassPackage = FindObjectFast<UPackage>( NULL, Import.ClassPackage, false, false ) )
 			{
-				if( UClass*	FindClass = FindObjectFast<UClass>( ClassPackage, Import.ClassName, false, false ) )
+				if( UClass*	FindClass = FindObjectFast<UClass>( ClassPackage, Import.ClassName, false, false ) ) // 
 				{
 					// Make sure the class has been loaded and linked before creating a CDO.
 					// This is an edge case, but can happen if a blueprint package has not finished creating exports for a class
@@ -3994,14 +4000,6 @@ UObject* FLinkerLoad::CreateImport( int32 Index )
 					// Import is regular import/ export.
 					else
 					{
-						if (FPlatformProperties::RequiresCookedData())
-						{
-							// At this point we know the class package and class of the import exist, this can potentially be a native blueprint
-							// (class is just a UClass and the package is /Script/CoreUObject)
-							// @todo: maybe we can further filter out classes by checking FindClass somehow (e.g. once it's a specialized UCookedBlueprintClass or something)
-							ConstructDynamicType(Import.ObjectName, Import.ClassName);
-						}
-						
 						// Find the imports' outer.
 						UObject* FindOuter = NULL;
 						// Import.
@@ -4038,7 +4036,7 @@ UObject* FLinkerLoad::CreateImport( int32 Index )
 						}
 	
 						// Find object now that we know it's class, outer and name.
-						FindObject = StaticFindObjectFast( FindClass, FindOuter, Import.ObjectName, false, false );
+						FindObject = FindImportFast(FindClass, FindOuter, Import.ObjectName);
 					}
 
 					if( FindObject )
@@ -4267,7 +4265,7 @@ void FLinkerLoad::DetachAllBulkData(bool bEnsureAllBulkDataIsLoaded)
  */
 bool FLinkerLoad::Precache( int64 PrecacheOffset, int64 PrecacheSize )
 {
-	return Loader->Precache( PrecacheOffset, PrecacheSize );
+	return bDynamicClassLinker || Loader->Precache(PrecacheOffset, PrecacheSize);
 }
 
 void FLinkerLoad::Seek( int64 InPos )
diff --git a/Engine/Source/Runtime/CoreUObject/Private/UObject/LinkerPlaceholderBase.cpp b/Engine/Source/Runtime/CoreUObject/Private/UObject/LinkerPlaceholderBase.cpp
--- a/Engine/Source/Runtime/CoreUObject/Private/UObject/LinkerPlaceholderBase.cpp
+++ b/Engine/Source/Runtime/CoreUObject/Private/UObject/LinkerPlaceholderBase.cpp
@@ -305,11 +305,21 @@ int32 TLinkerImportPlaceholder<UClass>::ResolvePropertyReferences(UClass* Replac
 				++ReplacementCount;
 			}
 
-			UClassProperty* ClassProperty = Cast<UClassProperty>(BaseObjProperty);
-			if ((ClassProperty != nullptr) && (ClassProperty->MetaClass == PlaceholderClass))
+			if (UClassProperty* ClassProperty = Cast<UClassProperty>(BaseObjProperty))
 			{
-				ClassProperty->MetaClass = ReplacementClass;
-				++ReplacementCount;
+				if (ClassProperty->MetaClass == PlaceholderClass)
+				{
+					ClassProperty->MetaClass = ReplacementClass;
+					++ReplacementCount;
+				}
+			}
+			else if (UAssetClassProperty* AssetClassProperty = Cast<UAssetClassProperty>(BaseObjProperty))
+			{
+				if (AssetClassProperty->MetaClass == PlaceholderClass)
+				{
+					AssetClassProperty->MetaClass = ReplacementClass;
+					++ReplacementCount;
+				}	
 			}
 		}	
 		else if (UInterfaceProperty* InterfaceProp = Cast<UInterfaceProperty>(Property))
diff --git a/Engine/Source/Runtime/CoreUObject/Private/UObject/Obj.cpp b/Engine/Source/Runtime/CoreUObject/Private/UObject/Obj.cpp
--- a/Engine/Source/Runtime/CoreUObject/Private/UObject/Obj.cpp
+++ b/Engine/Source/Runtime/CoreUObject/Private/UObject/Obj.cpp
@@ -46,7 +46,7 @@ static UPackage*			GObjTransientPkg								= NULL;
 #endif
 
 UObject::UObject( EStaticConstructor, EObjectFlags InFlags )
-: UObjectBaseUtility(InFlags | RF_MarkAsNative | RF_MarkAsRootSet)
+: UObjectBaseUtility(InFlags | (!(InFlags & RF_Dynamic) ? (RF_MarkAsNative | RF_MarkAsRootSet) : RF_NoFlags))
 {
 #if WITH_HOT_RELOAD_CTORS
 	EnsureNotRetrievingVTablePtr();
diff --git a/Engine/Source/Runtime/CoreUObject/Private/UObject/ObjectResource.cpp b/Engine/Source/Runtime/CoreUObject/Private/UObject/ObjectResource.cpp
--- a/Engine/Source/Runtime/CoreUObject/Private/UObject/ObjectResource.cpp
+++ b/Engine/Source/Runtime/CoreUObject/Private/UObject/ObjectResource.cpp
@@ -44,6 +44,7 @@ FObjectExport::FObjectExport()
 , bNotForEditorGame(true)
 , bIsAsset(false)
 , bExportLoadFailed(false)
+, bDynamicClass(false)
 , PackageGuid(FGuid(0, 0, 0, 0))
 , PackageFlags(0)
 {}
@@ -63,6 +64,7 @@ FObjectExport::FObjectExport( UObject* InObject )
 , bNotForEditorGame(true)
 , bIsAsset(false)
 , bExportLoadFailed(false)
+, bDynamicClass(false)
 , PackageGuid(FGuid(0, 0, 0, 0))
 , PackageFlags(0)
 {
@@ -131,6 +133,16 @@ FObjectImport::FObjectImport( UObject* InObject )
 {
 }
 
+FObjectImport::FObjectImport(UObject* InObject, UClass* InClass)
+:	FObjectResource	( InObject																)
+,	ClassPackage	( (InObject && InClass) ? InClass->GetOuter()->GetFName()	: NAME_None	)
+,	ClassName		( (InObject && InClass) ? InClass->GetFName()				: NAME_None	)
+,	XObject			( InObject																)
+,	SourceLinker	( NULL																	)
+,	SourceIndex		( INDEX_NONE															)
+{
+}
+
 FArchive& operator<<( FArchive& Ar, FObjectImport& I )
 {
 	Ar << I.ClassPackage << I.ClassName;
diff --git a/Engine/Source/Runtime/CoreUObject/Private/UObject/Package.cpp b/Engine/Source/Runtime/CoreUObject/Private/UObject/Package.cpp
--- a/Engine/Source/Runtime/CoreUObject/Private/UObject/Package.cpp
+++ b/Engine/Source/Runtime/CoreUObject/Private/UObject/Package.cpp
@@ -192,8 +192,9 @@ bool UPackage::IsFullyLoaded()
 	{
 		FString DummyFilename;
 		// Try to find matching package in package file cache.
-		if( !FPackageName::DoesPackageExist( *GetName(), NULL, &DummyFilename ) 
-		||	(GIsEditor && IFileManager::Get().FileSize(*DummyFilename) < 0) )
+		if (!GetConvertedDynamicPackageNameToTypeName().Contains(GetFName()) &&
+			  !FPackageName::DoesPackageExist( *GetName(), NULL, &DummyFilename ) || 
+				(GIsEditor && IFileManager::Get().FileSize(*DummyFilename) < 0) )
 		{
 			// Package has NOT been found, so we assume it's a newly created one and therefore fully loaded.
 			bHasBeenFullyLoaded = true;
diff --git a/Engine/Source/Runtime/CoreUObject/Private/UObject/Property.cpp b/Engine/Source/Runtime/CoreUObject/Private/UObject/Property.cpp
--- a/Engine/Source/Runtime/CoreUObject/Private/UObject/Property.cpp
+++ b/Engine/Source/Runtime/CoreUObject/Private/UObject/Property.cpp
@@ -202,6 +202,7 @@ struct TStructOpsTypeTraits<FRandomStream> : public TStructOpsTypeTraitsBase
 {
 	enum 
 	{
+		WithExportTextItem = true,
 		WithNoInitConstructor = true,
 		WithZeroConstructor = true,
 	};
@@ -348,15 +349,12 @@ UProperty::UProperty(const FObjectInitializer& ObjectInitializer, ECppProperty,
 
 void UProperty::Init()
 {
-	// properties created in C++ should always be marked RF_Transient so that when the package containing
-	// this property is saved, it doesn't try to save this UProperty into the ExportMap
-	SetFlags(RF_Transient);
-	SetInternalFlags(EInternalObjectFlags::Native);
 #if !WITH_EDITORONLY_DATA
 	[REDACTED_EMAIL]/PACKAGER: Until we have a cooker/packager step, this can fire when WITH_EDITORONLY_DATA is not defined!
 	//	checkSlow(!HasAnyPropertyFlags(CPF_EditorOnly));
 #endif // WITH_EDITORONLY_DATA
-	checkSlow(GetOuterUField()->HasAllFlags(RF_Transient) && GetOuterUField()->IsNative());
+	checkSlow(GetOuterUField()->HasAllFlags(RF_Transient));
+	checkSlow(HasAllFlags(RF_Transient));
 
 	GetOuterUField()->AddCppProperty(this);
 }
@@ -505,7 +503,7 @@ void UProperty::ExportCppDeclaration(FOutputDevice& Out, EExportedDeclaration::T
 	{
 		if (AdditionalExportCPPFlags & CPPF_BlueprintCppBackend)
 		{
-			NameCpp = UnicodeToCPPIdentifier(GetName(), HasAnyPropertyFlags(CPF_Deprecated), TEXT("bpv__"));
+			NameCpp = UnicodeToCPPIdentifier(GetName(), HasAnyPropertyFlags(CPF_Deprecated), HasAnyPropertyFlags(CPF_Parm) ? TEXT("bpp__") : TEXT("bpv__"));
 		}
 		else
 		{
@@ -1119,7 +1117,7 @@ const TCHAR* UProperty::ImportSingleProperty( const TCHAR* Str, void* DestData,
 			// strip whitespace after =
 			SkipWhitespace(Str);
 
-			if (!IsPropertyValueSpecified(Str))
+			if (!IsPropertyValueSpecified(Str) && ArrayProperty == nullptr)
 			{
 				// if we're not importing default properties for classes (i.e. we're pasting something in the editor or something)
 				// and there is no property value for this element, skip it, as that means that the value of this element matches
@@ -1150,7 +1148,7 @@ const TCHAR* UProperty::ImportSingleProperty( const TCHAR* Str, void* DestData,
 							Warn->Logf(ELogVerbosity::Warning,*ImportErrors[ErrorIndex]);
 						}
 					}
-					else if (Result == NULL || Result == Str)
+					else if ((Result == NULL && ArrayProperty == nullptr) || Result == Str)
 					{
 						Warn->Logf(ELogVerbosity::Warning, TEXT("Invalid property value in defaults: %s"), Start);
 					}
@@ -1182,7 +1180,7 @@ const TCHAR* UProperty::ImportSingleProperty( const TCHAR* Str, void* DestData,
 							Warn->Logf(ELogVerbosity::Warning, *ImportErrors[ErrorIndex]);
 						}
 					}
-					else if (Result == NULL || Result == Str)
+					else if ((Result == NULL && ArrayProperty == nullptr) || Result == Str)
 					{
 						Warn->Logf(ELogVerbosity::Warning, TEXT("Invalid property value in defaults: %s"), Start);
 					}
diff --git a/Engine/Source/Runtime/CoreUObject/Private/UObject/PropertyArray.cpp b/Engine/Source/Runtime/CoreUObject/Private/UObject/PropertyArray.cpp
--- a/Engine/Source/Runtime/CoreUObject/Private/UObject/PropertyArray.cpp
+++ b/Engine/Source/Runtime/CoreUObject/Private/UObject/PropertyArray.cpp
@@ -209,7 +209,7 @@ const TCHAR* UArrayProperty::ImportText_Internal( const TCHAR* Buffer, void* Dat
 
 	// If we export an empty array we export an empty string, so ensure that if we're passed an empty string
 	// we interpret it as an empty array.
-	if ( *Buffer == TCHAR('\0') )
+	if (*Buffer == TCHAR('\0') || *Buffer == TCHAR(')') || *Buffer == TCHAR(','))
 	{
 		ArrayHelper.EmptyValues();
 		return NULL;
diff --git a/Engine/Source/Runtime/CoreUObject/Private/UObject/PropertyAssetClass.cpp b/Engine/Source/Runtime/CoreUObject/Private/UObject/PropertyAssetClass.cpp
--- a/Engine/Source/Runtime/CoreUObject/Private/UObject/PropertyAssetClass.cpp
+++ b/Engine/Source/Runtime/CoreUObject/Private/UObject/PropertyAssetClass.cpp
@@ -1,19 +1,34 @@
 // Copyright 1998-2015 Epic Games, Inc. All Rights Reserved.
 
 #include "CoreUObjectPrivate.h"
+#include "LinkerPlaceholderClass.h"
 
 /*-----------------------------------------------------------------------------
 	UAssetClassProperty.
 -----------------------------------------------------------------------------*/
 
+void UAssetClassProperty::BeginDestroy()
+{
+#if USE_CIRCULAR_DEPENDENCY_LOAD_DEFERRING
+	if (ULinkerPlaceholderClass* PlaceholderClass = Cast<ULinkerPlaceholderClass>(MetaClass))
+	{
+		PlaceholderClass->RemoveReferencingProperty(this);
+	}
+#endif // USE_CIRCULAR_DEPENDENCY_LOAD_DEFERRING
+
+	Super::BeginDestroy();
+}
+
 FString UAssetClassProperty::GetCPPType(FString* ExtendedTypeText, uint32 CPPExportFlags) const
 {
-	return GetCPPTypeCustom(ExtendedTypeText, CPPExportFlags, MetaClass);
+	check(MetaClass);
+	return GetCPPTypeCustom(ExtendedTypeText, CPPExportFlags, 
+		FString::Printf(TEXT("%s%s"), MetaClass->GetPrefixCPP(), *MetaClass->GetName()));
 }
-FString UAssetClassProperty::GetCPPTypeCustom(FString* ExtendedTypeText, uint32 CPPExportFlags, UClass* ActualClass) const
+FString UAssetClassProperty::GetCPPTypeCustom(FString* ExtendedTypeText, uint32 CPPExportFlags, const FString& InnerNativeTypeName) const
 {
-	check(ActualClass);
-	return FString::Printf(TEXT("TAssetSubclassOf<%s%s> "), ActualClass->GetPrefixCPP(), *ActualClass->GetName());
+	ensure(!InnerNativeTypeName.IsEmpty());
+	return FString::Printf(TEXT("TAssetSubclassOf<%s> "), *InnerNativeTypeName);
 }
 FString UAssetClassProperty::GetCPPMacroType( FString& ExtendedTypeText ) const
 {
@@ -31,6 +46,16 @@ void UAssetClassProperty::Serialize( FArchive& Ar )
 	Super::Serialize( Ar );
 	Ar << MetaClass;
 
+#if USE_CIRCULAR_DEPENDENCY_LOAD_DEFERRING
+	if (Ar.IsLoading() || Ar.IsObjectReferenceCollector())
+	{
+		if (ULinkerPlaceholderClass* PlaceholderClass = Cast<ULinkerPlaceholderClass>(MetaClass))
+		{
+			PlaceholderClass->AddReferencingProperty(this);
+		}
+	}
+#endif // USE_CIRCULAR_DEPENDENCY_LOAD_DEFERRING
+
 	if( !(MetaClass||HasAnyFlags(RF_ClassDefaultObject)) )
 	{
 		// If we failed to load the MetaClass and we're not a CDO, that means we relied on a class that has been removed or doesn't exist.
@@ -44,6 +69,22 @@ void UAssetClassProperty::Serialize( FArchive& Ar )
 	}
 }
 
+#if USE_CIRCULAR_DEPENDENCY_LOAD_DEFERRING
+void UAssetClassProperty::SetMetaClass(UClass* NewMetaClass)
+{
+	if (ULinkerPlaceholderClass* NewPlaceholderClass = Cast<ULinkerPlaceholderClass>(NewMetaClass))
+	{
+		NewPlaceholderClass->AddReferencingProperty(this);
+	}
+
+	if (ULinkerPlaceholderClass* OldPlaceholderClass = Cast<ULinkerPlaceholderClass>(MetaClass))
+	{
+		OldPlaceholderClass->RemoveReferencingProperty(this);
+	}
+	MetaClass = NewMetaClass;
+}
+#endif // USE_CIRCULAR_DEPENDENCY_LOAD_DEFERRING
+
 void UAssetClassProperty::AddReferencedObjects(UObject* InThis, FReferenceCollector& Collector)
 {
 	UAssetClassProperty* This = CastChecked<UAssetClassProperty>(InThis);
diff --git a/Engine/Source/Runtime/CoreUObject/Private/UObject/PropertyAssetObject.cpp b/Engine/Source/Runtime/CoreUObject/Private/UObject/PropertyAssetObject.cpp
--- a/Engine/Source/Runtime/CoreUObject/Private/UObject/PropertyAssetObject.cpp
+++ b/Engine/Source/Runtime/CoreUObject/Private/UObject/PropertyAssetObject.cpp
@@ -6,10 +6,10 @@
 	UAssetObjectProperty.
 -----------------------------------------------------------------------------*/
 
-FString UAssetObjectProperty::GetCPPTypeCustom(FString* ExtendedTypeText, uint32 CPPExportFlags, UClass* ActualClass) const
+FString UAssetObjectProperty::GetCPPTypeCustom(FString* ExtendedTypeText, uint32 CPPExportFlags, const FString& InnerNativeTypeName) const
 {
-	check(ActualClass);
-	return FString::Printf(TEXT("TAssetPtr<%s%s>"), ActualClass->GetPrefixCPP(), *ActualClass->GetName());
+	ensure(!InnerNativeTypeName.IsEmpty());
+	return FString::Printf(TEXT("TAssetPtr<%s>"), *InnerNativeTypeName);
 }
 FString UAssetObjectProperty::GetCPPMacroType( FString& ExtendedTypeText ) const
 {
diff --git a/Engine/Source/Runtime/CoreUObject/Private/UObject/PropertyByte.cpp b/Engine/Source/Runtime/CoreUObject/Private/UObject/PropertyByte.cpp
--- a/Engine/Source/Runtime/CoreUObject/Private/UObject/PropertyByte.cpp
+++ b/Engine/Source/Runtime/CoreUObject/Private/UObject/PropertyByte.cpp
@@ -98,8 +98,14 @@ FString UByteProperty::GetCPPType( FString* ExtendedTypeText/*=NULL*/, uint32 CP
 			// been set, but we do this here in case existing code relies on it... somehow.
 			FullyQualifiedEnumName = Enum->GetName();
 		}
-
-		if ((CPPExportFlags & CPPF_ArgumentOrReturnValue) && ((PropertyFlags & CPF_ReturnParm) || !(PropertyFlags & CPF_OutParm)))
+		 
+		const bool bEnumClassForm = Enum->GetCppForm() == UEnum::ECppForm::EnumClass;
+		const bool bNonNativeEnum = Enum->GetClass() != UEnum::StaticClass(); // cannot use RF_Native flag, because in UHT the flag is not set
+		const bool bRawParam = (CPPExportFlags & CPPF_ArgumentOrReturnValue) 
+			&& (((PropertyFlags & CPF_ReturnParm) || !(PropertyFlags & CPF_OutParm)) 
+			|| bEnumClassForm || bNonNativeEnum);
+		const bool bConvertedCode = (CPPExportFlags & CPPF_BlueprintCppBackend) && (bEnumClassForm || bNonNativeEnum);
+		if (bRawParam || bConvertedCode)
 		{
 			return FullyQualifiedEnumName;
 		}
@@ -120,8 +126,17 @@ void UByteProperty::ExportTextItem( FString& ValueStr, const void* PropertyValue
 			const int32 ActualValue = *(const uint8*)PropertyValue;
 			const int32 MaxValue = Enum->GetMaxEnumValue();
 			const int32 GoodValue = Enum->IsValidEnumValue(ActualValue) ? ActualValue : MaxValue;
-			ValueStr += FString::Printf(TEXT("%s::%s"), *Enum->GetName(),
-				*Enum->GetEnumName(Enum->GetIndexByValue(GoodValue)));
+			if (GoodValue == MaxValue)
+			{
+				// not all native enums have Max value declared
+				const FString FullyQualifiedEnumName = Enum->CppType.IsEmpty() ? Enum->GetName() : Enum->CppType;
+				ValueStr += FString::Printf(TEXT("(%s)(%d)"), *FullyQualifiedEnumName, ActualValue);
+			}
+			else
+			{
+				ValueStr += FString::Printf(TEXT("%s::%s"), *Enum->GetName(),
+					*Enum->GetEnumName(Enum->GetIndexByValue(GoodValue)));
+			}
 		}
 		else
 		{
diff --git a/Engine/Source/Runtime/CoreUObject/Private/UObject/PropertyClass.cpp b/Engine/Source/Runtime/CoreUObject/Private/UObject/PropertyClass.cpp
--- a/Engine/Source/Runtime/CoreUObject/Private/UObject/PropertyClass.cpp
+++ b/Engine/Source/Runtime/CoreUObject/Private/UObject/PropertyClass.cpp
@@ -91,15 +91,17 @@ const TCHAR* UClassProperty::ImportText_Internal( const TCHAR* Buffer, void* Dat
 
 FString UClassProperty::GetCPPType(FString* ExtendedTypeText, uint32 CPPExportFlags) const
 {
-	return GetCPPTypeCustom(ExtendedTypeText, CPPExportFlags, MetaClass);
+	check(MetaClass);
+	return GetCPPTypeCustom(ExtendedTypeText, CPPExportFlags,
+		FString::Printf(TEXT("%s%s"), MetaClass->GetPrefixCPP(), *MetaClass->GetName()));
 }
 
-FString UClassProperty::GetCPPTypeCustom(FString* ExtendedTypeText, uint32 CPPExportFlags, UClass* ActualClass) const
+FString UClassProperty::GetCPPTypeCustom(FString* ExtendedTypeText, uint32 CPPExportFlags, const FString& InnerNativeTypeName) const
 {
 	if (PropertyFlags & CPF_UObjectWrapper)
 	{
-		check(ActualClass);
-		return FString::Printf(TEXT("TSubclassOf<%s%s> "), ActualClass->GetPrefixCPP(), *ActualClass->GetName());
+		ensure(!InnerNativeTypeName.IsEmpty());
+		return FString::Printf(TEXT("TSubclassOf<%s> "), *InnerNativeTypeName);
 	}
 	else
 	{
diff --git a/Engine/Source/Runtime/CoreUObject/Private/UObject/PropertyDelegate.cpp b/Engine/Source/Runtime/CoreUObject/Private/UObject/PropertyDelegate.cpp
--- a/Engine/Source/Runtime/CoreUObject/Private/UObject/PropertyDelegate.cpp
+++ b/Engine/Source/Runtime/CoreUObject/Private/UObject/PropertyDelegate.cpp
@@ -75,16 +75,30 @@ bool UDelegateProperty::NetSerializeItem( FArchive& Ar, UPackageMap* Map, void*
 FString UDelegateProperty::GetCPPType( FString* ExtendedTypeText/*=NULL*/, uint32 CPPExportFlags/*=0*/ ) const
 {
 	FString UnmangledFunctionName = SignatureFunction->GetName().LeftChop( FString( HEADER_GENERATED_DELEGATE_SIGNATURE_SUFFIX ).Len() );
-	if ((0 != (CPPExportFlags & EPropertyExportCPPFlags::CPPF_BlueprintCppBackend)) && !SignatureFunction->GetOwnerClass()->HasAnyClassFlags(CLASS_Native))
+	const bool bBlueprintCppBackend = (0 != (CPPExportFlags & EPropertyExportCPPFlags::CPPF_BlueprintCppBackend));
+	const bool bNative = SignatureFunction->IsNative();
+	if (bBlueprintCppBackend && bNative)
 	{
-		// the name must be unique
-		const FString OwnerName = UnicodeToCPPIdentifier(SignatureFunction->GetOwnerClass()->GetName(), false, TEXT(""));
-		const FString NewUnmangledFunctionName = FString::Printf(TEXT("%s__%s"), *UnmangledFunctionName, *OwnerName);
-		UnmangledFunctionName = NewUnmangledFunctionName;
+		UStruct* StructOwner = Cast<UStruct>(SignatureFunction->GetOuter());
+		if (StructOwner)
+		{
+			return FString::Printf(TEXT("%s%s::F%s"), StructOwner->GetPrefixCPP(), *StructOwner->GetName(), *UnmangledFunctionName);
+		}
 	}
-	if (0 != (CPPExportFlags & EPropertyExportCPPFlags::CPPF_CustomTypeName))
+	else
 	{
-		UnmangledFunctionName += TEXT("__SinglecastDelegate");
+		const bool NonNativeClassOwner = (SignatureFunction->GetOwnerClass() && !SignatureFunction->GetOwnerClass()->HasAnyClassFlags(CLASS_Native));
+		if (bBlueprintCppBackend && NonNativeClassOwner)
+		{
+			// the name must be unique
+			const FString OwnerName = UnicodeToCPPIdentifier(SignatureFunction->GetOwnerClass()->GetName(), false, TEXT(""));
+			const FString NewUnmangledFunctionName = FString::Printf(TEXT("%s__%s"), *UnmangledFunctionName, *OwnerName);
+			UnmangledFunctionName = NewUnmangledFunctionName;
+		}
+		if (0 != (CPPExportFlags & EPropertyExportCPPFlags::CPPF_CustomTypeName))
+		{
+			UnmangledFunctionName += TEXT("__SinglecastDelegate");
+		}
 	}
 	return FString(TEXT("F")) + UnmangledFunctionName;
 }
diff --git a/Engine/Source/Runtime/CoreUObject/Private/UObject/PropertyMap.cpp b/Engine/Source/Runtime/CoreUObject/Private/UObject/PropertyMap.cpp
--- a/Engine/Source/Runtime/CoreUObject/Private/UObject/PropertyMap.cpp
+++ b/Engine/Source/Runtime/CoreUObject/Private/UObject/PropertyMap.cpp
@@ -464,6 +464,7 @@ void UMapProperty::ExportTextItem(FString& ValueStr, const void* PropertyValue,
 
 	if (MapHelper.Num() == 0)
 	{
+		ValueStr += TEXT("()");
 		return;
 	}
 
diff --git a/Engine/Source/Runtime/CoreUObject/Private/UObject/PropertyMulticastDelegate.cpp b/Engine/Source/Runtime/CoreUObject/Private/UObject/PropertyMulticastDelegate.cpp
--- a/Engine/Source/Runtime/CoreUObject/Private/UObject/PropertyMulticastDelegate.cpp
+++ b/Engine/Source/Runtime/CoreUObject/Private/UObject/PropertyMulticastDelegate.cpp
@@ -123,10 +123,11 @@ FString UMulticastDelegateProperty::GetCPPType( FString* ExtendedTypeText/*=NULL
 #endif
 
 	FString UnmangledFunctionName = SignatureFunction->GetName().LeftChop( FString( HEADER_GENERATED_DELEGATE_SIGNATURE_SUFFIX ).Len() );
-	if ((0 != (CPPExportFlags & EPropertyExportCPPFlags::CPPF_BlueprintCppBackend)) && !SignatureFunction->GetOwnerClass()->HasAnyClassFlags(CLASS_Native))
+	const UClass* OwnerClass = SignatureFunction->GetOwnerClass();
+	if ((0 != (CPPExportFlags & EPropertyExportCPPFlags::CPPF_BlueprintCppBackend)) && OwnerClass && !OwnerClass->HasAnyClassFlags(CLASS_Native))
 	{
 		// the name must be unique
-		const FString OwnerName = UnicodeToCPPIdentifier(SignatureFunction->GetOwnerClass()->GetName(), false, TEXT(""));
+		const FString OwnerName = UnicodeToCPPIdentifier(OwnerClass->GetName(), false, TEXT(""));
 		const FString NewUnmangledFunctionName = FString::Printf(TEXT("%s__%s"), *UnmangledFunctionName, *OwnerName);
 		UnmangledFunctionName = NewUnmangledFunctionName;
 	}
diff --git a/Engine/Source/Runtime/CoreUObject/Private/UObject/PropertyObject.cpp b/Engine/Source/Runtime/CoreUObject/Private/UObject/PropertyObject.cpp
--- a/Engine/Source/Runtime/CoreUObject/Private/UObject/PropertyObject.cpp
+++ b/Engine/Source/Runtime/CoreUObject/Private/UObject/PropertyObject.cpp
@@ -8,10 +8,10 @@
 	UObjectProperty.
 -----------------------------------------------------------------------------*/
 
-FString UObjectProperty::GetCPPTypeCustom(FString* ExtendedTypeText, uint32 CPPExportFlags, UClass* ActualClass)  const
+FString UObjectProperty::GetCPPTypeCustom(FString* ExtendedTypeText, uint32 CPPExportFlags, const FString& InnerNativeTypeName)  const
 {
-	check(ActualClass);
-	return FString::Printf(TEXT("%s%s*"), ActualClass->GetPrefixCPP(), *ActualClass->GetName());
+	ensure(!InnerNativeTypeName.IsEmpty());
+	return FString::Printf(TEXT("%s*"), *InnerNativeTypeName);
 }
 
 FString UObjectProperty::GetCPPTypeForwardDeclaration() const
diff --git a/Engine/Source/Runtime/CoreUObject/Private/UObject/SavePackage.cpp b/Engine/Source/Runtime/CoreUObject/Private/UObject/SavePackage.cpp
--- a/Engine/Source/Runtime/CoreUObject/Private/UObject/SavePackage.cpp
+++ b/Engine/Source/Runtime/CoreUObject/Private/UObject/SavePackage.cpp
@@ -314,13 +314,11 @@ public:
 	 */
 	void AddReplacementsNames(UObject* Obj)
 	{
-		if (Obj && FReplaceCookedBPGC::Get().IsEnabled())
+		if (FScriptCookReplacementCoordinator* Coordinator = FScriptCookReplacementCoordinator::Get())
 		{
-			if (UObject* ReplObj = FReplaceCookedBPGC::Get().FindReplacementStub(Obj))
+			if (const UClass* ReplObjClass = Coordinator->FindReplacedClass(Obj))
 			{
-				MarkNameAsReferenced(ReplObj->GetFName());
-				MarkNameAsReferenced(ReplObj->GetClass()->GetFName());
-				MarkNameAsReferenced(ReplObj->GetOuter()->GetFName());
+				MarkNameAsReferenced(ReplObjClass->GetFName());
 			}
 		}
 	}
@@ -1582,7 +1580,7 @@ public:
 				else
 				{
 					// this import no longer exists in the new package
-					new(Imports)FObjectImport( NULL );
+					new(Imports)FObjectImport( nullptr );
 				}
 			}
 
@@ -3118,55 +3116,6 @@ TMap<UObject*, UObject*> UnmarkExportTagFromDuplicates()
 
 COREUOBJECT_API extern bool GOutputCookingWarnings;
 
-/**
- * Find imports to converted BPGC, that should replaced by a generated native class. 
- * For each converted class a stub is generated. The import object is replaced with the stub.
- */
-struct FImportReplacementsHelper
-{
-	TSet<UObject*> UniqueTagImpObjects;
-
-	FImportReplacementsHelper(const TArray<UObject*>& InTagImpObjects)
-	{
-		if (FReplaceCookedBPGC::Get().IsEnabled())
-		{
-			UniqueTagImpObjects = TSet<UObject*>(InTagImpObjects);
-		}
-	}
-
-	UObject* Handle(UObject* Obj, TArray<FObjectImport>& ImportMap)
-	{
-		if (FReplaceCookedBPGC::Get().IsEnabled())
-		{
-			UObject* ReplacedObjectToImport = FReplaceCookedBPGC::Get().FindReplacementStub(Obj);
-			Obj = ReplacedObjectToImport ? ReplacedObjectToImport : Obj;
-			if (ReplacedObjectToImport)
-			{
-				{
-					ensure(ReplacedObjectToImport->GetOuter() == ReplacedObjectToImport->GetOutermost());
-					bool bAlreadyAdded = false;
-					UniqueTagImpObjects.Add(ReplacedObjectToImport->GetOuter(), &bAlreadyAdded);
-					if (!bAlreadyAdded)
-					{
-						new(ImportMap)FObjectImport(ReplacedObjectToImport->GetOuter());
-					}
-				}
-
-									{
-										bool bAlreadyAdded = false;
-										UniqueTagImpObjects.Add(ReplacedObjectToImport->GetClass(), &bAlreadyAdded);
-										if (!bAlreadyAdded)
-										{
-											new(ImportMap)FObjectImport(ReplacedObjectToImport->GetClass());
-										}
-									}
-			}
-		}
-		return Obj;
-	}
-
-};
-
 class FDiffSerializeArchive : public FBufferArchive
 {
 private:
@@ -3243,7 +3192,7 @@ ESavePackageResult UPackage::Save(UPackage* InOuter, UObject* Base, EObjectFlags
 	const class ITargetPlatform* TargetPlatform, const FDateTime&  FinalTimeStamp, bool bSlowTask)
 {
 	UE_START_LOG_COOK_TIME( Filename );
-	
+
 	if (FPlatformProperties::HasEditorOnlyData())
 	{
 		if (GIsSavingPackage)
@@ -3271,14 +3220,6 @@ ESavePackageResult UPackage::Save(UPackage* InOuter, UObject* Base, EObjectFlags
 				UE_CLOG(!(SaveFlags & SAVE_NoError), LogSavePackage, Display, TEXT("Package marked as editor-only: %s. Package will not be saved."), *InOuter->GetName());
 				return ESavePackageResult::ReferencedOnlyByEditorOnlyData;
 			}
-#if WITH_EDITOR
-			else if (FReplaceCookedBPGC::Get().PackageShouldNotBeSaved(InOuter))
-			{
-				UE_LOG(LogSavePackage, Display, TEXT("Package %s contains assets, that were converted into native code. Package will not be saved."), *InOuter->GetName());
-				// Or replace with the NativeScriptPackage ?
-				return ESavePackageResult::ContainsConvertedAssets;
-			}
-#endif	// #if WITH_EDITOR
 		}
 #endif
 		// if we are cooking we should be doing it in the editor
@@ -3426,6 +3367,7 @@ ESavePackageResult UPackage::Save(UPackage* InOuter, UObject* Base, EObjectFlags
 		SlowTask.EnterProgressFrame();
 
 		bool Success = true;
+		bool bRequestStub = false;
 
 		ResetLoadersForSave(InOuter,Filename);
 
@@ -3625,6 +3567,23 @@ ESavePackageResult UPackage::Save(UPackage* InOuter, UObject* Base, EObjectFlags
 						UE_CLOG(!(SaveFlags & SAVE_NoError), LogSavePackage, Display, TEXT("No exports found (or all exports are editor-only) for %s. Package will not be saved."), *BaseFilename);
 						return ESavePackageResult::ContainsEditorOnlyData;
 					}
+
+#if WITH_EDITOR
+					if (FScriptCookReplacementCoordinator::Get())
+					{
+						EReplacementResult ReplacmentResult = FScriptCookReplacementCoordinator::Get()->IsTargetedForReplacement(InOuter);
+						if (ReplacmentResult == EReplacementResult::ReplaceCompletely)
+						{
+							return ESavePackageResult::ReplaceCompletely;
+						}
+						else if (ReplacmentResult == EReplacementResult::GenerateStub)
+						{
+							bRequestStub = true;
+						}
+						UE_LOG(LogSavePackage, Display, TEXT("Package %s contains assets, that were converted into native code. Package will not be saved."), *InOuter->GetName());
+						// Or replace with the NativeScriptPackage ?
+					}
+#endif
 				}
 
 				// Import objects & names.
@@ -3665,18 +3624,6 @@ ESavePackageResult UPackage::Save(UPackage* InOuter, UObject* Base, EObjectFlags
 						{
 							UE_LOG(LogSavePackage, Fatal, TEXT("%s"), *FString::Printf( TEXT("Transient object imported: %s"), *Obj->GetFullName() ) );
 						}
-#if WITH_EDITOR
-						if (FReplaceCookedBPGC::Get().IsEnabled())
-						{
-							for (UObject*& Dependency : ImportTagger.Dependencies)
-							{
-								if (UObject* NewDependency = FReplaceCookedBPGC::Get().FindReplacementStub(Dependency))
-								{
-									Dependency = NewDependency;
-								}
-							}
-						}
-#endif //WITH_EDITOR
 
 						if (Linker->IsCooking())
 						{
@@ -3985,6 +3932,7 @@ ESavePackageResult UPackage::Save(UPackage* InOuter, UObject* Base, EObjectFlags
 					InOuter->Guid = Linker->Summary.Guid;
 				}
 				new(Linker->Summary.Generations)FGenerationInfo(0, 0);
+
 				*Linker << Linker->Summary;
 				int32 OffsetAfterPackageFileSummary = Linker->Tell();
 		
@@ -4079,9 +4027,6 @@ ESavePackageResult UPackage::Save(UPackage* InOuter, UObject* Base, EObjectFlags
 				{
 					TArray<UObject*> TagImpObjects;
 					GetObjectsWithAnyMarks(TagImpObjects, OBJECTMARK_TagImp);
-#if WITH_EDITOR
-					FImportReplacementsHelper ImportReplacementsHelper(TagImpObjects);
-#endif //WITH_EDITOR
 
 					if (Linker->IsCooking())
 					{
@@ -4099,13 +4044,25 @@ ESavePackageResult UPackage::Save(UPackage* InOuter, UObject* Base, EObjectFlags
 					{
 						UObject* Obj = TagImpObjects[Index];
 						check(Obj->HasAnyMarks(OBJECTMARK_TagImp) || Linker->IsCooking());
-						
+						UClass* ObjClass = nullptr;
 #if WITH_EDITOR
-						Obj = ImportReplacementsHelper.Handle(Obj, Linker->ImportMap);
+						bool bReplaced = false;
+						if (FScriptCookReplacementCoordinator* Coordinator = FScriptCookReplacementCoordinator::Get())
+						{
+							if (UClass* ReplacedClass = Coordinator->FindReplacedClass(Obj))
+							{
+								ObjClass = ReplacedClass;
+								bReplaced = true;
+							}
+						}
+						if (!bReplaced)
 #endif //WITH_EDITOR
+						{
+							ObjClass = Obj->GetClass();
+						}
 						if (Obj->HasAnyMarks(OBJECTMARK_TagImp))
 						{
-							new( Linker->ImportMap )FObjectImport( Obj );
+							new(Linker->ImportMap)FObjectImport(Obj, ObjClass);
 						}
 					}
 				}
@@ -4352,7 +4309,9 @@ ESavePackageResult UPackage::Save(UPackage* InOuter, UObject* Base, EObjectFlags
 					UObject* Object = Linker->ImportMap[i].XObject;
 					if( Object != NULL )
 					{
-						Linker->ObjectIndicesMap.Add(Object, FPackageIndex::FromImport(i));
+						const FPackageIndex PackageIndex = FPackageIndex::FromImport(i);
+						//ensure(!Linker->ObjectIndicesMap.Contains(Object)); // this ensure will fail
+						Linker->ObjectIndicesMap.Add(Object, PackageIndex);
 					}
 					else
 					{
@@ -4489,16 +4448,6 @@ ESavePackageResult UPackage::Save(UPackage* InOuter, UObject* Base, EObjectFlags
 							// Set class index.
 							// If this is *exactly* a UClass, store null instead; for anything else, including UClass-derived classes, map it
 							UClass* ObjClass = Export.Object->GetClass();
-#if WITH_EDITOR
-							if (FReplaceCookedBPGC::Get().IsEnabled())
-							{
-								if (UObject* NewObjClass = FReplaceCookedBPGC::Get().FindReplacementStub(ObjClass))
-								{
-									ObjClass = CastChecked<UClass>(NewObjClass);
-								}
-							}
-#endif	// #if WITH_EDITOR
-
 							if (ObjClass != UClass::StaticClass())
 							{
 								Export.ClassIndex = Linker->MapObject(ObjClass);
@@ -4935,7 +4884,25 @@ ESavePackageResult UPackage::Save(UPackage* InOuter, UObject* Base, EObjectFlags
 
 		UE_LOG(LogSavePackage, Display, TEXT("Finished SavePackage %s"), Filename);
 
-		return Success ? ESavePackageResult::Success : ESavePackageResult::Error;
+		if (Success)
+		{
+			if (bRequestStub)
+			{
+				return ESavePackageResult::GenerateStub;
+			}
+			else
+			{
+				return ESavePackageResult::Success;
+			}
+		}
+		else
+		{
+			if (bRequestStub)
+			{
+				UE_LOG(LogSavePackage, Warning, TEXT("C++ stub requested, but package failed to save, may cause compile errors: %s"), Filename);
+			}
+			return ESavePackageResult::Error;
+		}
 	}
 	else
 	{
diff --git a/Engine/Source/Runtime/CoreUObject/Private/UObject/ScriptCore.cpp b/Engine/Source/Runtime/CoreUObject/Private/UObject/ScriptCore.cpp
--- a/Engine/Source/Runtime/CoreUObject/Private/UObject/ScriptCore.cpp
+++ b/Engine/Source/Runtime/CoreUObject/Private/UObject/ScriptCore.cpp
@@ -62,6 +62,8 @@ COREUOBJECT_API void GInitRunaway() {}
 // FBlueprintCoreDelegates
 
 FBlueprintCoreDelegates::FOnScriptDebuggingEvent FBlueprintCoreDelegates::OnScriptException;
+FBlueprintCoreDelegates::FOnScriptInstrumentEvent FBlueprintCoreDelegates::OnScriptProfilingEvent;
+FBlueprintCoreDelegates::FOnToggleScriptProfiler FBlueprintCoreDelegates::OnToggleScriptProfiler;
 
 void FBlueprintCoreDelegates::ThrowScriptException(const UObject* ActiveObject, const FFrame& StackFrame, const FBlueprintExceptionInfo& Info)
 {
@@ -114,6 +116,11 @@ void FBlueprintCoreDelegates::ThrowScriptException(const UObject* ActiveObject,
 	}
 }
 
+void FBlueprintCoreDelegates::InstrumentScriptEvent(const EScriptInstrumentationEvent& Info)
+{
+	OnScriptProfilingEvent.Broadcast(Info);
+}
+
 void FBlueprintCoreDelegates::SetScriptMaximumLoopIterations( const int32 MaximumLoopIterations )
 {
 	if (ensure(MaximumLoopIterations > 0))
@@ -159,17 +166,17 @@ FString ToValidCPPIdentifierChars(TCHAR Char)
 		RawValue = (RawValue - Digit) / 63;
 
 		TCHAR SafeChar;
-		if (Digit < 26)
+		if (Digit <= 25)
 		{
-			SafeChar = TCHAR(TCHAR('a') + (26 - Digit));
+			SafeChar = TCHAR(TCHAR('a') + (25 - Digit));
 		}
-		else if (Digit < 52)
+		else if (Digit <= 51)
 		{
-			SafeChar = TCHAR(TCHAR('A') + (52 - Digit));
+			SafeChar = TCHAR(TCHAR('A') + (51 - Digit));
 		}
-		else if (Digit < 62)
+		else if (Digit <= 61)
 		{
-			SafeChar = TCHAR(TCHAR('0') + (62 - Digit));
+			SafeChar = TCHAR(TCHAR('0') + (61 - Digit));
 		}
 		else
 		{
@@ -1011,6 +1018,14 @@ void UObject::ProcessEvent( UFunction* Function, void* Parms )
 	}
 	checkSlow((Function->ParmsSize == 0) || (Parms != NULL));
 
+#if !(UE_BUILD_SHIPPING || UE_BUILD_TEST)
+	if (GetClass()->HasInstrumentation())
+	{
+		EScriptInstrumentationEvent EventInstrumentationInfo(EScriptInstrumentation::Event, this);
+		FBlueprintCoreDelegates::InstrumentScriptEvent(EventInstrumentationInfo);
+	}
+#endif
+
 #if DO_BLUEPRINT_GUARD
 	FBlueprintExceptionTracker& BlueprintExceptionTracker = FBlueprintExceptionTracker::Get();
 	BlueprintExceptionTracker.ScriptEntryTag++;
@@ -1139,6 +1154,14 @@ void UObject::ProcessEvent( UFunction* Function, void* Parms )
 		}
 	}
 
+#if !(UE_BUILD_SHIPPING || UE_BUILD_TEST)
+	if (GetClass()->HasInstrumentation())
+	{
+		EScriptInstrumentationEvent EventInstrumentationInfo(EScriptInstrumentation::Stop, this);
+		FBlueprintCoreDelegates::InstrumentScriptEvent(EventInstrumentationInfo);
+	}
+#endif
+
 #if DO_BLUEPRINT_GUARD
 	--BlueprintExceptionTracker.ScriptEntryTag;
 #endif
@@ -1384,6 +1407,31 @@ void UObject::execWireTracepoint( FFrame& Stack, RESULT_DECL )
 }
 IMPLEMENT_VM_FUNCTION( EX_WireTracepoint, execWireTracepoint );
 
+void UObject::execInstrumentation( FFrame& Stack, RESULT_DECL )
+{
+#if !UE_BUILD_SHIPPING
+	const EScriptInstrumentation::Type EventType = static_cast<EScriptInstrumentation::Type>(Stack.ReadInt());
+#if WITH_EDITORONLY_DATA
+	if (GIsEditor)
+	{
+		if (EventType == EScriptInstrumentation::NodeEntry)
+		{
+			static FBlueprintExceptionInfo TracepointExceptionInfo(EBlueprintExceptionType::Tracepoint);
+			FBlueprintCoreDelegates::ThrowScriptException(this, Stack, TracepointExceptionInfo);
+		}
+		else if (EventType == EScriptInstrumentation::NodeExit)
+		{
+			static FBlueprintExceptionInfo WiretraceExceptionInfo(EBlueprintExceptionType::WireTracepoint);
+			FBlueprintCoreDelegates::ThrowScriptException(this, Stack, WiretraceExceptionInfo);
+		}
+	}
+#endif
+	EScriptInstrumentationEvent InstrumentationEventInfo(EventType, this, Stack);
+	FBlueprintCoreDelegates::InstrumentScriptEvent(InstrumentationEventInfo);
+#endif
+}
+IMPLEMENT_VM_FUNCTION( EX_InstrumentationEvent, execInstrumentation );
+
 void UObject::execEndFunctionParms( FFrame& Stack, RESULT_DECL )
 {
 	// For skipping over optional function parms without values specified.
diff --git a/Engine/Source/Runtime/CoreUObject/Private/UObject/UObjectBase.cpp b/Engine/Source/Runtime/CoreUObject/Private/UObject/UObjectBase.cpp
--- a/Engine/Source/Runtime/CoreUObject/Private/UObject/UObjectBase.cpp
+++ b/Engine/Source/Runtime/CoreUObject/Private/UObject/UObjectBase.cpp
@@ -390,7 +390,7 @@ TMap<FName, UScriptStruct *(*)()>& GetDynamicStructMap()
 	return DynamicStructMap;
 }
 
-void UObjectCompiledInDeferStruct(class UScriptStruct *(*InRegister)(), const TCHAR* PackageName, const TCHAR* Name, bool bDynamic)
+void UObjectCompiledInDeferStruct(class UScriptStruct *(*InRegister)(), const TCHAR* PackageName, const FName PathName, bool bDynamic)
 {
 	if (!bDynamic)
 	{
@@ -401,7 +401,7 @@ void UObjectCompiledInDeferStruct(class UScriptStruct *(*InRegister)(), const TC
 	}
 	else
 	{
-		GetDynamicStructMap().Add(FName(Name), InRegister);
+		GetDynamicStructMap().Add(PathName, InRegister);
 	}
 }
 
@@ -443,7 +443,7 @@ TMap<FName, UEnum *(*)()>& GetDynamicEnumMap()
 	return DynamicEnumMap;
 }
 
-void UObjectCompiledInDeferEnum(class UEnum *(*InRegister)(), const TCHAR* PackageName, const TCHAR* Name, bool bDynamic)
+void UObjectCompiledInDeferEnum(class UEnum *(*InRegister)(), const TCHAR* PackageName, const FName PathName, bool bDynamic)
 {
 	if (!bDynamic)
 	{
@@ -454,7 +454,7 @@ void UObjectCompiledInDeferEnum(class UEnum *(*InRegister)(), const TCHAR* Packa
 	}
 	else
 	{
-		GetDynamicEnumMap().Add(FName(Name), InRegister);
+		GetDynamicEnumMap().Add(PathName, InRegister);
 	}
 }
 
@@ -612,7 +612,7 @@ TMap<FName, FClassConstructFunctions>& GetDynamicClassMap()
 	return DynamicClassMap;
 }
 
-void UObjectCompiledInDefer(UClass *(*InRegister)(), UClass *(*InStaticClass)(), const TCHAR* Name, bool bDynamic)
+void UObjectCompiledInDefer(UClass *(*InRegister)(), UClass *(*InStaticClass)(), const TCHAR* Name, bool bDynamic, const TCHAR* DynamicPathName)
 {
 	if (!bDynamic)
 	{
@@ -630,7 +630,7 @@ void UObjectCompiledInDefer(UClass *(*InRegister)(), UClass *(*InStaticClass)(),
 		FClassConstructFunctions ClassFunctions;
 		ClassFunctions.ZConstructFn = InRegister;
 		ClassFunctions.StaticClassFn = InStaticClass;
-		GetDynamicClassMap().Add(FName(Name), ClassFunctions);
+		GetDynamicClassMap().Add(FName(DynamicPathName), ClassFunctions);
 	}
 }
 
@@ -1107,40 +1107,28 @@ namespace
 
 UScriptStruct* FindExistingStructIfHotReloadOrDynamic(UObject* Outer, const TCHAR* StructName, SIZE_T Size, uint32 Crc, bool bIsDynamic)
 {
-	if (UScriptStruct* Output = FindExistingStructOrEnumIfHotReload<UScriptStruct>(Outer, StructName, Size, Crc))
+	UScriptStruct* Result = FindExistingStructOrEnumIfHotReload<UScriptStruct>(Outer, StructName, Size, Crc);
+	if (!Result && bIsDynamic)
 	{
-		return Output;
+		Result = Cast<UScriptStruct>(StaticFindObjectFast(UScriptStruct::StaticClass(), Outer, StructName));
 	}
-
-	if (bIsDynamic)
-	{
-		return Cast<UScriptStruct>(StaticFindObjectFast(UScriptStruct::StaticClass(), Outer, StructName));
-	}
-
-	return nullptr;
-}
-
-UEnum* FindExistingEnumIfHotReload(UObject* Outer, const TCHAR* EnumName, SIZE_T Size, uint32 Crc)
-{
-	return FindExistingStructOrEnumIfHotReload<UEnum>(Outer, EnumName, Size, Crc);
+	return Result;
 }
 
-UClass* ConstructDynamicClass(FName ClassName)
+UEnum* FindExistingEnumIfHotReloadOrDynamic(UObject* Outer, const TCHAR* EnumName, SIZE_T Size, uint32 Crc, bool bIsDynamic)
 {
-	UClass* Result = nullptr;
-	FClassConstructFunctions* ClassConstructFn = GetDynamicClassMap().Find(ClassName);
-	if (ClassConstructFn)
+	UEnum* Result = FindExistingStructOrEnumIfHotReload<UEnum>(Outer, EnumName, Size, Crc);
+	if (!Result && bIsDynamic)
 	{
-		// This calls the class' StaticClass() function which calls GetPrivateStaticClassBody() which calls Z_Construct* function.
-		Result = ClassConstructFn->StaticClassFn();
+		Result = Cast<UEnum>(StaticFindObjectFast(UEnum::StaticClass(), Outer, EnumName));
 	}
 	return Result;
 }
 
-UClass::StaticClassFunctionType GetDynamicClassConstructFn(FName ClassName)
+UClass::StaticClassFunctionType GetDynamicClassConstructFn(FName ClassPathName)
 {
-	FClassConstructFunctions* ClassConstructFn = GetDynamicClassMap().Find(ClassName);
-	UE_CLOG(!ClassConstructFn, LogUObjectBase, Fatal, TEXT("Unable to find construct function pointer for dynamic class %s. Make sure dynamic class exists."), *ClassName.ToString());
+	FClassConstructFunctions* ClassConstructFn = GetDynamicClassMap().Find(ClassPathName);
+	UE_CLOG(!ClassConstructFn, LogUObjectBase, Fatal, TEXT("Unable to find construct function pointer for dynamic class %s. Make sure dynamic class exists."), *ClassPathName.ToString());
 	if (ClassConstructFn)
 	{
 		return ClassConstructFn->ZConstructFn;
@@ -1149,42 +1137,38 @@ UClass::StaticClassFunctionType GetDynamicClassConstructFn(FName ClassName)
 	return nullptr;
 }
 
-UScriptStruct* ConstructDynamicStruct(FName StructName)
+UObject* ConstructDynamicType(FName TypePathName)
 {
-	UScriptStruct* Result = nullptr;
-	UScriptStruct *(**StaticStructFNPtr)() = GetDynamicStructMap().Find(StructName);
-	if (StaticStructFNPtr)
+	UObject* Result = nullptr;
+	if (FClassConstructFunctions* ClassConstructFn = GetDynamicClassMap().Find(TypePathName))
+	{
+		Result = ClassConstructFn->StaticClassFn();
+	}
+	else if (UScriptStruct *(**StaticStructFNPtr)() = GetDynamicStructMap().Find(TypePathName))
 	{
 		Result = (*StaticStructFNPtr)();
 	}
-	return Result;
-}
-
-UEnum* ConstructDynamicEnum(FName EnumName)
-{
-	UEnum* Result = nullptr;
-	UEnum *(**StaticEnumFNPtr)() = GetDynamicEnumMap().Find(EnumName);
-	if (StaticEnumFNPtr)
+	else if (UEnum *(**StaticEnumFNPtr)() = GetDynamicEnumMap().Find(TypePathName))
 	{
 		Result = (*StaticEnumFNPtr)();
 	}
 	return Result;
 }
 
-UObject* ConstructDynamicType(FName TypeName, FName TypeClass)
+FName GetDynamicTypeClassName(FName TypePathName)
 {
-	UObject* Result = nullptr;
-	if ((TypeClass == UClass::StaticClass()->GetFName()) || (TypeClass == UDynamicClass::StaticClass()->GetFName()))
+	FName Result = NAME_None;
+	if (GetDynamicClassMap().Find(TypePathName))
 	{
-		Result = ConstructDynamicClass(TypeName);
+		Result = UDynamicClass::StaticClass()->GetFName();
 	}
-	else if (TypeClass == UScriptStruct::StaticClass()->GetFName())
+	else if (GetDynamicStructMap().Find(TypePathName))
 	{
-		Result = ConstructDynamicStruct(TypeName);
+		Result = UScriptStruct::StaticClass()->GetFName();
 	}
-	else if (TypeClass == UEnum::StaticClass()->GetFName())
+	else if (GetDynamicEnumMap().Find(TypePathName))
 	{
-		Result = ConstructDynamicEnum(TypeName);
+		Result = UEnum::StaticClass()->GetFName();
 	}
 	return Result;
 }
@@ -1199,4 +1183,10 @@ UPackage* FindOrConstructDynamicTypePackage(const TCHAR* PackageName)
 	}
 	check(Package);
 	return Package;
+}
+
+TMap<FName, FName>& GetConvertedDynamicPackageNameToTypeName()
+{
+	static TMap<FName, FName> ConvertedDynamicPackageNameToTypeName;
+	return ConvertedDynamicPackageNameToTypeName;
 }
\ No newline at end of file
diff --git a/Engine/Source/Runtime/CoreUObject/Private/UObject/UObjectGlobals.cpp b/Engine/Source/Runtime/CoreUObject/Private/UObject/UObjectGlobals.cpp
--- a/Engine/Source/Runtime/CoreUObject/Private/UObject/UObjectGlobals.cpp
+++ b/Engine/Source/Runtime/CoreUObject/Private/UObject/UObjectGlobals.cpp
@@ -2010,7 +2010,11 @@ UObject* StaticAllocateObject
 		{
 			// Should only get in here if we're NOT creating a subobject of a CDO.  CDO subobjects may still need to be serialized off of disk after being created by the constructor
 			// if really necessary there was code to allow replacement of object just needing postload, but lets not go there unless we have to
-			checkf(!Obj->HasAnyFlags(RF_NeedLoad|RF_NeedPostLoad|RF_ClassDefaultObject), *FText::Format( NSLOCTEXT( "Core", "ReplaceNotFullyLoaded_f", "Attempting to replace an object that hasn't been fully loaded: {0}" ), FText::FromString( Obj->GetFullName() ) ).ToString() );
+			checkf(!Obj->HasAnyFlags(RF_NeedLoad|RF_NeedPostLoad|RF_ClassDefaultObject),
+				*FText::Format(NSLOCTEXT("Core", "ReplaceNotFullyLoaded_f", "Attempting to replace an object that hasn't been fully loaded: {0} (Outer={1}, Flags={2})"),
+					FText::FromString(Obj->GetFullName()),
+					InOuter ? FText::FromString(InOuter->GetFullName()) : FText::FromString(TEXT("NULL")),
+					FText::FromString(FString::Printf(TEXT("0x%08x"), (int32)Obj->GetFlags()))).ToString());
 		}
 		// Subobjects are always created in the constructor, no need to re-create them here unless their archetype != CDO or they're blueprint generated.	
 		if (!bCreatingCDO && (!bCanRecycleSubobjects || !Obj->IsDefaultSubobject()))
diff --git a/Engine/Source/Runtime/CoreUObject/Public/Blueprint/BlueprintSupport.h b/Engine/Source/Runtime/CoreUObject/Public/Blueprint/BlueprintSupport.h
--- a/Engine/Source/Runtime/CoreUObject/Public/Blueprint/BlueprintSupport.h
+++ b/Engine/Source/Runtime/CoreUObject/Public/Blueprint/BlueprintSupport.h
@@ -2,6 +2,8 @@
 
 #pragma once
 
+#include "GCObject.h"
+
 /** 
  * This set of functions contains blueprint related UObject functionality.
  */
@@ -64,108 +66,90 @@ private:
 	FScopedClassDependencyGather();
 };
 
-/**
- * Stores info about BPCG, UDS, UDE, etc.. converted into native code.
- *
- * Native entities (generated from BP items) have "ReplaceConverted" with original object path.
- * This path is used to fix linker while loading.
- * 
- * This struct is used only when the code converted from BP is compiled in an editor build. 
- * It's usually only for a test/development purpose.
- */
-struct COREUOBJECT_API FReplaceConvertedAssetManager : public FGCObject
+enum class EReplacementResult
 {
-private:
-	TMap<FString, UObject*> ReplaceMap;
-	bool bIsEnabled;
-
-	FReplaceConvertedAssetManager();
+	/** Don't replace the provided package at all */
+	DontReplace,
 
-	void GatherOriginalPathsOfConvertedAssets();
-
-public:
-
-	static FReplaceConvertedAssetManager& Get();
-
-	void SetEnabled(bool InEnabled)
-	{
-		bIsEnabled = InEnabled;
-		if (bIsEnabled)
-		{
-			GatherOriginalPathsOfConvertedAssets();
-		}
-	}
-
-	bool IsEnabled()
-	{
-		return bIsEnabled;
-	}
+	/** Generate a stub file, but don't replace the package */
+	GenerateStub,
 
-	virtual void AddReferencedObjects(FReferenceCollector& Collector) override;
-
-	UPackage* FindPackageReplacement(const FString& OriginalPathName) const;
-
-	UObject* FindReplacement(const FString& OriginalPathName) const;
-
-	static void OnModulesChanged(FName ModuleThatChanged, EModuleChangeReason ReasonForChange);
+	/** Completely replace the file with generated code */
+	ReplaceCompletely
 };
 
 /**
- * The struct is used while saving cooked package to find replacements of converted BP assets.
+ * The struct is used while saving cooked package to find replacements for 
+ * converted Blueprint assets.
  */
-struct COREUOBJECT_API FReplaceCookedBPGC : public FGCObject
+struct COREUOBJECT_API FScriptCookReplacementCoordinator
 {
-private:
-	UPackage* NativeScriptPackage;
-	TMap<FString, UObject*> ReplaceMap;
-	TSet<FName> ConvertedPackagesNames;
-	bool bEnabled;
-
-	FReplaceCookedBPGC();
-
-	void AdditionalStubFieldInitialization(UField* Stub);
+	static FScriptCookReplacementCoordinator* Get();
+	static void Create(bool bEnabled, const TArray<FString>& ExcludedAssetTypes, const TArray<FString>& ExcludedBlueprintTypes, const TMap<UObject*, UClass*>& ReplacementMap);
 
 public:
-	static FReplaceCookedBPGC& Get();
-	virtual void AddReferencedObjects(FReferenceCollector& Collector) override;
-
-	// INITIALIZATION
-	void Initialize(const FString& NativePackageName);
-	void AddConvertedPackageName(FName PackageName);
-	template<typename FieldType>
-	void AddConvertedFieldStub(const FString& OriginalPathName, const FName Name)
-	{
-		UObject*& NativeObj = ReplaceMap.FindOrAdd(OriginalPathName);
-		if (!NativeObj)
-		{
-			check(NativeScriptPackage);
-			FieldType* NewField = NewObject<FieldType>(NativeScriptPackage, Name);
-			AdditionalStubFieldInitialization(NewField);
-			NewField->Bind();
-			NativeObj = NewField;
-		}
-	}
-	void AddConvertedFieldStub(UObject* Object);
+	/**
+	 * Enables this for handling (script) asset replacements (works only when  
+	 * running a commandlet - presumably the cooker).
+	 * 
+	 * @return True if this was successfully enabled, otherwise false.
+	 */
+	bool Initialize();
 
-	// USAGE
+	/** 
+	 * Checks to see if we're running with this on, aiming to swap out assets 
+	 * with native counterparts.
+	 */
 	bool IsEnabled() const
 	{
 		return bEnabled;
 	}
-	UObject* FindReplacementStub(UObject* Object);
 
-	// The object is of a convertible type, and the object is not listed in excluded assets.
-	bool CouldBeConverted(const UObject* Object) const;
-	bool PackageShouldNotBeSaved(const UPackage* InOuter) const;
-
-	DECLARE_DELEGATE_RetVal_OneParam(bool, FIsTargetedForConversionQuery, const UObject*);
-	FIsTargetedForConversionQuery& OnQueryIfAssetIsTargetedForConversion() { return IsTargetedForConversionDelegate; }
+	/**
+	 * Determines whether the provided package needs to be replaced (in part or completely)
+	 * 
+	 * @param Package	The package in question
+	 * @return Whether the package should be converted
+	 */
+	EReplacementResult IsTargetedForReplacement(const UPackage* Package) const;
+	
+	/**
+	* Determines whether the provided object needs to be replaced (in part or completely).
+	* Some objects in a package may require conversion and some may not. If any object 
+	* in a package wants to be converted then it is implied that all other objects will 
+	* be converted with it (no support for partial package conversion, beyond stubs)
+	*
+	* @param Object	The package in question
+	* @return Whether the object should be converted
+	*/
+	EReplacementResult IsTargetedForReplacement(const UObject* Object) const;
+
+	// Get class of converted asset. One that was specified in AddConvertedObject
+	UClass* FindReplacedClass(const UObject* Obj) const;
 
 private: 
-	FIsTargetedForConversionQuery IsTargetedForConversionDelegate;
+	/** Private so we can keep this as a singleton */
+	FScriptCookReplacementCoordinator(bool bEnabled, const TArray<FString>& ExcludedAssetTypes, const TArray<FString>& ExcludedBlueprintTypes, const TMap<UObject*, UClass*>& ReplacementMap);
+
+	/** 
+	 * Used to tell if replacements should be used (substituting native objects 
+	 * for assets). 
+	 */
+	bool bEnabled;
+
+	const TArray<FString> ExcludedAssetTypes;
+	const TArray<FString> ExcludedBlueprintTypes;
+
+	/** Tracks which assets has a replaced class */
+	TMap<UObject*, UClass*> ReplacementMap;
+
+	TSet<FString> ReplacedPackages;
+
+	/** singleton instance: */
+	static FScriptCookReplacementCoordinator* CoordinatorInstance;
 };
 
-#endif //WITH_EDITOR
+#endif // WITH_EDITOR
 
 /** 
  * A helper struct for storing FObjectInitializers that were not run on 
@@ -217,41 +201,30 @@ private:
 };
 
 
+struct COREUOBJECT_API FBlueprintDependencyData
+{
+	FName PackageName;
+	FName ObjectName;
+	FName ClassPackageName;
+	FName ClassName;
+};
+
 /**
  *	Stores info about dependencies of native classes converted from BPs
  */
 struct COREUOBJECT_API FConvertedBlueprintsDependencies
 {
-	typedef void(*GetDependenciesNamesFunc)(TArray<FName>&);
+	typedef void(*GetDependenciesNamesFunc)(TArray<FBlueprintDependencyData>&);
 
 private:
 
-	TMap<FName, GetDependenciesNamesFunc> ClassNameToGetter;
+	TMap<FName, GetDependenciesNamesFunc> PackageNameToGetter;
 
 public:
-
-#if WITH_EDITOR
-
-	typedef UField*(*StaticCreateClassFunc)();
-
-	TArray<StaticCreateClassFunc> CreateClassFunctions;
-
-	void RegisterClass(StaticCreateClassFunc CreateClass)
-	{
-		CreateClassFunctions.Add(CreateClass);
-	}
-
-#endif //WITH_EDITOR
-
 	static FConvertedBlueprintsDependencies& Get();
 
-	bool AnyClassRegistered() const
-	{
-		return ClassNameToGetter.Num() > 0;
-	}
-
-	void RegisterClass(FName ClassName, GetDependenciesNamesFunc GetAssets);
+	void RegisterClass(FName PackageName, GetDependenciesNamesFunc GetAssets);
 
 	/** Get all assets paths necessary for the class with the given class name and all converted classes that dependencies. */
-	void GetAssets(FName ClassName, TArray<FName>& OutPackagePaths) const;
+	void GetAssets(FName PackageName, TArray<FBlueprintDependencyData>& OutDependencies) const;
 };
\ No newline at end of file
diff --git a/Engine/Source/Runtime/CoreUObject/Public/UObject/Class.h b/Engine/Source/Runtime/CoreUObject/Public/UObject/Class.h
--- a/Engine/Source/Runtime/CoreUObject/Public/UObject/Class.h
+++ b/Engine/Source/Runtime/CoreUObject/Public/UObject/Class.h
@@ -2414,6 +2414,10 @@ public:
 	virtual void GetRequiredPreloadDependencies(TArray<UObject*>& DependenciesOut) {}
 
 	virtual UObject* GetArchetypeForCDO() const;
+
+	/** Returns true if this class implements script instrumentation. */
+	virtual bool HasInstrumentation() const { return false; }
+
 private:
 	#if UCLASS_FAST_ISA_IMPL & 2
 		// For UObjectBaseUtility
@@ -2493,6 +2497,9 @@ public:
 	TArray<UObject*> DynamicBindingObjects;
 	TArray<UObject*> ComponentTemplates;
 	TArray<UObject*> Timelines;
+
+	// IAnimClassInterface (UAnimClassData) or null
+	UObject* AnimClassImplementation;
 };
 
 /**
diff --git a/Engine/Source/Runtime/CoreUObject/Public/UObject/ConstructorHelpers.h b/Engine/Source/Runtime/CoreUObject/Public/UObject/ConstructorHelpers.h
--- a/Engine/Source/Runtime/CoreUObject/Public/UObject/ConstructorHelpers.h
+++ b/Engine/Source/Runtime/CoreUObject/Public/UObject/ConstructorHelpers.h
@@ -79,7 +79,12 @@ namespace ConstructorHelpersInternal
 				PathName += TCHAR('C');
 			}
 		}
-		return StaticLoadClass(BaseClass, NULL, *PathName);
+		UClass* LoadedClass = StaticLoadClass(BaseClass, NULL, *PathName);
+		if (LoadedClass)
+		{
+			LoadedClass->AddToRoot();
+		}
+		return LoadedClass;
 	}
 }
 
diff --git a/Engine/Source/Runtime/CoreUObject/Public/UObject/CoreObject.h b/Engine/Source/Runtime/CoreUObject/Public/UObject/CoreObject.h
--- a/Engine/Source/Runtime/CoreUObject/Public/UObject/CoreObject.h
+++ b/Engine/Source/Runtime/CoreUObject/Public/UObject/CoreObject.h
@@ -35,7 +35,9 @@ enum class ESavePackageResult
 	/** [When cooking] Package was not saved because it was referenced by editor-only properties */
 	ReferencedOnlyByEditorOnlyData, 
 	/** [When cooking] Package was not saved because it contains assets that were converted into native code */
-	ContainsConvertedAssets
+	ReplaceCompletely,
+	/** [When cooking] Package was saved, but we should generate a stub so that other converted packages can interface with it*/
+	GenerateStub
 };
 
 /**
@@ -45,7 +47,7 @@ PRAGMA_DISABLE_DEPRECATION_WARNINGS // Required for auto-generated functions ref
 class COREUOBJECT_API UPackage : public UObject
 {
 	// Have to unwind this macro to support the reference variable, can go back to commented declaration when removing the deprecated variable
-	// DECLARE_CASTED_CLASS_INTRINSIC(UPackage,UObject,0,CoreUObject, CASTCLASS_UPackage)
+	// DECLARE_CASTED_CLASS_INTRINSIC(UPackage, UObject, 0, TEXT("/Script/CoreUObject"), CASTCLASS_UPackage)
 
 #if WITH_HOT_RELOAD_CTORS
 	DECLARE_CASTED_CLASS_INTRINSIC_NO_CTOR_NO_VTABLE_CTOR( UPackage, UObject, 0, TEXT("/Script/CoreUObject"), CASTCLASS_UPackage, NO_API )
@@ -56,7 +58,7 @@ class COREUOBJECT_API UPackage : public UObject
 	{
 	};
 #else
-	DECLARE_CASTED_CLASS_INTRINSIC_NO_CTOR(UPackage,UObject,0,TEXT("/Script/CoreUObject"), CASTCLASS_UPackage, NO_API)
+	DECLARE_CASTED_CLASS_INTRINSIC_NO_CTOR(UPackage, UObject, 0, TEXT("/Script/CoreUObject"), CASTCLASS_UPackage, NO_API)
 #endif
 
 
diff --git a/Engine/Source/Runtime/CoreUObject/Public/UObject/Linker.h b/Engine/Source/Runtime/CoreUObject/Public/UObject/Linker.h
--- a/Engine/Source/Runtime/CoreUObject/Public/UObject/Linker.h
+++ b/Engine/Source/Runtime/CoreUObject/Public/UObject/Linker.h
@@ -275,6 +275,11 @@
 	*/
 	bool			bExportLoadFailed;
 
+	/**
+	* Export is a dynamic class.
+	*/
+	bool			bDynamicClass;
+
 	/** If this object is a top level package (which must have been forced into the export table via OBJECTMARK_ForceTagExp)
 	 * this is the GUID for the original package file
 	 * Serialized
@@ -343,6 +348,7 @@
 	 */
 	COREUOBJECT_API FObjectImport();
 	FObjectImport( UObject* InObject );
+	FObjectImport( UObject* InObject, UClass* InClass );
 	
 	/** I/O function */
 	friend COREUOBJECT_API FArchive& operator<<( FArchive& Ar, FObjectImport& I );
@@ -872,46 +878,6 @@
 		}
 		return NULL;
 	}
-
-	/** Gets the class name for the specified index in the export map. */
-	COREUOBJECT_API FName GetExportClassName( int32 ExportIdx );
-	/** Gets the class name for the specified index in the import map. */
-	FName GetExportClassName(FPackageIndex PackageIndex)
-	{
-		if (PackageIndex.IsExport())
-		{
-			return GetExportClassName(PackageIndex.ToExport());
-		}
-		return NAME_None;
-	}
-	/** Gets the class name for the specified index in the import map. */
-	FName GetImportClassName( int32 ImportIdx )
-	{
-		return ImportMap[ImportIdx].ClassName;
-	}
-	/** Gets the class name for the specified index in the import map. */
-	FName GetImportClassName(FPackageIndex PackageIndex)
-	{
-		if (PackageIndex.IsImport())
-		{
-			return GetImportClassName(PackageIndex.ToImport());
-		}
-		return NAME_None;
-	}
-	/** Gets the class name for the specified package index. */
-	FName GetClassName(FPackageIndex PackageIndex)
-	{
-		if (PackageIndex.IsImport())
-		{
-			return GetImportClassName(PackageIndex);
-		}
-		else if (PackageIndex.IsExport())
-		{
-			return GetExportClassName(PackageIndex);
-		}
-		return NAME_None;
-	}
-
 };
 
 
@@ -1002,6 +968,46 @@
 
 	virtual ~FLinker();
 
+	/** Gets the class name for the specified index in the export map. */
+	COREUOBJECT_API FName GetExportClassName(int32 ExportIdx);
+	/** Gets the class name for the specified index in the import map. */
+	FName GetExportClassName(FPackageIndex PackageIndex)
+	{
+		if (PackageIndex.IsExport())
+		{
+			return GetExportClassName(PackageIndex.ToExport());
+		}
+		return NAME_None;
+	}
+	/** Gets the class name for the specified index in the import map. */
+	FName GetImportClassName(int32 ImportIdx)
+	{
+		return ImportMap[ImportIdx].ClassName;
+	}
+	/** Gets the class name for the specified index in the import map. */
+	FName GetImportClassName(FPackageIndex PackageIndex)
+	{
+		if (PackageIndex.IsImport())
+		{
+			return GetImportClassName(PackageIndex.ToImport());
+		}
+		return NAME_None;
+	}
+	/** Gets the class name for the specified package index. */
+	FName GetClassName(FPackageIndex PackageIndex)
+	{
+		if (PackageIndex.IsImport())
+		{
+			return GetImportClassName(PackageIndex);
+		}
+		else if (PackageIndex.IsExport())
+		{
+			return GetExportClassName(PackageIndex);
+		}
+		return NAME_None;
+	}
+
+
 	FORCEINLINE ELinkerType::Type GetType() const
 	{
 		return LinkerType;
@@ -1298,6 +1304,8 @@
 	uint32					LoadFlags;
 	/** Indicates whether the imports for this loader have been verified													*/
 	bool					bHaveImportsBeenVerified;
+	/** Indicates that this linker was created for a dynamic class package and will not use Loader */
+	bool					bDynamicClassLinker;
 	/** Hash table for exports.																								*/
 	int32						ExportHash[256];
 #if WITH_EDITOR
@@ -2097,6 +2105,14 @@
 
 	bool HasPerformedFullExportResolvePass();
 
+	/** Finds import, tries to fall back to dynamic class if the object could not be found */
+	UObject* FindImport(UClass* ImportClass, UObject* ImportOuter, const TCHAR* Name);
+	/** Finds import, tries to fall back to dynamic class if the object could not be found */
+	UObject* FindImportFast(UClass* ImportClass, UObject* ImportOuter, FName Name);
+
+	/** Fills all necessary information for constructing dynamic type package linker */
+	void CreateDynamicTypeLoader();
+
 #if	USE_CIRCULAR_DEPENDENCY_LOAD_DEFERRING
 	/** 
 	 * For deferring dependency loads, we block CDO serialization until the 
diff --git a/Engine/Source/Runtime/CoreUObject/Public/UObject/ObjectBase.h b/Engine/Source/Runtime/CoreUObject/Public/UObject/ObjectBase.h
--- a/Engine/Source/Runtime/CoreUObject/Public/UObject/ObjectBase.h
+++ b/Engine/Source/Runtime/CoreUObject/Public/UObject/ObjectBase.h
@@ -92,7 +92,7 @@ enum EPackageFlags
 	PKG_CompiledIn					= 0x00000010,   // This package is from "compiled in" classes.
 	PKG_ForDiffing					= 0x00000020,	// This package was loaded just for the purposes of diff'ing
 	PKG_EditorOnly					= 0x00000040, // This is editor-only package (for example: editor module script package)
-//  PKG_Unused						= 0x00000080,
+	PKG_Developer					= 0x00000080,	// Developer module
 //	PKG_Unused						= 0x00000100,
 //	PKG_Unused						= 0x00000200,
 //	PKG_Unused						= 0x00000400,
diff --git a/Engine/Source/Runtime/CoreUObject/Public/UObject/Script.h b/Engine/Source/Runtime/CoreUObject/Public/UObject/Script.h
--- a/Engine/Source/Runtime/CoreUObject/Public/UObject/Script.h
+++ b/Engine/Source/Runtime/CoreUObject/Public/UObject/Script.h
@@ -241,6 +241,7 @@ enum EExprToken
 	EX_AssetConst			= 0x67,
 	EX_CallMath				= 0x68, // static pure function from on local call space
 	EX_SwitchValue			= 0x69,
+	EX_InstrumentationEvent	= 0x6A, // Instrumentation event
 	EX_Max					= 0x100,
 };
 
@@ -268,6 +269,24 @@ namespace EBlueprintExceptionType
 	};
 }
 
+// Script instrumentation event types
+namespace EScriptInstrumentation
+{
+	enum Type
+	{
+		Class = 0,
+		Instance,
+		Event,
+		Function,
+		Branch,
+		Macro,
+		NodeEntry,
+		NodeExit,
+		NodePin,
+		Stop
+	};
+}
+
 // Information about a blueprint exception
 struct FBlueprintExceptionInfo
 {
@@ -297,19 +316,60 @@ protected:
 	FString Description;
 };
 
+// Information about a blueprint instrumentation event
+struct EScriptInstrumentationEvent
+{
+public:
+
+	EScriptInstrumentationEvent(EScriptInstrumentation::Type InEventType, const UObject* InContextObject, const struct FFrame& InStackFrame)
+		: EventType(InEventType)
+		, ContextObject(InContextObject)
+		, StackFramePtr(&InStackFrame)
+	{
+	}
+
+	EScriptInstrumentationEvent(EScriptInstrumentation::Type InEventType, const UObject* InContextObject)
+		: EventType(InEventType)
+		, ContextObject(InContextObject)
+		, StackFramePtr(nullptr)
+	{
+	}
+
+	EScriptInstrumentation::Type GetType() const { return EventType; }
+	void SetType(EScriptInstrumentation::Type InType) { EventType = InType;	}
+	const UObject* GetContextObject() const { return ContextObject; }
+	bool IsStackFrameValid() const { return StackFramePtr != nullptr; }
+	const FFrame& GetStackFrame() const { return *StackFramePtr; }
+
+protected:
+
+	EScriptInstrumentation::Type EventType;
+	const UObject* ContextObject;
+	const struct FFrame* StackFramePtr;
+};
+
 // Blueprint core runtime delegates
 class COREUOBJECT_API FBlueprintCoreDelegates
 {
 public:
 	// Callback for debugging events such as a breakpoint (Object that triggered event, active stack frame, Info)
 	DECLARE_MULTICAST_DELEGATE_ThreeParams(FOnScriptDebuggingEvent, const UObject*, const struct FFrame&, const FBlueprintExceptionInfo&);
+	// Callback for blueprint profiling events
+	DECLARE_MULTICAST_DELEGATE_OneParam(FOnScriptInstrumentEvent, const EScriptInstrumentationEvent& );
+	// Callback for blueprint instrumentation enable/disable events
+	DECLARE_MULTICAST_DELEGATE_OneParam(FOnToggleScriptProfiler, bool );
 
 public:
 	// Called when a script exception occurs
 	static FOnScriptDebuggingEvent OnScriptException;
+	// Called when a script profiling event is fired
+	static FOnScriptInstrumentEvent OnScriptProfilingEvent;
+	// Called when a script profiler is enabled/disabled
+	static FOnToggleScriptProfiler OnToggleScriptProfiler;
 
 public:
 	static void ThrowScriptException(const UObject* ActiveObject, const struct FFrame& StackFrame, const FBlueprintExceptionInfo& Info);
+	static void InstrumentScriptEvent(const EScriptInstrumentationEvent& Info);
 	static void SetScriptMaximumLoopIterations( const int32 MaximumLoopIterations );
 };
 
diff --git a/Engine/Source/Runtime/CoreUObject/Public/UObject/ScriptSerialization.h b/Engine/Source/Runtime/CoreUObject/Public/UObject/ScriptSerialization.h
--- a/Engine/Source/Runtime/CoreUObject/Public/UObject/ScriptSerialization.h
+++ b/Engine/Source/Runtime/CoreUObject/Public/UObject/ScriptSerialization.h
@@ -214,6 +214,11 @@
 			}
 			break;
 		}
+		case EX_InstrumentationEvent:
+		{
+			iCode += sizeof(int32);
+			break;
+		}
 		case EX_Return:
 		{
 			SerializeExpr( iCode, Ar ); // Return expression.
diff --git a/Engine/Source/Runtime/CoreUObject/Public/UObject/UObject.h b/Engine/Source/Runtime/CoreUObject/Public/UObject/UObject.h
--- a/Engine/Source/Runtime/CoreUObject/Public/UObject/UObject.h
+++ b/Engine/Source/Runtime/CoreUObject/Public/UObject/UObject.h
@@ -1012,6 +1012,9 @@ public:
 	DECLARE_FUNCTION(execTracepoint);
 	DECLARE_FUNCTION(execWireTracepoint);
 
+	/** Instrumentation event for profiling; only observed in the builds with blueprint instrumentation */
+	DECLARE_FUNCTION(execInstrumentation);
+
 	DECLARE_FUNCTION(execEndOfScript);
 
 	/** failsafe for functions that return a value - returns the zero value for a property and logs that control reached the end of a non-void function */
diff --git a/Engine/Source/Runtime/CoreUObject/Public/UObject/UObjectBase.h b/Engine/Source/Runtime/CoreUObject/Public/UObject/UObjectBase.h
--- a/Engine/Source/Runtime/CoreUObject/Public/UObject/UObjectBase.h
+++ b/Engine/Source/Runtime/CoreUObject/Public/UObject/UObjectBase.h
@@ -309,26 +309,34 @@ struct TClassCompiledInDefer : public FFieldCompiledInInfo
 /**
  * Stashes the singleton function that builds a compiled in class. Later, this is executed.
  */
-COREUOBJECT_API void UObjectCompiledInDefer(class UClass *(*InRegister)(), class UClass *(*InStaticClass)(), const TCHAR* Name, bool bDynamic);
+COREUOBJECT_API void UObjectCompiledInDefer(class UClass *(*InRegister)(), class UClass *(*InStaticClass)(), const TCHAR* Name, bool bDynamic, const TCHAR* DynamicPathName);
 
 struct FCompiledInDefer
 {
-	FCompiledInDefer(class UClass *(*InRegister)(), class UClass *(*InStaticClass)(), const TCHAR* Name, bool bDynamic)
+	FCompiledInDefer(class UClass *(*InRegister)(), class UClass *(*InStaticClass)(), const TCHAR* Name, bool bDynamic, const TCHAR* DynamicPackageName = nullptr, const TCHAR* DynamicPathName = nullptr)
 	{
-		UObjectCompiledInDefer(InRegister, InStaticClass, Name, bDynamic);
+		if (bDynamic)
+		{
+			GetConvertedDynamicPackageNameToTypeName().Add(FName(DynamicPackageName), FName(Name));
+		}
+		UObjectCompiledInDefer(InRegister, InStaticClass, Name, bDynamic, DynamicPathName);
 	}
 };
 
 /**
  * Stashes the singleton function that builds a compiled in struct (StaticStruct). Later, this is executed.
  */
-COREUOBJECT_API void UObjectCompiledInDeferStruct(class UScriptStruct *(*InRegister)(), const TCHAR* PackageName, const TCHAR* Name, bool bDynamic);
+COREUOBJECT_API void UObjectCompiledInDeferStruct(class UScriptStruct *(*InRegister)(), const TCHAR* PackageName, const FName PathName, bool bDynamic);
 
 struct FCompiledInDeferStruct
 {
-	FCompiledInDeferStruct(class UScriptStruct *(*InRegister)(), const TCHAR* PackageName, const TCHAR* Name, bool bDynamic)
+	FCompiledInDeferStruct(class UScriptStruct *(*InRegister)(), const TCHAR* PackageName, const TCHAR* Name, bool bDynamic, const TCHAR* DynamicPackageName, const TCHAR* DynamicPathName)
 	{
-		UObjectCompiledInDeferStruct(InRegister, PackageName, Name, bDynamic);
+		if (bDynamic)
+		{
+			GetConvertedDynamicPackageNameToTypeName().Add(FName(DynamicPackageName), FName(Name));
+		}
+		UObjectCompiledInDeferStruct(InRegister, PackageName, DynamicPathName, bDynamic);
 	}
 };
 
@@ -340,13 +348,17 @@ COREUOBJECT_API class UScriptStruct *GetStaticStruct(class UScriptStruct *(*InRe
 /**
  * Stashes the singleton function that builds a compiled in enum. Later, this is executed.
  */
-COREUOBJECT_API void UObjectCompiledInDeferEnum(class UEnum *(*InRegister)(), const TCHAR* PackageName, const TCHAR* Name, bool bDynamic);
+COREUOBJECT_API void UObjectCompiledInDeferEnum(class UEnum *(*InRegister)(), const TCHAR* PackageName, const FName PathName, bool bDynamic);
 
 struct FCompiledInDeferEnum
 {
-	FCompiledInDeferEnum(class UEnum *(*InRegister)(), const TCHAR* PackageName, const TCHAR* Name, bool bDynamic)
+	FCompiledInDeferEnum(class UEnum *(*InRegister)(), const TCHAR* PackageName, const TCHAR* Name, bool bDynamic, const TCHAR* DynamicPackageName, const TCHAR* DynamicPathName)
 	{
-		UObjectCompiledInDeferEnum(InRegister, PackageName, Name, bDynamic);
+		if (bDynamic)
+		{
+			GetConvertedDynamicPackageNameToTypeName().Add(FName(DynamicPackageName), FName(Name));
+		}
+		UObjectCompiledInDeferEnum(InRegister, PackageName, DynamicPathName, bDynamic);
 	}
 };
 
@@ -356,7 +368,7 @@ struct FCompiledInDeferEnum
 COREUOBJECT_API class UEnum *GetStaticEnum(class UEnum *(*InRegister)(), UObject* EnumOuter, const TCHAR* EnumName);
 
 COREUOBJECT_API class UScriptStruct* FindExistingStructIfHotReloadOrDynamic(UObject* Outer, const TCHAR* StructName, SIZE_T Size, uint32 Crc, bool bIsDynamic);
-COREUOBJECT_API class UEnum* FindExistingEnumIfHotReload(UObject* Outer, const TCHAR* EnumName, SIZE_T Size, uint32 Crc);
+COREUOBJECT_API class UEnum* FindExistingEnumIfHotReloadOrDynamic(UObject* Outer, const TCHAR* EnumName, SIZE_T Size, uint32 Crc, bool bIsDynamic);
 
 /** @return	True if there are any newly-loaded UObjects that are waiting to be registered by calling ProcessNewlyLoadedUObjects() */
 COREUOBJECT_API bool AnyNewlyLoadedUObjects();
diff --git a/Engine/Source/Runtime/CoreUObject/Public/UObject/UObjectGlobals.h b/Engine/Source/Runtime/CoreUObject/Public/UObject/UObjectGlobals.h
--- a/Engine/Source/Runtime/CoreUObject/Public/UObject/UObjectGlobals.h
+++ b/Engine/Source/Runtime/CoreUObject/Public/UObject/UObjectGlobals.h
@@ -1851,9 +1851,15 @@ extern COREUOBJECT_API bool GShouldVerifyGCAssumptions;
 COREUOBJECT_API UScriptStruct* GetFallbackStruct();
 
 /** Constructs dynamic type of a given class. */
-COREUOBJECT_API UObject* ConstructDynamicType(FName TypeName, FName TypeClass);
+COREUOBJECT_API UObject* ConstructDynamicType(FName TypePathName);
+
+/** Given a dynamic type path name, returns that type's class name (can be either DynamicClass, ScriptStruct or Enum). */
+COREUOBJECT_API FName GetDynamicTypeClassName(FName TypePathName);
 
 /** Finds or constructs a package for dynamic type. */
 COREUOBJECT_API UPackage* FindOrConstructDynamicTypePackage(const TCHAR* PackageName);
 
+/** Get names of "virtual" packages, that contain Dynamic types  */
+COREUOBJECT_API TMap<FName, FName>& GetConvertedDynamicPackageNameToTypeName();
+
 #endif	// __UNOBJGLOBALS_H__
