Linux fixes, fallout from integration from Dev-Blueprints

diff --git a/Engine/Source/Editor/Kismet/Public/FindInBlueprintManager.h b/Engine/Source/Editor/Kismet/Public/FindInBlueprintManager.h
--- a/Engine/Source/Editor/Kismet/Public/FindInBlueprintManager.h
+++ b/Engine/Source/Editor/Kismet/Public/FindInBlueprintManager.h
@@ -529,7 +529,7 @@ protected:
 	TMap<FName, const UClass*> CachedAssetClasses;
 };
 
-KISMET_API struct FDisableGatheringDataOnScope
+struct KISMET_API FDisableGatheringDataOnScope
 {
 	bool bOriginallyEnabled;
 	FDisableGatheringDataOnScope() : bOriginallyEnabled(FFindInBlueprintSearchManager::Get().IsGatheringDataEnabled())
diff --git a/Engine/Source/Runtime/CoreUObject/Private/UObject/Package.cpp b/Engine/Source/Runtime/CoreUObject/Private/UObject/Package.cpp
--- a/Engine/Source/Runtime/CoreUObject/Private/UObject/Package.cpp
+++ b/Engine/Source/Runtime/CoreUObject/Private/UObject/Package.cpp
@@ -192,9 +192,12 @@ bool UPackage::IsFullyLoaded()
 	{
 		FString DummyFilename;
 		// Try to find matching package in package file cache.
-		if (!GetConvertedDynamicPackageNameToTypeName().Contains(GetFName()) &&
-			  !FPackageName::DoesPackageExist( *GetName(), NULL, &DummyFilename ) || 
-				(GIsEditor && IFileManager::Get().FileSize(*DummyFilename) < 0) )
+		if (	!GetConvertedDynamicPackageNameToTypeName().Contains(GetFName()) &&
+				(
+					!FPackageName::DoesPackageExist( *GetName(), NULL, &DummyFilename ) || 
+					(GIsEditor && IFileManager::Get().FileSize(*DummyFilename) < 0) 
+				)
+			)
 		{
 			// Package has NOT been found, so we assume it's a newly created one and therefore fully loaded.
 			bHasBeenFullyLoaded = true;
