Copying //UE4/Dev-Core to //UE4/Main

diff --git a/Engine/Source/Runtime/Core/Public/Stats/StatsMisc.h b/Engine/Source/Runtime/Core/Public/Stats/StatsMisc.h
--- a/Engine/Source/Runtime/Core/Public/Stats/StatsMisc.h
+++ b/Engine/Source/Runtime/Core/Public/Stats/StatsMisc.h
@@ -2,38 +2,67 @@
 
 #pragma once
 
-
-#if STATS
 /**
- * Utility class to capture time passed in seconds, adding delta time to passed
- * in variable. Not useful for reentrant functions
- */
+* Utility class to capture time passed in seconds, adding delta time to passed
+* in variable. Not useful for reentrant functions
+*/
 class FSimpleScopeSecondsCounter
 {
 public:
 	/** Ctor, capturing start time. */
-	FSimpleScopeSecondsCounter(double& InSeconds)
-	:	StartTime(FPlatformTime::Seconds())
-	,	Seconds(InSeconds)
+	FSimpleScopeSecondsCounter(double& InSeconds, bool bInEnabled = true)
+		: StartTime(FPlatformTime::Seconds())
+		, Seconds(InSeconds)
+		, bEnabled(bInEnabled)
+		, RecursionDepth(nullptr)
 	{}
+	/** Ctor, capturing start time. */
+	FSimpleScopeSecondsCounter(double& InSeconds, int32* InRecursionDepth)
+		: StartTime(FPlatformTime::Seconds())
+		, Seconds(InSeconds)
+		, bEnabled(*InRecursionDepth == 0)
+		, RecursionDepth(InRecursionDepth)
+	{
+		(*RecursionDepth)++;
+	}
 	/** Dtor, updating seconds with time delta. */
 	~FSimpleScopeSecondsCounter()
 	{
-		Seconds += FPlatformTime::Seconds() - StartTime;
+		if (bEnabled)
+		{
+			Seconds += FPlatformTime::Seconds() - StartTime;
+		}
+
+		if (RecursionDepth)
+		{
+			(*RecursionDepth)--;
+		}
 	}
 private:
 	/** Start time, captured in ctor. */
 	double StartTime;
 	/** Time variable to update. */
 	double& Seconds;
+	/** Is the timer enabled or disabled */
+	bool bEnabled;
+	/** Recursion depth */
+	int32* RecursionDepth;
 };
 
 /** Macro for updating a seconds counter without creating new scope. */
-#define SCOPE_SECONDS_COUNTER(Seconds) \
+#define SCOPE_SECONDS_COUNTER_BASE(Seconds) \
 	FSimpleScopeSecondsCounter SecondsCount_##Seconds(Seconds);
 
+#define SCOPE_SECONDS_COUNTER_RECURSION_SAFE_BASE(Seconds) \
+	static int32 SecondsCount_##Seconds##_RecursionCounter = 0; \
+	FSimpleScopeSecondsCounter SecondsCount_##Seconds(Seconds, &SecondsCount_##Seconds##_RecursionCounter);
+
+#if STATS
+#define SCOPE_SECONDS_COUNTER(Seconds) SCOPE_SECONDS_COUNTER_BASE(Seconds)
+#define SCOPE_SECONDS_COUNTER_RECURSION_SAFE(Seconds) SCOPE_SECONDS_COUNTER_RECURSION_SAFE_BASE(Seconds)
 #else
 #define SCOPE_SECONDS_COUNTER(Seconds)
+#define SCOPE_SECONDS_COUNTER_RECURSION_SAFE(Seconds)
 #endif 
 
 
diff --git a/Engine/Source/Runtime/Core/Public/Templates/SharedPointer.h b/Engine/Source/Runtime/Core/Public/Templates/SharedPointer.h
--- a/Engine/Source/Runtime/Core/Public/Templates/SharedPointer.h
+++ b/Engine/Source/Runtime/Core/Public/Templates/SharedPointer.h
@@ -306,7 +306,7 @@ public:
 		// Shared references must never be nullptr, so either pass a valid object or consider using TSharedPtr instead.
 		check( InRawPtrProxy.Object != nullptr );
 
-		*this = TSharedRef< ObjectType, Mode >( InRawPtrProxy.Object );
+		*this = TSharedRef< ObjectType, Mode >( InRawPtrProxy );
 		return *this;
 	}
 
@@ -680,7 +680,7 @@ public:
 	template< class OtherType >
 	FORCEINLINE TSharedPtr& operator=( SharedPointerInternals::FRawPtrProxy< OtherType > const& InRawPtrProxy )
 	{
-		*this = TSharedPtr< ObjectType, Mode >( InRawPtrProxy.Object );
+		*this = TSharedPtr< ObjectType, Mode >( InRawPtrProxy );
 		return *this;
 	}
 
diff --git a/Engine/Source/Runtime/Core/Public/Templates/UnrealTemplate.h b/Engine/Source/Runtime/Core/Public/Templates/UnrealTemplate.h
--- a/Engine/Source/Runtime/Core/Public/Templates/UnrealTemplate.h
+++ b/Engine/Source/Runtime/Core/Public/Templates/UnrealTemplate.h
@@ -14,18 +14,6 @@
 	Standard templates.
 -----------------------------------------------------------------------------*/
 
-/** 
- * CanConvertPointerFromTo<From, To>::Result is an enum value equal to 1 if From* is automatically convertable to a To* (without regard to const!)
- **/
-template<class From, class To>
-class CanConvertPointerFromTo
-{
-public:
-	DEPRECATED(4.8, "This trait is deprecated as it does not follow UE4 standards of syntax and hides a const conversion. Please use TPointerIsConvertibleFromTo instead.")
-	static const bool Result = TPointerIsConvertibleFromTo<From, const To>::Value;
-};
-
-
 /**
  * Aligns a value to the nearest higher multiple of 'Alignment', which must be a power of two.
  *
@@ -38,6 +26,19 @@ inline CONSTEXPR T Align( const T Ptr, int32 Alignment )
 {
 	return (T)(((PTRINT)Ptr + Alignment - 1) & ~(Alignment-1));
 }
+
+/**
+ * Checks if a pointer is aligned to the specified alignment.
+ *
+ * @param Ptr - The pointer to check.
+ *
+ * @return true if the pointer is aligned, false otherwise.
+ */
+static FORCEINLINE bool IsAligned(const volatile void* Ptr, const uint32 Alignment)
+{
+	return !(UPTRINT(Ptr) & (Alignment - 1));
+}
+
 /**
  * Aligns a value to the nearest higher multiple of 'Alignment'.
  *
diff --git a/Engine/Source/Runtime/Core/Public/UObject/ScriptDelegates.h b/Engine/Source/Runtime/Core/Public/UObject/ScriptDelegates.h
--- a/Engine/Source/Runtime/Core/Public/UObject/ScriptDelegates.h
+++ b/Engine/Source/Runtime/Core/Public/UObject/ScriptDelegates.h
@@ -500,7 +500,7 @@ protected:
 		const int32 NumFunctions = InvocationList.Num();
 		for( int32 CurFunctionIndex = 0; CurFunctionIndex < NumFunctions; ++CurFunctionIndex )
 		{
-			check( InvocationList[ CurFunctionIndex ] != InDelegate );
+			(void)ensure( InvocationList[ CurFunctionIndex ] != InDelegate );
 		}
 #endif // !(UE_BUILD_SHIPPING || UE_BUILD_TEST)
 		InvocationList.Add( InDelegate );
diff --git a/Engine/Source/Runtime/Core/Public/WinRT/WinRTARMPlatform.h b/Engine/Source/Runtime/Core/Public/WinRT/WinRTARMPlatform.h
--- a/Engine/Source/Runtime/Core/Public/WinRT/WinRTARMPlatform.h
+++ b/Engine/Source/Runtime/Core/Public/WinRT/WinRTARMPlatform.h
@@ -46,11 +46,12 @@ typedef FWinRTTypes FPlatformTypes;
 #define PLATFORM_HAS_BSD_SOCKET_FEATURE_WINSOCKETS			1
 
 // Function type macros.
-#define VARARGS			__cdecl					/* Functions with variable arguments */
-#define CDECL			__cdecl					/* Standard C function */
-#define STDCALL			__stdcall				/* Standard calling convention */
-#define FORCEINLINE		__forceinline			/* Force code to be inline */
-#define FORCENOINLINE	__declspec(noinline)	/* Force code to NOT be inline */
+#define VARARGS						__cdecl										/* Functions with variable arguments */
+#define CDECL						__cdecl										/* Standard C function */
+#define STDCALL						__stdcall									/* Standard calling convention */
+#define FORCEINLINE					__forceinline								/* Force code to be inline */
+#define FORCENOINLINE				__declspec(noinline)						/* Force code to NOT be inline */
+#define FUNCTION_CHECK_RETURN(...)	__declspec("SAL_checkReturn") __VA_ARGS__	/* Wrap a function signature in this to warn that callers should not ignore the return value. */
 
 // Hints compiler that expression is true; generally restricted to comparisons against constants
 #define ASSUME(expr)	__assume(expr)
diff --git a/Engine/Source/Runtime/Core/Public/WinRT/WinRTCriticalSection.h b/Engine/Source/Runtime/Core/Public/WinRT/WinRTCriticalSection.h
--- a/Engine/Source/Runtime/Core/Public/WinRT/WinRTCriticalSection.h
+++ b/Engine/Source/Runtime/Core/Public/WinRT/WinRTCriticalSection.h
@@ -53,4 +53,8 @@ public:
 	{
 		LeaveCriticalSection(&CriticalSection);
 	}
+
+private:
+	FWinRTCriticalSection(const FWinRTCriticalSection&);
+	FWinRTCriticalSection& operator=(const FWinRTCriticalSection&);
 };
diff --git a/Engine/Source/Runtime/Core/Public/WinRT/WinRTPlatform.h b/Engine/Source/Runtime/Core/Public/WinRT/WinRTPlatform.h
--- a/Engine/Source/Runtime/Core/Public/WinRT/WinRTPlatform.h
+++ b/Engine/Source/Runtime/Core/Public/WinRT/WinRTPlatform.h
@@ -59,11 +59,12 @@ typedef FWinRTTypes FPlatformTypes;
 #define PLATFORM_SUPPORTS_MESSAGEBUS						0
 
 // Function type macros.
-#define VARARGS			__cdecl					/* Functions with variable arguments */
-#define CDECL			__cdecl					/* Standard C function */
-#define STDCALL			__stdcall				/* Standard calling convention */
-#define FORCEINLINE		__forceinline			/* Force code to be inline */
-#define FORCENOINLINE	__declspec(noinline)	/* Force code to NOT be inline */
+#define VARARGS						__cdecl										/* Functions with variable arguments */
+#define CDECL						__cdecl										/* Standard C function */
+#define STDCALL						__stdcall									/* Standard calling convention */
+#define FORCEINLINE					__forceinline								/* Force code to be inline */
+#define FORCENOINLINE				__declspec(noinline)						/* Force code to NOT be inline */
+#define FUNCTION_CHECK_RETURN(...)	__declspec("SAL_checkReturn") __VA_ARGS__	/* Wrap a function signature in this to warn that callers should not ignore the return value. */
 
 // Hints compiler that expression is true; generally restricted to comparisons against constants
 #define ASSUME(expr)	__assume(expr)
diff --git a/Engine/Source/Runtime/Core/Public/Windows/WIndowsPlatform.h b/Engine/Source/Runtime/Core/Public/Windows/WIndowsPlatform.h
--- a/Engine/Source/Runtime/Core/Public/Windows/WIndowsPlatform.h
+++ b/Engine/Source/Runtime/Core/Public/Windows/WIndowsPlatform.h
@@ -84,13 +84,13 @@ typedef FWindowsPlatformTypes FPlatformTypes;
 #define PLATFORM_HAS_128BIT_ATOMICS							(!HACK_HEADER_GENERATOR && PLATFORM_64BITS && (WINVER >= 0x602))
 #define PLATFORM_USES_ANSI_STRING_FOR_EXTERNAL_PROFILING	0
 
-
 // Function type macros.
-#define VARARGS     __cdecl					/* Functions with variable arguments */
-#define CDECL	    __cdecl					/* Standard C function */
-#define STDCALL		__stdcall				/* Standard calling convention */
-#define FORCEINLINE __forceinline			/* Force code to be inline */
-#define FORCENOINLINE __declspec(noinline)	/* Force code to NOT be inline */
+#define VARARGS     __cdecl														/* Functions with variable arguments */
+#define CDECL	    __cdecl														/* Standard C function */
+#define STDCALL		__stdcall													/* Standard calling convention */
+#define FORCEINLINE __forceinline												/* Force code to be inline */
+#define FORCENOINLINE __declspec(noinline)										/* Force code to NOT be inline */
+#define FUNCTION_CHECK_RETURN(...) __declspec("SAL_checkReturn") __VA_ARGS__	/* Wrap a function signature in this to warn that callers should not ignore the return value. */
 
 // Hints compiler that expression is true; generally restricted to comparisons against constants
 #if !defined(__clang__)		// Clang doesn't support __assume (Microsoft specific)
diff --git a/Engine/Source/Runtime/Core/Public/Windows/WindowsCriticalSection.h b/Engine/Source/Runtime/Core/Public/Windows/WindowsCriticalSection.h
--- a/Engine/Source/Runtime/Core/Public/Windows/WindowsCriticalSection.h
+++ b/Engine/Source/Runtime/Core/Public/Windows/WindowsCriticalSection.h
@@ -48,25 +48,30 @@ public:
 	}
 
 	/**
-	* Quick test for seeing if the lock is already being used.
-	*/
+	 * Attempt to take a lock and returns whether or not a lock was taken.
+	 *
+	 * @return true if a lock was taken, false otherwise.
+	 */
 	FORCEINLINE bool TryLock()
 	{
 		if (TryEnterCriticalSection(&CriticalSection))
 		{
-			LeaveCriticalSection(&CriticalSection);
 			return true;
 		};
 		return false;
 	}
 
 	/**
-	 * Releases the lock on the critical seciton
+	 * Releases the lock on the critical section
 	 */
 	FORCEINLINE void Unlock()
 	{
 		LeaveCriticalSection(&CriticalSection);
 	}
+
+private:
+	FWindowsCriticalSection(const FWindowsCriticalSection&);
+	FWindowsCriticalSection& operator=(const FWindowsCriticalSection&);
 };
 
 /** System-Wide Critical Section for windows using mutex */
diff --git a/Engine/Source/Runtime/CoreUObject/Private/Misc/PackageName.cpp b/Engine/Source/Runtime/CoreUObject/Private/Misc/PackageName.cpp
--- a/Engine/Source/Runtime/CoreUObject/Private/Misc/PackageName.cpp
+++ b/Engine/Source/Runtime/CoreUObject/Private/Misc/PackageName.cpp
@@ -613,6 +613,65 @@ bool FPackageName::FindPackageFileWithoutExtension(const FString& InPackageFilen
 	return false;
 }
 
+void FPackageName::FindPackageFileAndLocalizationWithoutExtension(const FString& InNativePackageFilename, const FString& InLocalizedPackageFilename, FString& OutNativeFilename, FString& OutLocalizedFilename)
+{
+	auto& FileManager = IFileManager::Get();
+
+	static const FString* PackageExtensions[] =
+	{
+		&AssetPackageExtension,
+		&MapPackageExtension
+	};
+
+	uint8 FoundBits = 0;
+
+	// Loop through all known extensions and check if the file exist.
+	for (auto Extension : PackageExtensions)
+	{
+		FString   NativePackageFilename = InNativePackageFilename + *Extension;
+		FString   LocalizedPackageFilename = InLocalizedPackageFilename + *Extension;
+
+		FDateTime NativePackageTimestamp, LocalizedPackageTimeStamp;
+		FileManager.GetTimeStampPair(*NativePackageFilename, *LocalizedPackageFilename, NativePackageTimestamp, LocalizedPackageTimeStamp);
+		
+		if (NativePackageTimestamp != FDateTime::MinValue())
+		{			
+			FoundBits |= 1;
+		}
+
+		if (LocalizedPackageTimeStamp != FDateTime::MinValue())
+		{
+			FoundBits |= 2;
+		}
+
+		if (FoundBits)
+		{
+			if (FoundBits & 1)
+			{
+				OutNativeFilename = MoveTemp(NativePackageFilename);
+			}
+			else
+			{
+				OutNativeFilename.Empty();
+			}
+
+			if (FoundBits & 2)
+			{
+				OutLocalizedFilename = MoveTemp(LocalizedPackageFilename);
+			}
+			else
+			{
+				OutLocalizedFilename.Empty();
+			}
+
+			return;
+		}
+	}
+
+	OutNativeFilename.Empty();
+	OutLocalizedFilename.Empty();
+}
+
 bool FPackageName::DoesPackageExist(const FString& LongPackageName, const FGuid* Guid /*= NULL*/, FString* OutFilename /*= NULL*/, const bool ShouldGetLocalizedPackage /*= false*/)
 {
 	bool bFoundFile = false;
@@ -677,6 +736,82 @@ bool FPackageName::DoesPackageExist(const FString& LongPackageName, const FGuid*
 	return bFoundFile;
 }
 
+bool FPackageName::DoesPackageExistWithLocalization(const FString& LongPackageName, const FGuid* Guid, FString* OutNativeFilename, FString* OutLocalizedFilename)
+{
+	// Make sure passing filename as LongPackageName is supported.
+	FString PackageName;
+	FText Reason;
+
+	if (!FPackageName::TryConvertFilenameToLongPackageName(LongPackageName, PackageName))
+	{
+		verify(!FPackageName::IsValidLongPackageName(PackageName, true, &Reason));
+		UE_LOG(LogPackageName, Error, TEXT("Illegal call to DoesPackageExist: '%s' is not a standard unreal filename or a long path name. Reason: %s"), *LongPackageName, *Reason.ToString());
+		ensureMsgf(false, TEXT("Illegal call to DoesPackageExist: '%s' is not a standard unreal filename or a long path name. Reason: %s"), *LongPackageName, *Reason.ToString());
+		return false;
+	}
+	// Once we have the real Package Name, we can exit early if it's a script package - they exist only in memory.
+	if (IsScriptPackage(PackageName))
+	{
+		return false;
+	}
+
+	if (!FPackageName::IsValidLongPackageName(PackageName, true, &Reason))
+	{
+		UE_LOG(LogPackageName, Error, TEXT("DoesPackageExist: DoesPackageExist FAILED: '%s' is not a standard unreal filename or a long path name. Reason: %s"), *LongPackageName, *Reason.ToString());
+		return false;
+	}
+
+	// Convert to filename (no extension yet).
+	static const int32 NumFilenames = 2;
+	FString Filenames[NumFilenames];
+	Filenames[0] = LongPackageNameToFilename(PackageName, TEXT(""), false);
+	Filenames[1] = LongPackageNameToFilename(PackageName, TEXT(""), true);
+	FindPackageFileAndLocalizationWithoutExtension(Filenames[0], Filenames[1], Filenames[0], Filenames[1]);
+	
+	if (!FPlatformProperties::RequiresCookedData() && Guid != NULL)
+	{
+		for (int32 i = 0; i < NumFilenames; ++i)
+		{
+			//@todo-packageloc Load localized packages based on culture.
+
+			// On consoles, we don't support package downloading, so no need to waste any extra cycles/disk io dealing with it
+			if (Filenames[i].Len() > 0)
+			{
+				// @todo: If we could get to list of linkers here, it would be faster to check
+				// then to open the file and read it
+				FArchive* PackageReader = IFileManager::Get().CreateFileReader(*Filenames[i]);
+				// This had better open
+				check(PackageReader != NULL);
+
+				// Read in the package summary
+				FPackageFileSummary Summary;
+				*PackageReader << Summary;
+
+				// Compare Guids
+				if (Summary.Guid != *Guid)
+				{
+					Filenames[i].Empty();
+				}
+
+				// Close package file
+				delete PackageReader;
+			}
+		}
+	}
+
+	if (OutNativeFilename && Filenames[0].Len() > 0)
+	{
+		*OutNativeFilename = Filenames[0];
+	}
+
+	if (OutLocalizedFilename && Filenames[1].Len() > 0)
+	{
+		*OutLocalizedFilename = Filenames[1];
+	}
+
+	return Filenames[0].Len() > 0;
+}
+
 bool FPackageName::SearchForPackageOnDisk(const FString& PackageName, FString* OutLongPackageName, FString* OutFilename, bool bUseLocalizedNames)
 {
 	DECLARE_SCOPE_CYCLE_COUNTER(TEXT("FPackageName::SearchForPackageOnDisk"), STAT_PackageName_SearchForPackageOnDisk, STATGROUP_LoadTime);
diff --git a/Engine/Source/Runtime/CoreUObject/Private/Serialization/ArchiveUObject.cpp b/Engine/Source/Runtime/CoreUObject/Private/Serialization/ArchiveUObject.cpp
--- a/Engine/Source/Runtime/CoreUObject/Private/Serialization/ArchiveUObject.cpp
+++ b/Engine/Source/Runtime/CoreUObject/Private/Serialization/ArchiveUObject.cpp
@@ -136,3 +136,64 @@ void FSerializedPropertyScope::PopEditorOnlyProperty()
 	}
 }
 #endif
+
+void FArchiveReplaceObjectRefBase::SerializeObject(UObject* ObjectToSerialzie)
+{
+	// Simple FReferenceCollector proxy for FArchiveReplaceObjectRefBase
+	class FReplaceObjectRefCollector : public FReferenceCollector
+	{
+		FArchive& Ar;
+		bool bAllowReferenceElimination;
+	public:
+		FReplaceObjectRefCollector(FArchive& InAr)
+			: Ar(InAr)
+			, bAllowReferenceElimination(true)
+		{
+		}
+		virtual bool IsIgnoringArchetypeRef() const override
+		{
+			return Ar.IsIgnoringArchetypeRef();
+		}
+		virtual bool IsIgnoringTransient() const override
+		{
+			return false;
+		}
+		virtual void AllowEliminatingReferences(bool bAllow) override
+		{
+			bAllowReferenceElimination = bAllow;
+		}
+		virtual void HandleObjectReference(UObject*& InObject, const UObject* InReferencingObject, const UProperty* InReferencingProperty) override
+		{
+			if (bAllowReferenceElimination)
+			{
+				UProperty* NewSerializedProperty = const_cast<UProperty*>(InReferencingProperty);
+				FSerializedPropertyScope SerializedPropertyScope(Ar, NewSerializedProperty ? NewSerializedProperty : Ar.GetSerializedProperty());
+				Ar << InObject;
+			}
+		}
+	} ReplaceRefCollector(*this);
+
+	// serialization for class default objects must be deterministic (since class 
+	// default objects may be serialized during script compilation while the script
+	// and C++ versions of a class are not in sync), so use SerializeTaggedProperties()
+	// rather than the native Serialize() function
+	UClass* ObjectClass = ObjectToSerialzie->GetClass();
+	if (ObjectToSerialzie->HasAnyFlags(RF_ClassDefaultObject))
+	{		
+		StartSerializingDefaults();
+		if (!WantBinaryPropertySerialization() && (IsLoading() || IsSaving()))
+		{
+			ObjectClass->SerializeTaggedProperties(*this, (uint8*)ObjectToSerialzie, ObjectClass, NULL);
+		}
+		else
+		{
+			ObjectClass->SerializeBin(*this, ObjectToSerialzie);
+		}
+		StopSerializingDefaults();
+	}
+	else
+	{
+		ObjectToSerialzie->Serialize(*this);
+	}
+	ObjectClass->CallAddReferencedObjects(ObjectToSerialzie, ReplaceRefCollector);
+}
\ No newline at end of file
diff --git a/Engine/Source/Runtime/CoreUObject/Private/Serialization/AsyncLoading.cpp b/Engine/Source/Runtime/CoreUObject/Private/Serialization/AsyncLoading.cpp
--- a/Engine/Source/Runtime/CoreUObject/Private/Serialization/AsyncLoading.cpp
+++ b/Engine/Source/Runtime/CoreUObject/Private/Serialization/AsyncLoading.cpp
@@ -820,6 +820,7 @@ FAsyncLoadingThread::FAsyncLoadingThread()
 	QueuedRequestsEvent = FPlatformProcess::GetSynchEventFromPool();
 	CancelLoadingEvent = FPlatformProcess::GetSynchEventFromPool();
 	ThreadSuspendedEvent = FPlatformProcess::GetSynchEventFromPool();
+	ThreadResumedEvent = FPlatformProcess::GetSynchEventFromPool();
 	if (FAsyncLoadingThread::IsMultithreaded())
 	{
 		UE_LOG(LogStreaming, Log, TEXT("Async loading is multithreaded."));
@@ -844,6 +845,8 @@ FAsyncLoadingThread::~FAsyncLoadingThread()
 	CancelLoadingEvent = nullptr;
 	FPlatformProcess::ReturnSynchEventToPool(ThreadSuspendedEvent);
 	ThreadSuspendedEvent = nullptr;
+	FPlatformProcess::ReturnSynchEventToPool(ThreadResumedEvent);
+	ThreadResumedEvent = nullptr;	
 }
 
 bool FAsyncLoadingThread::Init()
@@ -860,7 +863,11 @@ uint32 FAsyncLoadingThread::Run()
 	{
 		if (IsLoadingSuspended.GetValue() == 0)
 		{
-			bWasSuspendedLastFrame = false;
+			if (bWasSuspendedLastFrame)
+			{
+				bWasSuspendedLastFrame = false;
+				ThreadResumedEvent->Trigger();
+			}			
 			TickAsyncThread(false, true, 0.0f);
 		}
 		else if (!bWasSuspendedLastFrame)
@@ -947,6 +954,10 @@ void FAsyncLoadingThread::ResumeLoading()
 	check(IsInGameThread());
 	const int32 SuspendCount = IsLoadingSuspended.Decrement();
 	UE_CLOG(SuspendCount < 0, LogStreaming, Fatal, TEXT("ResumeAsyncLoadingThread: Async loading was resumed more times than it was suspended."));
+	if (IsMultithreaded() && SuspendCount == 0)
+	{
+		ThreadResumedEvent->Wait();
+	}
 }
 
 float FAsyncLoadingThread::GetAsyncLoadPercentage(const FName& PackageName)
diff --git a/Engine/Source/Runtime/CoreUObject/Private/Serialization/AsyncLoadingThread.h b/Engine/Source/Runtime/CoreUObject/Private/Serialization/AsyncLoadingThread.h
--- a/Engine/Source/Runtime/CoreUObject/Private/Serialization/AsyncLoadingThread.h
+++ b/Engine/Source/Runtime/CoreUObject/Private/Serialization/AsyncLoadingThread.h
@@ -25,6 +25,8 @@ class FAsyncLoadingThread : public FRunnable
 	FEvent* CancelLoadingEvent;
 	/** [ASYNC/GAME THREAD] Event used to signal that the async loading thread should be suspended */
 	FEvent* ThreadSuspendedEvent;
+	/** [ASYNC/GAME THREAD] Event used to signal that the async loading thread has resumed */
+	FEvent* ThreadResumedEvent;
 	/** [ASYNC/GAME THREAD] List of queued packages to stream */
 	TArray<FAsyncPackageDesc*> QueuedPackages;
 #if THREADSAFE_UOBJECTS
diff --git a/Engine/Source/Runtime/CoreUObject/Private/UObject/FastReferenceCollector.h b/Engine/Source/Runtime/CoreUObject/Private/UObject/FastReferenceCollector.h
new file mode 100644
--- /dev/null
+++ b/Engine/Source/Runtime/CoreUObject/Private/UObject/FastReferenceCollector.h
@@ -0,0 +1,419 @@
+// Copyright 1998-2015 Epic Games, Inc. All Rights Reserved.
+
+#pragma once
+
+#include "UObject/GarbageCollection.h"
+
+/*=============================================================================
+	FastReferenceCollector.h: Unreal realtime garbage collection helpers
+=============================================================================*/
+
+/**
+ * Helper class that looks for UObject references by traversing UClass token stream and calls AddReferencedObjects.
+ * Provides a generic way of processing references that is used by Unreal Engine garbage collection.
+ * Can be used for fast (does not use serialization) reference collection purposes.
+ * 
+ * IT IS CRITICAL THIS CLASS DOES NOT CHANGE WITHOUT CONSIDERING PERFORMANCE IMPACT OF SAID CHANGES
+ *
+ * This class depends on three components: ReferenceProcessor, ReferenceCollector and ArrayPool.
+ * The assumptions for each of those components are as follows:
+ *
+   class FSampleReferenceProcessor
+   {
+   public:
+     int32 GetMinDesiredObjectsPerSubTask() const;
+     volatile bool IsRunningMultithreaded() const;
+		 void SetIsRunningMultithreaded(bool bIsParallel);
+		 void HandleTokenStreamObjectReference(TArray<UObject*>& ObjectsToSerialize, UObject* ReferencingObject, UObject*& Object, const int32 TokenIndex, bool bAllowReferenceElimination);
+		 void UpdateDetailedStats(UObject* CurrentObject, uint32 DeltaCycles);
+		 void LogDetailedStatsSummary();
+	 };
+
+	 class FSampleCollctor : public FReferenceCollector 
+	 {
+	   // Needs to implement FReferenceCollector pure virtual functions
+	 };
+   
+	 class FSampleArrayPool
+	 {
+	   static FSampleArrayPool& Get();
+		 TArray<UObject*>* GetArrayFromPool();
+		 void ReturnToPool(TArray<UObject*>* Array);
+	 };
+ */
+template <typename ReferenceProcessorType, typename CollectorType, typename ArrayPoolType>
+class TFastReferenceCollector
+{
+private:
+
+	/** Task graph task responsible for processing UObject array */
+	class FCollectorTask
+	{
+		TFastReferenceCollector*	Owner;
+		TArray<UObject*>*	ObjectsToSerialize;
+		ArrayPoolType& ArrayPool;
+
+	public:
+		FCollectorTask(TFastReferenceCollector* InOwner, const TArray<UObject*>* InObjectsToSerialize, int32 StartIndex, int32 NumObjects, ArrayPoolType& InArrayPool)
+			: Owner(InOwner)
+			, ObjectsToSerialize(InArrayPool.GetArrayFromPool())
+			, ArrayPool(InArrayPool)
+		{
+			ObjectsToSerialize->AddUninitialized(NumObjects);
+			FMemory::Memcpy(ObjectsToSerialize->GetData(), InObjectsToSerialize->GetData() + StartIndex, NumObjects * sizeof(UObject*));
+		}
+		~FCollectorTask()
+		{
+			ArrayPool.ReturnToPool(ObjectsToSerialize);
+		}
+		FORCEINLINE TStatId GetStatId() const
+		{
+			RETURN_QUICK_DECLARE_CYCLE_STAT(FCollectorTask, STATGROUP_TaskGraphTasks);
+		}
+		static ENamedThreads::Type GetDesiredThread()
+		{
+			return ENamedThreads::AnyThread;
+		}
+		static ESubsequentsMode::Type GetSubsequentsMode()
+		{
+			return ESubsequentsMode::TrackSubsequents;
+		}
+		void DoTask(ENamedThreads::Type CurrentThread, FGraphEventRef& MyCompletionGraphEvent)
+		{
+			Owner->ProcessObjectArray(*ObjectsToSerialize, MyCompletionGraphEvent);
+		}
+	};
+
+	/** Object that handles all UObject references */
+	ReferenceProcessorType& ReferenceProcessor;
+	/** Custom TArray allocator */
+	ArrayPoolType& ArrayPool;
+
+	/** Helper struct for stack based approach */
+	struct FStackEntry
+	{
+		/** Current data pointer, incremented by stride */
+		uint8*	Data;
+		/** Current stride */
+		int32		Stride;
+		/** Current loop count, decremented each iteration */
+		int32		Count;
+		/** First token index in loop */
+		int32		LoopStartIndex;
+	};
+
+
+
+public:
+	/** Default constructor, initializing all members. */
+	TFastReferenceCollector(ReferenceProcessorType& InReferenceProcessor, ArrayPoolType& InArrayPool)
+		: ReferenceProcessor(InReferenceProcessor)
+		, ArrayPool(InArrayPool)
+	{}
+
+	/**
+	* Performs reachability analysis.
+	*
+	* @param ObjectsToCollectReferencesFor List of objects which references should be collected
+	* @param bForceSingleThreaded Collect references on a single thread
+	*/
+	void CollectReferences(TArray<UObject*>& ObjectsToCollectReferencesFor, bool bForceSingleThreaded = false)
+	{
+		if (ObjectsToCollectReferencesFor.Num())
+		{
+			check(!ReferenceProcessor.IsRunningMultithreaded());
+
+			if (bForceSingleThreaded)
+			{
+				FGraphEventRef InvalidRef;
+				ReferenceProcessor.SetIsRunningMultithreaded(false);
+				ProcessObjectArray(ObjectsToCollectReferencesFor, InvalidRef);
+			}
+			else
+			{
+				ReferenceProcessor.SetIsRunningMultithreaded(true);
+
+				int32 NumChunks = FMath::Min<int32>(FTaskGraphInterface::Get().GetNumWorkerThreads(), ObjectsToCollectReferencesFor.Num());
+				int32 NumPerChunk = ObjectsToCollectReferencesFor.Num() / NumChunks;
+				check(NumPerChunk > 0);
+				FGraphEventArray ChunkTasks;
+				ChunkTasks.Empty(NumChunks);
+				int32 StartIndex = 0;
+				for (int32 Chunk = 0; Chunk < NumChunks; Chunk++)
+				{
+					if (Chunk + 1 == NumChunks)
+					{
+						NumPerChunk = ObjectsToCollectReferencesFor.Num() - StartIndex; // last chunk takes all remaining items
+					}
+					ChunkTasks.Add(TGraphTask< FCollectorTask >::CreateTask().ConstructAndDispatchWhenReady(this, &ObjectsToCollectReferencesFor, StartIndex, NumPerChunk, ArrayPool));
+					StartIndex += NumPerChunk;
+				}
+				QUICK_SCOPE_CYCLE_COUNTER(STAT_GC_Subtask_Wait);
+				FTaskGraphInterface::Get().WaitUntilTasksComplete(ChunkTasks, ENamedThreads::GameThread_Local);
+				ReferenceProcessor.SetIsRunningMultithreaded(false);
+			}
+		}
+	}
+
+private:
+
+	/**
+	 * Traverses UObject token stream to find existing references
+	 *
+	 * @param InObjectsToSerializeArray Objects to process
+	 * @param MyCompletionGraphEvent Task graph event
+	 */
+	void ProcessObjectArray(TArray<UObject*>& InObjectsToSerializeArray, FGraphEventRef& MyCompletionGraphEvent)
+	{
+		DECLARE_SCOPE_CYCLE_COUNTER(TEXT("TFastReferenceCollector::ProcessObjectArray"), STAT_FFastReferenceCollector_ProcessObjectArray, STATGROUP_GC);
+
+		UObject* CurrentObject = NULL;
+
+		const int32 MinDesiredObjectsPerSubTask = ReferenceProcessor.GetMinDesiredObjectsPerSubTask(); // sometimes there will be less, a lot less
+
+		/** Growing array of objects that require serialization */
+		TArray<UObject*>&	NewObjectsToSerializeArray = *ArrayPool.GetArrayFromPool();
+
+		// Ping-pong between these two arrays if there's not enough objects to spawn a new task
+		TArray<UObject*>& ObjectsToSerialize = InObjectsToSerializeArray;
+		TArray<UObject*>& NewObjectsToSerialize = NewObjectsToSerializeArray;
+
+		// Presized "recursion" stack for handling arrays and structs.
+		TArray<FStackEntry> Stack;
+		Stack.AddUninitialized(128); //@todo rtgc: need to add code handling more than 128 layers of recursion or at least assert
+
+		// it is necessary to have at least one extra item in the array memory block for the iffy prefetch code, below
+		ObjectsToSerialize.Reserve(ObjectsToSerialize.Num() + 1);
+
+		// Keep serializing objects till we reach the end of the growing array at which point
+		// we are done.
+		int32 CurrentIndex = 0;
+		do
+		{
+			CollectorType ReferenceCollector(ReferenceProcessor, NewObjectsToSerialize);
+			while (CurrentIndex < ObjectsToSerialize.Num())
+			{
+#if PERF_DETAILED_PER_CLASS_GC_STATS
+				uint32 StartCycles = FPlatformTime::Cycles();
+#endif
+				CurrentObject = ObjectsToSerialize[CurrentIndex++];
+
+				// GetData() used to avoiding bounds checking (min and max)
+				// FMath::Min used to avoid out of bounds (without branching) on last iteration. Though anything can be passed into PrefetchBlock, 
+				// reading ObjectsToSerialize out of bounds is not safe since ObjectsToSerialize[Num()] may be an unallocated/unsafe address.
+				const UObject * const NextObject = ObjectsToSerialize.GetData()[FMath::Min<int32>(CurrentIndex, ObjectsToSerialize.Num() - 1)];
+
+				// Prefetch the next object assuming that the property size of the next object is the same as the current one.
+				// This allows us to avoid a branch here.
+				FPlatformMisc::PrefetchBlock(NextObject, CurrentObject->GetClass()->GetPropertiesSize());
+
+				//@todo rtgc: we need to handle object references in struct defaults
+
+				// Make sure that token stream has been assembled at this point as the below code relies on it.
+				checkSlow(CurrentObject->GetClass()->HasAnyClassFlags(CLASS_TokenStreamAssembled));
+
+				// Get pointer to token stream and jump to the start.
+				FGCReferenceTokenStream* RESTRICT TokenStream = &CurrentObject->GetClass()->ReferenceTokenStream;
+				uint32 TokenStreamIndex = 0;
+				// Keep track of index to reference info. Used to avoid LHSs.
+				uint32 ReferenceTokenStreamIndex = 0;
+
+				// Create stack entry and initialize sane values.
+				FStackEntry* RESTRICT StackEntry = Stack.GetData();
+				uint8* StackEntryData = (uint8*)CurrentObject;
+				StackEntry->Data = StackEntryData;
+				StackEntry->Stride = 0;
+				StackEntry->Count = -1;
+				StackEntry->LoopStartIndex = -1;
+
+				// Keep track of token return count in separate integer as arrays need to fiddle with it.
+				int32 TokenReturnCount = 0;
+
+				// Parse the token stream.
+				while (true)
+				{
+					// Cache current token index as it is the one pointing to the reference info.
+					ReferenceTokenStreamIndex = TokenStreamIndex;
+
+					// Handle returning from an array of structs, array of structs of arrays of ... (yadda yadda)
+					for (int32 ReturnCount = 0; ReturnCount<TokenReturnCount; ReturnCount++)
+					{
+						// Make sure there's no stack underflow.
+						check(StackEntry->Count != -1);
+
+						// We pre-decrement as we're already through the loop once at this point.
+						if (--StackEntry->Count > 0)
+						{
+							// Point data to next entry.
+							StackEntryData = StackEntry->Data + StackEntry->Stride;
+							StackEntry->Data = StackEntryData;
+
+							// Jump back to the beginning of the loop.
+							TokenStreamIndex = StackEntry->LoopStartIndex;
+							ReferenceTokenStreamIndex = StackEntry->LoopStartIndex;
+							// We're not done with this token loop so we need to early out instead of backing out further.
+							break;
+						}
+						else
+						{
+							StackEntry--;
+							StackEntryData = StackEntry->Data;
+						}
+					}
+
+					TokenStreamIndex++;
+					FGCReferenceInfo ReferenceInfo = TokenStream->AccessReferenceInfo(ReferenceTokenStreamIndex);
+
+					if (ReferenceInfo.Type == GCRT_Object)
+					{
+						// We're dealing with an object reference.
+						UObject**	ObjectPtr = (UObject**)(StackEntryData + ReferenceInfo.Offset);
+						UObject*&	Object = *ObjectPtr;
+						TokenReturnCount = ReferenceInfo.ReturnCount;
+						ReferenceProcessor.HandleTokenStreamObjectReference(NewObjectsToSerialize, CurrentObject, Object, ReferenceTokenStreamIndex, true);
+					}
+					else if (ReferenceInfo.Type == GCRT_ArrayObject)
+					{
+						// We're dealing with an array of object references.
+						TArray<UObject*>& ObjectArray = *((TArray<UObject*>*)(StackEntryData + ReferenceInfo.Offset));
+						TokenReturnCount = ReferenceInfo.ReturnCount;
+						for (int32 ObjectIndex = 0, ObjectNum = ObjectArray.Num(); ObjectIndex < ObjectNum; ++ObjectIndex)
+						{
+							ReferenceProcessor.HandleTokenStreamObjectReference(NewObjectsToSerialize, CurrentObject, ObjectArray[ObjectIndex], ReferenceTokenStreamIndex, true);
+						}
+					}
+					else if (ReferenceInfo.Type == GCRT_ArrayStruct)
+					{
+						// We're dealing with a dynamic array of structs.
+						const FScriptArray& Array = *((FScriptArray*)(StackEntryData + ReferenceInfo.Offset));
+						StackEntry++;
+						StackEntryData = (uint8*)Array.GetData();
+						StackEntry->Data = StackEntryData;
+						StackEntry->Stride = TokenStream->ReadStride(TokenStreamIndex);
+						StackEntry->Count = Array.Num();
+
+						const FGCSkipInfo SkipInfo = TokenStream->ReadSkipInfo(TokenStreamIndex);
+						StackEntry->LoopStartIndex = TokenStreamIndex;
+
+						if (StackEntry->Count == 0)
+						{
+							// Skip empty array by jumping to skip index and set return count to the one about to be read in.
+							TokenStreamIndex = SkipInfo.SkipIndex;
+							TokenReturnCount = TokenStream->GetSkipReturnCount(SkipInfo);
+						}
+						else
+						{
+							// Loop again.
+							check(StackEntry->Data);
+							TokenReturnCount = 0;
+						}
+					}
+					else if (ReferenceInfo.Type == GCRT_PersistentObject)
+					{
+						// We're dealing with an object reference.
+						UObject**	ObjectPtr = (UObject**)(StackEntryData + ReferenceInfo.Offset);
+						UObject*&	Object = *ObjectPtr;
+						TokenReturnCount = ReferenceInfo.ReturnCount;
+						ReferenceProcessor.HandleTokenStreamObjectReference(NewObjectsToSerialize, CurrentObject, Object, ReferenceTokenStreamIndex, false);
+					}
+					else if (ReferenceInfo.Type == GCRT_FixedArray)
+					{
+						// We're dealing with a fixed size array
+						uint8* PreviousData = StackEntryData;
+						StackEntry++;
+						StackEntryData = PreviousData;
+						StackEntry->Data = PreviousData;
+						StackEntry->Stride = TokenStream->ReadStride(TokenStreamIndex);
+						StackEntry->Count = TokenStream->ReadCount(TokenStreamIndex);
+						StackEntry->LoopStartIndex = TokenStreamIndex;
+						TokenReturnCount = 0;
+					}
+					else if (ReferenceInfo.Type == GCRT_AddStructReferencedObjects)
+					{
+						// We're dealing with a function call
+						void const*	StructPtr = (void*)(StackEntryData + ReferenceInfo.Offset);
+						TokenReturnCount = ReferenceInfo.ReturnCount;
+						UScriptStruct::ICppStructOps::TPointerToAddStructReferencedObjects Func = (UScriptStruct::ICppStructOps::TPointerToAddStructReferencedObjects) TokenStream->ReadPointer(TokenStreamIndex);
+						Func(StructPtr, ReferenceCollector);
+					}
+					else if (ReferenceInfo.Type == GCRT_AddReferencedObjects)
+					{
+						// Static AddReferencedObjects function call.
+						void(*AddReferencedObjects)(UObject*, FReferenceCollector&) = (void(*)(UObject*, FReferenceCollector&))TokenStream->ReadPointer(TokenStreamIndex);
+						TokenReturnCount = ReferenceInfo.ReturnCount;
+						AddReferencedObjects(CurrentObject, ReferenceCollector);
+					}
+					else if (ReferenceInfo.Type == GCRT_AddTMapReferencedObjects)
+					{
+						void*         Map = StackEntryData + ReferenceInfo.Offset;
+						UMapProperty* MapProperty = (UMapProperty*)TokenStream->ReadPointer(TokenStreamIndex);
+						TokenReturnCount = ReferenceInfo.ReturnCount;
+						FSimpleObjectReferenceCollectorArchive CollectorArchive(CurrentObject, ReferenceCollector);
+						MapProperty->SerializeItem(CollectorArchive, Map, nullptr);
+					}
+					else if (ReferenceInfo.Type == GCRT_EndOfPointer)
+					{
+						TokenReturnCount = ReferenceInfo.ReturnCount;
+					}
+					else if (ReferenceInfo.Type == GCRT_EndOfStream)
+					{
+						// Break out of loop.
+						break;
+					}
+					else
+					{
+						UE_LOG(LogGarbage, Fatal, TEXT("Unknown token"));
+					}
+				}
+				check(StackEntry == Stack.GetData());
+
+				if (ReferenceProcessor.IsRunningMultithreaded() && NewObjectsToSerialize.Num() >= MinDesiredObjectsPerSubTask)
+				{
+					// This will start queueing task with objects from the end of array until there's less objects than worth to queue
+					const int32 ObjectsPerSubTask = FMath::Max<int32>(MinDesiredObjectsPerSubTask, NewObjectsToSerialize.Num() / FTaskGraphInterface::Get().GetNumWorkerThreads());
+					while (NewObjectsToSerialize.Num() >= MinDesiredObjectsPerSubTask)
+					{
+						const int32 StartIndex = FMath::Max(0, NewObjectsToSerialize.Num() - ObjectsPerSubTask);
+						const int32 NumThisTask = NewObjectsToSerialize.Num() - StartIndex;
+						MyCompletionGraphEvent->DontCompleteUntil(TGraphTask< FCollectorTask >::CreateTask().ConstructAndDispatchWhenReady(this, &NewObjectsToSerialize, StartIndex, NumThisTask, ArrayPool));
+						NewObjectsToSerialize.SetNumUnsafeInternal(StartIndex);
+					}
+				}
+#if PERF_DETAILED_PER_CLASS_GC_STATS
+				// Detailed per class stats should not be performed when parallel GC is running
+				check(!ReferenceProcessor.IsRunningMultithreaded());
+				ReferenceProcessor.UpdateDetailedStats(CurrentObject, FPlatformTime::Cycles() - StartCycles);
+			}
+			// Log summary stats.
+			ReferenceProcessor.LogDetailedStatsSummary();
+#else
+			}
+#endif
+			if (ReferenceProcessor.IsRunningMultithreaded() && NewObjectsToSerialize.Num() >= MinDesiredObjectsPerSubTask)
+			{
+				const int32 ObjectsPerSubTask = FMath::Max<int32>(MinDesiredObjectsPerSubTask, NewObjectsToSerialize.Num() / FTaskGraphInterface::Get().GetNumWorkerThreads());
+				int32 StartIndex = 0;
+				while (StartIndex < NewObjectsToSerialize.Num())
+				{
+					const int32 NumThisTask = FMath::Min<int32>(ObjectsPerSubTask, NewObjectsToSerialize.Num() - StartIndex);
+					MyCompletionGraphEvent->DontCompleteUntil(TGraphTask< FCollectorTask >::CreateTask().ConstructAndDispatchWhenReady(this, &NewObjectsToSerialize, StartIndex, NumThisTask, ArrayPool));
+					StartIndex += NumThisTask;
+				}
+				NewObjectsToSerialize.SetNumUnsafeInternal(0);
+			}
+			else if (NewObjectsToSerialize.Num())
+			{
+				// Don't spawn a new task, continue in the current one
+				// To avoid allocating and moving memory around swap ObjectsToSerialize and NewObjectsToSerialize arrays
+				Exchange(ObjectsToSerialize, NewObjectsToSerialize);
+				// Empty but don't free allocated memory
+				NewObjectsToSerialize.SetNumUnsafeInternal(0);
+
+				CurrentIndex = 0;
+			}
+		}
+		while (CurrentIndex < ObjectsToSerialize.Num());
+
+		ArrayPool.ReturnToPool(&NewObjectsToSerializeArray);
+	}
+};
\ No newline at end of file
diff --git a/Engine/Source/Runtime/CoreUObject/Private/UObject/GarbageCollection.cpp b/Engine/Source/Runtime/CoreUObject/Private/UObject/GarbageCollection.cpp
--- a/Engine/Source/Runtime/CoreUObject/Private/UObject/GarbageCollection.cpp
+++ b/Engine/Source/Runtime/CoreUObject/Private/UObject/GarbageCollection.cpp
@@ -14,11 +14,12 @@
    Garbage collection.
 -----------------------------------------------------------------------------*/
 
-DECLARE_STATS_GROUP( TEXT( "Garbage Collection" ), STATGROUP_GC, STATCAT_Advanced );
+#define PERF_DETAILED_PER_CLASS_GC_STATS				(LOOKING_FOR_PERF_ISSUES)
 
-DEFINE_LOG_CATEGORY_STATIC(LogGarbage, Warning, All);
+// FastReferenceCollector uses PERF_DETAILED_PER_CLASS_GC_STATS
+#include "UObject/FastReferenceCollector.h"
 
-#define PERF_DETAILED_PER_CLASS_GC_STATS				(LOOKING_FOR_PERF_ISSUES)
+DEFINE_LOG_CATEGORY(LogGarbage);
 
 // UE_BUILD_SHIPPING has GShouldVerifyGCAssumptions=false by default
 #define VERIFY_DISREGARD_GC_ASSUMPTIONS			!(UE_BUILD_SHIPPING || UE_BUILD_TEST)
@@ -178,10 +179,12 @@ private:
 #endif // UE_BUILD_DEBUG
 };
 
+void CleanupClusterArrayPools();
 /** Called on shutdown to free GC memory */
 void CleanupGCArrayPools()
 {
 	FGCArrayPool::Get().Cleanup();
+	CleanupClusterArrayPools();
 }
 
 /**
@@ -281,114 +284,249 @@ static void LogClassCountInfo( const TCHAR* LogText, TMap<const FName,uint32>& C
 #endif
 
 /**
- * Handles object reference, potentially NULL'ing
- *
- * @param Object						Object pointer passed by reference
- * @param ReferencingObject UObject which owns the reference (can be NULL)
- * @param bAllowReferenceElimination	Whether to allow NULL'ing the reference if RF_PendingKill is set
- */
-static FORCEINLINE void HandleObjectReference(TArray<UObject*>& ObjectsToSerialize, const UObject * const ReferencingObject, UObject*& Object, const bool bAllowReferenceElimination, const bool bStrongReference = true)
+* Handles UObject references found by TFastReferenceCollector
+*/
+class FGCReferenceProcessor
 {
-	// Disregard NULL objects and perform very fast check to see whether object is part of permanent
-	// object pool and should therefore be disregarded. The check doesn't touch the object and is
-	// cache friendly as it's just a pointer compare against to globals.
-	const bool IsInPermanentPool = GUObjectAllocator.ResidesInPermanentPool(Object);
+public:
 
-#if PERF_DETAILED_PER_CLASS_GC_STATS
-	if(IsInPermanentPool)
+	FGCReferenceProcessor()
 	{
-		GCurrentObjectDisregardedObjectRefs++;
 	}
-#endif
-	if(Object == nullptr || IsInPermanentPool)
+
+	FORCEINLINE int32 GetMinDesiredObjectsPerSubTask() const
 	{
-		return;
+		return GMinDesiredObjectsPerSubTask;
+	}
+
+	FORCEINLINE volatile bool IsRunningMultithreaded() const
+	{
+		return GIsRunningParallelReachability;
 	}
 
-	FUObjectItem* ObjectItem = GUObjectArray.ObjectToObjectItem(Object);
-	// Remove references to pending kill objects if we're allowed to do so.
-	if (ObjectItem->IsPendingKill() && bAllowReferenceElimination)
+	FORCEINLINE void SetIsRunningMultithreaded(bool bIsParallel)
 	{
-		// Null out reference.
-		Object = NULL;
+		GIsRunningParallelReachability = bIsParallel;
 	}
-	// Add encountered object reference to list of to be serialized objects if it hasn't already been added.
-	else if (ObjectItem->IsUnreachable())
-	{				
-		if( GIsRunningParallelReachability )
+
+	void UpdateDetailedStats(UObject* CurrentObject, uint32 DeltaCycles)
+	{
+#if PERF_DETAILED_PER_CLASS_GC_STATS
+		// Keep track of how many refs we encountered for the object's class.
+		const FName& ClassName = CurrentObject->GetClass()->GetFName();
+		// Refs to objects that reside in permanent object pool.
+		uint32 ClassDisregardedObjRefs = GClassToDisregardedObjectRefsMap.FindRef(ClassName);
+		GClassToDisregardedObjectRefsMap.Add(ClassName, ClassDisregardedObjRefs + GCurrentObjectDisregardedObjectRefs);
+		// Refs to regular objects.
+		uint32 ClassRegularObjRefs = GClassToRegularObjectRefsMap.FindRef(ClassName);
+		GClassToRegularObjectRefsMap.Add(ClassName, ClassRegularObjRefs + GCurrentObjectRegularObjectRefs);
+		// Track per class cycle count spent in GC.
+		uint32 ClassCycles = GClassToCyclesMap.FindRef(ClassName);
+		GClassToCyclesMap.Add(ClassName, ClassCycles + DeltaCycles);
+		// Reset current counts.
+		GCurrentObjectDisregardedObjectRefs = 0;
+		GCurrentObjectRegularObjectRefs = 0;
+#endif
+	}
+
+	void LogDetailedStatsSummary()
+	{
+#if PERF_DETAILED_PER_CLASS_GC_STATS
+		LogClassCountInfo(TEXT("references to regular objects from"), GClassToRegularObjectRefsMap, 20, 0);
+		LogClassCountInfo(TEXT("references to permanent objects from"), GClassToDisregardedObjectRefsMap, 20, 0);
+		LogClassCountInfo(TEXT("cycles for GC"), GClassToCyclesMap, 20, 0);
+#endif
+	}
+
+	template <bool bParallel>
+	FORCEINLINE void MarkReferencedClustersAsReachable(int32 ClusterRootIndex)
+	{
+		FUObjectCluster* Cluster = GUObjectClusters.FindChecked(ClusterRootIndex);
+		for (int32 ReferncedClusterIndex : Cluster->ReferencedClusters)
 		{
-			// Mark it as reachable.
-			if (ObjectItem->ThisThreadAtomicallyClearedRFUnreachable())
+			FUObjectItem* ReferencedClusterRootObjectItem = GUObjectArray.IndexToObjectUnsafeForGC(ReferncedClusterIndex);
+			// This condition should get collapsed by the compiler based on the template argument
+			if (bParallel)
 			{
-				// Add it to the list of objects to serialize.
-				ObjectsToSerialize.Add( Object );
+				ReferencedClusterRootObjectItem->ThisThreadAtomicallyClearedRFUnreachable();
 			}
-		}
-		else
-		{
-#if ENABLE_GC_DEBUG_OUTPUT
-			// this message is to help track down culprits behind "Object in PIE world still referenced" errors
-			if ( GIsEditor && !GIsPlayInEditorWorld && ReferencingObject != NULL && !ReferencingObject->RootPackageHasAnyFlags(PKG_PlayInEditor) && Object->RootPackageHasAnyFlags(PKG_PlayInEditor) )
+			else
 			{
-				UE_LOG(LogGarbage, Warning, TEXT("GC detected illegal reference to PIE object from content [possibly via [todo]]:"));
-				UE_LOG(LogGarbage, Warning, TEXT("      PIE object: %s"), *Object->GetFullName());
-				UE_LOG(LogGarbage, Warning, TEXT("  NON-PIE object: %s"), *ReferencingObject->GetFullName());
+				ReferencedClusterRootObjectItem->ClearUnreachable();
 			}
-#endif
-
-			// Mark it as reachable.
-			ObjectItem->ClearUnreachable();
-			// Add it to the list of objects to serialize.
-			ObjectsToSerialize.Add( Object );
 		}
 	}
 
-	if (bStrongReference)
+	/**
+	 * Handles object reference, potentially NULL'ing
+	 *
+	 * @param Object						Object pointer passed by reference
+	 * @param ReferencingObject UObject which owns the reference (can be NULL)
+	 * @param bAllowReferenceElimination	Whether to allow NULL'ing the reference if RF_PendingKill is set
+	*/
+	FORCEINLINE void HandleObjectReference(TArray<UObject*>& ObjectsToSerialize, const UObject * const ReferencingObject, UObject*& Object, const bool bAllowReferenceElimination, const bool bStrongReference = true)
 	{
-		ObjectItem->ClearNoStrongReference();
-	}
+		// Disregard NULL objects and perform very fast check to see whether object is part of permanent
+		// object pool and should therefore be disregarded. The check doesn't touch the object and is
+		// cache friendly as it's just a pointer compare against to globals.
+		const bool IsInPermanentPool = GUObjectAllocator.ResidesInPermanentPool(Object);
+
 #if PERF_DETAILED_PER_CLASS_GC_STATS
-	GCurrentObjectRegularObjectRefs++;
+		if (IsInPermanentPool)
+		{
+			GCurrentObjectDisregardedObjectRefs++;
+		}
 #endif
-}
+		if (Object == nullptr || IsInPermanentPool)
+		{
+			return;
+		}
 
-static FORCEINLINE void HandleTokenStreamObjectReference(TArray<UObject*>& ObjectsToSerialize, UObject* ReferencingObject, UObject*& Object, const int32 TokenIndex, bool bAllowReferenceElimination)
-{
-#if !(UE_BUILD_TEST || UE_BUILD_SHIPPING)
-	if (Object && !Object->IsValidLowLevelFast())
-	{
-		FString TokenDebugInfo;
-		if (UClass *Class = (ReferencingObject ? ReferencingObject->GetClass() : nullptr))
+		FUObjectItem* ObjectItem = GUObjectArray.ObjectToObjectItem(Object);
+		// Remove references to pending kill objects if we're allowed to do so.
+		if (ObjectItem->IsPendingKill() && bAllowReferenceElimination)
 		{
-			auto& TokenInfo = Class->DebugTokenMap.GetTokenInfo(TokenIndex);
-			TokenDebugInfo = FString::Printf(TEXT("ReferencingObjectClass: %s, Property Name: %s, Offset: %d"),
-				*Class->GetFullName(), *TokenInfo.Name.GetPlainNameString(), TokenInfo.Offset);
+			// Null out reference.
+			Object = NULL;
 		}
-		else
+		// Add encountered object reference to list of to be serialized objects if it hasn't already been added.
+		else if (ObjectItem->IsUnreachable())
+		{
+			if (GIsRunningParallelReachability)
+			{
+				// Mark it as reachable.
+				if (ObjectItem->ThisThreadAtomicallyClearedRFUnreachable())
+				{
+					// Objects that are part of a GC cluster should never have the unreachable flag set!
+					checkSlow(ObjectItem->GetOwnerIndex() == 0);
+
+					if (!ObjectItem->HasAnyFlags(EInternalObjectFlags::ClusterRoot))
+					{
+						// Add it to the list of objects to serialize.
+						ObjectsToSerialize.Add(Object);
+					}
+					else
+					{
+						// This is a cluster root reference so mark all referenced clusters as reachable
+						const int32 ObjectIndex = GUObjectArray.ObjectToIndex(Object);
+						MarkReferencedClustersAsReachable<true>(ObjectIndex);
+					}
+				}
+			}
+			else
+			{
+#if ENABLE_GC_DEBUG_OUTPUT
+				// this message is to help track down culprits behind "Object in PIE world still referenced" errors
+				if (GIsEditor && !GIsPlayInEditorWorld && ReferencingObject != NULL && !ReferencingObject->RootPackageHasAnyFlags(PKG_PlayInEditor) && Object->RootPackageHasAnyFlags(PKG_PlayInEditor))
+				{
+					UE_LOG(LogGarbage, Warning, TEXT("GC detected illegal reference to PIE object from content [possibly via [todo]]:"));
+					UE_LOG(LogGarbage, Warning, TEXT("      PIE object: %s"), *Object->GetFullName());
+					UE_LOG(LogGarbage, Warning, TEXT("  NON-PIE object: %s"), *ReferencingObject->GetFullName());
+				}
+#endif
+
+				// Mark it as reachable.
+				ObjectItem->ClearUnreachable();
+
+				// Objects that are part of a GC cluster should never have the unreachable flag set!
+				checkSlow(ObjectItem->GetOwnerIndex() == 0);
+
+				if (!ObjectItem->HasAnyFlags(EInternalObjectFlags::ClusterRoot))
+				{
+					// Add it to the list of objects to serialize.
+					ObjectsToSerialize.Add(Object);
+				}
+				else
+				{
+					// This is a cluster root reference so mark all referenced clusters as reachable
+					const int32 ObjectIndex = GUObjectArray.ObjectToIndex(Object);
+					MarkReferencedClustersAsReachable<false>(ObjectIndex);
+				}
+			}
+		}
+		else if (ObjectItem->GetOwnerIndex() && !ObjectItem->HasAnyFlags(EInternalObjectFlags::ReachableInCluster))
 		{
-			// This means this objects is most likely being referenced by AddReferencedObjects
-			TokenDebugInfo = TEXT("Native Reference");
+			ObjectItem->SetFlags(EInternalObjectFlags::ReachableInCluster);
+			// Make sure cluster root object is reachable too
+			const int32 OwnerIndex = ObjectItem->GetOwnerIndex();
+			FUObjectItem* RootObjectItem = GUObjectArray.IndexToObjectUnsafeForGC(OwnerIndex);
+			checkSlow(RootObjectItem->HasAnyFlags(EInternalObjectFlags::ClusterRoot));
+			if (GIsRunningParallelReachability)
+			{
+				if (RootObjectItem->ThisThreadAtomicallyClearedRFUnreachable())
+				{
+					// Make sure all referenced clusters are marked as reachable too
+					MarkReferencedClustersAsReachable<true>(OwnerIndex);
+				}
+			}
+			else if (RootObjectItem->IsUnreachable())
+			{
+				RootObjectItem->ClearUnreachable();
+				// Make sure all referenced clusters are marked as reachable too
+				MarkReferencedClustersAsReachable<false>(OwnerIndex);
+			}
 		}
 
-		UE_LOG(LogGarbage, Fatal, TEXT("Invalid object in GC: 0x%016llx, ReferencingObject: %s, %s, TokenIndex: %d"),
-			(int64)(PTRINT)Object,
-			ReferencingObject ? *ReferencingObject->GetFullName() : TEXT("NULL"),
-			*TokenDebugInfo, TokenIndex);
+		if (bStrongReference)
+		{
+			ObjectItem->ClearNoStrongReference();
+		}
+#if PERF_DETAILED_PER_CLASS_GC_STATS
+		GCurrentObjectRegularObjectRefs++;
+#endif
 	}
+
+	/**
+	* Handles UObject reference from the token stream.
+	*
+	* @param ObjectsToSerialize An array of remaining objects to serialize.
+	* @param ReferencingObject Object referencing the object to process.
+	* @param TokenIndex Index to the token stream where the reference was found.
+	* @param bAllowReferenceElimination True if reference elimination is allowed.
+	*/
+	FORCEINLINE void HandleTokenStreamObjectReference(TArray<UObject*>& ObjectsToSerialize, UObject* ReferencingObject, UObject*& Object, const int32 TokenIndex, bool bAllowReferenceElimination)
+	{
+#if !(UE_BUILD_TEST || UE_BUILD_SHIPPING)
+		if (Object && !Object->IsValidLowLevelFast())
+		{
+			FString TokenDebugInfo;
+			if (UClass *Class = (ReferencingObject ? ReferencingObject->GetClass() : nullptr))
+			{
+				auto& TokenInfo = Class->DebugTokenMap.GetTokenInfo(TokenIndex);
+				TokenDebugInfo = FString::Printf(TEXT("ReferencingObjectClass: %s, Property Name: %s, Offset: %d"),
+					*Class->GetFullName(), *TokenInfo.Name.GetPlainNameString(), TokenInfo.Offset);
+			}
+			else
+			{
+				// This means this objects is most likely being referenced by AddReferencedObjects
+				TokenDebugInfo = TEXT("Native Reference");
+			}
+
+			UE_LOG(LogGarbage, Fatal, TEXT("Invalid object in GC: 0x%016llx, ReferencingObject: %s, %s, TokenIndex: %d"),
+				(int64)(PTRINT)Object,
+				ReferencingObject ? *ReferencingObject->GetFullName() : TEXT("NULL"),
+				*TokenDebugInfo, TokenIndex);
+		}
 #endif
-	HandleObjectReference(ObjectsToSerialize, ReferencingObject, Object, bAllowReferenceElimination);
-}
+		HandleObjectReference(ObjectsToSerialize, ReferencingObject, Object, bAllowReferenceElimination);
+	}
+};
 
+/**
+* Specialized FReferenceCollector that uses FGCReferenceProcessor to mark objects as reachable.
+*/
 class FGCCollector : public FReferenceCollector
 {
+	FGCReferenceProcessor& ReferenceProcessor;
 	TArray<UObject*>& ObjectArray;
 	bool bAllowEliminatingReferences;
 	bool bShouldHandleAsWeakRef;
 
 public:
 
-	FGCCollector(TArray<UObject*>& InObjectArray)
-		: ObjectArray(InObjectArray)
+	FGCCollector(FGCReferenceProcessor& InProcessor, TArray<UObject*>& InObjectArray)
+		: ReferenceProcessor(InProcessor)
+		, ObjectArray(InObjectArray)
 		, bAllowEliminatingReferences(true)
 		, bShouldHandleAsWeakRef(false)
 	{
@@ -405,8 +543,26 @@ public:
 				ReferencingProperty ? *ReferencingProperty->GetFullName() : TEXT("NULL"));
 		}
 #endif
-		::HandleObjectReference(ObjectArray, const_cast<UObject*>(ReferencingObject), Object, bAllowEliminatingReferences, !bShouldHandleAsWeakRef);
+		ReferenceProcessor.HandleObjectReference(ObjectArray, const_cast<UObject*>(ReferencingObject), Object, bAllowEliminatingReferences, !bShouldHandleAsWeakRef);
+	}
+	virtual void HandleObjectReferences(UObject** InObjects, const int32 ObjectNum, const UObject* InReferencingObject, const UProperty* InReferencingProperty) override
+	{
+		for (int32 ObjectIndex = 0; ObjectIndex < ObjectNum; ++ObjectIndex)
+		{
+			UObject*& Object = InObjects[ObjectIndex];
+#if !(UE_BUILD_TEST || UE_BUILD_SHIPPING)
+			if (Object && !Object->IsValidLowLevelFast())
+			{
+				UE_LOG(LogGarbage, Fatal, TEXT("Invalid object in GC: 0x%016llx, ReferencingObject: %s, ReferencingProperty: %s"),
+					(int64)(PTRINT)Object,
+					InReferencingObject ? *InReferencingObject->GetFullName() : TEXT("NULL"),
+					InReferencingProperty ? *InReferencingProperty->GetFullName() : TEXT("NULL"));
+			}
+#endif
+			ReferenceProcessor.HandleObjectReference(ObjectArray, const_cast<UObject*>(InReferencingObject), Object, bAllowEliminatingReferences, !bShouldHandleAsWeakRef);
+		}
 	}
+
 	virtual bool IsIgnoringArchetypeRef() const override
 	{
 		return false;
@@ -503,61 +659,11 @@ bool IsTokenStreamDirty();
  * is used to deal with object references from types that aren't supported by the reflectable type system.
  * interface doesn't make sense to implement for.
  */
-class FArchiveRealtimeGC
-{
-private:
-
-	/** Helper struct for stack based approach */
-	struct FStackEntry
-	{
-		/** Current data pointer, incremented by stride */
-		uint8*	Data;
-		/** Current stride */
-		int32		Stride;
-		/** Current loop count, decremented each iteration */
-		int32		Count;
-		/** First token index in loop */
-		int32		LoopStartIndex;
-	};
-
-	class FGCTask
+class FRealtimeGC
 	{
-		FArchiveRealtimeGC*	Owner;
-		TArray<UObject*>*	ObjectsToSerialize;
-
-	public:
-		FGCTask(FArchiveRealtimeGC* InOwner, const TArray<UObject*>* InObjectsToSerialize,int32 StartIndex, int32 NumObjects)
-			: Owner(InOwner)
-			, ObjectsToSerialize(FGCArrayPool::Get().GetArrayFromPool())
-		{
-			ObjectsToSerialize->AddUninitialized(NumObjects);
-			FMemory::Memcpy(ObjectsToSerialize->GetData(), InObjectsToSerialize->GetData() + StartIndex, NumObjects * sizeof(UObject*));
-		}
-		~FGCTask()
-		{
-			FGCArrayPool::Get().ReturnToPool(ObjectsToSerialize);
-		}
-		FORCEINLINE TStatId GetStatId() const
-		{
-			RETURN_QUICK_DECLARE_CYCLE_STAT(FGCTask, STATGROUP_TaskGraphTasks);
-		}
-		static ENamedThreads::Type GetDesiredThread()
-		{
-			return ENamedThreads::AnyThread;
-		}
-		static ESubsequentsMode::Type GetSubsequentsMode() 
-		{ 
-			return ESubsequentsMode::TrackSubsequents; 
-		}
-		void DoTask(ENamedThreads::Type CurrentThread, FGraphEventRef& MyCompletionGraphEvent)
-		{
-			Owner->ProcessObjectArray(*ObjectsToSerialize, MyCompletionGraphEvent);			
-		}
-	};
-
 public:
 	/** Default constructor, initializing all members. */
-	FArchiveRealtimeGC()
+	FRealtimeGC()
 	{}
 
 	/** 
@@ -582,7 +688,7 @@ public:
 
 			// Keep track of how many objects are around.
 			GObjectCountDuringLastMarkPhase++;
-
+			ObjectItem->ClearFlags(EInternalObjectFlags::ReachableInCluster);
 			// Special case handling for objects that are part of the root set.
 			if (ObjectItem->IsRootSet())
 			{
@@ -593,7 +699,7 @@ public:
 				ObjectsToSerialize.Add(Object);
 			}
 			// Regular objects.
-			else
+			else if (ObjectItem->GetOwnerIndex() == 0)
 			{
 				bool bMarkAsUnreachable = true;
 				if (!ObjectItem->IsPendingKill())
@@ -646,7 +752,7 @@ public:
 	 */
 	void PerformReachabilityAnalysis(EObjectFlags KeepFlags, bool bForceSingleThreaded = false)
 	{		
-		DECLARE_SCOPE_CYCLE_COUNTER( TEXT( "FArchiveRealtimeGC::PerformReachabilityAnalysis" ), STAT_FArchiveRealtimeGC_PerformReachabilityAnalysis, STATGROUP_GC );
+		DECLARE_SCOPE_CYCLE_COUNTER( TEXT( "FRealtimeGC::PerformReachabilityAnalysis" ), STAT_FArchiveRealtimeGC_PerformReachabilityAnalysis, STATGROUP_GC );
 
 		/** Growing array of objects that require serialization */
 		TArray<UObject*>& ObjectsToSerialize = *FGCArrayPool::Get().GetArrayFromPool();
@@ -672,38 +778,10 @@ public:
 			MarkObjectsAsUnreachable<true>(ObjectsToSerialize, KeepFlags);
 		}
 
-		if( ObjectsToSerialize.Num() )
 		{
-			check(!GIsRunningParallelReachability);
-
-			if ( bForceSingleThreaded )
-			{
-				FGraphEventRef InvalidRef;
-				ProcessObjectArray( ObjectsToSerialize, InvalidRef );
-			}
-			else
-			{				
-				GIsRunningParallelReachability = true;
-
-				int32 NumChunks = FMath::Min<int32>(FTaskGraphInterface::Get().GetNumWorkerThreads(), ObjectsToSerialize.Num());
-				int32 NumPerChunk = ObjectsToSerialize.Num() / NumChunks;
-				check(NumPerChunk > 0);
-				FGraphEventArray ChunkTasks;
-				ChunkTasks.Empty(NumChunks);
-				int32 StartIndex = 0;
-				for (int32 Chunk = 0; Chunk < NumChunks; Chunk++)
-				{
-					if (Chunk + 1 == NumChunks)
-					{
-						NumPerChunk = ObjectsToSerialize.Num() - StartIndex; // last chunk takes all remaining items
-					}
-					ChunkTasks.Add(TGraphTask<FGCTask>::CreateTask().ConstructAndDispatchWhenReady(this, &ObjectsToSerialize, StartIndex, NumPerChunk));
-					StartIndex += NumPerChunk;
-				}
-				QUICK_SCOPE_CYCLE_COUNTER(STAT_GC_Subtask_Wait);
-				FTaskGraphInterface::Get().WaitUntilTasksComplete(ChunkTasks, ENamedThreads::GameThread_Local);
-				GIsRunningParallelReachability = false;
-			}
+			FGCReferenceProcessor ReferenceProcessor;
+			TFastReferenceCollector<FGCReferenceProcessor, FGCCollector, FGCArrayPool> ReferenceCollector(ReferenceProcessor, FGCArrayPool::Get());
+			ReferenceCollector.CollectReferences(ObjectsToSerialize, bForceSingleThreaded);
 		}
 
 		FGCArrayPool::Get().ReturnToPool(&ObjectsToSerialize);
@@ -712,285 +790,6 @@ public:
 		FGCArrayPool::Get().CheckLeaks();
 #endif
 	}
-
-	void DispatchObjectTasks(TArray<UObject*>& ObjectsToSerialize)
-	{
-	}
-
-	void ProcessObjectArray(TArray<UObject*>& InObjectsToSerializeArray, FGraphEventRef& MyCompletionGraphEvent)
-	{		
-		DECLARE_SCOPE_CYCLE_COUNTER( TEXT( "FArchiveRealtimeGC::ProcessObjectArray" ), STAT_FArchiveRealtimeGC_ProcessObjectArray, STATGROUP_GC );
-
-		UObject* CurrentObject = NULL;
-
-		const int32 MinDesiredObjectsPerSubTask = GMinDesiredObjectsPerSubTask; // sometimes there will be less, a lot less
-
-		/** Growing array of objects that require serialization */
-		TArray<UObject*>&	NewObjectsToSerializeArray = *FGCArrayPool::Get().GetArrayFromPool();
-
-		// Ping-pong between these two arrays if there's not enough objects to spawn a new task
-		TArray<UObject*>& ObjectsToSerialize = InObjectsToSerializeArray;
-		TArray<UObject*>& NewObjectsToSerialize = NewObjectsToSerializeArray;
-
-		// Presized "recursion" stack for handling arrays and structs.
-		TArray<FStackEntry> Stack;
-		Stack.AddUninitialized( 128 ); //@todo rtgc: need to add code handling more than 128 layers of recursion or at least assert
-
-		// it is necessary to have at least one extra item in the array memory block for the iffy prefetch code, below
-		ObjectsToSerialize.Reserve(ObjectsToSerialize.Num() + 1);
-	
-		// Keep serializing objects till we reach the end of the growing array at which point
-		// we are done.
-		int32 CurrentIndex = 0;
-		do
-		{
-			FGCCollector ReferenceCollector( NewObjectsToSerialize );
-			while( CurrentIndex < ObjectsToSerialize.Num() )
-			{
-#if PERF_DETAILED_PER_CLASS_GC_STATS
-				uint32 StartCycles = FPlatformTime::Cycles();
-#endif
-				CurrentObject = ObjectsToSerialize[CurrentIndex++];
-
-				// GetData() used to avoiding bounds checking (min and max)
-				// FMath::Min used to avoid out of bounds (without branching) on last iteration. Though anything can be passed into PrefetchBlock, 
-				// reading ObjectsToSerialize out of bounds is not safe since ObjectsToSerialize[Num()] may be an unallocated/unsafe address.
-				const UObject * const NextObject = ObjectsToSerialize.GetData()[FMath::Min<int32>(CurrentIndex, ObjectsToSerialize.Num() - 1)];
-
-				// Prefetch the next object assuming that the property size of the next object is the same as the current one.
-				// This allows us to avoid a branch here.
-				FPlatformMisc::PrefetchBlock(NextObject, CurrentObject->GetClass()->GetPropertiesSize());
-
-				//@todo rtgc: we need to handle object references in struct defaults
-
-				// Make sure that token stream has been assembled at this point as the below code relies on it.
-				checkSlow( CurrentObject->GetClass()->HasAnyClassFlags(CLASS_TokenStreamAssembled) );
-
-				// Get pointer to token stream and jump to the start.
-				FGCReferenceTokenStream* RESTRICT TokenStream = &CurrentObject->GetClass()->ReferenceTokenStream;
-				uint32 TokenStreamIndex			= 0;
-				// Keep track of index to reference info. Used to avoid LHSs.
-				uint32 ReferenceTokenStreamIndex	= 0;
-
-				// Create stack entry and initialize sane values.
-				FStackEntry* RESTRICT StackEntry = Stack.GetData();
-				uint8* StackEntryData		= (uint8*) CurrentObject;
-				StackEntry->Data			= StackEntryData;
-				StackEntry->Stride			= 0;
-				StackEntry->Count			= -1;
-				StackEntry->LoopStartIndex	= -1;
-			
-				// Keep track of token return count in separate integer as arrays need to fiddle with it.
-				int32 TokenReturnCount		= 0;
-
-				// Parse the token stream.
-				while( true )
-				{
-					// Cache current token index as it is the one pointing to the reference info.
-					ReferenceTokenStreamIndex = TokenStreamIndex;
-
-					// Handle returning from an array of structs, array of structs of arrays of ... (yadda yadda)
-					for( int32 ReturnCount=0; ReturnCount<TokenReturnCount; ReturnCount++ )
-					{
-						// Make sure there's no stack underflow.
-						check( StackEntry->Count != -1 );
-
-						// We pre-decrement as we're already through the loop once at this point.
-						if( --StackEntry->Count > 0 )
-						{
-							// Point data to next entry.
-							StackEntryData	 = StackEntry->Data + StackEntry->Stride;
-							StackEntry->Data = StackEntryData;
-
-							// Jump back to the beginning of the loop.
-							TokenStreamIndex = StackEntry->LoopStartIndex;
-							ReferenceTokenStreamIndex = StackEntry->LoopStartIndex;
-							// We're not done with this token loop so we need to early out instead of backing out further.
-							break;
-						}
-						else
-						{
-							StackEntry--;
-							StackEntryData = StackEntry->Data;
-						}
-					}
-
-					// Instead of reading information about reference from stream and caching it like below we access
-					// the same memory address over and over and over again to avoid a nasty LHS penalty. Not reading 
-					// the reference info means we need to manually increment the token index to skip to the next one.
-					TokenStreamIndex++;
-					// Helper to make code more readable and hide the ugliness that is avoiding LHSs from caching.
-					#define	REFERENCE_INFO TokenStream->AccessReferenceInfo( ReferenceTokenStreamIndex )
-
-					if( REFERENCE_INFO.Type == GCRT_Object )
-					{	
-						// We're dealing with an object reference.
-						UObject**	ObjectPtr	= (UObject**)(StackEntryData + REFERENCE_INFO.Offset);
-						UObject*&	Object		= *ObjectPtr;
-						TokenReturnCount		= REFERENCE_INFO.ReturnCount;
-						HandleTokenStreamObjectReference(NewObjectsToSerialize, CurrentObject, Object, ReferenceTokenStreamIndex, true);
-					}
-					else if( REFERENCE_INFO.Type == GCRT_ArrayObject )
-					{
-						// We're dealing with an array of object references.
-						TArray<UObject*>& ObjectArray = *((TArray<UObject*>*)(StackEntryData + REFERENCE_INFO.Offset));
-						TokenReturnCount = REFERENCE_INFO.ReturnCount;
-						for( int32 ObjectIndex = 0, ObjectNum = ObjectArray.Num(); ObjectIndex < ObjectNum; ++ObjectIndex )
-						{
-							HandleTokenStreamObjectReference(NewObjectsToSerialize, CurrentObject, ObjectArray[ObjectIndex], ReferenceTokenStreamIndex, true);
-						}
-					}
-					else if( REFERENCE_INFO.Type == GCRT_ArrayStruct )
-					{
-						// We're dealing with a dynamic array of structs.
-						const FScriptArray& Array = *((FScriptArray*)(StackEntryData + REFERENCE_INFO.Offset));
-						StackEntry++;
-						StackEntryData				= (uint8*) Array.GetData();
-						StackEntry->Data			= StackEntryData;
-						StackEntry->Stride			= TokenStream->ReadStride( TokenStreamIndex );
-						StackEntry->Count			= Array.Num();
-					
-						const FGCSkipInfo SkipInfo	= TokenStream->ReadSkipInfo( TokenStreamIndex );
-						StackEntry->LoopStartIndex	= TokenStreamIndex;
-					
-						if( StackEntry->Count == 0 )
-						{
-							// Skip empty array by jumping to skip index and set return count to the one about to be read in.
-							TokenStreamIndex		= SkipInfo.SkipIndex;
-							TokenReturnCount		= TokenStream->GetSkipReturnCount( SkipInfo );
-						}
-						else
-						{	
-							// Loop again.
-							check( StackEntry->Data );
-							TokenReturnCount		= 0;
-						}
-					}
-					else if( REFERENCE_INFO.Type == GCRT_PersistentObject )
-					{
-						// We're dealing with an object reference.
-						UObject**	ObjectPtr	= (UObject**)(StackEntryData + REFERENCE_INFO.Offset);
-						UObject*&	Object		= *ObjectPtr;
-						TokenReturnCount		= REFERENCE_INFO.ReturnCount;
-						HandleTokenStreamObjectReference(NewObjectsToSerialize, CurrentObject, Object, ReferenceTokenStreamIndex, false);
-					}
-					else if( REFERENCE_INFO.Type == GCRT_FixedArray )
-					{
-						// We're dealing with a fixed size array
-						uint8* PreviousData	= StackEntryData;
-						StackEntry++;
-						StackEntryData				= PreviousData;
-						StackEntry->Data			= PreviousData;
-						StackEntry->Stride			= TokenStream->ReadStride( TokenStreamIndex );
-						StackEntry->Count			= TokenStream->ReadCount( TokenStreamIndex );
-						StackEntry->LoopStartIndex	= TokenStreamIndex;
-						TokenReturnCount			= 0;
-					}
-					else if( REFERENCE_INFO.Type == GCRT_AddStructReferencedObjects )
-					{
-						// We're dealing with a function call
-						void const*	StructPtr	= (void*)(StackEntryData + REFERENCE_INFO.Offset);
-						TokenReturnCount		= REFERENCE_INFO.ReturnCount;
-						UScriptStruct::ICppStructOps::TPointerToAddStructReferencedObjects Func = (UScriptStruct::ICppStructOps::TPointerToAddStructReferencedObjects) TokenStream->ReadPointer( TokenStreamIndex );
-						Func(StructPtr, ReferenceCollector);
-					}
-					else if( REFERENCE_INFO.Type == GCRT_AddReferencedObjects )
-					{
-						// Static AddReferencedObjects function call.
-						void (*AddReferencedObjects)(UObject*, FReferenceCollector&) = (void(*)(UObject*, FReferenceCollector&))TokenStream->ReadPointer( TokenStreamIndex );
-						TokenReturnCount = REFERENCE_INFO.ReturnCount;
-						AddReferencedObjects(CurrentObject, ReferenceCollector);
-					}
-					else if( REFERENCE_INFO.Type == GCRT_AddTMapReferencedObjects )
-					{
-						void*         Map         = StackEntryData + REFERENCE_INFO.Offset;
-						UMapProperty* MapProperty = (UMapProperty*)TokenStream->ReadPointer( TokenStreamIndex );
-						TokenReturnCount = REFERENCE_INFO.ReturnCount;
-						FSimpleObjectReferenceCollectorArchive CollectorArchive(CurrentObject, ReferenceCollector);
-						MapProperty->SerializeItem(CollectorArchive, Map, nullptr);
-					}
-					else if (REFERENCE_INFO.Type == GCRT_EndOfPointer)
-					{
-						TokenReturnCount = REFERENCE_INFO.ReturnCount;
-					}
-					else if( REFERENCE_INFO.Type == GCRT_EndOfStream )
-					{
-						// Break out of loop.
-						break;
-					}
-					else
-					{
-						UE_LOG(LogGarbage, Fatal,TEXT("Unknown token"));
-					}
-				}
-				check(StackEntry == Stack.GetData());
-
-				if (GIsRunningParallelReachability && NewObjectsToSerialize.Num() >= MinDesiredObjectsPerSubTask)
-				{
-					// This will start queueing task with objects from the end of array until there's less objects than worth to queue
-					const int32 ObjectsPerSubTask = FMath::Max<int32>(MinDesiredObjectsPerSubTask, NewObjectsToSerialize.Num() / FTaskGraphInterface::Get().GetNumWorkerThreads());
-					while (NewObjectsToSerialize.Num() >= MinDesiredObjectsPerSubTask)
-					{						
-						const int32 StartIndex = FMath::Max(0, NewObjectsToSerialize.Num() - ObjectsPerSubTask);
-						const int32 NumThisTask = NewObjectsToSerialize.Num() - StartIndex;
-						MyCompletionGraphEvent->DontCompleteUntil(TGraphTask<FGCTask>::CreateTask().ConstructAndDispatchWhenReady(this, &NewObjectsToSerialize, StartIndex, NumThisTask));
-						NewObjectsToSerialize.SetNumUnsafeInternal(StartIndex);
-					}
-				}
-#if PERF_DETAILED_PER_CLASS_GC_STATS
-				// Detailed per class stats should not be performed when parallel GC is running
-				check( !GIsRunningParallelReachability );
-
-				uint32 DeltaCycles = FPlatformTime::Cycles() - StartCycles;
-				// Keep track of how many refs we encountered for the object's class.
-				const FName& ClassName = CurrentObject->GetClass()->GetFName();
-				// Refs to objects that reside in permanent object pool.
-				uint32 ClassDisregardedObjRefs = GClassToDisregardedObjectRefsMap.FindRef( ClassName );
-				GClassToDisregardedObjectRefsMap.Add( ClassName, ClassDisregardedObjRefs + GCurrentObjectDisregardedObjectRefs );
-				// Refs to regular objects.
-				uint32 ClassRegularObjRefs = GClassToRegularObjectRefsMap.FindRef( ClassName );
-				GClassToRegularObjectRefsMap.Add( ClassName, ClassRegularObjRefs + GCurrentObjectRegularObjectRefs );
-				// Track per class cycle count spent in GC.
-				uint32 ClassCycles = GClassToCyclesMap.FindRef( ClassName );
-				GClassToCyclesMap.Add( ClassName, ClassCycles + DeltaCycles );
-				// Reset current counts.
-				GCurrentObjectDisregardedObjectRefs = 0;
-				GCurrentObjectRegularObjectRefs		= 0;
-			}
-			// Log summary stats.
-			LogClassCountInfo( TEXT("references to regular objects from"), GClassToRegularObjectRefsMap, 20, 0 );
-			LogClassCountInfo( TEXT("references to permanent objects from"), GClassToDisregardedObjectRefsMap, 20, 0 );
-			LogClassCountInfo( TEXT("cycles for GC"), GClassToCyclesMap, 20, 0 );
-#else
-			}
-#endif
-			if( GIsRunningParallelReachability && NewObjectsToSerialize.Num() >= MinDesiredObjectsPerSubTask )
-			{			
-				const int32 ObjectsPerSubTask = FMath::Max<int32>(MinDesiredObjectsPerSubTask,NewObjectsToSerialize.Num() / FTaskGraphInterface::Get().GetNumWorkerThreads());
-				int32 StartIndex = 0;
-				while (StartIndex < NewObjectsToSerialize.Num())
-				{
-					const int32 NumThisTask = FMath::Min<int32>(ObjectsPerSubTask, NewObjectsToSerialize.Num() - StartIndex);
-					MyCompletionGraphEvent->DontCompleteUntil(TGraphTask<FGCTask>::CreateTask().ConstructAndDispatchWhenReady(this, &NewObjectsToSerialize, StartIndex, NumThisTask));
-					StartIndex += NumThisTask;
-				}
-				NewObjectsToSerialize.SetNumUnsafeInternal(0);
-			}
-			else if( NewObjectsToSerialize.Num() )
-			{
-				// Don't spawn a new task, continue in the current one
-				// To avoid allocating and moving memory around swap ObjectsToSerialize and NewObjectsToSerialize arrays
-				Exchange( ObjectsToSerialize, NewObjectsToSerialize );
-				// Empty but don't free allocated memory
-				NewObjectsToSerialize.SetNumUnsafeInternal(0);
-
-				CurrentIndex = 0;
-			}
-		}
-		while( CurrentIndex < ObjectsToSerialize.Num() );
-
-		FGCArrayPool::Get().ReturnToPool(&NewObjectsToSerializeArray);
-	}
 };
 
 /**
@@ -1340,19 +1139,20 @@ void CollectGarbageInternal(EObjectFlags KeepFlags, bool bPerformFullPurge)
 #if VERIFY_DISREGARD_GC_ASSUMPTIONS
 	FUObjectArray& UObjectArray = GUObjectArray;
 	// Only verify assumptions if option is enabled. This avoids false positives in the Editor or commandlets.
-	if (UObjectArray.DisregardForGCEnabled() && GShouldVerifyGCAssumptions)
+	if ((UObjectArray.DisregardForGCEnabled() || GUObjectClusters.Num()) && GShouldVerifyGCAssumptions)
 	{
 		DECLARE_SCOPE_CYCLE_COUNTER( TEXT( "CollectGarbageInternal.VerifyGCAssumptions" ), STAT_CollectGarbageInternal_VerifyGCAssumptions, STATGROUP_GC );
 
 		// Verify that objects marked to be disregarded for GC are not referencing objects that are not part of the root set.
-		for( FObjectIterator It; It; ++It )
+		for (FRawObjectIterator It(false); It; ++It)
 		{
 			bool bShouldAssert = false;
-			UObject* Object = *It;
+			FUObjectItem* ObjectItem = *It;
+			UObject* Object = (UObject*)ObjectItem->Object;
 			// Don't require UGCObjectReferencer's references to adhere to the assumptions.
 			// Although we want the referencer itself to sit in the disregard for gc set, most of the objects
 			// it's referencing will not be in the root set.
-			if (UObjectArray.IsDisregardForGC(Object) && !Object->IsA(UGCObjectReferencer::StaticClass()))
+			if ((ObjectItem->HasAnyFlags(EInternalObjectFlags::ClusterRoot) || UObjectArray.IsDisregardForGC(Object)) && !Object->IsA(UGCObjectReferencer::StaticClass()))
 			{
 				// Serialize object with reference collector.
 				TArray<UObject*> CollectedReferences;
@@ -1363,7 +1163,11 @@ void CollectGarbageInternal(EObjectFlags KeepFlags, bool bPerformFullPurge)
 				for( int32 ReferenceIndex=0; ReferenceIndex<CollectedReferences.Num(); ReferenceIndex++ )
 				{
 					UObject* ReferencedObject = CollectedReferences[ReferenceIndex];
-					if (ReferencedObject && !(ReferencedObject->IsRooted() || UObjectArray.IsDisregardForGC(ReferencedObject)))
+					if (ReferencedObject && 
+						!(ReferencedObject->IsRooted() || 
+						  UObjectArray.IsDisregardForGC(ReferencedObject) || 
+							UObjectArray.ObjectToObjectItem(ReferencedObject)->GetOwnerIndex() ||
+							UObjectArray.ObjectToObjectItem(ReferencedObject)->HasAnyFlags(EInternalObjectFlags::ClusterRoot)))
 					{
 						UE_LOG(LogGarbage, Warning, TEXT("Disregard for GC object %s referencing %s which is not part of root set"),
 							*Object->GetFullName(),
@@ -1394,7 +1198,7 @@ void CollectGarbageInternal(EObjectFlags KeepFlags, bool bPerformFullPurge)
 	// Perform reachability analysis.
 	{
 		const double StartTime = FPlatformTime::Seconds();
-		FArchiveRealtimeGC TagUsedRealtimeGC;
+		FRealtimeGC TagUsedRealtimeGC;
 		TagUsedRealtimeGC.PerformReachabilityAnalysis( KeepFlags, bForceSingleThreadedGC );
 		UE_LOG(LogGarbage, Log, TEXT("%f ms for GC"), (FPlatformTime::Seconds() - StartTime) * 1000 );
 	}
@@ -1411,6 +1215,7 @@ void CollectGarbageInternal(EObjectFlags KeepFlags, bool bPerformFullPurge)
 
 		// Unhash all unreachable objects.
 		const double StartTime = FPlatformTime::Seconds();
+		int32 ClustersRemoved = 0;
 		for ( FRawObjectIterator It(true); It; ++It )
 		{
 			//@todo UE4 - A prefetch was removed here. Re-add it. It wasn't right anyway, since it was ten items ahead and the consoles on have 8 prefetch slots
@@ -1419,6 +1224,34 @@ void CollectGarbageInternal(EObjectFlags KeepFlags, bool bPerformFullPurge)
 			checkSlow(ObjectItem);
 			if (ObjectItem->IsUnreachable())
 			{
+				if ((ObjectItem->GetFlags() & EInternalObjectFlags::ClusterRoot) == EInternalObjectFlags::ClusterRoot)
+				{
+					// Nuke the entire cluster
+					ObjectItem->ClearFlags(EInternalObjectFlags::ClusterRoot|EInternalObjectFlags::NoStrongReference);					
+					const int32 ClusterRootIndex = It.GetIndex();
+					FUObjectCluster* Cluster = GUObjectClusters.FindChecked(ClusterRootIndex);
+					checkSlow(Cluster);
+					for (int32 ClusterObjectIndex : Cluster->Objects)
+					{
+						FUObjectItem* ClusterObjectItem = GUObjectArray.IndexToObjectUnsafeForGC(ClusterObjectIndex);
+						ClusterObjectItem->ClearFlags(EInternalObjectFlags::NoStrongReference);
+						ClusterObjectItem->SetOwnerIndex(0);
+
+						if (!ClusterObjectItem->HasAnyFlags(EInternalObjectFlags::ReachableInCluster))
+						{
+							ClusterObjectItem->SetFlags(EInternalObjectFlags::Unreachable);
+							if (ClusterObjectIndex < ClusterRootIndex)
+							{
+								UObject* ClusterObject = (UObject*)ClusterObjectItem->Object;
+								ClusterObject->ConditionalBeginDestroy();
+							}
+						}
+					}
+					delete Cluster;
+					GUObjectClusters.Remove(ClusterRootIndex);
+					ClustersRemoved++;
+				}
+
 				// Begin the object's asynchronous destruction.
 				UObject* Object = (UObject*)ObjectItem->Object;
 				Object->ConditionalBeginDestroy();
@@ -1429,7 +1262,7 @@ void CollectGarbageInternal(EObjectFlags KeepFlags, bool bPerformFullPurge)
 				ObjectItem->SetPendingKill();
 			}
 		}
-		UE_LOG(LogGarbage, Log, TEXT("%f ms for unhashing unreachable objects"), (FPlatformTime::Seconds() - StartTime) * 1000 );
+		UE_LOG(LogGarbage, Log, TEXT("%f ms for unhashing unreachable objects. Clusters removed: %d."), (FPlatformTime::Seconds() - StartTime) * 1000, ClustersRemoved);
 	}
 
 	// Set flag to indicate that we are relying on a purge to be performed.
diff --git a/Engine/Source/Runtime/CoreUObject/Private/UObject/Linker.cpp b/Engine/Source/Runtime/CoreUObject/Private/UObject/Linker.cpp
--- a/Engine/Source/Runtime/CoreUObject/Private/UObject/Linker.cpp
+++ b/Engine/Source/Runtime/CoreUObject/Private/UObject/Linker.cpp
@@ -15,6 +15,7 @@ DEFINE_LOG_CATEGORY(LogLinker);
 
 #define LOCTEXT_NAMESPACE "Linker"
 
+#define ENABLE_PAIRED_PACKAGE_LOCALIZATION_LOOKUP 1
 
 /*-----------------------------------------------------------------------------
 	Helper functions.
@@ -85,7 +86,7 @@ FName FLinker::GetExportClassName( int32 i )
 	if (ExportMap.IsValidIndex(i))
 	{
 		FObjectExport& Export = ExportMap[i];
-		if (!Export.ClassIndex.IsNull())
+		if( !Export.ClassIndex.IsNull() )
 		{
 			return ImpExp(Export.ClassIndex).ObjectName;
 		}
@@ -468,6 +469,24 @@ static bool DoesPackageExistForGetPackageLinker(const FString& LongPackageName,
 	}
 }
 
+/** Customized version of FPackageName::DoesPackageExist that takes dynamic native class packages into account */
+static void DoesPackageAndLocalizationExistForGetPackageLinker(const FString& LongPackageName, const FGuid* Guid, FString& OutNativeFilename, FString& OutLocalizedFilename)
+{
+	if (
+#if WITH_EDITORONLY_DATA
+		GLinkerAllowDynamicClasses &&
+#endif
+		GetConvertedDynamicPackageNameToTypeName().Contains(*LongPackageName))
+	{
+		OutNativeFilename = FPackageName::LongPackageNameToFilename(LongPackageName);
+		OutLocalizedFilename.Empty();
+	}
+	else
+	{
+		FPackageName::DoesPackageExistWithLocalization(LongPackageName, Guid, &OutNativeFilename, &OutLocalizedFilename);
+	}
+}
+
 //
 // Find or create the linker for a package.
 //
@@ -502,10 +521,17 @@ FLinkerLoad* GetPackageLinker
 			return nullptr;
 		}
 	
+		// Verify that the file exists.
 		FString NativeFilename;
 		FString LocalizedFilename;
+#if ENABLE_PAIRED_PACKAGE_LOCALIZATION_LOOKUP
+		DoesPackageAndLocalizationExistForGetPackageLinker(InOuter->GetName(), CompatibleGuid, NativeFilename, LocalizedFilename);
+		const bool DoesNativePackageExist = NativeFilename.Len() > 0;
+		const bool DoesLocalizedPackageExist = LocalizedFilename.Len() > 0;
+#else
 		const bool DoesNativePackageExist = DoesPackageExistForGetPackageLinker(InOuter->GetName(), CompatibleGuid, NativeFilename, false);
 		const bool DoesLocalizedPackageExist = DoesPackageExistForGetPackageLinker(InOuter->GetName(), CompatibleGuid, LocalizedFilename, true);
+#endif
 
 		// If we are the editor, we must have a native package. If we are the game, we must have a localized package or a native package.
 		if ( (GIsEditor && !DoesNativePackageExist) || (!GIsEditor && !DoesLocalizedPackageExist && !DoesNativePackageExist) )
@@ -583,8 +609,14 @@ FLinkerLoad* GetPackageLinker
 		// Verify that the file exists.
 		FString NativeFilename;
 		FString LocalizedFilename;
+#if ENABLE_PAIRED_PACKAGE_LOCALIZATION_LOOKUP
+		DoesPackageAndLocalizationExistForGetPackageLinker(PackageName, CompatibleGuid, NativeFilename, LocalizedFilename);
+		const bool DoesNativePackageExist = NativeFilename.Len() > 0;
+		const bool DoesLocalizedPackageExist = LocalizedFilename.Len() > 0;
+#else
 		const bool DoesNativePackageExist = DoesPackageExistForGetPackageLinker(PackageName, CompatibleGuid, NativeFilename, false);
 		const bool DoesLocalizedPackageExist = DoesPackageExistForGetPackageLinker(PackageName, CompatibleGuid, LocalizedFilename, true);
+#endif
 
 		if( (GIsEditor && !DoesNativePackageExist) || (!GIsEditor && !DoesLocalizedPackageExist && !DoesNativePackageExist) )
 		{
diff --git a/Engine/Source/Runtime/CoreUObject/Private/UObject/Obj.cpp b/Engine/Source/Runtime/CoreUObject/Private/UObject/Obj.cpp
--- a/Engine/Source/Runtime/CoreUObject/Private/UObject/Obj.cpp
+++ b/Engine/Source/Runtime/CoreUObject/Private/UObject/Obj.cpp
@@ -727,6 +727,12 @@ void UObject::ConditionalPostLoad()
 			PostLoad();
 		}
 
+		extern int32 GCreateGCClusters;
+		if (GCreateGCClusters && FPlatformProperties::RequiresCookedData() && !GIsInitialLoad && CanBeClusterRoot())
+		{
+			CreateCluster();
+		}
+
 #if !(UE_BUILD_SHIPPING || UE_BUILD_TEST)
 		if (ThreadContext.DebugPostLoad.Contains(this))
 		{
diff --git a/Engine/Source/Runtime/CoreUObject/Private/UObject/UObjectArray.cpp b/Engine/Source/Runtime/CoreUObject/Private/UObject/UObjectArray.cpp
--- a/Engine/Source/Runtime/CoreUObject/Private/UObject/UObjectArray.cpp
+++ b/Engine/Source/Runtime/CoreUObject/Private/UObject/UObjectArray.cpp
@@ -8,6 +8,8 @@
 
 DEFINE_LOG_CATEGORY_STATIC(LogUObjectArray, Log, All);
 
+TMap<int32, FUObjectCluster*> GUObjectClusters;
+
 void FUObjectArray::AllocateObjectPool(int32 MaxUObjects, int32 MaxObjectsNotConsideredByGC)
 {
 	check(IsInGameThread());
@@ -229,7 +231,7 @@ void FUObjectArray::ShutdownUObjectArray()
 {
 }
 
-UObjectBase*** FUObjectArray::GetObjectArrayForDebugVisualizers()
+FFixedUObjectArray* FUObjectArray::GetObjectArrayForDebugVisualizers()
 {
-	return GUObjectArray.ObjObjects.GetRootBlockForDebuggerVisualizers();
+	return &GUObjectArray.ObjObjects;
 }
diff --git a/Engine/Source/Runtime/CoreUObject/Private/UObject/UObjectBaseUtility.cpp b/Engine/Source/Runtime/CoreUObject/Private/UObject/UObjectBaseUtility.cpp
--- a/Engine/Source/Runtime/CoreUObject/Private/UObject/UObjectBaseUtility.cpp
+++ b/Engine/Source/Runtime/CoreUObject/Private/UObject/UObjectBaseUtility.cpp
@@ -7,7 +7,7 @@
 #include "CoreUObjectPrivate.h"
 #include "Interface.h"
 #include "ModuleManager.h"
-
+#include "FastReferenceCollector.h"
 
 /***********************/
 /******** Names ********/
@@ -417,4 +417,3 @@ bool UObjectBaseUtility::IsDefaultSubobject() const
 	const bool bIsInstanced = GetOuter() && (GetOuter()->HasAnyFlags(RF_ClassDefaultObject) || ((UObject*)this)->GetArchetype() != GetClass()->GetDefaultObject());
 	return bIsInstanced;
 }
-
