Copying //UE4/Dev-Core to //UE4/Main

diff --git a/Engine/Source/Runtime/CoreUObject/Private/UObject/UObjectClusters.cpp b/Engine/Source/Runtime/CoreUObject/Private/UObject/UObjectClusters.cpp
new file mode 100644
--- /dev/null
+++ b/Engine/Source/Runtime/CoreUObject/Private/UObject/UObjectClusters.cpp
@@ -0,0 +1,412 @@
+// Copyright 1998-2015 Epic Games, Inc. All Rights Reserved.
+
+/*=============================================================================
+	UObjectClusters.cpp: Unreal UObject Cluster helper functions
+=============================================================================*/
+
+#include "CoreUObjectPrivate.h"
+#include "Interface.h"
+#include "ModuleManager.h"
+#include "FastReferenceCollector.h"
+
+int32 GCreateGCClusters = 1;
+static FAutoConsoleVariableRef CCreateGCClusters(
+	TEXT("gc.CreateGCClusters"),
+	GCreateGCClusters,
+	TEXT("If true, the engine will attempt to create clusters of objects for better garbage collection performance."),
+	ECVF_Default
+	);
+
+int32 GMergeGCClusters = 0;
+static FAutoConsoleVariableRef CMergeGCClusters(
+	TEXT("gc.MergeGCClusters"),
+	GMergeGCClusters,
+	TEXT("If true, when creating clusters, the clusters referenced from another cluster will get merged into one big cluster."),
+	ECVF_Default
+	);
+
+// Dumps all clusters to log.
+static void ListClusters(const TArray<FString>& Args)
+{
+	const bool bHierarchy = Args.Num() && Args[0] == TEXT("Hierarchy");
+	int32 MaxInterClusterReferences = 0;
+	int32 TotalInterClusterReferences = 0;
+	int32 MaxClusterSize = 0;
+	int32 TotalClusterObjects = 0;	
+	for (TPair<int32, FUObjectCluster*>& Pair : GUObjectClusters)
+	{
+		FUObjectItem* RootItem = GUObjectArray.IndexToObjectUnsafeForGC(Pair.Key);
+		UObject* RootObject = static_cast<UObject*>(RootItem->Object);
+		UE_LOG(LogObj, Display, TEXT("%s (Index: %d), Size: %d, ReferencedClusters: %d"), *RootObject->GetFullName(), Pair.Key, Pair.Value->Objects.Num(), Pair.Value->ReferencedClusters.Num());
+		MaxInterClusterReferences = FMath::Max(MaxInterClusterReferences, Pair.Value->ReferencedClusters.Num());
+		TotalInterClusterReferences += Pair.Value->ReferencedClusters.Num();
+		MaxClusterSize = FMath::Max(MaxClusterSize, Pair.Value->Objects.Num());
+		TotalClusterObjects += Pair.Value->Objects.Num();
+		if (bHierarchy)
+		{
+			int32 Index = 0;
+			for (int32 ObjectIndex : Pair.Value->Objects)
+			{
+				FUObjectItem* ObjectItem = GUObjectArray.IndexToObjectUnsafeForGC(ObjectIndex);
+				UObject* Object = static_cast<UObject*>(ObjectItem->Object);
+				UE_LOG(LogObj, Display, TEXT("    [%.4d]: %s (Index: %d)"), Index++, *Object->GetFullName(), ObjectIndex);
+			}
+			for (int32 ClusterIndex : Pair.Value->ReferencedClusters)
+			{
+				FUObjectItem* ClusterRootItem = GUObjectArray.IndexToObjectUnsafeForGC(ClusterIndex);
+				UObject* ClusterRootObject = static_cast<UObject*>(ClusterRootItem->Object);
+				UE_LOG(LogObj, Display, TEXT("    -> %s (Index: %d)"), *ClusterRootObject->GetFullName(), ClusterIndex);
+			}
+		}
+	}
+	UE_LOG(LogObj, Display, TEXT("Number of clusters: %d"), GUObjectClusters.Num());
+	UE_LOG(LogObj, Display, TEXT("Maximum cluster size: %d"), MaxClusterSize);
+	UE_LOG(LogObj, Display, TEXT("Average cluster size: %d"), GUObjectClusters.Num() ? (TotalClusterObjects / GUObjectClusters.Num()) : 0);
+	UE_LOG(LogObj, Display, TEXT("Number of objects in GC clusters: %d"), TotalClusterObjects);
+	UE_LOG(LogObj, Display, TEXT("Maximum number of custer-to-cluster references: %d"), MaxInterClusterReferences);
+	UE_LOG(LogObj, Display, TEXT("Average number of custer-to-cluster references: %d"), GUObjectClusters.Num() ? (TotalInterClusterReferences / GUObjectClusters.Num()) : 0);
+}
+
+static FAutoConsoleCommand ListClustersCommand(
+	TEXT("gc.ListClusters"),
+	TEXT("Dumps all clusters do output log. When 'Hiearchy' argument is specified lists all objects inside clusters."),
+	FConsoleCommandWithArgsDelegate::CreateStatic(ListClusters)
+	);
+
+/**
+* Pool for reducing allocations when constructing clusters
+*/
+class FClusterArrayPool
+{
+public:
+
+	/**
+	* Gets the singleton instance of the FObjectArrayPool
+	* @return Pool singleton.
+	*/
+	FORCEINLINE static FClusterArrayPool& Get()
+	{
+		static FClusterArrayPool Singleton;
+		return Singleton;
+	}
+
+	/**
+	* Gets an event from the pool or creates one if necessary.
+	*
+	* @return The array.
+	* @see ReturnToPool
+	*/
+	FORCEINLINE TArray<UObject*>* GetArrayFromPool()
+	{
+		TArray<UObject*>* Result = Pool.Pop();
+		if (!Result)
+		{
+			Result = new TArray<UObject*>();
+		}
+		check(Result);
+#if UE_BUILD_DEBUG
+		NumberOfUsedArrays.Increment();
+#endif // UE_BUILD_DEBUG
+		return Result;
+	}
+
+	/**
+	* Returns an array to the pool.
+	*
+	* @param Array The array to return.
+	* @see GetArrayFromPool
+	*/
+	FORCEINLINE void ReturnToPool(TArray<UObject*>* Array)
+	{
+#if UE_BUILD_DEBUG
+		const int32 CheckUsedArrays = NumberOfUsedArrays.Decrement();
+		checkSlow(CheckUsedArrays >= 0);
+#endif // UE_BUILD_DEBUG
+		check(Array);
+		Array->Reset();
+		Pool.Push(Array);
+	}
+
+	/** Performs memory cleanup */
+	void Cleanup()
+	{
+#if UE_BUILD_DEBUG
+		const int32 CheckUsedArrays = NumberOfUsedArrays.GetValue();
+		checkSlow(CheckUsedArrays == 0);
+#endif // UE_BUILD_DEBUG
+
+		uint32 FreedMemory = 0;
+		TArray< TArray<UObject*>* > AllArrays;
+		Pool.PopAll(AllArrays);
+		for (TArray<UObject*>* Array : AllArrays)
+		{
+			FreedMemory += Array->GetAllocatedSize();
+			delete Array;
+		}
+		UE_LOG(LogObj, Log, TEXT("Freed %ub from %d cluster array pools."), FreedMemory, AllArrays.Num());
+	}
+
+#if UE_BUILD_DEBUG
+	void CheckLeaks()
+	{
+		// This function is called after cluster has been created so at this point there should be no
+		// arrays used by cluster creation code and all should be returned to the pool
+		const int32 LeakedPoolArrays = NumberOfUsedArrays.GetValue();
+		checkSlow(LeakedPoolArrays == 0);
+	}
+#endif
+
+private:
+
+	/** Holds the collection of recycled arrays. */
+	TLockFreePointerListLIFO< TArray<UObject*> > Pool;
+
+#if UE_BUILD_DEBUG
+	/** Number of arrays currently acquired from the pool by clusters */
+	FThreadSafeCounter NumberOfUsedArrays;
+#endif // UE_BUILD_DEBUG
+};
+
+/** Called on shutdown to free cluster memory */
+void CleanupClusterArrayPools()
+{
+	FClusterArrayPool::Get().Cleanup();
+}
+
+/**
+ * Handles UObject references found by TFastReferenceCollector
+ */
+class FClusterReferenceProcessor
+{
+	int32 ClusterRootIndex;
+	FUObjectCluster& Cluster;
+	volatile bool bIsRunningMultithreaded;
+public:
+
+	FClusterReferenceProcessor(int32 InClusterRootIndex, FUObjectCluster& InCluster)
+		: ClusterRootIndex(InClusterRootIndex)
+		, Cluster(InCluster)
+		, bIsRunningMultithreaded(false)
+	{}
+
+	FORCEINLINE int32 GetMinDesiredObjectsPerSubTask() const
+	{
+		// We're not running the processor in parallel when creating clusters
+		return 0;
+	}
+
+	FORCEINLINE volatile bool IsRunningMultithreaded() const
+	{
+		// This should always be false
+		return bIsRunningMultithreaded;
+	}
+
+	FORCEINLINE void SetIsRunningMultithreaded(bool bIsParallel)
+	{
+		check(!bIsParallel);
+		bIsRunningMultithreaded = bIsParallel;
+	}
+
+	void UpdateDetailedStats(UObject* CurrentObject, uint32 DeltaCycles)
+	{
+	}
+
+	void LogDetailedStatsSummary()
+	{
+	}
+
+	/**
+	 * Adds an object to cluster (if possible)
+	 *
+	 * @param ObjectIndex UObject index in GUObjectArray
+	 * @param ObjectItem UObject's entry in GUObjectArray
+	 * @param Obj The object to add to cluster
+	 * @param ObjectsToSerialize An array of remaining objects to serialize (Obj must be added to it if Obj can be added to cluster)
+	 * @param bOuterAndClass If true, the Obj's Outer and Class will also be added to the cluster
+	 */
+	void AddObjectToCluster(int32 ObjectIndex, FUObjectItem* ObjectItem, UObject* Obj, TArray<UObject*>& ObjectsToSerialize, bool bOuterAndClass)
+	{
+		if (ObjectItem->GetOwnerIndex() == 0 && !ObjectItem->IsRootSet() && !GUObjectArray.IsDisregardForGC(Obj) && Obj->CanBeInCluster())
+		{
+			ObjectsToSerialize.Add(Obj);
+			ObjectItem->SetOwnerIndex(ClusterRootIndex);
+			Cluster.Objects.Add(ObjectIndex);
+
+			if (bOuterAndClass)
+			{
+				UObject* ObjOuter = Obj->GetOuter();
+				if (ObjOuter)
+				{
+					int32 OuterIndex = GUObjectArray.ObjectToIndex(ObjOuter);
+					AddObjectToCluster(OuterIndex, GUObjectArray.IndexToObjectUnsafeForGC(OuterIndex), ObjOuter, ObjectsToSerialize, false);
+				}
+				if (!Obj->GetClass()->HasAllClassFlags(CLASS_Native))
+				{
+					int32 ClassIndex = GUObjectArray.ObjectToIndex(Obj->GetClass());
+					AddObjectToCluster(ClassIndex, GUObjectArray.IndexToObjectUnsafeForGC(ClassIndex), Obj->GetClass(), ObjectsToSerialize, false);
+
+					int32 ClassOuterIndex = GUObjectArray.ObjectToIndex(Obj->GetClass()->GetOuter());
+					AddObjectToCluster(ClassOuterIndex, GUObjectArray.IndexToObjectUnsafeForGC(ClassOuterIndex), Obj->GetClass()->GetOuter(), ObjectsToSerialize, false);
+				}
+			}
+		}		
+	}
+
+	/**
+	* Merges an existing cluster with the currently constructed one
+	*
+	* @param ObjectItem UObject's entry in GUObjectArray. This is either the other cluster root or one if the cluster objects.
+	* @param Obj The object to add to cluster
+	* @param ObjectsToSerialize An array of remaining objects to serialize (Obj must be added to it if Obj can be added to cluster)
+	*/
+	void MergeCluster(FUObjectItem* ObjectItem, UObject* Object, TArray<UObject*>& ObjectsToSerialize)
+	{
+		FUObjectItem* OtherClusterRootItem = nullptr;
+		int32 OtherClusterRootIndex = INDEX_NONE;
+		// First find the other cluster root item and its index
+		if (ObjectItem->HasAnyFlags(EInternalObjectFlags::ClusterRoot))
+		{
+			OtherClusterRootIndex = GUObjectArray.ObjectToIndex(Object);
+			OtherClusterRootItem = ObjectItem;
+		}
+		else
+		{
+			OtherClusterRootIndex = ObjectItem->GetOwnerIndex();
+			OtherClusterRootItem = GUObjectArray.IndexToObjectUnsafeForGC(OtherClusterRootIndex);
+		}
+		// This is another cluster, merge it with this one
+		FUObjectCluster* ClusterToMerge = GUObjectClusters.FindChecked(OtherClusterRootIndex);
+		for (int32 OtherClusterObjectIndex : ClusterToMerge->Objects)
+		{
+			FUObjectItem* OtherClusterObjectItem = GUObjectArray.IndexToObjectUnsafeForGC(OtherClusterObjectIndex);
+			OtherClusterObjectItem->SetOwnerIndex(0);
+			AddObjectToCluster(OtherClusterObjectIndex, OtherClusterObjectItem, static_cast<UObject*>(OtherClusterObjectItem->Object), ObjectsToSerialize, true);
+		}		
+		GUObjectClusters.Remove(OtherClusterRootIndex);
+		delete ClusterToMerge;
+
+		// Make sure the root object is also added to the current cluster
+		OtherClusterRootItem->ClearFlags(EInternalObjectFlags::ClusterRoot);
+		OtherClusterRootItem->SetOwnerIndex(0);
+		AddObjectToCluster(OtherClusterRootIndex, OtherClusterRootItem, static_cast<UObject*>(OtherClusterRootItem->Object), ObjectsToSerialize, true);
+
+		// Sanity check so that we make sure the object was actually in the lister it said it belonged to
+		check(ObjectItem->GetOwnerIndex() == ClusterRootIndex);
+	}
+
+	/**
+	* Handles UObject reference from the token stream. Performance is critical here so we're FORCEINLINING this function.
+	*
+	* @param ObjectsToSerialize An array of remaining objects to serialize (Obj must be added to it if Obj can be added to cluster)
+	* @param ReferencingObject Object referencing the object to process.
+	* @param TokenIndex Index to the token stream where the reference was found.
+	* @param bAllowReferenceElimination True if reference elimination is allowed (ignored when constructing clusters).
+	*/
+	FORCEINLINE void HandleTokenStreamObjectReference(TArray<UObject*>& ObjectsToSerialize, UObject* ReferencingObject, UObject*& Object, const int32 TokenIndex, bool bAllowReferenceElimination)
+	{
+		if (Object)
+		{
+			FUObjectItem* ObjectItem = GUObjectArray.ObjectToObjectItem(Object);
+
+			// Add encountered object reference to list of to be serialized objects if it hasn't already been added.
+			if (ObjectItem->GetOwnerIndex() != ClusterRootIndex)
+			{
+				if (ObjectItem->HasAnyFlags(EInternalObjectFlags::ClusterRoot) || ObjectItem->GetOwnerIndex() != 0)
+				{					
+					if (GMergeGCClusters)
+					{
+						// This is an existing cluster, merge it with the current one.
+						MergeCluster(ObjectItem, Object, ObjectsToSerialize);
+					}
+					else
+					{
+						// Simply reference this cluster and all clusters it's referencing
+						const int32 OtherClusterRootIndex = ObjectItem->HasAnyFlags(EInternalObjectFlags::ClusterRoot) ? GUObjectArray.ObjectToIndex(Object) : ObjectItem->GetOwnerIndex();
+						FUObjectCluster* OtherCluster = GUObjectClusters.FindChecked(OtherClusterRootIndex);
+						Cluster.ReferencedClusters.AddUnique(OtherClusterRootIndex);
+						for (int32 OtherClusterReferencedCluster : OtherCluster->ReferencedClusters)
+						{
+							check(OtherClusterReferencedCluster != ClusterRootIndex);
+							Cluster.ReferencedClusters.AddUnique(OtherClusterReferencedCluster);
+						}
+					}
+				}
+				else if (ObjectItem->GetOwnerIndex() == 0 && !ObjectItem->IsRootSet() && !GUObjectArray.IsDisregardForGC(Object) && Object->CanBeInCluster())
+				{
+					// New object, add it to the cluster.
+					AddObjectToCluster(GUObjectArray.ObjectToIndex(Object), ObjectItem, Object, ObjectsToSerialize, true);
+				}
+			}
+		}
+	}
+};
+
+/**
+ * Specialized FReferenceCollector that uses FClusterReferenceProcessor to construct the cluster
+ */
+class FClusterCollector : public FReferenceCollector
+{
+	FClusterReferenceProcessor& Processor;
+	TArray<UObject*>& ObjectArray;
+
+public:
+	FClusterCollector(FClusterReferenceProcessor& InProcessor, TArray<UObject*>& InObjectArray)
+		: Processor(InProcessor)
+		, ObjectArray(InObjectArray)
+	{
+	}
+	virtual void HandleObjectReference(UObject*& Object, const UObject* ReferencingObject, const UProperty* ReferencingProperty) override
+	{
+		Processor.HandleTokenStreamObjectReference(ObjectArray, const_cast<UObject*>(ReferencingObject), Object, INDEX_NONE, false);
+	}
+	virtual void HandleObjectReferences(UObject** InObjects, const int32 ObjectNum, const UObject* ReferencingObject, const UProperty* InReferencingProperty) override
+	{
+		for (int32 ObjectIndex = 0; ObjectIndex < ObjectNum; ++ObjectIndex)
+		{
+			UObject*& Object = InObjects[ObjectIndex];
+			Processor.HandleTokenStreamObjectReference(ObjectArray, const_cast<UObject*>(ReferencingObject), Object, INDEX_NONE, false);
+		}
+	}
+	virtual bool IsIgnoringArchetypeRef() const override
+	{
+		return false;
+	}
+	virtual bool IsIgnoringTransient() const override
+	{
+		return false;
+	}
+};
+
+void UObjectBaseUtility::CreateCluster()
+{
+	FUObjectItem* RootItem = GUObjectArray.IndexToObject(InternalIndex);
+	if (RootItem->GetOwnerIndex() != 0)
+	{
+		return;
+	}
+
+	// Create a new cluster, reserve an arbitrary amount of memory for it.
+	FUObjectCluster* Cluster = new FUObjectCluster;
+	Cluster->Objects.Reserve(64);
+
+	// Collect all objects referenced by cluster root and by all objects it's referencing
+	FClusterReferenceProcessor Processor(InternalIndex, *Cluster);
+	TFastReferenceCollector<FClusterReferenceProcessor, FClusterCollector, FClusterArrayPool> ReferenceCollector(Processor, FClusterArrayPool::Get());
+	TArray<UObject*> ObjectsToProcess;
+	ObjectsToProcess.Add(static_cast<UObject*>(this));
+	ReferenceCollector.CollectReferences(ObjectsToProcess, true);
+#if UE_BUILD_DEBUG
+	FClusterArrayPool::Get().CheckLeaks();
+#endif
+
+	if (Cluster->Objects.Num())
+	{
+		// Add new cluster to the global cluster map.
+		GUObjectClusters.Add(InternalIndex, Cluster);
+		RootItem->SetFlags(EInternalObjectFlags::ClusterRoot);
+	}
+	else
+	{
+		delete Cluster;
+	}
+}
\ No newline at end of file
diff --git a/Engine/Source/Runtime/CoreUObject/Private/UObject/WeakObjectPtr.cpp b/Engine/Source/Runtime/CoreUObject/Private/UObject/WeakObjectPtr.cpp
--- a/Engine/Source/Runtime/CoreUObject/Private/UObject/WeakObjectPtr.cpp
+++ b/Engine/Source/Runtime/CoreUObject/Private/UObject/WeakObjectPtr.cpp
@@ -9,11 +9,9 @@
 DEFINE_LOG_CATEGORY_STATIC(LogWeakObjectPtr, Log, All);
 
 /*-----------------------------------------------------------------------------------------------------------
-	Base serial number management. This is complicated by the need for thread safety and lock-free implementation
+	Base serial number management.
 -------------------------------------------------------------------------------------------------------------*/
 
-int32** GSerialNumberBlocksForDebugVisualizersRoot = 0;
-
 /*-----------------------------------------------------------------------------------------------------------
 	FWeakObjectPtr
 -------------------------------------------------------------------------------------------------------------*/
diff --git a/Engine/Source/Runtime/CoreUObject/Public/Misc/PackageName.h b/Engine/Source/Runtime/CoreUObject/Public/Misc/PackageName.h
--- a/Engine/Source/Runtime/CoreUObject/Public/Misc/PackageName.h
+++ b/Engine/Source/Runtime/CoreUObject/Public/Misc/PackageName.h
@@ -195,6 +195,15 @@ public:
 	 **/
 	static bool DoesPackageExist(const FString& LongPackageName, const FGuid* Guid = NULL, FString* OutFilename = NULL, const bool ShouldGetLocalizedPackage = false);
 
+	/**
+	* Checks if the package and any localized version exist on disk.
+	*
+	* @param LongPackageName Package name.
+	* @param OutFilenameNative Package package filename on disk. Set to empty string if the package is not found.
+	* @param OutFilenameLocalized Localized package filename on disk. Set to empty string if the package is not found
+	**/
+	static bool DoesPackageExistWithLocalization(const FString& LongPackageName, const FGuid* Guid = NULL, FString* OutNativeFilename = NULL, FString* OutLocalizedFilename = NULL);
+
 	/**
 	 * Attempts to find a package given its short name on disk (very slow).
 	 * 
@@ -364,6 +373,11 @@ public:
 	*/
 	static bool FindPackageFileWithoutExtension(const FString& InPackageFilename, FString& OutFilename);
 
+	/**
+	* TODO
+	*/
+	static void FindPackageFileAndLocalizationWithoutExtension(const FString& InNativePackageFilename, const FString& InLocalizedPackageFilename, FString& OutNativeFilename, FString& OutLocalizedFilename);
+
 private:
 
 	static FString AssetPackageExtension;
diff --git a/Engine/Source/Runtime/CoreUObject/Public/Serialization/ArchiveUObject.h b/Engine/Source/Runtime/CoreUObject/Public/Serialization/ArchiveUObject.h
--- a/Engine/Source/Runtime/CoreUObject/Public/Serialization/ArchiveUObject.h
+++ b/Engine/Source/Runtime/CoreUObject/Public/Serialization/ArchiveUObject.h
@@ -1026,6 +1026,16 @@ public:
 		uint32 InPortFlags);
 };
 
+/** Base class for object replacement archives */ 
+class COREUOBJECT_API FArchiveReplaceObjectRefBase : public FArchiveUObject
+{
+protected:
+	/**
+	* Serializes a single object
+	*/
+	void SerializeObject(UObject* ObjectToSerialzie);
+};
+
 /*----------------------------------------------------------------------------
 	FArchiveReplaceObjectRef.
 ----------------------------------------------------------------------------*/
@@ -1038,7 +1048,7 @@ public:
  * NOTE: The template type must be a child of UObject or this class will not compile.
  */
 template< class T >
-class FArchiveReplaceObjectRef : public FArchiveUObject
+class FArchiveReplaceObjectRef : public FArchiveReplaceObjectRefBase
 {
 public:
 	/**
@@ -1089,29 +1099,7 @@ public:
 		{
 			// start the initial serialization
 			SerializedObjects.Add(SearchObject);
-
-			// serialization for class default objects must be deterministic (since class 
-			// default objects may be serialized during script compilation while the script
-			// and C++ versions of a class are not in sync), so use SerializeTaggedProperties()
-			// rather than the native Serialize() function
-			if ( SearchObject->HasAnyFlags(RF_ClassDefaultObject) )
-			{
-				UClass* ObjectClass = SearchObject->GetClass();
-				StartSerializingDefaults();
-				if ( !WantBinaryPropertySerialization() && (IsLoading() || IsSaving()) )
-				{
-					ObjectClass->SerializeTaggedProperties(*this, (uint8*)SearchObject, ObjectClass, NULL);
-				}
-				else
-				{
-					ObjectClass->SerializeBin(*this, SearchObject);
-				}
-				StopSerializingDefaults();
-			}
-			else
-			{
-				SearchObject->Serialize(*this);
-			}
+			SerializeObject(SearchObject);
 		}
 	}
 
@@ -1164,29 +1152,7 @@ public:
 				if (!bAlreadyAdded)
 				{
 					// otherwise recurse down into the object if it is contained within the initial search object
-	
-					// serialization for class default objects must be deterministic (since class 
-					// default objects may be serialized during script compilation while the script
-					// and C++ versions of a class are not in sync), so use SerializeTaggedProperties()
-					// rather than the native Serialize() function
-					if ( Obj->HasAnyFlags(RF_ClassDefaultObject) )
-					{
-						UClass* ObjectClass = Obj->GetClass();
-						StartSerializingDefaults();
-						if ( !WantBinaryPropertySerialization() && (IsLoading() || IsSaving()) )
-						{
-							ObjectClass->SerializeTaggedProperties(*this, (uint8*)Obj, ObjectClass, NULL);
-						}
-						else
-						{
-							ObjectClass->SerializeBin(*this, Obj);
-						}
-						StopSerializingDefaults();
-					}
-					else
-					{
-						Obj->Serialize(*this);
-					}
+					SerializeObject(Obj);
 				}
 			}
 			else if ( bNullPrivateReferences && !Obj->HasAnyFlags(RF_Public) )
diff --git a/Engine/Source/Runtime/CoreUObject/Public/UObject/GarbageCollection.h b/Engine/Source/Runtime/CoreUObject/Public/UObject/GarbageCollection.h
--- a/Engine/Source/Runtime/CoreUObject/Public/UObject/GarbageCollection.h
+++ b/Engine/Source/Runtime/CoreUObject/Public/UObject/GarbageCollection.h
@@ -8,6 +8,9 @@
 /** Context sensitive keep flags for garbage collection */
 #define GARBAGE_COLLECTION_KEEPFLAGS	(GIsEditor ? RF_Standalone : RF_NoFlags)
 
+COREUOBJECT_API DECLARE_LOG_CATEGORY_EXTERN(LogGarbage, Warning, All);
+DECLARE_STATS_GROUP(TEXT("Garbage Collection"), STATGROUP_GC, STATCAT_Advanced);
+
 /*-----------------------------------------------------------------------------
 	Realtime garbage collection helper classes.
 -----------------------------------------------------------------------------*/
@@ -476,4 +479,4 @@ class COREUOBJECT_API FGCScopeGuard
 public:
 	FGCScopeGuard();
 	~FGCScopeGuard();
-};
\ No newline at end of file
+};
diff --git a/Engine/Source/Runtime/CoreUObject/Public/UObject/ObjectBase.h b/Engine/Source/Runtime/CoreUObject/Public/UObject/ObjectBase.h
--- a/Engine/Source/Runtime/CoreUObject/Public/UObject/ObjectBase.h
+++ b/Engine/Source/Runtime/CoreUObject/Public/UObject/ObjectBase.h
@@ -510,6 +510,8 @@ enum class EInternalObjectFlags : int32
 {
 	None = 0,
 	// All the other bits are reserved, DO NOT ADD NEW FLAGS HERE!
+	ReachableInCluster = 1 << 23, /// External reference to object in cluster exists
+	ClusterRoot = 1 << 24, ///< Root of a cluster
 	Native = 1 << 25, ///< Native (UClass only).
 	Async = 1 << 26, ///< Object exists only on a different thread than the game thread.
 	AsyncLoading = 1 << 27, ///< Object is being asynchronously loaded.
@@ -520,7 +522,7 @@ enum class EInternalObjectFlags : int32
 
 	GarbageCollectionKeepFlags = Native | Async | AsyncLoading,
 	// Make sure this is up to date!
-	AllFlags = Native | Async | AsyncLoading | Unreachable | PendingKill | RootSet | NoStrongReference
+	AllFlags = ReachableInCluster | ClusterRoot | Native | Async | AsyncLoading | Unreachable | PendingKill | RootSet | NoStrongReference
 };
 ENUM_CLASS_FLAGS(EInternalObjectFlags);
 
diff --git a/Engine/Source/Runtime/CoreUObject/Public/UObject/UObjectArray.h b/Engine/Source/Runtime/CoreUObject/Public/UObject/UObjectArray.h
--- a/Engine/Source/Runtime/CoreUObject/Public/UObject/UObjectArray.h
+++ b/Engine/Source/Runtime/CoreUObject/Public/UObject/UObjectArray.h
@@ -13,9 +13,6 @@
 */
 #define UE_GC_TRACK_OBJ_AVAILABLE (WITH_EDITOR)
 
-/** Used to test for stale weak pointers in the debug visualizers **/
-extern COREUOBJECT_API int32** GSerialNumberBlocksForDebugVisualizersRoot;
-
 /**
 * Single item in the UObject array.
 */
@@ -24,10 +21,22 @@ struct FUObjectItem
 	// Pointer to the allocated object
 	class UObjectBase* Object;
 	// UObject internal flags
-	int32 Flags;
+	int32 ClusterAndFlags;
 	// Weak Object Pointer Serial number associated with the object
 	int32 SerialNumber;
 
+	FORCEINLINE void SetOwnerIndex(int32 OwnerIndex)
+	{
+		check(OwnerIndex >= 0 && (OwnerIndex & int32(EInternalObjectFlags::AllFlags)) == 0);
+		ClusterAndFlags &= int32(EInternalObjectFlags::AllFlags);
+		ClusterAndFlags |= OwnerIndex;
+	}
+
+	FORCEINLINE int32 GetOwnerIndex() const
+	{
+		return ClusterAndFlags & ~int32(EInternalObjectFlags::AllFlags);
+	}
+
 	FORCEINLINE int32 GetSerialNumber() const
 	{
 		return SerialNumber;
@@ -35,31 +44,33 @@ struct FUObjectItem
 
 	FORCEINLINE void SetFlags(EInternalObjectFlags FlagsToSet)
 	{
-		Flags |= int32(FlagsToSet);
+		check((int32(FlagsToSet) & ~int32(EInternalObjectFlags::AllFlags)) == 0);
+		ClusterAndFlags |= int32(FlagsToSet);
 	}
 
 	FORCEINLINE EInternalObjectFlags GetFlags() const
 	{
-		return EInternalObjectFlags(Flags & int32(EInternalObjectFlags::AllFlags));
+		return EInternalObjectFlags(ClusterAndFlags & int32(EInternalObjectFlags::AllFlags));
 	}
 
 	FORCEINLINE void ClearFlags(EInternalObjectFlags FlagsToClear)
 	{
-		Flags &= ~int32(FlagsToClear);
+		check((int32(FlagsToClear) & ~int32(EInternalObjectFlags::AllFlags)) == 0);
+		ClusterAndFlags &= ~int32(FlagsToClear);
 	}
 
 	FORCEINLINE bool ThisThreadAtomicallyClearedFlag(EInternalObjectFlags FlagToClear)
 	{
-		static_assert(sizeof(int32) == sizeof(Flags), "Flags must be 32-bit for atomics.");
+		static_assert(sizeof(int32) == sizeof(ClusterAndFlags), "Flags must be 32-bit for atomics.");
 		bool bIChangedIt = false;
 		while (1)
 		{
-			int32 StartValue = int32(Flags);
+			int32 StartValue = int32(ClusterAndFlags);
 			if (!(StartValue & int32(FlagToClear)))
 			{
 				break;
 			}
-			int32 OldValue = (int32)FPlatformAtomics::InterlockedCompareExchange((int32*)&Flags, StartValue & ~int32(FlagToClear), StartValue);
+			int32 OldValue = (int32)FPlatformAtomics::InterlockedCompareExchange((int32*)&ClusterAndFlags, StartValue & ~int32(FlagToClear), StartValue);
 			if (OldValue == StartValue)
 			{
 				bIChangedIt = true;
@@ -73,20 +84,20 @@ struct FUObjectItem
 
 	FORCEINLINE bool HasAnyFlags(EInternalObjectFlags InFlags) const
 	{
-		return !!(Flags & int32(InFlags));
+		return !!(ClusterAndFlags & int32(InFlags));
 	}
 
 	FORCEINLINE void SetUnreachable()
 	{
-		Flags |= int32(EInternalObjectFlags::Unreachable);
+		ClusterAndFlags |= int32(EInternalObjectFlags::Unreachable);
 	}
 	FORCEINLINE void ClearUnreachable()
 	{
-		Flags &= ~int32(EInternalObjectFlags::Unreachable);
+		ClusterAndFlags &= ~int32(EInternalObjectFlags::Unreachable);
 	}
 	FORCEINLINE bool IsUnreachable() const
 	{
-		return !!(Flags & int32(EInternalObjectFlags::Unreachable));
+		return !!(ClusterAndFlags & int32(EInternalObjectFlags::Unreachable));
 	}
 	FORCEINLINE bool ThisThreadAtomicallyClearedRFUnreachable()
 	{
@@ -95,45 +106,45 @@ struct FUObjectItem
 
 	FORCEINLINE void SetPendingKill()
 	{
-		Flags |= int32(EInternalObjectFlags::PendingKill);
+		ClusterAndFlags |= int32(EInternalObjectFlags::PendingKill);
 	}
 	FORCEINLINE void ClearPendingKill()
 	{
-		Flags &= ~int32(EInternalObjectFlags::PendingKill);
+		ClusterAndFlags &= ~int32(EInternalObjectFlags::PendingKill);
 	}
 	FORCEINLINE bool IsPendingKill() const
 	{
-		return !!(Flags & int32(EInternalObjectFlags::PendingKill));
+		return !!(ClusterAndFlags & int32(EInternalObjectFlags::PendingKill));
 	}
 
 	FORCEINLINE void SetRootSet()
 	{
-		Flags |= int32(EInternalObjectFlags::RootSet);
+		ClusterAndFlags |= int32(EInternalObjectFlags::RootSet);
 	}
 	FORCEINLINE void ClearRootSet()
 	{
-		Flags &= ~int32(EInternalObjectFlags::RootSet);
+		ClusterAndFlags &= ~int32(EInternalObjectFlags::RootSet);
 	}
 	FORCEINLINE bool IsRootSet() const
 	{
-		return !!(Flags & int32(EInternalObjectFlags::RootSet));
+		return !!(ClusterAndFlags & int32(EInternalObjectFlags::RootSet));
 	}
 
 	FORCEINLINE void SetNoStrongReference()
 	{
-		Flags |= int32(EInternalObjectFlags::NoStrongReference);
+		ClusterAndFlags |= int32(EInternalObjectFlags::NoStrongReference);
 	}
 	FORCEINLINE void ClearNoStrongReference()
 	{
-		Flags &= ~int32(EInternalObjectFlags::NoStrongReference);
+		ClusterAndFlags &= ~int32(EInternalObjectFlags::NoStrongReference);
 	}
 	FORCEINLINE bool IsNoStrongReference() const
 	{
-		return !!(Flags & int32(EInternalObjectFlags::NoStrongReference));
+		return !!(ClusterAndFlags & int32(EInternalObjectFlags::NoStrongReference));
 	}
 	FORCEINLINE void ResetSerialNumberAndFlags()
 	{
-		Flags = 0;
+		ClusterAndFlags = 0;
 		SerialNumber = 0;
 	}
 };
@@ -671,7 +682,7 @@ private:
 	/**
 	 * return the object array for use by debug visualizers
 	 */
-	static UObjectBase*** GetObjectArrayForDebugVisualizers();
+	static FFixedUObjectArray* GetObjectArrayForDebugVisualizers();
 
 	// note these variables are left with the Obj prefix so they can be related to the historical GObj versions
 
@@ -707,8 +718,15 @@ private:
 	FThreadSafeCounter	MasterSerialNumber;
 };
 
+struct FUObjectCluster
+{
+	TArray<int32> Objects;
+	TArray<int32> ReferencedClusters;
+};
+
 /** Global UObject allocator							*/
 extern COREUOBJECT_API FUObjectArray GUObjectArray;
+extern COREUOBJECT_API TMap<int32, FUObjectCluster* > GUObjectClusters;
 
 /**
 	* Static version of IndexToObject for use with TWeakObjectPtr.
diff --git a/Engine/Source/Runtime/CoreUObject/Public/UObject/UObjectBaseUtility.h b/Engine/Source/Runtime/CoreUObject/Public/UObject/UObjectBaseUtility.h
--- a/Engine/Source/Runtime/CoreUObject/Public/UObject/UObjectBaseUtility.h
+++ b/Engine/Source/Runtime/CoreUObject/Public/UObject/UObjectBaseUtility.h
@@ -286,11 +286,38 @@ public:
 	 */
 	FString GetPathName( const UObject* StopOuter=NULL ) const;
 
+	/**
+	* Called after load to determine if the object can be a cluster root
+	*
+	* @return	true if this object can be a cluster root
+	*/
+	virtual bool CanBeClusterRoot() const
+	{
+		return false;
+	}
+
+	/**
+	* Called during cluster construction if the object can be added to a cluster
+	*
+	* @return	true if this object can be inside of a cluster
+	*/
+	virtual bool CanBeInCluster() const
+	{
+		return true;
+	}
+
 protected:
+
 	/**
 	 * Internal version of GetPathName() that eliminates lots of copies.
 	 */
 	void GetPathName( const UObject* StopOuter, FString& ResultString ) const;
+
+	/**
+	* Called during PostLoad to create UObject cluster
+	*/
+	virtual void CreateCluster();
+
 public:
 	/**
 	 * Walks up the chain of packages until it reaches the top level, which it ignores.
diff --git a/Engine/Source/Runtime/CoreUObject/Public/UObject/UObjectGlobals.h b/Engine/Source/Runtime/CoreUObject/Public/UObject/UObjectGlobals.h
--- a/Engine/Source/Runtime/CoreUObject/Public/UObject/UObjectGlobals.h
+++ b/Engine/Source/Runtime/CoreUObject/Public/UObject/UObjectGlobals.h
@@ -1556,10 +1556,7 @@ public:
 	template<class UObjectType>
 	void AddReferencedObjects(TArray<UObjectType>& ObjectArray, const UObject* ReferencingObject = NULL, const UProperty* ReferencingProperty = NULL)
 	{
-		for (auto& Object : ObjectArray)
-		{
-			HandleObjectReference(*(UObject**)&Object, ReferencingObject, ReferencingProperty);
-		}
+		HandleObjectReferences(reinterpret_cast<UObject**>(ObjectArray.GetData()), ObjectArray.Num(), ReferencingObject, ReferencingProperty);
 	}
 
 	/**
@@ -1663,6 +1660,23 @@ protected:
 	 * @param ReferencingProperty Referencing property (if available).
 	 */
 	virtual void HandleObjectReference(UObject*& InObject, const UObject* InReferencingObject, const UProperty* InReferencingProperty) = 0;
+
+	/**
+	* Handle multiple object references. Called by AddReferencedObjects.
+	* DEFAULT IMPLEMENTAION IS SLOW as it calls HandleObjectReference multiple times. In order to optimize it, provide your own implementation.
+	*
+	* @param Object Referenced object.
+	* @param ReferencingObject Referencing object (if available).
+	* @param ReferencingProperty Referencing property (if available).
+	*/
+	virtual void HandleObjectReferences(UObject** InObjects, const int32 ObjectNum, const UObject* InReferencingObject, const UProperty* InReferencingProperty)
+	{
+		for (int32 ObjectIndex = 0; ObjectIndex < ObjectNum; ++ObjectIndex)
+		{
+			UObject*& Object = InObjects[ObjectIndex];
+			HandleObjectReference(Object, InReferencingObject, InReferencingProperty);
+		}
+	}
 };
 
 /**
diff --git a/Engine/Source/Runtime/Engine/Classes/Animation/AnimInstance.h b/Engine/Source/Runtime/Engine/Classes/Animation/AnimInstance.h
--- a/Engine/Source/Runtime/Engine/Classes/Animation/AnimInstance.h
+++ b/Engine/Source/Runtime/Engine/Classes/Animation/AnimInstance.h
@@ -298,10 +298,18 @@ struct FNativeStateBinding
 	/** Delegate to use when checking transition */
 	FOnGraphStateChanged NativeStateDelegate;
 
-	FNativeStateBinding(const FName& InMachineName, const FName& InStateName, const FOnGraphStateChanged& InNativeStateDelegate)
+#if WITH_EDITORONLY_DATA
+	/** Name of this binding */
+	FName BindingName;
+#endif
+
+	FNativeStateBinding(const FName& InMachineName, const FName& InStateName, const FOnGraphStateChanged& InNativeStateDelegate, const FName& InBindingName = NAME_None)
 		: MachineName(InMachineName)
 		, StateName(InStateName)
 		, NativeStateDelegate(InNativeStateDelegate)
+#if WITH_EDITORONLY_DATA
+		, BindingName(InBindingName)
+#endif
 	{
 	}
 };
diff --git a/Engine/Source/Runtime/Engine/Classes/Engine/CoreSettings.h b/Engine/Source/Runtime/Engine/Classes/Engine/CoreSettings.h
--- a/Engine/Source/Runtime/Engine/Classes/Engine/CoreSettings.h
+++ b/Engine/Source/Runtime/Engine/Classes/Engine/CoreSettings.h
@@ -136,6 +136,16 @@ protected:
 		ToolTip = "If enabled, garbage collection will use multiple threads."))
 	uint32 AllowParallelGC : 1;
 
+	UPROPERTY(EditAnywhere, config, Category = Optimization, meta = (
+		ConsoleVariable = "gc.CreateGCClusters", DisplayName = "Create Garbage Collector UObject Clusters",
+		ToolTip = "If true, the engine will attempt to create clusters of objects for better garbage collection performance."))
+	uint32 CreateGCClusters : 1;
+
+	UPROPERTY(EditAnywhere, config, Category = Optimization, meta = (
+		ConsoleVariable = "gc.MergeGCClusters", DisplayName = "Merge GC Clusters",
+		ToolTip = "If true, when creating clusters, the clusters referenced from another cluster will get merged into one big cluster."))
+		uint32 MergeGCClusters : 1;
+
 	UPROPERTY(EditAnywhere, config, Category = General, meta = (
 		ConsoleVariable = "gc.NumRetriesBeforeForcingGC", DisplayName = "Number Of Retries Before Forcing GC",
 		ToolTip = "Maximum number of times GC can be skipped if worker threads are currently modifying UObject state. 0 = never force GC"))
diff --git a/Engine/Source/Runtime/Engine/Classes/Materials/Material.h b/Engine/Source/Runtime/Engine/Classes/Materials/Material.h
--- a/Engine/Source/Runtime/Engine/Classes/Materials/Material.h
+++ b/Engine/Source/Runtime/Engine/Classes/Materials/Material.h
@@ -841,6 +841,7 @@ public:
 	ENGINE_API virtual void FinishDestroy() override;
 	ENGINE_API virtual SIZE_T GetResourceSize(EResourceSizeMode::Type Mode) override;
 	ENGINE_API static void AddReferencedObjects(UObject* InThis, FReferenceCollector& Collector);
+	ENGINE_API virtual bool CanBeClusterRoot() const override;
 	//~ End UObject Interface
 
 #if WITH_EDITOR
diff --git a/Engine/Source/Runtime/Engine/Classes/Particles/ParticleSystem.h b/Engine/Source/Runtime/Engine/Classes/Particles/ParticleSystem.h
--- a/Engine/Source/Runtime/Engine/Classes/Particles/ParticleSystem.h
+++ b/Engine/Source/Runtime/Engine/Classes/Particles/ParticleSystem.h
@@ -335,6 +335,7 @@ public:
 	virtual void PostLoad() override;
 	virtual SIZE_T GetResourceSize(EResourceSizeMode::Type Mode) override;
 	virtual void GetAssetRegistryTags(TArray<FAssetRegistryTag>& OutTags) const override;
+	virtual bool CanBeClusterRoot() const override;
 	//~ End UObject Interface.
 
 
diff --git a/Engine/Source/Runtime/Engine/Private/Animation/AnimInstanceProxy.cpp b/Engine/Source/Runtime/Engine/Private/Animation/AnimInstanceProxy.cpp
--- a/Engine/Source/Runtime/Engine/Private/Animation/AnimInstanceProxy.cpp
+++ b/Engine/Source/Runtime/Engine/Private/Animation/AnimInstanceProxy.cpp
@@ -100,7 +100,7 @@ void FAnimInstanceProxy::InitializeRootNode()
 		GameThreadPreUpdateNodes.Reset();
 
 		// cache any state machine descriptions we have
-		for (UStructProperty* Property : AnimClassInterface->GetAnimNodeProperties())
+		for(UStructProperty* Property : AnimClassInterface->GetAnimNodeProperties())
 		{
 			if (Property->Struct->IsChildOf(FAnimNode_Base::StaticStruct()))
 			{
@@ -112,12 +112,12 @@ void FAnimInstanceProxy::InitializeRootNode()
 						GameThreadPreUpdateNodes.Add(AnimNode);
 					}
 
-					if (Property->Struct->IsChildOf(FAnimNode_StateMachine::StaticStruct()))
-					{
+			if(Property->Struct->IsChildOf(FAnimNode_StateMachine::StaticStruct()))
+			{
 						FAnimNode_StateMachine* StateMachine = static_cast<FAnimNode_StateMachine*>(AnimNode);
-						StateMachine->CacheMachineDescription(AnimClassInterface);
-					}
-				}
+				StateMachine->CacheMachineDescription(AnimClassInterface);
+			}
+		}
 			}
 		}
 
@@ -1174,17 +1174,10 @@ bool FAnimInstanceProxy::HasNativeTransitionBinding(const FName& MachineName, co
 	{
 		if(Binding.MachineName == MachineName && Binding.PreviousStateName == PrevStateName && Binding.NextStateName == NextStateName)
 		{
-			IDelegateInstance* DelegateInstance = Binding.NativeTransitionDelegate.GetDelegateInstance();
-			if(DelegateInstance)
-			{
-				OutBindingName = DelegateInstance->GetFunctionName();
-			}
 #if WITH_EDITORONLY_DATA
-			// if the function wasnt forthcoming with a name, use the explicit name supplied in the binding
-			if(OutBindingName == NAME_None)
-			{
 				OutBindingName = Binding.TransitionName;
-			}
+#else
+			OutBindingName = NAME_None;
 #endif
 			return true;
 		}
@@ -1193,9 +1186,9 @@ bool FAnimInstanceProxy::HasNativeTransitionBinding(const FName& MachineName, co
 	return false;
 }
 
-void FAnimInstanceProxy::AddNativeStateEntryBinding(const FName& MachineName, const FName& StateName, const FOnGraphStateChanged& NativeEnteredDelegate)
+void FAnimInstanceProxy::AddNativeStateEntryBinding(const FName& MachineName, const FName& StateName, const FOnGraphStateChanged& NativeEnteredDelegate, const FName& BindingName)
 {
-	NativeStateEntryBindings.Add(FNativeStateBinding(MachineName, StateName, NativeEnteredDelegate));
+	NativeStateEntryBindings.Add(FNativeStateBinding(MachineName, StateName, NativeEnteredDelegate, BindingName));
 }
 	
 bool FAnimInstanceProxy::HasNativeStateEntryBinding(const FName& MachineName, const FName& StateName, FName& OutBindingName)
@@ -1204,11 +1197,11 @@ bool FAnimInstanceProxy::HasNativeStateEntryBinding(const FName& MachineName, co
 	{
 		if(Binding.MachineName == MachineName && Binding.StateName == StateName)
 		{
-			IDelegateInstance* DelegateInstance = Binding.NativeStateDelegate.GetDelegateInstance();
-			if(DelegateInstance)
-			{
-				OutBindingName = DelegateInstance->GetFunctionName();
-			}
+#if WITH_EDITORONLY_DATA
+			OutBindingName = Binding.BindingName;
+#else
+			OutBindingName = NAME_None;
+#endif
 			return true;
 		}
 	}
@@ -1216,9 +1209,9 @@ bool FAnimInstanceProxy::HasNativeStateEntryBinding(const FName& MachineName, co
 	return false;
 }
 
-void FAnimInstanceProxy::AddNativeStateExitBinding(const FName& MachineName, const FName& StateName, const FOnGraphStateChanged& NativeExitedDelegate)
+void FAnimInstanceProxy::AddNativeStateExitBinding(const FName& MachineName, const FName& StateName, const FOnGraphStateChanged& NativeExitedDelegate, const FName& BindingName)
 {
-	NativeStateExitBindings.Add(FNativeStateBinding(MachineName, StateName, NativeExitedDelegate));
+	NativeStateExitBindings.Add(FNativeStateBinding(MachineName, StateName, NativeExitedDelegate, BindingName));
 }
 
 bool FAnimInstanceProxy::HasNativeStateExitBinding(const FName& MachineName, const FName& StateName, FName& OutBindingName)
@@ -1227,11 +1220,11 @@ bool FAnimInstanceProxy::HasNativeStateExitBinding(const FName& MachineName, con
 	{
 		if(Binding.MachineName == MachineName && Binding.StateName == StateName)
 		{
-			IDelegateInstance* DelegateInstance = Binding.NativeStateDelegate.GetDelegateInstance();
-			if(DelegateInstance)
-			{
-				OutBindingName = DelegateInstance->GetFunctionName();
-			}
+#if WITH_EDITORONLY_DATA
+			OutBindingName = Binding.BindingName;
+#else
+			OutBindingName = NAME_None;
+#endif
 			return true;
 		}
 	}
diff --git a/Engine/Source/Runtime/Engine/Private/CoreSettings.cpp b/Engine/Source/Runtime/Engine/Private/CoreSettings.cpp
--- a/Engine/Source/Runtime/Engine/Private/CoreSettings.cpp
+++ b/Engine/Source/Runtime/Engine/Private/CoreSettings.cpp
@@ -70,7 +70,7 @@ UStreamingSettings::UStreamingSettings()
 	AsyncLoadingUseFullTimeLimit = true;
 	PriorityAsyncLoadingExtraTime = 20.0f;
 	LevelStreamingActorsUpdateTimeLimit = 5.0f;
-	LevelStreamingComponentsRegistrationGranularity = 10;
+	LevelStreamingComponentsRegistrationGranularity = 10;	
 }
 
 void UStreamingSettings::PostInitProperties()
@@ -110,6 +110,8 @@ UGarbageCollectionSettings::UGarbageCollectionSettings()
 	SizeOfPermanentObjectPool = 0;
 	MaxObjectsInEditor = 12 * 1024 * 1024;
 	MaxObjectsInGame = 2 * 1024 * 1024;
+	CreateGCClusters = true;
+	MergeGCClusters = false;
 }
 
 void UGarbageCollectionSettings::PostInitProperties()
diff --git a/Engine/Source/Runtime/Engine/Private/Materials/Material.cpp b/Engine/Source/Runtime/Engine/Private/Materials/Material.cpp
--- a/Engine/Source/Runtime/Engine/Private/Materials/Material.cpp
+++ b/Engine/Source/Runtime/Engine/Private/Materials/Material.cpp
@@ -3211,6 +3211,11 @@ void UMaterial::AddReferencedObjects(UObject* InThis, FReferenceCollector& Colle
 	Super::AddReferencedObjects(This, Collector);
 }
 
+bool UMaterial::CanBeClusterRoot() const 
+{
+	return true;
+}
+
 #if WITH_EDITOR
 void UMaterial::CancelOutstandingCompilation()
 {
diff --git a/Engine/Source/Runtime/Engine/Private/Particles/ParticleComponents.cpp b/Engine/Source/Runtime/Engine/Private/Particles/ParticleComponents.cpp
--- a/Engine/Source/Runtime/Engine/Private/Particles/ParticleComponents.cpp
+++ b/Engine/Source/Runtime/Engine/Private/Particles/ParticleComponents.cpp
@@ -2328,6 +2328,10 @@ void UParticleSystem::GetAssetRegistryTags(TArray<FAssetRegistryTag>& OutTags) c
 	Super::GetAssetRegistryTags(OutTags);
 }
 
+bool UParticleSystem::CanBeClusterRoot() const
+{
+	return true;
+}
 
 bool UParticleSystem::CalculateMaxActiveParticleCounts()
 {
diff --git a/Engine/Source/Runtime/Engine/Private/TimerManager.cpp b/Engine/Source/Runtime/Engine/Private/TimerManager.cpp
--- a/Engine/Source/Runtime/Engine/Private/TimerManager.cpp
+++ b/Engine/Source/Runtime/Engine/Private/TimerManager.cpp
@@ -32,11 +32,9 @@ FString FTimerUnifiedDelegate::ToString() const
 
 	if (FuncDelegate.IsBound())
 	{
-		Object = FuncDelegate.GetUObject();
-		if (IDelegateInstance* Ptr = FuncDelegate.GetDelegateInstance())
-		{
-			FunctionName = Ptr->GetFunctionName();
-		}
+#if USE_DELEGATE_TRYGETBOUNDFUNCTIONNAME
+		FunctionName = FuncDelegate.TryGetBoundFunctionName();
+#endif
 	}
 	else if (FuncDynDelegate.IsBound())
 	{
diff --git a/Engine/Source/Runtime/Engine/Private/UnrealEngine.cpp b/Engine/Source/Runtime/Engine/Private/UnrealEngine.cpp
--- a/Engine/Source/Runtime/Engine/Private/UnrealEngine.cpp
+++ b/Engine/Source/Runtime/Engine/Private/UnrealEngine.cpp
@@ -40,6 +40,7 @@
 #include "IMotionController.h"
 #include "Scalability.h"
 #include "StatsData.h"
+#include "StatsFile.h"
 #include "ScreenRendering.h"
 #include "RHIStaticStates.h"
 #include "AudioDeviceManager.h"
@@ -7351,6 +7352,7 @@ void DrawStatsHUD( UWorld* World, FViewport* Viewport, FCanvas* Canvas, UCanvas*
 	const int32 StatsXOffset = 100;// FPlatformProperties::SupportsWindowedMode() ? 4 : 100;
 
 	int32 MessageY = 35;
+	const int32 FontSizeY = 20;
 
 	if( !GIsEditor )
 	{
@@ -7386,7 +7388,7 @@ void DrawStatsHUD( UWorld* World, FViewport* Viewport, FCanvas* Canvas, UCanvas*
 				}
 				SmallTextItem.Text =  FText::FromString( FString::Printf(TEXT("LIGHTING NEEDS TO BE REBUILT (%u unbuilt object(s))"), World->NumLightingUnbuiltObjects) );				
 				Canvas->DrawItem( SmallTextItem, FVector2D( MessageX, MessageY ) );
-				MessageY += 20;
+				MessageY += FontSizeY;
 			}
 			
 #if !(UE_BUILD_SHIPPING || UE_BUILD_TEST)
@@ -7398,7 +7400,7 @@ void DrawStatsHUD( UWorld* World, FViewport* Viewport, FCanvas* Canvas, UCanvas*
 					SmallTextItem.SetColor(FLinearColor::Red);
 					SmallTextItem.Text = FText::FromString(FString::Printf(TEXT("TEXTURE STREAMING POOL OVER %0.2f MB"), (float)MemOver / 1024.0f / 1024.0f));
 					Canvas->DrawItem(SmallTextItem, FVector2D(MessageX, MessageY));
-					MessageY += 20;
+					MessageY += FontSizeY;
 				}
 			}
 #endif
@@ -7416,7 +7418,7 @@ void DrawStatsHUD( UWorld* World, FViewport* Viewport, FCanvas* Canvas, UCanvas*
 				SmallTextItem.SetColor( FLinearColor::White );
 				SmallTextItem.Text =  LOCTEXT("NAVMESHERROR", "NAVMESH NEEDS TO BE REBUILT");				
 				Canvas->DrawItem( SmallTextItem, FVector2D( MessageX, MessageY ) );
-				MessageY += 20;
+				MessageY += FontSizeY;
 			}
 
 			if( World->bKismetScriptError )
@@ -7424,7 +7426,7 @@ void DrawStatsHUD( UWorld* World, FViewport* Viewport, FCanvas* Canvas, UCanvas*
 				SmallTextItem.Text = LOCTEXT("BlueprintInLevelHadCompileErrorMessage", "BLUEPRINT COMPILE ERROR" );
 				SmallTextItem.SetColor( FLinearColor::Red );
 				Canvas->DrawItem( SmallTextItem, FVector2D( MessageX, MessageY ) );
-				MessageY += 20;
+				MessageY += FontSizeY;
 			}
 
 			SmallTextItem.SetColor( FLinearColor::White );
@@ -7433,7 +7435,7 @@ void DrawStatsHUD( UWorld* World, FViewport* Viewport, FCanvas* Canvas, UCanvas*
 			{
 				SmallTextItem.Text = FText::FromString(FString::Printf(TEXT("Shaders Compiling (%u)"), GShaderCompilingManager->GetNumRemainingJobs()));
 				Canvas->DrawItem(SmallTextItem, FVector2D(MessageX, MessageY));
-				MessageY += 20;
+				MessageY += FontSizeY;
 			}
 
 #if ENABLE_VISUAL_LOG
@@ -7464,7 +7466,7 @@ void DrawStatsHUD( UWorld* World, FViewport* Viewport, FCanvas* Canvas, UCanvas*
 					GEngine->GetSmallFont(),
 					FColor(128,128,128)
 					);
-				MessageY += 20;
+				MessageY += FontSizeY;
 			}
 			*/
 
@@ -7472,14 +7474,14 @@ void DrawStatsHUD( UWorld* World, FViewport* Viewport, FCanvas* Canvas, UCanvas*
 			{
 				SmallTextItem.Text =  LOCTEXT("Levelstreamingfrozen", "Level streaming frozen..." );
 				Canvas->DrawItem( SmallTextItem, FVector2D( MessageX, MessageY ) );
-				MessageY += 20;
+				MessageY += FontSizeY;
 			}
 
 			if (GIsPrepareMapChangeBroken)
 			{
 				SmallTextItem.Text =  LOCTEXT("PrepareMapChangeError", "PrepareMapChange had a bad level name! Check the log (tagged with PREPAREMAPCHANGE) for info" );
 				Canvas->DrawItem( SmallTextItem, FVector2D( MessageX, MessageY ) );
-				MessageY += 20;
+				MessageY += FontSizeY;
 			}
 
 #if STATS
@@ -7490,13 +7492,22 @@ void DrawStatsHUD( UWorld* World, FViewport* Viewport, FCanvas* Canvas, UCanvas*
 				{				
 					SmallTextItem.Text =  LOCTEXT("AIPROFILINGWARNING", "PROFILING WITH AI LOGGING ON!" );
 					Canvas->DrawItem( SmallTextItem, FVector2D( MessageX, MessageY ) );
-					MessageY += 20;
+					MessageY += FontSizeY;
 				}
 				if (GShouldVerifyGCAssumptions)
 				{
 					SmallTextItem.Text =  LOCTEXT("GCPROFILINGWARNING", "PROFILING WITH GC VERIFY ON!" );
 					Canvas->DrawItem( SmallTextItem, FVector2D( MessageX, MessageY ) );					
-					MessageY += 20;
+					MessageY += FontSizeY;
+				}
+
+				const bool bIsStatsFileActive = FCommandStatsFile::Get().IsStatFileActive();
+				if (bIsStatsFileActive)
+				{
+					SmallTextItem.SetColor( FLinearColor::White );
+					SmallTextItem.Text = FCommandStatsFile::Get().GetFileMetaDesc();
+					Canvas->DrawItem( SmallTextItem, FVector2D( MessageX, MessageY ) );
+					MessageY += FontSizeY;
 				}
 			}
 #endif
diff --git a/Engine/Source/Runtime/Engine/Public/Animation/AnimInstanceProxy.h b/Engine/Source/Runtime/Engine/Public/Animation/AnimInstanceProxy.h
--- a/Engine/Source/Runtime/Engine/Public/Animation/AnimInstanceProxy.h
+++ b/Engine/Source/Runtime/Engine/Public/Animation/AnimInstanceProxy.h
@@ -434,13 +434,13 @@ protected:
 	bool HasNativeTransitionBinding(const FName& MachineName, const FName& PrevStateName, const FName& NextStateName, FName& OutBindingName);
 
 	// Sets up a native state entry delegate from state with StateName, in the state machine with name MachineName.
-	void AddNativeStateEntryBinding(const FName& MachineName, const FName& StateName, const FOnGraphStateChanged& NativeEnteredDelegate);
+	void AddNativeStateEntryBinding(const FName& MachineName, const FName& StateName, const FOnGraphStateChanged& NativeEnteredDelegate, const FName& BindingName = NAME_None);
 	
 	// Check for whether a native entry delegate is bound to the specified state
 	bool HasNativeStateEntryBinding(const FName& MachineName, const FName& StateName, FName& OutBindingName);
 
 	// Sets up a native state exit delegate from state with StateName, in the state machine with name MachineName.
-	void AddNativeStateExitBinding(const FName& MachineName, const FName& StateName, const FOnGraphStateChanged& NativeExitedDelegate);
+	void AddNativeStateExitBinding(const FName& MachineName, const FName& StateName, const FOnGraphStateChanged& NativeExitedDelegate, const FName& BindingName = NAME_None);
 
 	// Check for whether a native exit delegate is bound to the specified state
 	bool HasNativeStateExitBinding(const FName& MachineName, const FName& StateName, FName& OutBindingName);
diff --git a/Engine/Source/Runtime/PakFile/Public/IPlatformFilePak.h b/Engine/Source/Runtime/PakFile/Public/IPlatformFilePak.h
--- a/Engine/Source/Runtime/PakFile/Public/IPlatformFilePak.h
+++ b/Engine/Source/Runtime/PakFile/Public/IPlatformFilePak.h
@@ -1072,6 +1072,27 @@ public:
 		return Result;
 	}
 
+	virtual void GetTimeStampPair(const TCHAR* FilenameA, const TCHAR* FilenameB, FDateTime& OutTimeStampA, FDateTime& OutTimeStampB) override
+	{
+		FPakFile* PakFileA = nullptr;
+		FPakFile* PakFileB = nullptr;
+		FindFileInPakFiles(FilenameA, &PakFileA);
+		FindFileInPakFiles(FilenameB, &PakFileB);
+
+		// If either file exists, we'll assume both should exist here and therefore we can skip the
+		// request to the lower level platform file.
+		if (PakFileA != nullptr || PakFileB != nullptr)
+		{
+			OutTimeStampA = PakFileA != nullptr ? PakFileA->GetTimestamp() : FDateTime::MinValue();
+			OutTimeStampB = PakFileB != nullptr ? PakFileB->GetTimestamp() : FDateTime::MinValue();
+		}
+		else
+		{
+			// Fall back to lower level.
+			LowerLevel->GetTimeStampPair(FilenameA, FilenameB, OutTimeStampA, OutTimeStampB);
+		}
+	}
+
 	virtual void SetTimeStamp(const TCHAR* Filename, FDateTime DateTime) override
 	{
 		// No modifications allowed on files from pak (although we could theoretically allow this one).
diff --git a/Engine/Source/Runtime/Windows/D3D12RHI/Private/D3D12Device.cpp b/Engine/Source/Runtime/Windows/D3D12RHI/Private/D3D12Device.cpp
--- a/Engine/Source/Runtime/Windows/D3D12RHI/Private/D3D12Device.cpp
+++ b/Engine/Source/Runtime/Windows/D3D12RHI/Private/D3D12Device.cpp
@@ -735,7 +735,6 @@ void FD3D12Device::SetupAfterDeviceCreation()
 {
 	CreateSignatures();
 
-    PipelineStateCache = FD3D12PipelineStateCache(this);
     FString GraphicsCacheFile = FPaths::GameSavedDir() / TEXT("D3DGraphics.ushaderprecache");
     FString ComputeCacheFile = FPaths::GameSavedDir() / TEXT("D3DCompute.ushaderprecache");
 
diff --git a/Engine/Source/Runtime/Windows/D3D12RHI/Private/D3D12DirectCommandListManager.cpp b/Engine/Source/Runtime/Windows/D3D12RHI/Private/D3D12DirectCommandListManager.cpp
--- a/Engine/Source/Runtime/Windows/D3D12RHI/Private/D3D12DirectCommandListManager.cpp
+++ b/Engine/Source/Runtime/Windows/D3D12RHI/Private/D3D12DirectCommandListManager.cpp
@@ -246,7 +246,11 @@ void FD3D12CommandListManager::ExecuteCommandLists(TArray<FD3D12CommandListHandl
 	if (NeedsResourceBarriers)
 	{
 #if UE_BUILD_DEBUG	
-		if (!ResourceStateCS.TryLock())
+		if (ResourceStateCS.TryLock())
+		{
+			ResourceStateCS.Unlock();
+		}
+		else
 		{
 			FD3D12DynamicRHI::GetD3DRHI()->SubmissionLockStalls++;
 			// We don't think this will get hit but it's possible. If we do see this happen,
diff --git a/Engine/Source/Runtime/Windows/D3D12RHI/Private/D3D12StateCachePrivate.h b/Engine/Source/Runtime/Windows/D3D12RHI/Private/D3D12StateCachePrivate.h
--- a/Engine/Source/Runtime/Windows/D3D12RHI/Private/D3D12StateCachePrivate.h
+++ b/Engine/Source/Runtime/Windows/D3D12RHI/Private/D3D12StateCachePrivate.h
@@ -710,6 +710,10 @@ public:
 	~FD3D12PipelineStateCache();
 
 	static const bool bUseAPILibaries = false;
+
+private:
+	FD3D12PipelineStateCache(const FD3D12PipelineStateCache&);
+	FD3D12PipelineStateCache& operator=(const FD3D12PipelineStateCache&);
 };
 
 class FD3D12BitArray
