Copying //UE4/Orion-Staging to //UE4/Main (Origin: //Orion/Dev-General [REDACTED_HANDLE])

diff --git a/Engine/Source/Runtime/Core/Public/Misc/CoreMisc.h b/Engine/Source/Runtime/Core/Public/Misc/CoreMisc.h
--- a/Engine/Source/Runtime/Core/Public/Misc/CoreMisc.h
+++ b/Engine/Source/Runtime/Core/Public/Misc/CoreMisc.h
@@ -72,6 +72,11 @@ struct CORE_API FCommandLine
 	 */
 	static const TCHAR* Get();
 	
+	/**
+	 * Returns an edited version of the executable's command line. 
+	 */
+	static const TCHAR* GetForLogging();
+
 	/**
 	 * Returns the command line originally passed to the executable.
 	 */
@@ -127,9 +132,12 @@ private:
 	static void WhitelistCommandLines();
 	/** Filters any command line args that aren't on the approved list */
 	static TArray<FString> FilterCommandLine(TCHAR* CommandLine);
+	/** Filters any command line args that are on the to-strip list */
+	static TArray<FString> FilterCommandLineForLogging(TCHAR* CommandLine);
 	/** Rebuilds the command line using the filtered args */
 	static void BuildWhitelistCommandLine(TCHAR* CommandLine, uint32 Length, const TArray<FString>& FilteredArgs);
 	static TArray<FString> ApprovedArgs;
+	static TArray<FString> FilterArgsForLogging;
 #else
 #define WhitelistCommandLines()
 #endif
@@ -140,6 +148,8 @@ private:
 	static TCHAR CmdLine[MaxCommandLineSize];
 	/** character buffer containing the original command line */
 	static TCHAR OriginalCmdLine[MaxCommandLineSize];
+	/** character buffer containing the command line filtered for logging purposes */
+	static TCHAR LoggingCmdLine[MaxCommandLineSize];
 	/** subprocess command line */
 	static FString SubprocessCommandLine;
 };
diff --git a/Engine/Source/Runtime/Core/Public/Misc/MemStack.h b/Engine/Source/Runtime/Core/Public/Misc/MemStack.h
--- a/Engine/Source/Runtime/Core/Public/Misc/MemStack.h
+++ b/Engine/Source/Runtime/Core/Public/Misc/MemStack.h
@@ -26,36 +26,13 @@ public:
 		PageSize = 64 * 1024,
 		SmallPageSize = 1024
 	};
-	FORCEINLINE static void *Alloc()
-	{
-		void *Result = TheAllocator.Allocate();
-		STAT(UpdateStats());
-		return Result;
-	}
-	FORCEINLINE static void Free(void *Mem)
-	{
-		TheAllocator.Free(Mem);
-		STAT(UpdateStats());
-	}
-	FORCEINLINE static void *AllocSmall()
-	{
-		void *Result = TheSmallAllocator.Allocate();
-		STAT(UpdateStats());
-		return Result;
-	}
-	FORCEINLINE static void FreeSmall(void *Mem)
-	{
-		TheSmallAllocator.Free(Mem);
-		STAT(UpdateStats());
-	}
-	static uint64 BytesUsed()
-	{
-		return uint64(TheAllocator.GetNumUsed().GetValue()) * PageSize + uint64(TheSmallAllocator.GetNumUsed().GetValue()) * SmallPageSize;
-	}
-	static uint64 BytesFree()
-	{
-		return uint64(TheAllocator.GetNumFree().GetValue()) * PageSize + uint64(TheSmallAllocator.GetNumFree().GetValue()) * SmallPageSize;
-	}
+	static void *Alloc();
+	static void Free(void *Mem);
+	static void *AllocSmall();
+	static void FreeSmall(void *Mem);
+	static uint64 BytesUsed();
+	static uint64 BytesFree();
+	static void LatchProtectedMode();
 private:
 
 #if STATS
diff --git a/Engine/Source/Runtime/Core/Public/Misc/Paths.h b/Engine/Source/Runtime/Core/Public/Misc/Paths.h
--- a/Engine/Source/Runtime/Core/Public/Misc/Paths.h
+++ b/Engine/Source/Runtime/Core/Public/Misc/Paths.h
@@ -324,6 +324,9 @@ public:
 	// Returns the path in front of the filename
 	static FString GetPath(const FString& InPath);
 
+	// Changes the extension of the given filename
+	static FString ChangeExtension(const FString& InPath, const FString& InNewExtension);
+
 	/** @return true if this file was found, false otherwise */
 	static bool FileExists(const FString& InPath);
 
diff --git a/Engine/Source/Runtime/Core/Public/ProfilingDebugging/ProfilingHelpers.h b/Engine/Source/Runtime/Core/Public/ProfilingDebugging/ProfilingHelpers.h
--- a/Engine/Source/Runtime/Core/Public/ProfilingDebugging/ProfilingHelpers.h
+++ b/Engine/Source/Runtime/Core/Public/ProfilingDebugging/ProfilingHelpers.h
@@ -48,6 +48,14 @@ CORE_API void SendDataToPCViaUnrealConsole( const FString& NotifyType, const FSt
  **/
 CORE_API FString CreateProfileFilename( const FString& InFileExtension, bool bIncludeDateForDirectoryName );
 
+/**
+* This will generate the profiling file name that will work with limited filename sizes on consoles.
+* We want a uniform naming convention so we will all just call this function.
+*
+*
+**/
+CORE_API FString CreateProfileFilename(const FString& InFilename, const FString& InFileExtension, bool bIncludeDateForDirectoryName);
+
 /** 
  * This will create the directories and the file name all in one function
  **/
diff --git a/Engine/Source/Runtime/Core/Public/Stats/Stats2.h b/Engine/Source/Runtime/Core/Public/Stats/Stats2.h
--- a/Engine/Source/Runtime/Core/Public/Stats/Stats2.h
+++ b/Engine/Source/Runtime/Core/Public/Stats/Stats2.h
@@ -2018,6 +2018,7 @@ DECLARE_STATS_GROUP(TEXT("Engine"),STATGROUP_Engine, STATCAT_Advanced);
 DECLARE_STATS_GROUP(TEXT("FPS Chart"),STATGROUP_FPSChart, STATCAT_Advanced);
 DECLARE_STATS_GROUP(TEXT("GPU Particles"),STATGROUP_GPUParticles, STATCAT_Advanced);
 DECLARE_STATS_GROUP(TEXT("Game"),STATGROUP_Game, STATCAT_Advanced);
+DECLARE_STATS_GROUP(TEXT("GPU Defrag"), STATGROUP_GPUDEFRAG, STATCAT_Advanced);
 DECLARE_STATS_GROUP(TEXT("Gnm"),STATGROUP_PS4RHI, STATCAT_Advanced);
 DECLARE_STATS_GROUP_VERBOSE(TEXT("GnmVerbose"), STATGROUP_PS4RHIVERBOSE, STATCAT_Advanced);
 DECLARE_STATS_GROUP(TEXT("Init Views"),STATGROUP_InitViews, STATCAT_Advanced);
diff --git a/Engine/Source/Runtime/Core/Public/Stats/StatsData.h b/Engine/Source/Runtime/Core/Public/Stats/StatsData.h
--- a/Engine/Source/Runtime/Core/Public/Stats/StatsData.h
+++ b/Engine/Source/Runtime/Core/Public/Stats/StatsData.h
@@ -287,6 +287,12 @@ struct CORE_API FRawStatStackNode
 	/** Print this tree to the log **/
 	void DebugPrint(TCHAR const* Filter = nullptr, int32 MaxDepth = MAX_int32, int32 Depth = 0) const;
 
+	/** Print this tree to the log **/
+	void DebugPrintLeafFilter(TCHAR const* Filter) const;
+
+	/** Print this tree to the log **/
+	void DebugPrintLeafFilterInner(TCHAR const* Filter, int32 Depth, TArray<FString>& Stack) const;
+
 	/** Condense this tree into a flat list using EStatOperation::ChildrenStart, EStatOperation::ChildrenEnd, EStatOperation::Leaf **/
 	void Encode(TArray<FStatMessage>& OutStats) const;
 
diff --git a/Engine/Source/Runtime/Core/Public/Windows/WindowsPlatformMemory.h b/Engine/Source/Runtime/Core/Public/Windows/WindowsPlatformMemory.h
--- a/Engine/Source/Runtime/Core/Public/Windows/WindowsPlatformMemory.h
+++ b/Engine/Source/Runtime/Core/Public/Windows/WindowsPlatformMemory.h
@@ -38,6 +38,7 @@ struct CORE_API FWindowsPlatformMemory
 		MCR_GPU, // memory directly a GPU (graphics card, etc)
 		MCR_GPUSystem, // system memory directly accessible by a GPU
 		MCR_TexturePool, // presized texture pools
+		MCR_GPUDefragPool, // presized pool of memory that can be defragmented.
 		MCR_SamplePlatformSpecifcMemoryRegion, 
 		MCR_MAX
 	};
diff --git a/Engine/Source/Runtime/Core/Public/Windows/WindowsPlatformMisc.h b/Engine/Source/Runtime/Core/Public/Windows/WindowsPlatformMisc.h
--- a/Engine/Source/Runtime/Core/Public/Windows/WindowsPlatformMisc.h
+++ b/Engine/Source/Runtime/Core/Public/Windows/WindowsPlatformMisc.h
@@ -196,6 +196,7 @@ struct CORE_API FWindowsPlatformMisc
 	static FString GetCPUVendor();
 	static FString GetCPUBrand();
 	static FString GetPrimaryGPUBrand();
+	static void GetGPUDriverInfo(const FString DeviceDescription, FString& InternalDriverVersion, FString& UserDriverVersion, FString& DriverDate);
 	static void GetOSVersions( FString& out_OSVersionLabel, FString& out_OSSubVersionLabel );
 	static bool GetDiskTotalAndFreeSpace( const FString& InPath, uint64& TotalNumberOfBytes, uint64& NumberOfFreeBytes );
 
diff --git a/Engine/Source/Runtime/CoreUObject/Private/Serialization/AsyncLoading.cpp b/Engine/Source/Runtime/CoreUObject/Private/Serialization/AsyncLoading.cpp
--- a/Engine/Source/Runtime/CoreUObject/Private/Serialization/AsyncLoading.cpp
+++ b/Engine/Source/Runtime/CoreUObject/Private/Serialization/AsyncLoading.cpp
@@ -1889,6 +1889,7 @@ EAsyncPackageState::Type FAsyncPackage::PostLoadObjects()
 		}
 		// All object must be finalized on the game thread
 		DeferredFinalizeObjects.Add(Object);
+		check(Object->IsValidLowLevelFast());
 		// Make sure all objects in DeferredFinalizeObjects are referenced too
 		FAsyncObjectsReferencer::Get().AddObject(Object);
 	}
diff --git a/Engine/Source/Runtime/CoreUObject/Private/Serialization/BulkData.cpp b/Engine/Source/Runtime/CoreUObject/Private/Serialization/BulkData.cpp
--- a/Engine/Source/Runtime/CoreUObject/Private/Serialization/BulkData.cpp
+++ b/Engine/Source/Runtime/CoreUObject/Private/Serialization/BulkData.cpp
@@ -825,6 +825,11 @@ void FUntypedBulkData::Serialize( FArchive& Ar, UObject* Owner, int32 Idx )
 						SerializeBulkData(Ar, BulkData.Get());
 					}
 				}
+				else if (BulkDataFlags & BULKDATA_PayloadInSeperateFile)
+				{
+					Filename = FPaths::ChangeExtension(Filename, TEXT(".ubulk"));
+					checkSlow(FPaths::FileExists(Filename));
+				}
 			}
 			// Serialize the bulk data right away.
 			else
@@ -849,15 +854,28 @@ void FUntypedBulkData::Serialize( FArchive& Ar, UObject* Owner, int32 Idx )
 					else
 					{
 						// if the payload is NOT stored inline ...
-
-						// store the current file offset
-						int64 CurOffset = Ar.Tell();
-						// seek to the location in the file where the payload is stored
-						Ar.Seek(BulkDataOffsetInFile);
-						// serialize the payload
-						SerializeBulkData( Ar, BulkData.Get() );
-						// seek to the location we came from
-						Ar.Seek(CurOffset);
+						if (BulkDataFlags & BULKDATA_PayloadInSeperateFile)
+						{
+							// open seperate bulk data file
+							FArchive* TargetArchive = IFileManager::Get().CreateFileReader(*Filename);
+							// seek to the location in the file where the payload is stored
+							TargetArchive->Seek(BulkDataOffsetInFile);
+							// serialize the payload
+							SerializeBulkData(*TargetArchive, BulkData.Get());
+							// cleanup file
+							delete TargetArchive;
+						}
+						else
+						{
+							// store the current file offset
+							int64 CurOffset = Ar.Tell();
+							// seek to the location in the file where the payload is stored
+							Ar.Seek(BulkDataOffsetInFile);
+							// serialize the payload
+							SerializeBulkData(Ar, BulkData.Get());
+							// seek to the location we came from
+							Ar.Seek(CurOffset);
+						}
 					}
   				}
 			}
@@ -925,6 +943,7 @@ void FUntypedBulkData::Serialize( FArchive& Ar, UObject* Owner, int32 Idx )
 
 				BulkStore.BulkDataOffsetInFilePos = SavedBulkDataOffsetInFilePos;
 				BulkStore.BulkDataSizeOnDiskPos = SavedBulkDataSizeOnDiskPos;
+				BulkStore.BulkDataFlagsPos = SavedBulkDataFlagsPos;
 				BulkStore.BulkDataFlags = BulkDataFlags;
 				BulkStore.BulkData = this;
 				
@@ -1282,7 +1301,7 @@ void FUntypedBulkData::LoadDataIntoMemory( void* Dest )
 	AttachedAr->Seek( PushedPos );
 #else
 	bool bWasLoadedSuccessfully = false;
-	if ((IsInGameThread() || IsInAsyncLoadingThread()) && Package.IsValid() && Package->LinkerLoad && Package->LinkerLoad->GetOwnerThreadId() == FPlatformTLS::GetCurrentThreadId())
+	if ((IsInGameThread() || IsInAsyncLoadingThread()) && Package.IsValid() && Package->LinkerLoad && Package->LinkerLoad->GetOwnerThreadId() == FPlatformTLS::GetCurrentThreadId() && ((BulkDataFlags & BULKDATA_PayloadInSeperateFile) == 0))
 	{
 		FLinkerLoad* LinkerLoad = Package->LinkerLoad;
 		if (LinkerLoad && LinkerLoad->Loader && !LinkerLoad->IsCompressed())
diff --git a/Engine/Source/Runtime/CoreUObject/Private/UObject/LazyObjectPtr.cpp b/Engine/Source/Runtime/CoreUObject/Private/UObject/LazyObjectPtr.cpp
--- a/Engine/Source/Runtime/CoreUObject/Private/UObject/LazyObjectPtr.cpp
+++ b/Engine/Source/Runtime/CoreUObject/Private/UObject/LazyObjectPtr.cpp
@@ -131,7 +131,8 @@ void FLazyObjectPtr::PossiblySerializeObjectGuid(UObject *Object, FArchive& Ar)
 				{
 					if (OtherObject != NULL)
 					{
-						UE_CLOG(!((FApp::IsGame() || GIsPlayInEditorWorld) && Package && Package->ContainsMap()), LogUObjectGlobals, Warning, TEXT("Guid is in use by %s and %s, which should never happen in the editor but could happen at runtime with duplicate level loading or PIE"), *OtherObject->GetFullName(), !!Object ? *Object->GetFullName() : TEXT("NULL"));
+						// IsGame returns true for GIsPlayInEditorWorld
+						UE_CLOG(!((FApp::IsGame() || IsRunningCommandlet()) && Package && Package->ContainsMap()), LogUObjectGlobals, Warning, TEXT("Guid is in use by %s and %s, which should never happen in the editor but could happen at runtime with duplicate level loading or PIE"), *OtherObject->GetFullName(), !!Object ? *Object->GetFullName() : TEXT("NULL"));
 						// This guid is in use, which should never happen in the editor but could happen at runtime with duplicate level loading or PIE. If so give it a new GUID to avoid crashing
 						Guid = FGuid::NewGuid();
 					}
diff --git a/Engine/Source/Runtime/CoreUObject/Private/UObject/Obj.cpp b/Engine/Source/Runtime/CoreUObject/Private/UObject/Obj.cpp
--- a/Engine/Source/Runtime/CoreUObject/Private/UObject/Obj.cpp
+++ b/Engine/Source/Runtime/CoreUObject/Private/UObject/Obj.cpp
@@ -1615,6 +1615,9 @@ void UObject::LoadConfig( UClass* ConfigClass/*=NULL*/, const TCHAR* InFilename/
 		UArrayProperty* Array = dynamic_cast<UArrayProperty*>( Property );
 		if( Array == NULL )
 		{
+			// todo: UE-22989 - This special case handling for FText properties can be removed when UE-22989 is fixed so that Import/ExportText on an FText property doesn't lose the namespace/key data
+			UTextProperty* TextProp = dynamic_cast<UTextProperty*>( Property );
+
 			for( int32 i=0; i<Property->ArrayDim; i++ )
 			{
 				if( Property->ArrayDim!=1 )
@@ -1622,18 +1625,32 @@ void UObject::LoadConfig( UClass* ConfigClass/*=NULL*/, const TCHAR* InFilename/
 					Key = FString::Printf(TEXT("%s[%i]"), *Property->GetName(), i);
 				}
 
-				FString Value;
-				bool bFoundValue = GConfig->GetString( *ClassSection, *Key, Value, *PropFileName );
-				if (bFoundValue)
+				bool bFoundValue = false;
+				if (TextProp)
 				{
-					if (Property->ImportText(*Value, Property->ContainerPtrToValuePtr<uint8>(this, i), PortFlags, this) == NULL)
+					FText Value;
+					bFoundValue = GConfig->GetText( *ClassSection, *Key, Value, *PropFileName );
+					if (bFoundValue)
 					{
-						// this should be an error as the properties from the .ini / .int file are not correctly being read in and probably are affecting things in subtle ways
-						UE_LOG(LogObj, Error, TEXT("LoadConfig (%s): import failed for %s in: %s"), *GetPathName(), *Property->GetName(), *Value);
+						TextProp->SetPropertyValue_InContainer(this, Value, i);
+					}
+				}
+				else
+				{
+					FString Value;
+					bFoundValue = GConfig->GetString( *ClassSection, *Key, Value, *PropFileName );
+					if (bFoundValue)
+					{
+						if (Property->ImportText(*Value, Property->ContainerPtrToValuePtr<uint8>(this, i), PortFlags, this) == NULL)
+						{
+							// this should be an error as the properties from the .ini / .int file are not correctly being read in and probably are affecting things in subtle ways
+							UE_LOG(LogObj, Error, TEXT("LoadConfig (%s): import failed for %s in: %s"), *GetPathName(), *Property->GetName(), *Value);
+						}
 					}
 				}
+
 #if !UE_BUILD_SHIPPING
-				else if (!FPlatformProperties::RequiresCookedData())
+				if (!bFoundValue && !FPlatformProperties::RequiresCookedData())
 				{
 					CheckMissingSection(ClassSection, PropFileName);
 				}
@@ -1661,11 +1678,30 @@ void UObject::LoadConfig( UClass* ConfigClass/*=NULL*/, const TCHAR* InFilename/
 				// Only override default properties if there is something to override them with.
 				if ( List.Num() > 0 )
 				{
+					// todo: UE-22989 - This special case handling for FText properties can be removed when UE-22989 is fixed so that Import/ExportText on an FText property doesn't lose the namespace/key data
+					UTextProperty* TextProp = dynamic_cast<UTextProperty*>( Array->Inner );
+
 					ArrayHelper.EmptyAndAddValues(List.Num());
 
-					for( int32 i=List.Num()-1,c=0; i>=0; i--,c++ )
+					if (TextProp)
 					{
-						Array->Inner->ImportText( *List[i], ArrayHelper.GetRawPtr(c), PortFlags, this );
+						for( int32 i=List.Num()-1,c=0; i>=0; i--,c++ )
+						{
+							FText Value;
+							if (!FParse::Text(*List[i], Value, *ClassSection))
+							{
+								Value = FText::FromString(*List[i]);
+							}
+
+							TextProp->SetPropertyValue(ArrayHelper.GetRawPtr(c), Value);
+						}
+					}
+					else
+					{
+						for( int32 i=List.Num()-1,c=0; i>=0; i--,c++ )
+						{
+							Array->Inner->ImportText( *List[i], ArrayHelper.GetRawPtr(c), PortFlags, this );
+						}
 					}
 				}
 				else
@@ -1818,12 +1854,78 @@ void UObject::SaveConfig( uint64 Flags, const TCHAR* InFilename, FConfigCacheIni
 					// Default ini's require the array syntax to be applied to the property name
 					FString CompleteKey = FString::Printf(TEXT("%s%s"), bIsADefaultIniWrite ? TEXT("+") : TEXT(""), *Key);
 
+					// todo: UE-22989 - This special case handling for FText properties can be removed when UE-22989 is fixed so that Import/ExportText on an FText property doesn't lose the namespace/key data
+					UTextProperty* TextProp = dynamic_cast<UTextProperty*>( Array->Inner );
+
 					FScriptArrayHelper_InContainer ArrayHelper(Array, this);
-					for( int32 i=0; i<ArrayHelper.Num(); i++ )
+
+					if (TextProp)
+					{
+						auto TextToConfigString = [](const FText& InText) -> FString
+						{
+							if (InText.IsTransient())
+							{
+								UE_LOG(LogConfig, Warning, TEXT( "UObject::SaveConfig() Transient FText"));
+								return FString(TEXT("Error: Transient FText"));
+							}
+
+							if (InText.IsCultureInvariant())
+							{
+								// todo: UE-22989 - This isn't really the correct way to handle culture invariant strings, but it mirrors the logic in FTextFriendHelper
+								return FString::Printf(
+									TEXT("NSLOCTEXT(\"\",\"\",\"%s\")" ), 
+									*InText.ToString().ReplaceCharWithEscapedChar().ReplaceQuotesWithEscapedQuotes()
+									);
+							}
+							else
+							{
+								const FString* SourceString = FTextInspector::GetSourceString(InText);
+								if (SourceString)
+								{
+									FTextDisplayStringRef DisplayString = FTextInspector::GetSharedDisplayString(InText);
+
+									FString LocNamespace;
+									FString LocKey;
+									const bool FoundNamespaceAndKey = FTextLocalizationManager::Get().FindNamespaceAndKeyFromDisplayString(DisplayString, LocNamespace, LocKey);
+
+									// If this has no namespace or key, attempt to give it a GUID for a key and register it.
+									if (!FoundNamespaceAndKey && GIsEditor)
+									{
+										LocKey = FGuid::NewGuid().ToString();
+										if (!FTextLocalizationManager::Get().AddDisplayString(DisplayString, LocNamespace, LocKey))
+										{
+											// Could not add display string, reset namespace and key.
+											LocNamespace.Empty();
+											LocKey.Empty();
+										}
+									}
+
+									return FString::Printf(
+										TEXT("NSLOCTEXT(\"%s\",\"%s\",\"%s\")" ), 
+										*LocNamespace, 
+										*LocKey, 
+										*InText.ToString().ReplaceCharWithEscapedChar().ReplaceQuotesWithEscapedQuotes()
+										);
+								}
+							}
+
+							return InText.ToString();
+						};
+
+						for( int32 i=0; i<ArrayHelper.Num(); i++ )
+						{
+							const FText* TextValuePtr = TextProp->GetPropertyValuePtr(ArrayHelper.GetRawPtr(i));
+							Sec->Add(*CompleteKey, TextToConfigString(*TextValuePtr));
+						}
+					}
+					else
 					{
-						FString	Buffer;
-						Array->Inner->ExportTextItem( Buffer, ArrayHelper.GetRawPtr(i), ArrayHelper.GetRawPtr(i), this, PortFlags );
-						Sec->Add(*CompleteKey, *Buffer);
+						for( int32 i=0; i<ArrayHelper.Num(); i++ )
+						{
+							FString	Buffer;
+							Array->Inner->ExportTextItem( Buffer, ArrayHelper.GetRawPtr(i), ArrayHelper.GetRawPtr(i), this, PortFlags );
+							Sec->Add(*CompleteKey, *Buffer);
+						}
 					}
 				}
 				else if( Property->Identical_InContainer(this, SuperClassDefaultObject) )
@@ -1850,9 +1952,20 @@ void UObject::SaveConfig( uint64 Flags, const TCHAR* InFilename, FConfigCacheIni
 							Key = TempKey;
 						}
 
-						FString	Value;
-						Property->ExportText_InContainer( Index, Value, this, this, this, PortFlags );
-						Config->SetString( *Section, *Key, *Value, *PropFileName );
+						// todo: UE-22989 - This special case handling for FText properties can be removed when UE-22989 is fixed so that Import/ExportText on an FText property doesn't lose the namespace/key data
+						UTextProperty* TextProp = dynamic_cast<UTextProperty*>( Property );
+
+						if (TextProp)
+						{
+							const FText* TextValuePtr = TextProp->GetPropertyValuePtr_InContainer(this, Index);
+							Config->SetText( *Section, *Key, *TextValuePtr, *PropFileName );
+						}
+						else
+						{
+							FString	Value;
+							Property->ExportText_InContainer( Index, Value, this, this, this, PortFlags );
+							Config->SetString( *Section, *Key, *Value, *PropFileName );
+						}
 					}
 					else if( Property->Identical_InContainer(this, SuperClassDefaultObject, Index) )
 					{
diff --git a/Engine/Source/Runtime/CoreUObject/Private/UObject/SavePackage.cpp b/Engine/Source/Runtime/CoreUObject/Private/UObject/SavePackage.cpp
--- a/Engine/Source/Runtime/CoreUObject/Private/UObject/SavePackage.cpp
+++ b/Engine/Source/Runtime/CoreUObject/Private/UObject/SavePackage.cpp
@@ -4523,9 +4523,35 @@ ESavePackageResult UPackage::Save(UPackage* InOuter, UObject* Base, EObjectFlags
 				int64 StartOfBulkDataArea = Linker->Tell();
 				Linker->Summary.BulkDataStartOffset = StartOfBulkDataArea;
 
+				if (Linker->BulkDataToAppend.Num() > 0)
 				{
 					FScopedSlowTask BulkDataFeedback(Linker->BulkDataToAppend.Num());
 
+					FArchive* TargetArchive = Linker;
+					FArchive* BulkArchive = nullptr;
+					uint32 ExtraBulkDataFlags = 0;
+
+					static struct FUseSeperateBulkDataFiles
+					{
+						bool bEnable;
+						FUseSeperateBulkDataFiles()
+						{
+							if (!GConfig->GetBool(TEXT("Core.System"), TEXT("UseSeperateBulkDataFiles"), bEnable, GEngineIni))
+							{
+								bEnable = false;
+							}
+						}
+					} ShouldUseSeperateBulkDataFiles;
+
+					const bool bShouldUseSeparateBulkFile = ShouldUseSeperateBulkDataFiles.bEnable && Linker->IsCooking();
+
+					if (bShouldUseSeparateBulkFile)
+					{
+						FString BulkFilename = FPaths::ChangeExtension(Filename, TEXT(".ubulk"));
+						TargetArchive = BulkArchive = IFileManager::Get().CreateFileWriter(*BulkFilename);
+						ExtraBulkDataFlags = BULKDATA_PayloadInSeperateFile;
+					}
+
 					for (int32 i=0; i < Linker->BulkDataToAppend.Num(); ++i)
 					{
 						BulkDataFeedback.EnterProgressFrame();
@@ -4534,23 +4560,29 @@ ESavePackageResult UPackage::Save(UPackage* InOuter, UObject* Base, EObjectFlags
 
 						// Set bulk data flags to what they were during initial serialization (they might have changed after that)
 						const uint32 OldBulkDataFlags = BulkDataStorageInfo.BulkData->GetBulkDataFlags();
-						BulkDataStorageInfo.BulkData->SetBulkDataFlags(BulkDataStorageInfo.BulkDataFlags);
+						uint32 ModifiedBulkDataFlags = BulkDataStorageInfo.BulkDataFlags | ExtraBulkDataFlags;
+						BulkDataStorageInfo.BulkData->SetBulkDataFlags(ModifiedBulkDataFlags);
 
-						int64 BulkStartOffset = Linker->Tell();
+						int64 BulkStartOffset = TargetArchive->Tell();
 						int64 StoredBulkStartOffset = BulkStartOffset - StartOfBulkDataArea;
 
-						BulkDataStorageInfo.BulkData->SerializeBulkData(*Linker, BulkDataStorageInfo.BulkData->Lock(LOCK_READ_ONLY));
+						BulkDataStorageInfo.BulkData->SerializeBulkData(*TargetArchive, BulkDataStorageInfo.BulkData->Lock(LOCK_READ_ONLY));
+
+						int64 BulkEndOffset = TargetArchive->Tell();
+						int64 LinkerEndOffset = Linker->Tell();
 
-						int64 BulkEndOffset = Linker->Tell();
 						int32 SizeOnDisk = (int32)(BulkEndOffset - BulkStartOffset);
 				
+						Linker->Seek(BulkDataStorageInfo.BulkDataFlagsPos);
+						*Linker << ModifiedBulkDataFlags;
+
 						Linker->Seek(BulkDataStorageInfo.BulkDataOffsetInFilePos);
 						*Linker << StoredBulkStartOffset;
 
 						Linker->Seek(BulkDataStorageInfo.BulkDataSizeOnDiskPos);
 						*Linker << SizeOnDisk;
 
-						Linker->Seek(BulkEndOffset);
+						Linker->Seek(LinkerEndOffset);
 
 						// Restore BulkData flags to before serialization started
 						BulkDataStorageInfo.BulkData->ClearBulkDataFlags(0xFFFFFFFF);
@@ -4558,6 +4590,12 @@ ESavePackageResult UPackage::Save(UPackage* InOuter, UObject* Base, EObjectFlags
 
 						BulkDataStorageInfo.BulkData->Unlock();
 					}
+
+					if (BulkArchive)
+					{
+						BulkArchive->Close();
+						delete BulkArchive;
+					}
 				}
 
 				Linker->BulkDataToAppend.Empty();
diff --git a/Engine/Source/Runtime/CoreUObject/Public/Serialization/BulkData.h b/Engine/Source/Runtime/CoreUObject/Public/Serialization/BulkData.h
--- a/Engine/Source/Runtime/CoreUObject/Public/Serialization/BulkData.h
+++ b/Engine/Source/Runtime/CoreUObject/Public/Serialization/BulkData.h
@@ -28,6 +28,8 @@ enum EBulkDataFlags
 	BULKDATA_SerializeCompressed				= (BULKDATA_SerializeCompressedZLIB),
 	/** Forces the payload to be always streamed, regardless of its size */
 	BULKDATA_ForceStreamPayload = 1 << 7,
+	/** If set, payload is stored in a .upack file alongside the uasset				*/
+	BULKDATA_PayloadInSeperateFile				= 1 << 8,
 
 };
 
diff --git a/Engine/Source/Runtime/CoreUObject/Public/UObject/Linker.h b/Engine/Source/Runtime/CoreUObject/Public/UObject/Linker.h
--- a/Engine/Source/Runtime/CoreUObject/Public/UObject/Linker.h
+++ b/Engine/Source/Runtime/CoreUObject/Public/UObject/Linker.h
@@ -2189,6 +2189,8 @@ public:
 		int64 BulkDataOffsetInFilePos;
 		/** Offset to the location where the payload size is stored */
 		int64 BulkDataSizeOnDiskPos;
+		/** Offset to the location where the bulk data flags are stored */
+		int64 BulkDataFlagsPos;
 		/** Bulk data flags at the time of serialization */
 		uint32 BulkDataFlags;
 		/** The bulkdata */
diff --git a/Engine/Source/Runtime/Engine/Classes/AI/Navigation/NavigationSystem.h b/Engine/Source/Runtime/Engine/Classes/AI/Navigation/NavigationSystem.h
--- a/Engine/Source/Runtime/Engine/Classes/AI/Navigation/NavigationSystem.h
+++ b/Engine/Source/Runtime/Engine/Classes/AI/Navigation/NavigationSystem.h
@@ -706,7 +706,7 @@ public:
 	/** Used to display "navigation building in progress" notify */
 	bool IsNavigationBuildInProgress(bool bCheckDirtyToo = true);
 
-	void OnNavigationGenerationFinished(ANavigationData& NavData);
+	virtual void OnNavigationGenerationFinished(ANavigationData& NavData);
 
 	/** Used to display "navigation building in progress" counter */
 	int32 GetNumRemainingBuildTasks() const;
diff --git a/Engine/Source/Runtime/Engine/Classes/Animation/AnimInstance.h b/Engine/Source/Runtime/Engine/Classes/Animation/AnimInstance.h
--- a/Engine/Source/Runtime/Engine/Classes/Animation/AnimInstance.h
+++ b/Engine/Source/Runtime/Engine/Classes/Animation/AnimInstance.h
@@ -711,6 +711,7 @@ public:
 	TArray<FMontageEvaluationState> MontageEvaluationData;
 
 	virtual void OnMontageInstanceStopped(FAnimMontageInstance & StoppedMontageInstance);
+	void ClearMontageInstanceReferences(FAnimMontageInstance& InMontageInstance);
 
 protected:
 	/** Map between Active Montages and their FAnimMontageInstance */
diff --git a/Engine/Source/Runtime/Engine/Classes/Animation/AnimSequenceBase.h b/Engine/Source/Runtime/Engine/Classes/Animation/AnimSequenceBase.h
--- a/Engine/Source/Runtime/Engine/Classes/Animation/AnimSequenceBase.h
+++ b/Engine/Source/Runtime/Engine/Classes/Animation/AnimSequenceBase.h
@@ -83,12 +83,6 @@ class UAnimSequenceBase : public UAnimationAsset
 	/** Evaluate curve data to Instance at the time of CurrentTime **/
 	ENGINE_API virtual void EvaluateCurveData(FBlendedCurve& OutCurve, float CurrentTime) const;
 
-	/**
-	 * return true if this is valid additive animation
-	 * false otherwise
-	 */
-	virtual bool IsValidAdditive() const { return false; }
-
 #if WITH_EDITOR
 	/** Return Number of Frames **/
 	virtual int32 GetNumberOfFrames() const;
diff --git a/Engine/Source/Runtime/Engine/Classes/Animation/AnimationAsset.h b/Engine/Source/Runtime/Engine/Classes/Animation/AnimationAsset.h
--- a/Engine/Source/Runtime/Engine/Classes/Animation/AnimationAsset.h
+++ b/Engine/Source/Runtime/Engine/Classes/Animation/AnimationAsset.h
@@ -751,6 +751,12 @@ public:
 	/** Return a list of unique marker names for blending compatibility */
 	ENGINE_API virtual TArray<FName>* GetUniqueMarkerNames() { return NULL; }
 
+	/**
+	* return true if this is valid additive animation
+	* false otherwise
+	*/
+	virtual bool IsValidAdditive() const { return false; }
+
 #if WITH_EDITORONLY_DATA
 	/** Information for thumbnail rendering */
 	UPROPERTY(VisibleAnywhere, Instanced, Category = Thumbnail)
diff --git a/Engine/Source/Runtime/Engine/Classes/Animation/BlendSpaceBase.h b/Engine/Source/Runtime/Engine/Classes/Animation/BlendSpaceBase.h
--- a/Engine/Source/Runtime/Engine/Classes/Animation/BlendSpaceBase.h
+++ b/Engine/Source/Runtime/Engine/Classes/Animation/BlendSpaceBase.h
@@ -274,7 +274,7 @@ public:
 	virtual void ReplaceReferredAnimations(const TMap<UAnimSequence*, UAnimSequence*>& ReplacementMap) override;
 	virtual int32 GetMarkerUpdateCounter() const;
 #endif
-	virtual TArray<FName>* GetUniqueMarkerNames() override { return (bAllSequencesHaveMatchingMarkers && SampleData.Num() > 0) ? SampleData[0].Animation->GetUniqueMarkerNames() : nullptr; }
+	virtual TArray<FName>* GetUniqueMarkerNames() override { return (SampleIndexWithMarkers != INDEX_NONE && SampleData.Num() > 0) ? SampleData[SampleIndexWithMarkers].Animation->GetUniqueMarkerNames() : nullptr; }
 	//~ End UAnimationAsset Interface
 
 	void TickFollowerSamples(TArray<FBlendSampleData> &SampleDataList, const int32 HighestWeightIndex, FAnimAssetTickContext &Context, bool bResetMarkerDataOnFollowers) const
@@ -289,7 +289,11 @@ public:
 				{
 					SampleDataItem.MarkerTickRecord.Reset();
 				}
-				Sample.Animation->TickByMarkerAsFollower(SampleDataItem.MarkerTickRecord, Context.MarkerTickContext, SampleDataItem.Time, SampleDataItem.PreviousTime, Context.GetLeaderDelta(), true);
+
+				if (Sample.Animation->AuthoredSyncMarkers.Num() > 0) // Update followers who can do marker sync, others will be handled later in TickAssetPlayer
+				{
+					Sample.Animation->TickByMarkerAsFollower(SampleDataItem.MarkerTickRecord, Context.MarkerTickContext, SampleDataItem.Time, SampleDataItem.PreviousTime, Context.GetLeaderDelta(), true);
+				}
 			}
 		}
 	}
@@ -432,8 +436,9 @@ protected:
 	 */
 	virtual void GetRawSamplesFromBlendInput(const FVector &BlendInput, TArray<FGridBlendSample, TInlineAllocator<4> > & OutBlendSamples) const {}
 
-	/** Track whether all our sequences have the same marker set for blending (if false use original scale based on length approach) */
-	bool bAllSequencesHaveMatchingMarkers;
+	/** Track index to get marker data from. Samples are tested for the suitability of marker based sync
+	    during load and if we can use marker based sync we cache an index to a representative sample here */
+	int32 SampleIndexWithMarkers;
 
 #if WITH_EDITOR
 private:
diff --git a/Engine/Source/Runtime/Engine/Classes/Camera/CameraModifier_CameraShake.h b/Engine/Source/Runtime/Engine/Classes/Camera/CameraModifier_CameraShake.h
--- a/Engine/Source/Runtime/Engine/Classes/Camera/CameraModifier_CameraShake.h
+++ b/Engine/Source/Runtime/Engine/Classes/Camera/CameraModifier_CameraShake.h
@@ -36,14 +36,21 @@ public:
 	/**
 	 * Stops and removes the camera shake of the given class from the camera.
 	 * @param Shake - the camera shake class to remove.
+	 * @param bImmediately		If true, shake stops right away regardless of blend out settings. If false, shake may blend out according to its settings.
 	 */
-	virtual void RemoveCameraShake(UCameraShake* ShakeInst);
+	virtual void RemoveCameraShake(UCameraShake* ShakeInst, bool bImmediately = true);
 
-	/** Stops and removes all camera shakes of the given class from the camera. */
-	virtual void RemoveAllCameraShakesOfClass(TSubclassOf<class UCameraShake> ShakeClass);
+	/**
+	 * Stops and removes all camera shakes of the given class from the camera. 
+	 * @param bImmediately		If true, shake stops right away regardless of blend out settings. If false, shake may blend out according to its settings.
+	 */
+	virtual void RemoveAllCameraShakesOfClass(TSubclassOf<class UCameraShake> ShakeClass, bool bImmediately = true);
 
-	/** Stops and removes all camera shakes from the camera. */
-	virtual void RemoveAllCameraShakes();
+	/** 
+	 * Stops and removes all camera shakes from the camera. 
+	 * @param bImmediately		If true, shake stops right away regardless of blend out settings. If false, shake may blend out according to its settings.
+	 */
+	virtual void RemoveAllCameraShakes(bool bImmediately = true);
 	
 	//~ Begin UCameraModifer Interface
 	virtual bool ModifyCamera(float DeltaTime, struct FMinimalViewInfo& InOutPOV) override;
diff --git a/Engine/Source/Runtime/Engine/Classes/Camera/CameraShake.h b/Engine/Source/Runtime/Engine/Classes/Camera/CameraShake.h
--- a/Engine/Source/Runtime/Engine/Classes/Camera/CameraShake.h
+++ b/Engine/Source/Runtime/Engine/Classes/Camera/CameraShake.h
@@ -239,16 +239,23 @@ public:
 	UFUNCTION(BlueprintNativeEvent, Category = CameraShake)
 	bool ReceiveIsFinished() const;
 
-	/** Called when the shake is explicitly stopped. */
+	/** 
+	 * Called when the shake is explicitly stopped. 
+	 * @param bImmediatly		If true, shake stops right away regardless of blend out settings. If false, shake may blend out according to its settings.
+	 */
 	UFUNCTION(BlueprintImplementableEvent, Category = CameraShake)
-	void ReceiveStopShake();
+	void ReceiveStopShake(bool bImmediately);
 
 	// Native API
-
 	virtual void UpdateAndApplyCameraShake(float DeltaTime, float Alpha, FMinimalViewInfo& InOutPOV);
 	virtual void PlayShake(class APlayerCameraManager* Camera, float Scale, ECameraAnimPlaySpace::Type InPlaySpace, FRotator UserPlaySpaceRot = FRotator::ZeroRotator);
 	virtual bool IsFinished() const;
-	virtual void StopShake();
+
+	/** 
+	 * Stops this shake from playing. 
+	 * @param bImmediatly		If true, shake stops right away regardless of blend out settings. If false, shake may blend out according to its settings.
+	 */
+	virtual void StopShake(bool bImmediately = true);
 
 	// Returns true if this camera shake will loop forever
 	bool IsLooping() const;
diff --git a/Engine/Source/Runtime/Engine/Classes/Camera/PlayerCameraManager.h b/Engine/Source/Runtime/Engine/Classes/Camera/PlayerCameraManager.h
--- a/Engine/Source/Runtime/Engine/Classes/Camera/PlayerCameraManager.h
+++ b/Engine/Source/Runtime/Engine/Classes/Camera/PlayerCameraManager.h
@@ -648,15 +648,15 @@ public:
 	
 	/** Immediately stops the given shake instance and invalidates it. */
 	UFUNCTION(BlueprintCallable, Category = "Camera Shakes")
-	virtual void StopCameraShake(class UCameraShake* ShakeInstance);
+	virtual void StopCameraShake(class UCameraShake* ShakeInstance, bool bImmediately = true);
 
 	/** Stops playing CameraShake of the given class. */
 	UFUNCTION(BlueprintCallable, Category = "Camera Shakes")
-	virtual void StopAllInstancesOfCameraShake(TSubclassOf<class UCameraShake> Shake);
+	virtual void StopAllInstancesOfCameraShake(TSubclassOf<class UCameraShake> Shake, bool bImmediately = true);
 
 	/** Stops all active camera shakes on this camera. */
 	UFUNCTION(BlueprintCallable, Category = "Camera Shakes")
-	virtual void StopAllCameraShakes();
+	virtual void StopAllCameraShakes(bool bImmediately = true);
 
 	//
 	//  CameraAnim fades.
diff --git a/Engine/Source/Runtime/Engine/Classes/Components/ActorComponent.h b/Engine/Source/Runtime/Engine/Classes/Components/ActorComponent.h
--- a/Engine/Source/Runtime/Engine/Classes/Components/ActorComponent.h
+++ b/Engine/Source/Runtime/Engine/Classes/Components/ActorComponent.h
@@ -529,10 +529,10 @@ public:
 	/** Marks the transform as dirty - will be sent to the render thread at the end of the frame*/
 	void MarkRenderTransformDirty();
 
-	/** If we belong to a world, mark this for a deffered update, otherwise do it now. */
+	/** If we belong to a world, mark this for a deferred update, otherwise do it now. */
 	void MarkForNeededEndOfFrameUpdate();
 
-	/** If we belong to a world, mark this for a deffered update, otherwise do it now. */
+	/** If we belong to a world, mark this for a deferred update, otherwise do it now. */
 	void MarkForNeededEndOfFrameRecreate();
 
 	/** return true if this component requires end of frame updates to happen from the game thread. */
diff --git a/Engine/Source/Runtime/Engine/Classes/Curves/CurveBase.h b/Engine/Source/Runtime/Engine/Classes/Curves/CurveBase.h
--- a/Engine/Source/Runtime/Engine/Classes/Curves/CurveBase.h
+++ b/Engine/Source/Runtime/Engine/Classes/Curves/CurveBase.h
@@ -509,6 +509,9 @@ public:
 	/** Gets a copy of the keys, so indices and handles can't be meddled with */
 	TArray<FRichCurveKey> GetCopyOfKeys() const;
 
+	/** Gets a const reference of the keys, so indices and handles can't be meddled with */
+	const TArray<FRichCurveKey>& GetConstRefOfKeys() const;
+
 	/** Const iterator for the keys, so the indices and handles stay valid */
 	TArray<FRichCurveKey>::TConstIterator GetKeyIterator() const;
 	
diff --git a/Engine/Source/Runtime/Engine/Classes/Engine/DemoNetDriver.h b/Engine/Source/Runtime/Engine/Classes/Engine/DemoNetDriver.h
--- a/Engine/Source/Runtime/Engine/Classes/Engine/DemoNetDriver.h
+++ b/Engine/Source/Runtime/Engine/Classes/Engine/DemoNetDriver.h
@@ -134,8 +134,8 @@ public:
 	*/
 	void GotoTimeInSeconds(const float TimeInSeconds, const FOnGotoTimeDelegate& InOnGotoTimeDelegate = FOnGotoTimeDelegate());
 
-	bool IsRecording();
-	bool IsPlaying();
+	bool IsRecording() const;
+	bool IsPlaying() const;
 
 public:
 
diff --git a/Engine/Source/Runtime/Engine/Classes/Engine/EndUserSettings.h b/Engine/Source/Runtime/Engine/Classes/Engine/EndUserSettings.h
--- a/Engine/Source/Runtime/Engine/Classes/Engine/EndUserSettings.h
+++ b/Engine/Source/Runtime/Engine/Classes/Engine/EndUserSettings.h
@@ -21,6 +21,14 @@ class UEndUserSettings : public UObject, public IImportantToggleSettingInterface
 	UPROPERTY(EditAnywhere, config, Category=Privacy)
 	bool bSendAnonymousUsageDataToEpic;
 
+	/** Determines whether the engine sends anonymous crash/abnormal-shutdown data about game sessions to Epic Games in order to improve Unreal Engine. Information will never be shared with 3rd parties. */
+	UPROPERTY(EditAnywhere, config, Category=PrivacyDetails)
+	bool bSendMeanTimeBetweenFailureDataToEpic;
+
+	/** If enabled, adds user identifying data to the otherwise anonymous reports sent to Epic Games. */
+	UPROPERTY(EditAnywhere, config, Category = PrivacyDetails)
+	bool bAllowUserIdInUsageData;
+
 public:
 	/** Sets bSendAnonymousUsageDataToEpic. It MUST be set this way when the end user changes the value to trigger the correct state change in the analytics system. */
 	void SetSendAnonymousUsageDataToEpic(bool bEnable);
diff --git a/Engine/Source/Runtime/Engine/Classes/Engine/Engine.h b/Engine/Source/Runtime/Engine/Classes/Engine/Engine.h
--- a/Engine/Source/Runtime/Engine/Classes/Engine/Engine.h
+++ b/Engine/Source/Runtime/Engine/Classes/Engine/Engine.h
@@ -1695,6 +1695,11 @@ public:
 	bool HandleDumpParticleMemCommand( const TCHAR* Cmd, FOutputDevice& Ar );
 #endif
 
+#if WITH_PROFILEGPU
+	bool HandleProfileGPUCommand( const TCHAR* Cmd, FOutputDevice& Ar );
+	bool HandleShowMaterialDrawEventsCommand( const TCHAR* Cmd, FOutputDevice& Ar );
+#endif
+
 	// Compile in Debug or Development
 #if !(UE_BUILD_SHIPPING || UE_BUILD_TEST) && WITH_HOT_RELOAD
 	bool HandleHotReloadCommand( const TCHAR* Cmd, FOutputDevice& Ar );
@@ -1703,7 +1708,6 @@ public:
 	// Compile in Debug, Development, and Test
 #if !UE_BUILD_SHIPPING
 	bool HandleDumpConsoleCommandsCommand( const TCHAR* Cmd, FOutputDevice& Ar, UWorld* InWorld );
-	bool HandleShowMaterialDrawEventsCommand( const TCHAR* Cmd, FOutputDevice& Ar );
 	bool HandleDumpAvailableResolutionsCommand( const TCHAR* Cmd, FOutputDevice& Ar );
 	bool HandleAnimSeqStatsCommand( const TCHAR* Cmd, FOutputDevice& Ar );
 	bool HandleCountDisabledParticleItemsCommand( const TCHAR* Cmd, FOutputDevice& Ar );
@@ -1718,7 +1722,6 @@ public:
 	bool HandleRecompileGlobalShadersCommand( const TCHAR* Cmd, FOutputDevice& Ar );
 	bool HandleDumpShaderStatsCommand( const TCHAR* Cmd, FOutputDevice& Ar );
 	bool HandleDumpMaterialStatsCommand( const TCHAR* Cmd, FOutputDevice& Ar );
-	bool HandleProfileGPUCommand( const TCHAR* Cmd, FOutputDevice& Ar );
 	bool HandleProfileCommand( const TCHAR* Cmd, FOutputDevice& Ar );
 	bool HandleProfileGPUHitchesCommand( const TCHAR* Cmd, FOutputDevice& Ar );
 	bool HandleShaderComplexityCommand( const TCHAR* Cmd, FOutputDevice& Ar );
@@ -2654,6 +2657,10 @@ public:
 	virtual bool AreEditorAnalyticsEnabled() const { return false; }
 	/** @return true if in-game analytics are enabled */
 	bool AreGameAnalyticsEnabled() const;
+	/** @return true if in-game analytics are sent with an anonymous GUID rather than real account and machine ids */
+	bool AreGameAnalyticsAnonymous() const;
+	/** @return true if in-game MTBF analytics are enabled */
+	bool AreGameMTBFEventsEnabled() const;
 	virtual void CreateStartupAnalyticsAttributes( TArray<struct FAnalyticsEventAttribute>& StartSessionAttributes ) const {}
 	
 	/** @return true if the engine is autosaving a package */
diff --git a/Engine/Source/Runtime/Engine/Classes/Engine/EngineTypes.h b/Engine/Source/Runtime/Engine/Classes/Engine/EngineTypes.h
--- a/Engine/Source/Runtime/Engine/Classes/Engine/EngineTypes.h
+++ b/Engine/Source/Runtime/Engine/Classes/Engine/EngineTypes.h
@@ -4052,3 +4052,24 @@ struct FMaterialMergeData
 		, EmissiveScale(0.0f)
 	{}
 };
+
+/**
+ * The description of a user activity
+ */
+USTRUCT()
+struct FUserActivity
+{
+	GENERATED_USTRUCT_BODY()
+
+	/** Describes the user's activity */
+	UPROPERTY(EditAnywhere, Category = "Activity")
+	FString ActionName;
+
+	/** Default constructor. */
+	FUserActivity() { }
+
+	/** Creates and initializes a new instance. */
+	FUserActivity(const FString& InActionName)
+		: ActionName(InActionName)
+	{ }
+};
diff --git a/Engine/Source/Runtime/Engine/Classes/Engine/LODActor.h b/Engine/Source/Runtime/Engine/Classes/Engine/LODActor.h
--- a/Engine/Source/Runtime/Engine/Classes/Engine/LODActor.h
+++ b/Engine/Source/Runtime/Engine/Classes/Engine/LODActor.h
@@ -27,7 +27,7 @@ public:
 	TArray<class AActor*> SubActors;
 
 	/** what distance do you want this to show up instead of SubActors */
-	UPROPERTY(Category=LODActor, EditAnywhere, meta=(ClampMin = "0", UIMin = "0"))
+	UPROPERTY()
 	float LODDrawDistance;
 	
 	/** what distance do you want this to show up instead of SubActors */
@@ -100,6 +100,9 @@ public:
 
 	/** Cleans the SubActor array (removes all NULL entries) */
 	void CleanSubActorArray();
+
+	/** Recalculates the drawing distance according to a fixed FOV of 90 and the transition screen size*/
+	void RecalculateDrawingDistance(const float TransitionScreenSize );
 #endif // WITH_EDITOR
 	
 	//~ Begin UObject Interface.
@@ -130,11 +133,29 @@ public:
 	UPROPERTY()
 	uint32 NumTrianglesInMergedMesh;
 
-	UPROPERTY(EditAnywhere, Category = MaterialMerge)
+	/** Flag whether or not to use the override MaterialSettings when creating the proxy mesh */
+	UPROPERTY(EditAnywhere, Category = HierarchicalLODSettings)
 	bool bOverrideMaterialMergeSettings;
 
-	UPROPERTY(EditAnywhere, Category = MaterialMerge, meta = (editcondition = "bOverrideMaterialMergeSettings"))
+	/** Override Material Settings, used when creating the proxy mesh */
+	UPROPERTY(EditAnywhere, Category = HierarchicalLODSettings, meta = (editcondition = "bOverrideMaterialMergeSettings"))
 	FMaterialProxySettings MaterialSettings;
+
+	/** Flag whether or not to use the override TransitionScreenSize for this proxy mesh */
+	UPROPERTY(EditAnywhere, Category = HierarchicalLODSettings)
+	bool bOverrideTransitionScreenSize;
+
+	/** Override transition screen size value, determines the screen size / draw distance at which the proxy mesh is visible */
+	UPROPERTY(EditAnywhere, Category = HierarchicalLODSettings, meta = (editcondition = "bOverrideTransitionScreenSize"))
+	float TransitionScreenSize;
+
+	/** Flag whether or not to use the override ScreenSize when creating the proxy mesh */
+	UPROPERTY(EditAnywhere, Category = HierarchicalLODSettings)
+	bool bOverrideScreenSize;
+
+	/** Override screen size value, used when creating the proxy mesh */
+	UPROPERTY(EditAnywhere, Category = HierarchicalLODSettings, meta = (editcondition = "bOverrideScreenSize"))
+	int32 ScreenSize;
 #endif // WITH_EDITORONLY_DATA
 
 	/** Returns StaticMeshComponent subobject **/
diff --git a/Engine/Source/Runtime/Engine/Classes/Engine/LevelStreaming.h b/Engine/Source/Runtime/Engine/Classes/Engine/LevelStreaming.h
--- a/Engine/Source/Runtime/Engine/Classes/Engine/LevelStreaming.h
+++ b/Engine/Source/Runtime/Engine/Classes/Engine/LevelStreaming.h
@@ -323,15 +323,25 @@ private:
 	/** @return Name of the LOD package on disk to load to the new package named PackageName, Name_None otherwise					*/
 	FName GetLODPackageNameToLoad() const;
 
+	enum EReqLevelBlock
+	{
+		/** Block load AlwaysLoaded levels. Otherwise Async load. */
+		BlockAlwaysLoadedLevelsOnly,	
+		/** Block all loads */
+		AlwaysBlock,
+		/** Never block loads */
+		NeverBlock,
+	};
+
 	/** 
 	 * Try to find loaded level in memory, issue a loading request otherwise
 	 *
 	 * @param	PersistentWorld			Persistent world
 	 * @param	bAllowLevelLoadRequests	Whether to allow level load requests
-	 * @param	bBlockOnLoad			Whether loading operation should block
+	 * @param	BlockPolicy				Whether loading operation should block
 	 * @return							true if the load request was issued or a package was already loaded
 	 */
-	bool RequestLevel(UWorld* PersistentWorld, bool bAllowLevelLoadRequests, bool bBlockOnLoad);
+	bool RequestLevel(UWorld* PersistentWorld, bool bAllowLevelLoadRequests, EReqLevelBlock BlockPolicy);
 	
 	/** Sets the value of LoadedLevel */
 	void SetLoadedLevel(class ULevel* Level);
diff --git a/Engine/Source/Runtime/Engine/Classes/Engine/UserInterfaceSettings.h b/Engine/Source/Runtime/Engine/Classes/Engine/UserInterfaceSettings.h
--- a/Engine/Source/Runtime/Engine/Classes/Engine/UserInterfaceSettings.h
+++ b/Engine/Source/Runtime/Engine/Classes/Engine/UserInterfaceSettings.h
@@ -60,6 +60,9 @@ public:
 	UPROPERTY(config, EditAnywhere, Category = "Cursors", meta = (MetaClass = "UserWidget", ToolTip = "Widget to use when the Crosshairs Cursor is requested."))
 	FStringClassReference CrosshairsCursor;
 
+	UPROPERTY(config, EditAnywhere, Category = "Cursors", meta = (MetaClass = "UserWidget", ToolTip = "Widget to use when the Hand Cursor is requested."))
+	FStringClassReference HandCursor;
+
 	UPROPERTY(config, EditAnywhere, Category = "Cursors", meta = (MetaClass = "UserWidget", ToolTip = "Widget to use when the GrabHand Cursor is requested."))
 	FStringClassReference GrabHandCursor;
 	
diff --git a/Engine/Source/Runtime/Engine/Classes/Engine/World.h b/Engine/Source/Runtime/Engine/Classes/Engine/World.h
--- a/Engine/Source/Runtime/Engine/Classes/Engine/World.h
+++ b/Engine/Source/Runtime/Engine/Classes/Engine/World.h
@@ -142,6 +142,8 @@ private:
 	bool bNeedCancelCleanUp;
 	/** The context we are running in. Can be used to get the FWorldContext from Engine*/
 	FName WorldContextHandle;
+	/** Real time which we started traveling at  */
+	double SeamlessTravelStartTime = 0.f;
 
 	/** copy data between the old world and the new world */
 	void CopyWorldData();
@@ -2414,6 +2416,9 @@ public:
 	/** Returns whether all the 'always loaded' levels are loaded. */
 	bool AreAlwaysLoadedLevelsLoaded() const;
 
+	/** Requests async loading of any 'always loaded' level. Used in seamless travel to prevent blocking in the first UpdateLevelStreaming.  */
+	void AsyncLoadAlwaysLoadedLevelsForSeamlessTravel();
+
 	/**
 	 * Returns whether the level streaming code is allowed to issue load requests.
 	 *
@@ -3149,7 +3154,7 @@ public:
 	 * @param bAbsolute whether we are using relative or absolute travel
 	 * @param bShouldSkipGameNotify whether to notify the clients/game or not
 	 */
-	virtual void ServerTravel(const FString& InURL, bool bAbsolute = false, bool bShouldSkipGameNotify = false);
+	virtual bool ServerTravel(const FString& InURL, bool bAbsolute = false, bool bShouldSkipGameNotify = false);
 
 	/** seamlessly travels to the given URL by first loading the entry level in the background,
 	 * switching to it, and then loading the specified level. Does not disrupt network communication or disconnect clients.
diff --git a/Engine/Source/Runtime/Engine/Classes/GameFramework/Actor.h b/Engine/Source/Runtime/Engine/Classes/GameFramework/Actor.h
--- a/Engine/Source/Runtime/Engine/Classes/GameFramework/Actor.h
+++ b/Engine/Source/Runtime/Engine/Classes/GameFramework/Actor.h
@@ -437,9 +437,15 @@ public:
 	UPROPERTY()
 	TArray< FName > Layers;
 
+private:
 	/** The Actor that owns the UChildActorComponent that owns this Actor. */
 	UPROPERTY()
-	TWeakObjectPtr<AActor> ParentComponentActor;	
+	TWeakObjectPtr<AActor> ParentComponentActor_DEPRECATED;	
+
+	/** The UChildActorComponent that owns this Actor. */
+	TWeakObjectPtr<UChildActorComponent> ParentComponent;	
+
+public:
 
 #if WITH_EDITORONLY_DATA
 
@@ -1450,6 +1456,16 @@ public:
 		return false;
 	}
 
+	/*~
+	 * Returns transform of the RootComponent 
+	 * this is a template for no other reason than to delay compilation until USceneComponent is defined
+	 */ 
+	template<class T>
+	static FORCEINLINE FTransform GetActorTransform(const T* RootComponent)
+	{
+		return (RootComponent != nullptr) ? RootComponent->GetComponentTransform() : FTransform();
+	}
+
 	/*~
 	 * Returns location of the RootComponent 
 	 * this is a template for no other reason than to delay compilation until USceneComponent is defined
@@ -1509,6 +1525,12 @@ public:
 	 */
 	bool SetRootComponent(class USceneComponent* NewRootComponent);
 
+	/** Returns the transform of the RootComponent of this Actor*/ 
+	FORCEINLINE FTransform GetActorTransform() const
+	{
+		return GetActorTransform(RootComponent);
+	}
+
 	/** Returns the location of the RootComponent of this Actor*/ 
 	FORCEINLINE FVector GetActorLocation() const
 	{
@@ -1614,9 +1636,10 @@ public:
 	virtual void SetIsTemporarilyHiddenInEditor( bool bIsHidden );
 
 	/**
+	 * @param  bIncludeParent - Whether to recurse up child actor hierarchy or not
 	 * @return Whether or not this actor is hidden in the editor for the duration of the current editor session
 	 */
-	bool IsTemporarilyHiddenInEditor() const { return bHiddenEdTemporary; }
+	bool IsTemporarilyHiddenInEditor(bool bIncludeParent = false) const;
 
 	/** @return	Returns true if this actor is allowed to be displayed, selected and manipulated by the editor. */
 	bool IsEditable() const;
@@ -2001,6 +2024,13 @@ public:
 	UFUNCTION(BlueprintAuthorityOnly, BlueprintCallable, Category="Networking")
 	void FlushNetDormancy();
 
+	/** Returns whether this Actor was spawned by a child actor component */
+	UFUNCTION(BlueprintCallable, Category="Actor")
+	bool IsChildActor() const;
+
+	UFUNCTION(BlueprintCallable, Category="Actor")
+	UChildActorComponent* GetParentComponent() const;
+
 	/** Ensure that all the components in the Components array are registered */
 	virtual void RegisterAllComponents();
 
@@ -2651,6 +2681,7 @@ private:
 	void InternalDispatchBlockingHit(UPrimitiveComponent* MyComp, UPrimitiveComponent* OtherComp, bool bSelfMoved, FHitResult const& Hit);
 
 	friend struct FMarkActorIsBeingDestroyed;
+	friend struct FActorParentComponentSetter;
 };
 
 struct FMarkActorIsBeingDestroyed
diff --git a/Engine/Source/Runtime/Engine/Classes/GameFramework/GameUserSettings.h b/Engine/Source/Runtime/Engine/Classes/GameFramework/GameUserSettings.h
--- a/Engine/Source/Runtime/Engine/Classes/GameFramework/GameUserSettings.h
+++ b/Engine/Source/Runtime/Engine/Classes/GameFramework/GameUserSettings.h
@@ -54,6 +54,10 @@ public:
 	UFUNCTION(BlueprintPure, Category=Settings)
 	FIntPoint GetLastConfirmedScreenResolution() const;
 
+	/** Returns user's desktop resolution, in pixels. */
+	UFUNCTION(BlueprintPure, Category = Settings)
+	FIntPoint GetDesktopResolution() const;
+
 	/** Sets the user setting for game screen resolution, in pixels. */
 	UFUNCTION(BlueprintCallable, Category=Settings)
 	void SetScreenResolution(FIntPoint Resolution);
@@ -303,6 +307,9 @@ protected:
 	UPROPERTY(config)
 	float FrameRateLimit;
 
+	/** Min resolution scale we allow in current display mode */
+	int32 MinResolutionScale;
+
 	/**
 	 * Check if the current version of the game user settings is valid. Sub-classes can override this to provide game-specific versioning as necessary.
 	 * @return True if the current version is valid, false if it is not
@@ -311,4 +318,8 @@ protected:
 
 	/** Update the version of the game user settings to the current version */
 	virtual void UpdateVersion();
+
+private:
+
+	void UpdateMinResolutionScaling();
 };
diff --git a/Engine/Source/Runtime/Engine/Classes/GameFramework/PlayerController.h b/Engine/Source/Runtime/Engine/Classes/GameFramework/PlayerController.h
--- a/Engine/Source/Runtime/Engine/Classes/GameFramework/PlayerController.h
+++ b/Engine/Source/Runtime/Engine/Classes/GameFramework/PlayerController.h
@@ -838,7 +838,7 @@ public:
 
 	/** Stop camera shake on client.  */
 	UFUNCTION(reliable, client, BlueprintCallable, Category="Game|Feedback")
-	void ClientStopCameraShake(TSubclassOf<class UCameraShake> Shake);
+	void ClientStopCameraShake(TSubclassOf<class UCameraShake> Shake, bool bImmediately = true);
 
 	/** 
 	 * Play a force feedback pattern on the player's controller
diff --git a/Engine/Source/Runtime/Engine/Classes/Kismet/KismetSystemLibrary.h b/Engine/Source/Runtime/Engine/Classes/Kismet/KismetSystemLibrary.h
--- a/Engine/Source/Runtime/Engine/Classes/Kismet/KismetSystemLibrary.h
+++ b/Engine/Source/Runtime/Engine/Classes/Kismet/KismetSystemLibrary.h
@@ -1362,6 +1362,13 @@ class ENGINE_API UKismetSystemLibrary : public UBlueprintFunctionLibrary
 	UFUNCTION(BlueprintCallable, Category="Rendering")
 	static bool GetSupportedFullscreenResolutions(TArray<FIntPoint>& Resolutions);
 
+	/**
+	 * Gets the smallest Y resolution we want to support in the UI, clamped within reasons
+	 * @return value in pixels
+	 */
+	UFUNCTION(BlueprintPure, Category="Rendering", meta=(UnsafeDuringActorConstruction = "true"))
+	static int32 GetMinYResolutionForUI();
+
 	// Opens the specified URL in the platform's web browser of choice
 	UFUNCTION(BlueprintCallable, Category = "Utilities|Platform")
 	static void LaunchURL(const FString& URL);
@@ -1500,4 +1507,10 @@ class ENGINE_API UKismetSystemLibrary : public UBlueprintFunctionLibrary
 	 */
 	UFUNCTION(BlueprintCallable, Category = "Utilities|Platform")
 	static void RegisterForRemoteNotifications();
+
+	/**
+	 * Tells the engine what the user is doing for debug, analytics, etc.
+	 */
+	UFUNCTION(BlueprintCallable, Category = "Utilities")
+	static void SetUserActivity(const FUserActivity& UserActivity);
 };
diff --git a/Engine/Source/Runtime/Engine/Classes/Particles/ParticleEmitter.h b/Engine/Source/Runtime/Engine/Classes/Particles/ParticleEmitter.h
--- a/Engine/Source/Runtime/Engine/Classes/Particles/ParticleEmitter.h
+++ b/Engine/Source/Runtime/Engine/Classes/Particles/ParticleEmitter.h
@@ -80,6 +80,9 @@ struct FParticleBurst
 	
 };
 
+DECLARE_STATS_GROUP(TEXT("Emitters"), STATGROUP_Emitters, STATCAT_Advanced);
+DECLARE_DWORD_COUNTER_STAT_EXTERN(TEXT("STAT_EmittersStatGroupTester"), STAT_EmittersStatGroupTester, STATGROUP_Emitters, ENGINE_API);
+
 UCLASS(hidecategories=Object, editinlinenew, abstract, MinimalAPI)
 class UParticleEmitter : public UObject
 {
@@ -172,6 +175,14 @@ class UParticleEmitter : public UObject
 	/** When true, if the current LOD is disabled the emitter will be kept alive. Otherwise, the emitter will be considered complete if the current LOD is disabled. */
 	UPROPERTY(EditAnywhere, Category = Particle)
 	uint32 bDisabledLODsKeepEmitterAlive : 1;
+	
+		/** When true, emitters deemed insignificant will have their tick and render disabled Instantly. When false they will simple stop spawning new particles. */
+	UPROPERTY(EditAnywhere, Category = Significance)
+	uint32 bDisableWhenInsignficant : 1;
+
+	/** The significance level required of this emitter's owner for this emitter to be active. */
+	UPROPERTY(EditAnywhere, Category = Significance)
+	EParticleSignificanceLevel SignificanceLevel;
 
 	//////////////////////////////////////////////////////////////////////////
 	// Below is information udpated by calling CacheEmitterModuleInfo
@@ -354,7 +365,35 @@ class UParticleEmitter : public UObject
 
 	/** Returns true if the is emitter has any enabled LODs, false otherwise. */
 	bool HasAnyEnabledLODs()const;
-};
 
+	/** Stat id of this object, 0 if nobody asked for it yet */
+	STAT(mutable TStatId StatID;)
+		
+	/**
+	* Returns the stat ID of the object...
+	* We can't use the normal version of this because those names are meaningless; we need the special name in the emitter
+	**/
+	FORCEINLINE TStatId GetStatID(bool bForDeferredUse = false) const
+	{
+#if STATS
+		// this is done to avoid even registering stats for a disabled group (unless we plan on using it later)
+		if (bForDeferredUse || FThreadStats::IsCollectingData(GET_STATID(STAT_EmittersStatGroupTester)))
+		{
+			if (!StatID.IsValidStat())
+			{
+				CreateStatID();
+			}
+			return StatID;
+		}
+#endif
+		return TStatId(); // not doing stats at the moment, or ever
+	}
+	/**
+	* Creates this stat ID for the emitter...and handle a null this pointer
+	**/
+#if STATS
+	void CreateStatID() const;
+#endif
+};
 
 
diff --git a/Engine/Source/Runtime/Engine/Classes/Particles/ParticleSystem.h b/Engine/Source/Runtime/Engine/Classes/Particles/ParticleSystem.h
--- a/Engine/Source/Runtime/Engine/Classes/Particles/ParticleSystem.h
+++ b/Engine/Source/Runtime/Engine/Classes/Particles/ParticleSystem.h
@@ -2,6 +2,8 @@
 
 
 #pragma once
+#include "ParticleHelper.h"
+
 #include "ParticleSystem.generated.h"
 
 /**
@@ -296,6 +298,13 @@ public:
 	UPROPERTY(EditAnywhere, Category = Performance, meta = (ToolTip = "Minimum duration between ticks; 33=tick at max. 30FPS, 16=60FPS, 8=120FPS"))
 	uint32 MinTimeBetweenTicks;
 
+	UPROPERTY(EditAnywhere, Category = Performance, meta = (ToolTip = "Time this system can be deemed insignificant before it deactivates. "))
+	float InsignificantTimeTillDeactivate;
+
+	/** The maximum level of significance for emitters in this system. Any emitters with a higher significance will be capped at this significance level. */
+	UPROPERTY(EditAnywhere, Category = Performance)
+	EParticleSignificanceLevel MaxSignificanceLevel;
+
 	/** Local space position that UVs generated with the ParticleMacroUV material node will be centered on. */
 	UPROPERTY(EditAnywhere, Category=MacroUV)
 	FVector MacroUVPosition;
diff --git a/Engine/Source/Runtime/Engine/Classes/Particles/ParticleSystemComponent.h b/Engine/Source/Runtime/Engine/Classes/Particles/ParticleSystemComponent.h
--- a/Engine/Source/Runtime/Engine/Classes/Particles/ParticleSystemComponent.h
+++ b/Engine/Source/Runtime/Engine/Classes/Particles/ParticleSystemComponent.h
@@ -6,6 +6,7 @@
 #include "Materials/MaterialInterface.h"
 #include "Particles/Emitter.h"
 #include "Particles/ParticleSystem.h"
+#include "Particles/ParticleEmitter.h"
 #include "ParticleSystemComponent.generated.h"
 
 //
@@ -73,6 +74,7 @@ enum EParticleEventType
 // Called when the particle system is done
 DECLARE_DYNAMIC_MULTICAST_DELEGATE_OneParam( FOnSystemFinished, class UParticleSystemComponent*, PSystem );
 
+DECLARE_DYNAMIC_MULTICAST_DELEGATE_OneParam(FOnSystemPreActivate, class UParticleSystemComponent*, PSC);
 
 /** Struct used for a particular named instance parameter for this ParticleSystemComponent. */
 USTRUCT(BlueprintType)
@@ -313,7 +315,20 @@ public:
 	 */
 	UPROPERTY(EditAnywhere, BlueprintReadOnly, Category=Attachment)
 	uint32 bAutoManageAttachment:1;
+	
+	/** The significance this component requires of it's emitters for them to be enabled. */
+	UPROPERTY()
+	EParticleSignificanceLevel RequiredSignificance;
 
+	/** True if this component has emitters that weren't created or inited due to significance. If set we need to re-init when significance increases. */
+	bool bMissingEmittersDueToSignificance;
+
+	/** The significance level we ran InitParticles() with. If we change to require less significance than this we need to re-run InitParticles() to be sure we have created all emitters. */
+	EParticleSignificanceLevel InitRequiredSignificanceLevel;
+	
+	/** Time in seconds since we were last considered significant. */
+	float LastSignificantTime;
+	
 private:
 	/** Did we auto attach during activation? Used to determine if we should restore the relative transform during detachment. */
 	uint32 bDidAutoAttach:1;
@@ -473,6 +488,9 @@ public:
 	UPROPERTY(BlueprintAssignable)
 	FOnSystemFinished OnSystemFinished;
 
+	UPROPERTY(BlueprintAssignable)
+	FOnSystemPreActivate OnSystemPreActivate;
+
 public:
 	/**
 	 * Component we automatically attach to when activated, if bAutoManageAttachment is true.
@@ -523,6 +541,8 @@ private:
 	bool bParallelRenderThreadUpdate;
 public:
 
+	void SetSignificance(float OldSignificance, float NewSignificance);
+
 	FORCEINLINE const FTransform& GetAsyncComponentToWorld()
 	{
 		if (!bParallelRenderThreadUpdate && !IsInGameThread())
@@ -964,7 +984,7 @@ public:
 	}
 
 	// If particles have not already been initialised (ie. EmitterInstances created) do it now.
-	virtual void InitParticles();
+	virtual void InitParticles(bool bReInitExistingEmitters=true);
 
 
 	// @todo document
diff --git a/Engine/Source/Runtime/Engine/Classes/Sound/SoundBase.h b/Engine/Source/Runtime/Engine/Classes/Sound/SoundBase.h
--- a/Engine/Source/Runtime/Engine/Classes/Sound/SoundBase.h
+++ b/Engine/Source/Runtime/Engine/Classes/Sound/SoundBase.h
@@ -6,6 +6,7 @@
  * The base class for a playable sound object 
  */
 
+#include "Audio.h"
 #include "Sound/SoundConcurrency.h"
 #include "SoundBase.generated.h"
 
@@ -109,6 +110,9 @@ public:
 	virtual float GetVolumeMultiplier();
 	virtual float GetPitchMultiplier();
 
+	/** Returns whether or not this sound is looping. */
+	const bool IsLooping() { return (GetDuration() >= INDEFINITELY_LOOPING_DURATION); }
+
 	/** 
 	 * Parses the Sound to generate the WaveInstances to play
 	 */
diff --git a/Engine/Source/Runtime/Engine/Private/Actor.cpp b/Engine/Source/Runtime/Engine/Private/Actor.cpp
--- a/Engine/Source/Runtime/Engine/Private/Actor.cpp
+++ b/Engine/Source/Runtime/Engine/Private/Actor.cpp
@@ -21,6 +21,7 @@
 #include "GameFramework/Pawn.h"
 #include "Components/PawnNoiseEmitterComponent.h"
 #include "Components/TimelineComponent.h"
+#include "Components/ChildActorComponent.h"
 #include "Camera/CameraComponent.h"
 #include "GameFramework/DamageType.h"
 #include "Kismet/GameplayStatics.h"
@@ -527,6 +528,19 @@ void AActor::PostLoad()
 		bExchangedRoles = false;
 	}
 
+	if (AActor* ParentActor = ParentComponentActor_DEPRECATED.Get())
+	{
+		TInlineComponentArray<UChildActorComponent*> ParentChildActorComponents(ParentActor);
+		for (UChildActorComponent* ChildActorComponent : ParentChildActorComponents)
+		{
+			if (ChildActorComponent->ChildActor == this)
+			{
+				ParentComponent = ChildActorComponent;
+				break;
+			}
+		}
+	}
+
 	if ( GIsEditor )
 	{
 #if WITH_EDITORONLY_DATA
@@ -1377,14 +1391,14 @@ void AActor::OnRep_AttachmentReplication()
 	{
 		if (RootComponent)
 		{
-			USceneComponent* ParentComponent = (AttachmentReplication.AttachComponent ? AttachmentReplication.AttachComponent : AttachmentReplication.AttachParent->GetRootComponent());
+			USceneComponent* AttachParentComponent = (AttachmentReplication.AttachComponent ? AttachmentReplication.AttachComponent : AttachmentReplication.AttachParent->GetRootComponent());
 
-			if (ParentComponent)
+			if (AttachParentComponent)
 			{
 				RootComponent->RelativeLocation = AttachmentReplication.LocationOffset;
 				RootComponent->RelativeRotation = AttachmentReplication.RotationOffset;
 				RootComponent->RelativeScale3D = AttachmentReplication.RelativeScale3D;
-				RootComponent->AttachTo(ParentComponent, AttachmentReplication.AttachSocket);
+				RootComponent->AttachTo(AttachParentComponent, AttachmentReplication.AttachSocket);
 			}
 		}
 	}
@@ -3677,8 +3691,37 @@ void AActor::DispatchPhysicsCollisionHit(const FRigidBodyCollisionInfo& MyInfo,
 	}
 }
 
-// COMPONENTS
+#if WITH_EDITOR
+bool AActor::IsTemporarilyHiddenInEditor(const bool bIncludeParent) const
+{
+	if (bHiddenEdTemporary)
+	{
+		return true;
+	}
+
+	if (bIncludeParent)
+	{
+		if (UChildActorComponent* ParentCAC = ParentComponent.Get())
+		{
+			return ParentCAC->GetOwner()->IsTemporarilyHiddenInEditor(true);
+		}
+	}
+
+	return false;
+}
+#endif
+
+bool AActor::IsChildActor() const
+{
+	return ParentComponent.IsValid();
+}
 
+UChildActorComponent* AActor::GetParentComponent() const
+{
+	return ParentComponent.Get();
+}
+
+// COMPONENTS
 
 void AActor::UnregisterAllComponents()
 {
@@ -3773,19 +3816,19 @@ bool AActor::IncrementalRegisterComponents(int32 NumComponentsToRegister)
 		if(!Component->IsRegistered() && Component->bAutoRegister && !Component->IsPendingKill())
 		{
 			// Ensure that all parent are registered first
-			USceneComponent* ParentComponent = GetUnregisteredParent(Component);
-			if (ParentComponent)
+			USceneComponent* UnregisteredParentComponent = GetUnregisteredParent(Component);
+			if (UnregisteredParentComponent)
 			{
 				bool bParentAlreadyHandled = false;
-				RegisteredParents.Add(ParentComponent, &bParentAlreadyHandled);
+				RegisteredParents.Add(UnregisteredParentComponent, &bParentAlreadyHandled);
 				if (bParentAlreadyHandled)
 				{
-					UE_LOG(LogActor, Error, TEXT("AActor::IncrementalRegisterComponents parent component '%s' cannot be registered in actor '%s'"), *GetPathNameSafe(ParentComponent), *GetPathName());
+					UE_LOG(LogActor, Error, TEXT("AActor::IncrementalRegisterComponents parent component '%s' cannot be registered in actor '%s'"), *GetPathNameSafe(UnregisteredParentComponent), *GetPathName());
 					break;
 				}
 
 				// Register parent first, then return to this component on a next iteration
-				Component = ParentComponent;
+				Component = UnregisteredParentComponent;
 				CompIdx--;
 				NumTotalRegisteredComponents--; // because we will try to register the parent again later...
 			}
@@ -3864,15 +3907,20 @@ void AActor::UpdateComponentTransforms()
 
 void AActor::MarkComponentsRenderStateDirty()
 {
-	TInlineComponentArray<UActorComponent*> Components;
-	GetComponents(Components);
+	TInlineComponentArray<UActorComponent*> Components(this);
 
-	for (int32 Idx = 0; Idx < Components.Num(); Idx++)
+	for (UActorComponent* ActorComp : Components)
 	{
-		UActorComponent* ActorComp = Components[Idx];
 		if (ActorComp->IsRegistered())
 		{
 			ActorComp->MarkRenderStateDirty();
+			if (UChildActorComponent* ChildActorComponent = Cast<UChildActorComponent>(ActorComp))
+			{
+				if (ChildActorComponent->ChildActor)
+				{
+					ChildActorComponent->ChildActor->MarkComponentsRenderStateDirty();
+				}
+			}
 		}
 	}
 }
@@ -3923,10 +3971,9 @@ void AActor::DrawDebugComponents(FColor const& BaseColor) const
 		DrawDebugCoordinateSystem(MyWorld, Loc, Rot, 10.f);
 
 		// draw line from me to my parent
-		USceneComponent const* const ParentComponent = Cast<USceneComponent>(Component->AttachParent);
-		if (ParentComponent)
+		if (Component->AttachParent)
 		{
-			DrawDebugLine(MyWorld, ParentComponent->GetComponentLocation(), Loc, BaseColor);
+			DrawDebugLine(MyWorld, Component->AttachParent->GetComponentLocation(), Loc, BaseColor);
 		}
 
 		// draw component name
diff --git a/Engine/Source/Runtime/Engine/Private/ActorConstruction.cpp b/Engine/Source/Runtime/Engine/Private/ActorConstruction.cpp
--- a/Engine/Source/Runtime/Engine/Private/ActorConstruction.cpp
+++ b/Engine/Source/Runtime/Engine/Private/ActorConstruction.cpp
@@ -6,6 +6,7 @@
 #include "ComponentInstanceDataCache.h"
 #include "Engine/LevelScriptActor.h"
 #include "Engine/CullDistanceVolume.h"
+#include "Components/ChildActorComponent.h"
 
 #if WITH_EDITOR
 #include "Editor.h"
@@ -194,8 +195,8 @@ void AActor::RerunConstructionScripts()
 		
 		FTransform OldTransform = FTransform::Identity;
 		FName  SocketName;
-		AActor* Parent = NULL;
-		USceneComponent* ParentComponent = NULL;
+		AActor* Parent = nullptr;
+		USceneComponent* AttachParentComponent = nullptr;
 
 		bool bUseRootComponentProperties = true;
 
@@ -223,7 +224,7 @@ void AActor::RerunConstructionScripts()
 				if (Parent)
 				{
 					USceneComponent* AttachParent = ActorTransactionAnnotation->RootComponentData.AttachedParentInfo.AttachParent.Get();
-					ParentComponent = (AttachParent ? AttachParent : FindObjectFast<USceneComponent>(Parent, ActorTransactionAnnotation->RootComponentData.AttachedParentInfo.AttachParentName));
+					AttachParentComponent = (AttachParent ? AttachParent : FindObjectFast<USceneComponent>(Parent, ActorTransactionAnnotation->RootComponentData.AttachedParentInfo.AttachParentName));
 					SocketName = ActorTransactionAnnotation->RootComponentData.AttachedParentInfo.SocketName;
 					DetachRootComponentFromParent();
 				}
@@ -259,7 +260,7 @@ void AActor::RerunConstructionScripts()
 			for (AActor* AttachedActor : AttachedActors)
 			{
 				// We don't need to detach child actors, that will be handled by component tear down
-				if (!AttachedActor->ParentComponentActor.IsValid())
+				if (!AttachedActor->ParentComponent.IsValid())
 				{
 					USceneComponent* EachRoot = AttachedActor->GetRootComponent();
 					// If the component we are attached to is about to go away...
@@ -279,7 +280,7 @@ void AActor::RerunConstructionScripts()
 				}
 				else
 				{
-					check(AttachedActor->ParentComponentActor == this);
+					check(AttachedActor->ParentComponent->GetOwner() == this);
 				}
 			}
 		}
@@ -296,7 +297,7 @@ void AActor::RerunConstructionScripts()
 					UE_LOG(LogActor, Warning, TEXT("RerunConstructionScripts: RootComponent (%s) attached to another component in this Actor (%s)."), *RootComponent->GetPathName(), *Parent->GetPathName());
 					Parent = NULL;
 				}
-				ParentComponent = RootComponent->AttachParent;
+				AttachParentComponent = RootComponent->AttachParent;
 				SocketName = RootComponent->AttachSocketName;
 				//detach it to remove any scaling 
 				RootComponent->DetachFromParent(true);
@@ -356,13 +357,13 @@ void AActor::RerunConstructionScripts()
 		if(Parent)
 		{
 			USceneComponent* ChildRoot = GetRootComponent();
-			if (ParentComponent == NULL)
+			if (AttachParentComponent == nullptr)
 			{
-				ParentComponent = Parent->GetRootComponent();
+				AttachParentComponent = Parent->GetRootComponent();
 			}
-			if (ChildRoot != NULL && ParentComponent != NULL)
+			if (ChildRoot != nullptr && AttachParentComponent != nullptr)
 			{
-				ChildRoot->AttachTo(ParentComponent, SocketName, EAttachLocation::KeepWorldPosition);
+				ChildRoot->AttachTo(AttachParentComponent, SocketName, EAttachLocation::KeepWorldPosition);
 			}
 		}
 
@@ -370,7 +371,7 @@ void AActor::RerunConstructionScripts()
 		for(FAttachedActorInfo& Info : AttachedActorInfos)
 		{
 			// If this actor is no longer attached to anything, reattach
-			if (!Info.AttachedActor->IsPendingKill() && Info.AttachedActor->GetAttachParentActor() == NULL)
+			if (!Info.AttachedActor->IsPendingKill() && Info.AttachedActor->GetAttachParentActor() == nullptr)
 			{
 				USceneComponent* ChildRoot = Info.AttachedActor->GetRootComponent();
 				if (ChildRoot && ChildRoot->AttachParent != RootComponent)
diff --git a/Engine/Source/Runtime/Engine/Private/Animation/AnimInstance.cpp b/Engine/Source/Runtime/Engine/Private/Animation/AnimInstance.cpp
--- a/Engine/Source/Runtime/Engine/Private/Animation/AnimInstance.cpp
+++ b/Engine/Source/Runtime/Engine/Private/Animation/AnimInstance.cpp
@@ -227,10 +227,16 @@ void UAnimInstance::UninitializeAnimation()
 
 	for(int32 Index = 0; Index < MontageInstances.Num(); ++Index)
 	{
-		delete MontageInstances[Index];
+		FAnimMontageInstance* MontageInstance = MontageInstances[Index];
+		if (ensure(MontageInstance != nullptr))
+		{
+			ClearMontageInstanceReferences(*MontageInstance);
+			delete MontageInstance;
+		}
 	}
 
 	MontageInstances.Empty();
+	ActiveMontagesMap.Empty();
 
 	USkeletalMeshComponent* SkelMeshComp = GetSkelMeshComponent();
 	if (SkelMeshComp)
@@ -1434,6 +1440,9 @@ void UAnimInstance::Montage_Advance(float DeltaSeconds)
 
 			if (!MontageInstance->IsValid())
 			{
+				// Make sure we've cleared our references before deleting memory
+				ClearMontageInstanceReferences(*MontageInstance);
+
 				delete MontageInstance;
 				MontageInstances.RemoveAt(InstanceIndex);
 				--InstanceIndex;
@@ -2245,27 +2254,45 @@ void UAnimInstance::StopAllMontagesByGroupName(FName InGroupName, const FAlphaBl
 
 void UAnimInstance::OnMontageInstanceStopped(FAnimMontageInstance& StoppedMontageInstance)
 {
-	UAnimMontage* MontageStopped = StoppedMontageInstance.Montage;
-	ensure(MontageStopped != NULL);
+	ClearMontageInstanceReferences(StoppedMontageInstance);
+}
 
-	// Remove instance for Active List.
-	FAnimMontageInstance** AnimInstancePtr = ActiveMontagesMap.Find(MontageStopped);
-	if (AnimInstancePtr && (*AnimInstancePtr == &StoppedMontageInstance))
+void UAnimInstance::ClearMontageInstanceReferences(FAnimMontageInstance& InMontageInstance)
+{
+	if (UAnimMontage* MontageStopped = InMontageInstance.Montage)
 	{
-		ActiveMontagesMap.Remove(MontageStopped);
+		// Remove instance for Active List.
+		FAnimMontageInstance** AnimInstancePtr = ActiveMontagesMap.Find(MontageStopped);
+		if (AnimInstancePtr && (*AnimInstancePtr == &InMontageInstance))
+		{
+			ActiveMontagesMap.Remove(MontageStopped);
+		}
+	}
+	else
+	{
+		// If Montage ref is nullptr, it's possible the instance got terminated already and that is fine.
+		// Make sure it's been removed from our ActiveMap though
+		if (ActiveMontagesMap.FindKey(&InMontageInstance) != nullptr)
+		{
+			UE_LOG(LogAnimation, Warning, TEXT("%s: null montage found in the montage instance array!!"), *GetName());
+		}
 	}
 
 	// Clear RootMotionMontageInstance
-	if (RootMotionMontageInstance == &StoppedMontageInstance)
+	if (RootMotionMontageInstance == &InMontageInstance)
 	{
-		RootMotionMontageInstance = NULL;
+		RootMotionMontageInstance = nullptr;
 	}
+
+	// Clear any active synchronization
+	InMontageInstance.MontageSync_StopFollowing();
+	InMontageInstance.MontageSync_StopLeading();
 }
 
 FAnimMontageInstance* UAnimInstance::GetActiveInstanceForMontage(UAnimMontage const& Montage) const
 {
 	FAnimMontageInstance* const* FoundInstancePtr = ActiveMontagesMap.Find(&Montage);
-	return FoundInstancePtr ? *FoundInstancePtr : NULL;
+	return FoundInstancePtr ? *FoundInstancePtr : nullptr;
 }
 
 FAnimMontageInstance* UAnimInstance::GetRootMotionMontageInstance() const
diff --git a/Engine/Source/Runtime/Engine/Private/Animation/AnimInstanceProxy.cpp b/Engine/Source/Runtime/Engine/Private/Animation/AnimInstanceProxy.cpp
--- a/Engine/Source/Runtime/Engine/Private/Animation/AnimInstanceProxy.cpp
+++ b/Engine/Source/Runtime/Engine/Private/Animation/AnimInstanceProxy.cpp
@@ -133,7 +133,7 @@ void FAnimInstanceProxy::Uninitialize(UAnimInstance* InAnimInstance)
 	MaterialParameterCurves.Reset();
 }
 
-void FAnimInstanceProxy::PreUpdate(const UAnimInstance* InAnimInstance, float DeltaSeconds)
+void FAnimInstanceProxy::PreUpdate(UAnimInstance* InAnimInstance, float DeltaSeconds)
 {
 	CurrentDeltaSeconds = DeltaSeconds;
 	RootMotionMode = InAnimInstance->RootMotionMode;
diff --git a/Engine/Source/Runtime/Engine/Private/Animation/AnimMontage.cpp b/Engine/Source/Runtime/Engine/Private/Animation/AnimMontage.cpp
--- a/Engine/Source/Runtime/Engine/Private/Animation/AnimMontage.cpp
+++ b/Engine/Source/Runtime/Engine/Private/Animation/AnimMontage.cpp
@@ -1262,8 +1262,7 @@ void FAnimMontageInstance::Terminate()
 	}
 
 	UAnimMontage* OldMontage = Montage;
-	Montage = NULL;
-
+	
 	UAnimInstance* Inst = AnimInstance.Get();
 	if (Inst)
 	{
@@ -1277,16 +1276,17 @@ void FAnimMontageInstance::Terminate()
 
 		// terminating, trigger end
 		Inst->QueueMontageEndedEvent(FQueuedMontageEndedEvent(OldMontage, bInterrupted, OnMontageEnded));
-	}
 
-	// Clear any active synchronization
-	MontageSync_StopFollowing();
-	MontageSync_StopLeading();
+		// Clear references to this MontageInstance. Needs to happen before Montage is cleared to nullptr, as TMaps can use that as a key.
+		Inst->ClearMontageInstanceReferences(*this);
+	}
 
 	// clear Blend curve
 	Blend.SetCustomCurve(NULL);
 	Blend.SetBlendOption(EAlphaBlendOption::Linear);
 
+	Montage = nullptr;
+
 	UE_LOG(LogAnimMontage, Verbose, TEXT("Terminating: AnimMontage: %s"), *GetNameSafe(OldMontage));
 }
 
diff --git a/Engine/Source/Runtime/Engine/Private/Animation/AnimSequence.cpp b/Engine/Source/Runtime/Engine/Private/Animation/AnimSequence.cpp
--- a/Engine/Source/Runtime/Engine/Private/Animation/AnimSequence.cpp
+++ b/Engine/Source/Runtime/Engine/Private/Animation/AnimSequence.cpp
@@ -276,7 +276,6 @@ struct FScopedAnimSequenceCompressedDataCache
 		RawAnimationData = Src->RawAnimationData;
 		AnimationTrackNames = Src->AnimationTrackNames;
 		TrackToSkeletonMapTable = Src->TrackToSkeletonMapTable;
-		CompressedTrackOffsets = MoveTemp(Src->CompressedTrackOffsets);
 		TranslationData = MoveTemp(Src->TranslationData);
 		RotationData = MoveTemp(Src->RotationData);
 		ScaleData = MoveTemp(Src->ScaleData);
@@ -299,7 +298,6 @@ struct FScopedAnimSequenceCompressedDataCache
 		Src->RawAnimationData = MoveTemp(RawAnimationData);
 		Src->AnimationTrackNames = MoveTemp(AnimationTrackNames);
 		Src->TrackToSkeletonMapTable = MoveTemp(TrackToSkeletonMapTable);
-		Src->CompressedTrackOffsets = MoveTemp(CompressedTrackOffsets);
 		Src->TranslationData = MoveTemp(TranslationData);
 		Src->RotationData = MoveTemp(RotationData);
 		Src->ScaleData = MoveTemp(ScaleData);
@@ -4144,7 +4142,10 @@ void UAnimSequence::AdvanceMarkerPhaseAsLeader(bool bLooping, float MoveDelta, c
 	bool bOffsetInitialized = false;
 	float MarkerTimeOffset = 0.f;
 
-	check(CurrentTime >= 0.f && CurrentTime <= SequenceLength);
+	// Hard to reproduce issue triggering this, ensure & clamp for now
+	ensureMsgf(CurrentTime >= 0.f && CurrentTime <= SequenceLength, TEXT("Current time inside of AdvanceMarkerPhaseAsLeader is out of range %.3f of 0.0 to %.3f\n    Sequence: %s"), CurrentTime, SequenceLength, *GetFullName());
+
+	CurrentTime = FMath::Clamp(CurrentTime, 0.f, SequenceLength);
 
 	if (bPlayingForwards)
 	{
@@ -4609,10 +4610,13 @@ void UAnimSequence::GetMarkerIndicesForPosition(const FMarkerSyncAnimPosition& S
 		const FAnimSyncMarker& PrevMarker = AuthoredSyncMarkers[PrevMarkerIdx];
 		if (PrevMarker.MarkerName == SyncPosition.PreviousMarkerName)
 		{
-			int32 EndCount = bLooping ? AuthoredSyncMarkers.Num() : (AuthoredSyncMarkers.Num() - (PrevMarkerIdx + 1));
-			for (int32 NextMarkerCount = 0; NextMarkerCount < EndCount; ++NextMarkerCount)
+			const int32 EndMarkerSearchStart = PrevMarkerIdx + 1;
+
+			const int32 EndCount = bLooping ? AuthoredSyncMarkers.Num() + EndMarkerSearchStart : AuthoredSyncMarkers.Num();
+			for (int32 NextMarkerCount = EndMarkerSearchStart; NextMarkerCount < EndCount; ++NextMarkerCount)
 			{
-				int32 NextMarkerIdx = (PrevMarkerIdx + 1 + NextMarkerCount) % AuthoredSyncMarkers.Num();
+				const int32 NextMarkerIdx = NextMarkerCount % AuthoredSyncMarkers.Num();
+
 				if (AuthoredSyncMarkers[NextMarkerIdx].MarkerName == SyncPosition.NextMarkerName)
 				{
 					float NextMarkerTime = AuthoredSyncMarkers[NextMarkerIdx].Time;
@@ -4638,6 +4642,26 @@ void UAnimSequence::GetMarkerIndicesForPosition(const FMarkerSyncAnimPosition& S
 					break;
 				}
 			}
+
+			// If we get here and we haven't found a match and we are not looping then there 
+			// is no point running the rest of the loop set up something as relevant as we can and carry on
+			if (OutPrevMarker.MarkerIndex == MarkerIndexSpecialValues::Unitialized)
+			{
+				//Find nearest previous marker that is earlier than our current time
+				DiffToCurrentTime = OutCurrentTime - PrevMarker.Time;
+				int32 PrevMarkerToUse = PrevMarkerIdx + 1;
+				while (DiffToCurrentTime > 0.f && PrevMarkerToUse < AuthoredSyncMarkers.Num())
+				{
+					DiffToCurrentTime = OutCurrentTime - AuthoredSyncMarkers[PrevMarkerToUse].Time;
+					++PrevMarkerToUse;
+				}
+				OutPrevMarker.MarkerIndex = PrevMarkerToUse - 1;	// We always go one past the marker we actually want to use
+				
+				OutNextMarker.MarkerIndex = -1;						// This goes to minus one as the very fact we are here means
+																	// that there is no next marker to use
+				OutCurrentTime = GetCurrentTimeFromMarkers(OutPrevMarker, OutNextMarker, SyncPosition.PositionBetweenMarkers);
+				break; // no need to keep searching, we are done
+			}
 		}
 	}
 	// Should have found a markers above!
diff --git a/Engine/Source/Runtime/Engine/Private/Animation/AnimSequenceBase.cpp b/Engine/Source/Runtime/Engine/Private/Animation/AnimSequenceBase.cpp
--- a/Engine/Source/Runtime/Engine/Private/Animation/AnimSequenceBase.cpp
+++ b/Engine/Source/Runtime/Engine/Private/Animation/AnimSequenceBase.cpp
@@ -140,7 +140,8 @@ void UAnimSequenceBase::GetAnimNotifies(const float& StartTime, const float& Del
 		const ETypeAdvanceAnim AdvanceType = FAnimationRuntime::AdvanceTime(false, DesiredDeltaMove, CurrentPosition, SequenceLength);
 
 		// Verify position assumptions
-		checkf(bPlayingBackwards ? (CurrentPosition <= PreviousPosition) : (CurrentPosition >= PreviousPosition), TEXT("in Animsequence %s"), *GetName());
+		ensureMsgf(bPlayingBackwards ? (CurrentPosition <= PreviousPosition) : (CurrentPosition >= PreviousPosition), TEXT("in Animsequence %s : bPlayingBackwards(%d), PreviousPosition(%0.2f), Current Position(%0.2f)"), 
+			*GetName(), bPlayingBackwards, PreviousPosition, CurrentPosition);
 		
 		GetAnimNotifiesFromDeltaPositions(PreviousPosition, CurrentPosition, OutActiveNotifies);
 	
diff --git a/Engine/Source/Runtime/Engine/Private/Animation/AnimSingleNodeInstanceProxy.cpp b/Engine/Source/Runtime/Engine/Private/Animation/AnimSingleNodeInstanceProxy.cpp
--- a/Engine/Source/Runtime/Engine/Private/Animation/AnimSingleNodeInstanceProxy.cpp
+++ b/Engine/Source/Runtime/Engine/Private/Animation/AnimSingleNodeInstanceProxy.cpp
@@ -24,6 +24,13 @@ void FAnimSingleNodeInstanceProxy::Initialize(UAnimInstance* InAnimInstance)
 
 bool FAnimSingleNodeInstanceProxy::Evaluate(FPoseContext& Output)
 {
+	const bool bCanProcessAdditiveAnimationsLocal
+#if WITH_EDITOR
+		= bCanProcessAdditiveAnimations;
+#else
+		= false;
+#endif
+
 	if (CurrentAsset != NULL)
 	{
 		//@TODO: animrefactor: Seems like more code duplication than we need
@@ -37,7 +44,14 @@ bool FAnimSingleNodeInstanceProxy::Evaluate(FPoseContext& Output)
 			{
 				FAnimExtractContext ExtractionContext(CurrentTime, Sequence->bEnableRootMotion);
 
-				Sequence->GetAdditiveBasePose(Output.Pose, Output.Curve, ExtractionContext);
+				if (bCanProcessAdditiveAnimationsLocal)
+				{
+					Sequence->GetAdditiveBasePose(Output.Pose, Output.Curve, ExtractionContext);
+				}
+				else
+				{
+					Output.ResetToRefPose();
+				}
 
 				FCompactPose AdditivePose;
 				FBlendedCurve AdditiveCurve;
@@ -63,7 +77,7 @@ bool FAnimSingleNodeInstanceProxy::Evaluate(FPoseContext& Output)
 			if (AnimTrack.IsAdditive())
 			{
 #if WITH_EDITORONLY_DATA
-				if (Composite->PreviewBasePose)
+				if (bCanProcessAdditiveAnimationsLocal && Composite->PreviewBasePose)
 				{
 					Composite->PreviewBasePose->GetAdditiveBasePose(Output.Pose, Output.Curve, ExtractionContext);
 				}
@@ -105,7 +119,7 @@ bool FAnimSingleNodeInstanceProxy::Evaluate(FPoseContext& Output)
 				{
 #if WITH_EDITORONLY_DATA
 					// if montage is additive, we need to have base pose for the slot pose evaluate
-					if (Montage->PreviewBasePose && Montage->SequenceLength > 0.f)
+					if (bCanProcessAdditiveAnimationsLocal && Montage->PreviewBasePose && Montage->SequenceLength > 0.f)
 					{
 						Montage->PreviewBasePose->GetBonePose(SourcePose, SourceCurve, FAnimExtractContext(CurrentTime));
 					}
@@ -287,6 +301,15 @@ void FAnimSingleNodeInstanceProxy::SetAnimationAsset(class UAnimationAsset* NewA
 		CurrentAsset = NewAsset;
 	}
 
+	if (
+#if WITH_EDITOR
+		!bCanProcessAdditiveAnimations &&
+#endif
+		NewAsset && NewAsset->IsValidAdditive())
+	{
+		UE_LOG(LogAnimation, Warning, TEXT("Setting an additve animation (%s) on an AnimSingleNodeInstance is not allowed. This will not function correctly in cooked builds!"), *NewAsset->GetName());
+	}
+
 	if (MeshComponent)
 	{
 		if (MeshComponent->SkeletalMesh == NULL)
@@ -309,6 +332,7 @@ void FAnimSingleNodeInstanceProxy::SetAnimationAsset(class UAnimationAsset* NewA
 	PlayRate = InPlayRate;
 	CurrentTime = 0.f;
 	BlendSpaceInput = FVector::ZeroVector;
+	UpdateBlendspaceSamples(BlendSpaceInput);
 	MarkerTickRecord.Reset();
 
 #if WITH_EDITORONLY_DATA
diff --git a/Engine/Source/Runtime/Engine/Private/Animation/AnimationRuntime.cpp b/Engine/Source/Runtime/Engine/Private/Animation/AnimationRuntime.cpp
--- a/Engine/Source/Runtime/Engine/Private/Animation/AnimationRuntime.cpp
+++ b/Engine/Source/Runtime/Engine/Private/Animation/AnimationRuntime.cpp
@@ -124,6 +124,19 @@ FORCEINLINE void BlendCurves(const TFixedSizeArrayView<FBlendedCurve>& SourceCur
 	}
 }
 
+FORCEINLINE void BlendCurves(const TFixedSizeArrayView<FBlendedCurve>& SourceCurves, const TFixedSizeArrayView<float>& SourceWeights, const TFixedSizeArrayView<int32>& SourceWeightsIndices, FBlendedCurve& OutCurve)
+{
+	if (SourceCurves.Num() > 0)
+	{
+		OutCurve.Override(SourceCurves[0], SourceWeights[SourceWeightsIndices[0]]);
+
+		for (int32 CurveIndex = 1; CurveIndex<SourceCurves.Num(); ++CurveIndex)
+		{
+			OutCurve.Accumulate(SourceCurves[CurveIndex], SourceWeights[SourceWeightsIndices[CurveIndex]]);
+		}
+	}
+}
+
 FORCEINLINE void BlendCurves(const TFixedSizeArrayView<const FBlendedCurve*>& SourceCurves, const TFixedSizeArrayView<float>& SourceWeights, FBlendedCurve& OutCurve)
 {
 	if(SourceCurves.Num() > 0)
@@ -230,7 +243,7 @@ void FAnimationRuntime::BlendPosesTogether(const TFixedSizeArrayView<FCompactPos
 	// curve blending if exists
 	if (SourceCurves.Num() > 0)
 	{
-		BlendCurves(SourceCurves, SourceWeights, ResultCurve);
+		BlendCurves(SourceCurves, SourceWeights, SourceWeightsIndices, ResultCurve);
 	}
 }
 
diff --git a/Engine/Source/Runtime/Engine/Private/Animation/BlendSpaceBase.cpp b/Engine/Source/Runtime/Engine/Private/Animation/BlendSpaceBase.cpp
--- a/Engine/Source/Runtime/Engine/Private/Animation/BlendSpaceBase.cpp
+++ b/Engine/Source/Runtime/Engine/Private/Animation/BlendSpaceBase.cpp
@@ -85,7 +85,7 @@ struct FBlendSpaceScratchData : public TThreadSingleton<FBlendSpaceScratchData>
 UBlendSpaceBase::UBlendSpaceBase(const FObjectInitializer& ObjectInitializer)
 	: Super(ObjectInitializer)
 {
-	bAllSequencesHaveMatchingMarkers = false;
+	SampleIndexWithMarkers = INDEX_NONE;
 }
 
 void UBlendSpaceBase::PostLoad()
@@ -158,6 +158,24 @@ int32 GetHighestWeightSample(TArray<FBlendSampleData> &SampleDataList)
 	return HighestWeightIndex;
 }
 
+int32 GetHighestWeightMarkerSyncSample(const TArray<FBlendSampleData> &SampleDataList, const TArray<struct FBlendSample>& BlendSamples)
+{
+	int32 HighestWeightIndex = -1;
+	float HighestWeight = FLT_MIN;
+
+	for (int32 I = 0; I < SampleDataList.Num(); I++)
+	{
+		const FBlendSampleData& SampleData = SampleDataList[I];
+		if (SampleData.GetWeight() > HighestWeight &&
+			BlendSamples[SampleData.SampleDataIndex].Animation->AuthoredSyncMarkers.Num() > 0)
+		{
+			HighestWeightIndex = I;
+			HighestWeight = SampleData.GetWeight();
+		}
+	}
+	return HighestWeightIndex;
+}
+
 /////////////////////////////////////////////////////
 void UBlendSpaceBase::TickAssetPlayer(FAnimTickRecord& Instance, struct FAnimNotifyQueue& NotifyQueue, FAnimAssetTickContext& Context) const
 {
@@ -274,7 +292,7 @@ void UBlendSpaceBase::TickAssetPlayer(FAnimTickRecord& Instance, struct FAnimNot
 				SampleDataList.Append(NewSampleDataList);
 			}
 
-			bool bCanDoMarkerSync = bAllSequencesHaveMatchingMarkers && (Context.IsSingleAnimationContext() || (Instance.bCanUseMarkerSync && Context.CanUseMarkerPosition()));
+			bool bCanDoMarkerSync = (SampleIndexWithMarkers != INDEX_NONE) && (Context.IsSingleAnimationContext() || (Instance.bCanUseMarkerSync && Context.CanUseMarkerPosition()));
 
 			if (bCanDoMarkerSync)
 			{
@@ -321,10 +339,16 @@ void UBlendSpaceBase::TickAssetPlayer(FAnimTickRecord& Instance, struct FAnimNot
 				// advance current time - blend spaces hold normalized time as when dealing with changing anim length it would be possible to go backwards
 				UE_LOG(LogAnimation, Verbose, TEXT("BlendSpace(%s) - BlendInput(%s) : AnimLength(%0.5f) "), *GetName(), *BlendInput.ToString(), NewAnimLength);
 				
+				const int32 HighestMarkerSyncWeightIndex = bCanDoMarkerSync ? GetHighestWeightMarkerSyncSample(SampleDataList, SampleData) : -1;
+
+				if (HighestMarkerSyncWeightIndex == -1)
+				{
+					bCanDoMarkerSync = false;
+				}
+
 				if (bCanDoMarkerSync)
 				{
-					const int32 HighestWeightIndex = GetHighestWeightSample(SampleDataList);
-					FBlendSampleData& SampleDataItem = SampleDataList[HighestWeightIndex];
+					FBlendSampleData& SampleDataItem = SampleDataList[HighestMarkerSyncWeightIndex];
 					const FBlendSample& Sample = SampleData[SampleDataItem.SampleDataIndex];
 
 					bool bResetMarkerDataOnFollowers = false;
@@ -343,7 +367,8 @@ void UBlendSpaceBase::TickAssetPlayer(FAnimTickRecord& Instance, struct FAnimNot
 					{
 						Context.SetLeaderDelta(NewDeltaTime);
 						Sample.Animation->TickByMarkerAsLeader(SampleDataItem.MarkerTickRecord, Context.MarkerTickContext, SampleDataItem.Time, SampleDataItem.PreviousTime, NewDeltaTime, true);
-						TickFollowerSamples(SampleDataList, HighestWeightIndex, Context, bResetMarkerDataOnFollowers);
+						check(Context.MarkerTickContext.IsMarkerSyncStartValid());
+						TickFollowerSamples(SampleDataList, HighestMarkerSyncWeightIndex, Context, bResetMarkerDataOnFollowers);
 					}
 					NormalizedCurrentTime = SampleDataItem.Time / Sample.Animation->SequenceLength;
 					*Instance.MarkerTickRecord = SampleDataItem.MarkerTickRecord;
@@ -365,6 +390,7 @@ void UBlendSpaceBase::TickAssetPlayer(FAnimTickRecord& Instance, struct FAnimNot
 				{
 					bCanDoMarkerSync = false;
 				}
+
 				if (bCanDoMarkerSync)
 				{
 					const int32 HighestWeightIndex = GetHighestWeightSample(SampleDataList);
@@ -416,7 +442,7 @@ void UBlendSpaceBase::TickAssetPlayer(FAnimTickRecord& Instance, struct FAnimNot
 							float PrevSampleDataTime;
 							float& CurrentSampleDataTime = SampleEntry.Time;
 
-							if (!bCanDoMarkerSync) //Have already updated time if we are doing marker sync
+							if (!bCanDoMarkerSync || Sample.Animation->AuthoredSyncMarkers.Num() == 0) //Have already updated time if we are doing marker sync
 							{
 								const float SampleNormalizedPreviousTime = Sample.Animation->RateScale >= 0.f ? ClampedNormalizedPreviousTime : 1.f - ClampedNormalizedPreviousTime;
 								const float SampleNormalizedCurrentTime = Sample.Animation->RateScale >= 0.f ? ClampedNormalizedCurrentTime : 1.f - ClampedNormalizedCurrentTime;
@@ -859,10 +885,11 @@ void UBlendSpaceBase::ValidateSampleData()
 	bool bSampleDataChanged=false;
 	AnimLength = 0.f;
 
-	bool bNoAnimationsHaveMarkers = true;
 	bool bAllMarkerPatternsMatch = true;
 	FSyncPattern BlendSpacePattern;
 
+	int32 SampleWithMarkers = INDEX_NONE;
+
 	for (int32 I=0; I<SampleData.Num(); ++I)
 	{
 		FBlendSample& Sample = SampleData[I];
@@ -905,7 +932,11 @@ void UBlendSpaceBase::ValidateSampleData()
 
 		if (Sample.Animation->AuthoredSyncMarkers.Num() > 0)
 		{
-			bNoAnimationsHaveMarkers = false;
+			if (SampleWithMarkers == INDEX_NONE)
+			{
+				SampleWithMarkers = I;
+			}
+
 			if (BlendSpacePattern.MarkerNames.Num() == 0)
 			{
 				PopulateMarkerNameArray(BlendSpacePattern.MarkerNames, Sample.Animation->AuthoredSyncMarkers);
@@ -922,7 +953,7 @@ void UBlendSpaceBase::ValidateSampleData()
 		}
 	}
 
-	bAllSequencesHaveMatchingMarkers = bAllMarkerPatternsMatch && !bNoAnimationsHaveMarkers;
+	SampleIndexWithMarkers = bAllMarkerPatternsMatch ? SampleWithMarkers : INDEX_NONE;
 
 	if (bSampleDataChanged)
 	{
diff --git a/Engine/Source/Runtime/Engine/Private/AnimationAsset.cpp b/Engine/Source/Runtime/Engine/Private/AnimationAsset.cpp
--- a/Engine/Source/Runtime/Engine/Private/AnimationAsset.cpp
+++ b/Engine/Source/Runtime/Engine/Private/AnimationAsset.cpp
@@ -63,7 +63,7 @@ void FAnimGroupInstance::TestMontageTickRecordForLeadership()
 
 void FAnimGroupInstance::Finalize(const FAnimGroupInstance* PreviousGroup)
 {
-	if (!PreviousGroup || PreviousGroup->GroupLeaderIndex != GroupLeaderIndex || ValidMarkers != PreviousGroup->ValidMarkers
+	if (!PreviousGroup || PreviousGroup->GroupLeaderIndex != GroupLeaderIndex
 		|| (PreviousGroup->MontageLeaderWeight > 0.f && MontageLeaderWeight == 0.f/*if montage disappears, we should reset as well*/))
 	{
 		UE_LOG(LogAnimMarkerSync, Log, TEXT("Resetting Marker Sync Groups"));
@@ -87,8 +87,6 @@ void FAnimGroupInstance::Prepare(const FAnimGroupInstance* PreviousGroup)
 		ActivePlayers[0].bCanUseMarkerSync = true;
 		bCanUseMarkerSync = true;
 
-		int32 PlayerIndexToResetMarkers = INDEX_NONE;
-
 		//filter markers based on what exists in the other animations
 		for ( int32 ActivePlayerIndex = 0; ActivePlayerIndex < ActivePlayers.Num(); ++ActivePlayerIndex )
 		{
@@ -128,8 +126,6 @@ void FAnimGroupInstance::Prepare(const FAnimGroupInstance* PreviousGroup)
 						if ( !PlayerMarkerNames->Contains(MarkerName) )
 						{
 							ValidMarkers.RemoveAtSwap(ValidMarkerIndex, 1, false);
-
-							PlayerIndexToResetMarkers = ActivePlayerIndex;
 						}
 					}
 				}
@@ -140,12 +136,9 @@ void FAnimGroupInstance::Prepare(const FAnimGroupInstance* PreviousGroup)
 
 		ValidMarkers.Sort();
 
-		// if we have list of markers that needs to rest
-		if (PlayerIndexToResetMarkers > 0)
+		if (PreviousGroup && (ValidMarkers != PreviousGroup->ValidMarkers))
 		{
-			// if you removed valid markers, we also should reset previous candidate marker tick record
-			// because they might contain previous marker sets
-			for (int32 InternalActivePlayerIndex = 0; InternalActivePlayerIndex < PlayerIndexToResetMarkers; ++InternalActivePlayerIndex)
+			for (int32 InternalActivePlayerIndex = 0; InternalActivePlayerIndex < ActivePlayers.Num(); ++InternalActivePlayerIndex)
 			{
 				ActivePlayers[InternalActivePlayerIndex].MarkerTickRecord->Reset();
 			}
diff --git a/Engine/Source/Runtime/Engine/Private/AudioDecompress.cpp b/Engine/Source/Runtime/Engine/Private/AudioDecompress.cpp
--- a/Engine/Source/Runtime/Engine/Private/AudioDecompress.cpp
+++ b/Engine/Source/Runtime/Engine/Private/AudioDecompress.cpp
@@ -331,6 +331,8 @@ void FAsyncAudioDecompressWorker::DoWork()
 			}
 			else
 			{
+				check(Wave->DecompressionType == DTYPE_Native);
+
 				Wave->RawPCMDataSize = QualityInfo.SampleDataSize;
 				check(Wave->RawPCMData == nullptr);
 				Wave->RawPCMData = ( uint8* )FMemory::Malloc( Wave->RawPCMDataSize );
@@ -338,10 +340,11 @@ void FAsyncAudioDecompressWorker::DoWork()
 				// Decompress all the sample data into preallocated memory
 				AudioInfo->ExpandFile(Wave->RawPCMData, &QualityInfo);
 
-				const SIZE_T ResSize = Wave->GetResourceSize(EResourceSizeMode::Exclusive);
-				Wave->TrackedMemoryUsage += ResSize;
-				INC_DWORD_STAT_BY( STAT_AudioMemorySize, ResSize );
-				INC_DWORD_STAT_BY( STAT_AudioMemory, ResSize );
+				// Only track the RawPCMDataSize at this point since we haven't yet
+				// removed the compressed asset from memory and GetResourceSize() would add that at this point
+				Wave->TrackedMemoryUsage += Wave->RawPCMDataSize;
+				INC_DWORD_STAT_BY(STAT_AudioMemorySize, Wave->RawPCMDataSize);
+				INC_DWORD_STAT_BY(STAT_AudioMemory, Wave->RawPCMDataSize);
 			}
 		}
 		else if (Wave->DecompressionType == DTYPE_RealTime)
