Copying //UE4/Orion-Staging to //UE4/Main (originated from //Orion/Dev-General @ 2831630)

==========================
MAJOR FEATURES + CHANGES
==========================

Change 2831624 on 2016/01/17 by [REDACTED_PERSON]

	Merge disable of FCachedReadPlatformData on PS4.  Reduces memory spikes. 2830986
	#test none

Change 2831402 on 2016/01/17 by [REDACTED_PERSON]

	HLOD priority and streamout changes.
	Give texture pool an extra 200MB which we can afford thanks to [REDACTED_PERSON]/[REDACTED_PERSON]
	#test run agora, notice nice textures.

Change 2831398 on 2016/01/17 by [REDACTED_PERSON]

	Fix 3 logic bugs with Relocate
	#test run game, look for corruption.

Change 2831372 on 2016/01/16 by [REDACTED_PERSON]

	Update param.sfo's and lockdown version in prep for good PS4 playtest build.
	#test build from last night...

Change 2831274 on 2016/01/16 by [REDACTED_PERSON]

	Disable platform file cache wrapper on PS4

	#tests ran cooked ps4 build, timed loading (no real change), measured memory used for file handles (small)

Change 2831237 on 2016/01/16 by [REDACTED_PERSON]

	Fix PS4 compile error

	#tests none

Change 2831219 on 2016/01/16 by [REDACTED_PERSON]

	Fix possible invalid access to shared  movie player resource across threads causing startup crash.

Change 2831218 on 2016/01/16 by [REDACTED_PERSON]

	Fix bad warning case.
	#test none

Change 2831201 on 2016/01/16 by [REDACTED_PERSON]

	Added extra info about referencer to missing asset reference message
	#tests cooked, ran editor

Change 2831183 on 2016/01/16 by [REDACTED_PERSON]

	#OSS #PS4 #Purchasing #StoreV2
	- Force failure if we have no receipts after a "successful" checkout.
	- Report consumed entitlements as well as unconsumed but leave ValidationInfo empty so we can tell the difference at the application level
	- Convert productIds to skuIds at checkout time
	- Added PS4 Implementation of IOnlineStoreV2
	- Bugfix: set bSuccessfullyStartedUp=false when InitNPGameSettings() fails
	- Adjusted FOnlineStoreOffer to use FText::AsCurrencyBase
	#TESTS: login, purchase redemption, store MTX purchasing on PS4 & PC

Change 2831129 on 2016/01/16 by [REDACTED_PERSON]

	#MCP
	- Added a ctor to make converting from FOnlineError to FMcpQueryResult easier (for stuff that was already using FMcpQueryResult).
	#TESTS: frontend

Change 2830986 on 2016/01/15 by [REDACTED_PERSON]

	PS4: Disabling FCachedReadPlatformFile on PS4 to significantly reduce high watermark memory consumption during blocking loads
	#tests Ran Paragon PS4 down a bad path that currently does a blocking map and hero load

Change 2830943 on 2016/01/15 by [REDACTED_PERSON]

	Sequencer: Fix bug introduced with preroll. It was also causing a crash in particle track instance.

	#tests Master sequence trailer plays without crashing

Change 2830912 on 2016/01/15 by [REDACTED_PERSON]

	Rendering: Exposed GRHIDeviceId (only filled in on D3D11 and D3D12 RHI's under the same circumstances as GRHIAdapterName, etc..., 0 otherwise)
	#tests Tested printing the value out

Change 2830910 on 2016/01/15 by [REDACTED_PERSON]

	Rendering: Improved GPU driver detection logic to handle more cases
	#tests Tested on my machine which was previous reporting Unknown for the values as some entries contained the key in the Settings subfolder

Change 2830776 on 2016/01/15 by [REDACTED_PERSON]

	from Dev-Rendering
	added ensure to track down multiple issues like
	OR-11771 CRASH: User Crashed when pressing the Play button
	OR-12430 CRASH: OT2 user crashed with FRHIResource::AddRef()

diff --git a/Engine/Config/BaseEngine.ini b/Engine/Config/BaseEngine.ini
--- a/Engine/Config/BaseEngine.ini
+++ b/Engine/Config/BaseEngine.ini
@@ -183,7 +183,8 @@ InitialButtonRepeatDelay=0.2
 ButtonRepeatDelay=0.1
 NetDriverDefinitions=(DefName="GameNetDriver",DriverClassName="/Script/OnlineSubsystemUtils.IpNetDriver",DriverClassNameFallback="/Script/OnlineSubsystemUtils.IpNetDriver")
 +NetDriverDefinitions=(DefName="DemoNetDriver",DriverClassName="/Script/Engine.DemoNetDriver",DriverClassNameFallback="/Script/Engine.DemoNetDriver")
-
+SerializationOutOfBoundsErrorMessage=NSLOCTEXT("","SerializationOutOfBoundsErrorMessage","Corrupt data found, please verify your installation.")
+SerializationOutOfBoundsErrorMessageCaption=NSLOCTEXT("","SerializationOutOfBoundsErrorMessageCaption","Serialization Error : Action Needed")
 bSmoothFrameRate=true
 SmoothedFrameRateRange=(LowerBound=(Type="ERangeBoundTypes::Inclusive",Value=22),UpperBound=(Type="ERangeBoundTypes::Exclusive",Value=62))
 bCheckForMultiplePawnsSpawnedInAFrame=false
diff --git a/Engine/Plugins/Developer/PerforceSourceControl/Source/PerforceSourceControl/Private/PerforceSourceControlOperations.cpp b/Engine/Plugins/Developer/PerforceSourceControl/Source/PerforceSourceControl/Private/PerforceSourceControlOperations.cpp
--- a/Engine/Plugins/Developer/PerforceSourceControl/Source/PerforceSourceControl/Private/PerforceSourceControlOperations.cpp
+++ b/Engine/Plugins/Developer/PerforceSourceControl/Source/PerforceSourceControl/Private/PerforceSourceControlOperations.cpp
@@ -397,7 +397,7 @@ bool FPerforceMarkForAddWorker::Execute(FPerforceSourceControlCommand& InCommand
 		if(!IFileManager::Get().FileExists(*FileToAdd))
 		{
 			bHasMissingFiles = true;
-			InCommand.ErrorMessages.Add(FText::Format(LOCTEXT("Error_FailedToMarkFileForAdd_FileMissing", "Failed mark the file '%s' for add. The file doesn't exist on disk."), FText::FromString(FileToAdd)));
+			InCommand.ErrorMessages.Add(FText::Format(LOCTEXT("Error_FailedToMarkFileForAdd_FileMissing", "Failed mark the file '{0}' for add. The file doesn't exist on disk."), FText::FromString(FileToAdd)));
 		}
 	}
 	if(bHasMissingFiles)
diff --git a/Engine/Shaders/DitheredTransitionStencil.usf b/Engine/Shaders/DitheredTransitionStencil.usf
--- a/Engine/Shaders/DitheredTransitionStencil.usf
+++ b/Engine/Shaders/DitheredTransitionStencil.usf
@@ -9,9 +9,10 @@
 float DitheredTransitionFactor;
 
 void Main(
-	float4 SvPosition : SV_Position,
+	FScreenVertexOutput Interpolators,
 	out float OutColor : SV_Target0)
 {
+	float4 SvPosition = Interpolators.Position;
 	if (abs(DitheredTransitionFactor) > 0.001f)
 	{
 		float RandCos = cos(dot(floor(SvPosition.xy), float2(347.83451793f,3343.28371963f)));
diff --git a/Engine/Shaders/HeightFogVertexShader.usf b/Engine/Shaders/HeightFogVertexShader.usf
--- a/Engine/Shaders/HeightFogVertexShader.usf
+++ b/Engine/Shaders/HeightFogVertexShader.usf
@@ -13,10 +13,12 @@ void Main(
 	out float4 OutPosition : SV_POSITION
 	)
 {	
+	ResolvedView = ResolveView();
+
 	// screenspace position from vb
 	OutPosition = float4(InPosition,0,1);
 	// texture coord from vb
-	OutTexCoord = InPosition * View.ScreenPositionScaleBias.xy + View.ScreenPositionScaleBias.wz;
+	OutTexCoord = InPosition * ResolvedView.ScreenPositionScaleBias.xy + ResolvedView.ScreenPositionScaleBias.wz;
 
 	// deproject to world space
 	OutScreenVector = mul(float4(InPosition,1,0), ResolvedView.ScreenToTranslatedWorld).xyz;
diff --git a/Engine/Shaders/PostProcessBokehDOF.usf b/Engine/Shaders/PostProcessBokehDOF.usf
--- a/Engine/Shaders/PostProcessBokehDOF.usf
+++ b/Engine/Shaders/PostProcessBokehDOF.usf
@@ -442,8 +442,7 @@ void MainRecombinePS(
 
 #if RECOMBINE_METHOD == 2 || RECOMBINE_METHOD == 3
 	// sampling from pixel centers, so we need to offset if we're using lower resolution separate translucency
-	float Offset = 1.0-(1.0/SeparateTranslucencyResMult.xxxx);
-	float4 SeparateTranslucency = Texture2DSample(PostprocessInput2, PostprocessInput2Sampler, (PixelPosCenter+Offset) * PostprocessInput2Size.zw * SeparateTranslucencyResMult.xxxx );
+	float4 SeparateTranslucency = Texture2DSample(PostprocessInput2, PostprocessInput2Sampler, PixelPosCenter * PostprocessInput2Size.zw * SeparateTranslucencyResMult.xxxx );
 
 	// add RGB, darken by A (this allows to represent translucent and additive blending)
 	OutColor.rgb = OutColor.rgb * SeparateTranslucency.a + SeparateTranslucency.rgb;
diff --git a/Engine/Shaders/PostProcessCommon.usf b/Engine/Shaders/PostProcessCommon.usf
--- a/Engine/Shaders/PostProcessCommon.usf
+++ b/Engine/Shaders/PostProcessCommon.usf
@@ -126,22 +126,23 @@ half3 ColorCorrection(half3 InLDRColor)
 
 // Generate a mask to darken the screen borders.
 // (Camera artifact and artistic effect)
-// @param ScreenSpacePos -1..1
-float ComputeVignetteMask(float2 ScreenSpacePos, float Intensity)
+// @param VignetteCircleSpacePos from VignetteSpace()
+float ComputeVignetteMask(float2 VignetteCircleSpacePos, float Intensity)
 {
 	// Natural vignetting
 	// cosine-fourth law
-	ScreenSpacePos *= Intensity;
-	float Tan2Angle = dot( ScreenSpacePos, ScreenSpacePos );
+	VignetteCircleSpacePos *= Intensity;
+	float Tan2Angle = dot( VignetteCircleSpacePos, VignetteCircleSpacePos );
 	float Cos4Angle = Square( rcp( Tan2Angle + 1 ) );
 	return Cos4Angle;
 }
 
-// Scale {-1 to 1} space to vignette circle space.
+// Scale {-1 to 1} of the viewport space to vignette circle space.
 // Vignette space is scaled such that regardless of viewport aspect ratio, 
 // corners are at the same brightness on a circle.
 float2 VignetteSpace(float2 Pos)
 {
+	// could be optimized but this computation should be done in the vertex shader (3 or 4 vertices)
 	float Scale = sqrt(2.0) / sqrt(1.0 + (ViewportSize.y * ViewportSize.z) * (ViewportSize.y * ViewportSize.z));
 	return Pos * float2(1.0,  ViewportSize.y * ViewportSize.z) * Scale;
 }
diff --git a/Engine/Shaders/PostProcessTonemap.usf b/Engine/Shaders/PostProcessTonemap.usf
--- a/Engine/Shaders/PostProcessTonemap.usf
+++ b/Engine/Shaders/PostProcessTonemap.usf
@@ -77,9 +77,16 @@ float PseudoRandomTex(float2 xy)
 
 float SwitchVerticalAxis;
 
+float4 FringeUVParams;
 
+// @param ViewportPos in viewport -1..1 1..-1
+// @return UV in PostprocessInput
+float2 ViewportPosToUV(float2 ViewportPos)
+{
+	float2 ViewportUV = ViewportPos * float2(0.5f, -0.5f) + 0.5f;
 
-float4 FringeUVParams;
+	return (ViewportUV * Frame.ViewSizeAndInvSize.xy + Frame.ViewRectMin.xy) * Frame.BufferSizeAndInvSize.zw;
+}
 
 // vertex shader entry point
 void MainVS(
@@ -101,23 +108,21 @@ void MainVS(
 	// Eye adaptation is not yet supported
 	OutExposureScaleVignette.x = 1.0f;
 #endif
+	
+	// -1..1 1..-1 in the viewport
+	float2 ViewportPos = OutPosition.xy;
 
 	// Scale vignette to always be a circle with consistent corner intensity.
-	OutExposureScaleVignette.yz = VignetteSpace(OutPosition.xy);
+	float2 VignetteSpacePos = VignetteSpace(ViewportPos);
+	OutExposureScaleVignette.yz = VignetteSpacePos;
 
 	// Grain
 	OutGrainUV.xy = OutTexCoord + PostprocessInput0Size.zw * float2(-0.5,0.5);
 	OutGrainUV.zw = OutTexCoord + GrainRandomFull.xy;
 
 	// Fringe
-	// Center of the viewport in pixels and UV
-	float2 CenterPx = float2(0.5, 0.5) * (Frame.ViewSizeAndInvSize.xy) + Frame.ViewRectMin.xy;
-	float2 CenterUV = CenterPx *  PostprocessInput0Size.zw;
-
-	// Scale FringeUVs about the view port center for fringe
-	// NB: FringUVParams.{r,g} are less than 1
-	OutFringe.xy = FringeUVParams.r * (OutTexCoord - CenterUV) + CenterUV;
-	OutFringe.zw = FringeUVParams.g * (OutTexCoord - CenterUV) + CenterUV;
+	OutFringe.xy = ViewportPosToUV(ViewportPos * FringeUVParams.r);
+	OutFringe.zw = ViewportPosToUV(ViewportPos * FringeUVParams.g);
 }
 
 // Function graphing
diff --git a/Engine/Source/Developer/SynthBenchmark/Private/SynthBenchmarkPrivatePCH.cpp b/Engine/Source/Developer/SynthBenchmark/Private/SynthBenchmarkPrivatePCH.cpp
--- a/Engine/Source/Developer/SynthBenchmark/Private/SynthBenchmarkPrivatePCH.cpp
+++ b/Engine/Source/Developer/SynthBenchmark/Private/SynthBenchmarkPrivatePCH.cpp
@@ -22,6 +22,7 @@ class FSynthBenchmark : public ISynthBenchmark
 	
 	/** ISynthBenchmark implementation */
 	virtual void Run(FSynthBenchmarkResults& InOut, bool bGPUBenchmark, float WorkScale) const override;
+	virtual void GetRHIDisplay(FGPUAdpater& Out) const override;
 };
 
 IMPLEMENT_MODULE( FSynthBenchmark, SynthBenchmark )
@@ -223,4 +224,23 @@ void FSynthBenchmark::Run(FSynthBenchmarkResults& InOut, bool bGPUBenchmark, flo
 
 	UE_LOG(LogSynthBenchmark, Display, TEXT(""));
 	UE_LOG(LogSynthBenchmark, Display, TEXT("         ... Total Time: %f sec"),  (float)(FPlatformTime::Seconds() - StartTime));
-}
\ No newline at end of file
+}
+
+// could be moved to a more central place, from FWindowsPlatformSurvey::WriteFStringToResults
+static void WriteFStringToResults(TCHAR* OutBuffer, const FString& InString)
+{
+	FMemory::Memset( OutBuffer, 0, sizeof(TCHAR) * FHardwareSurveyResults::MaxStringLength );
+	TCHAR* Cursor = OutBuffer;
+	for (int32 i = 0; i < FMath::Min(InString.Len(), FHardwareSurveyResults::MaxStringLength - 1); i++)
+	{
+		*Cursor++ = InString[i];
+	}
+}
+
+void FSynthBenchmark::GetRHIDisplay(FGPUAdpater& Out) const
+{
+	WriteFStringToResults(Out.AdapterName, GRHIAdapterName);
+	WriteFStringToResults(Out.AdapterInternalDriverVersion, GRHIAdapterInternalDriverVersion);
+	WriteFStringToResults(Out.AdapterUserDriverVersion, GRHIAdapterUserDriverVersion);
+	WriteFStringToResults(Out.AdapterDriverDate, GRHIAdapterDriverDate);
+}
diff --git a/Engine/Source/Developer/SynthBenchmark/Public/SynthBenchmark.h b/Engine/Source/Developer/SynthBenchmark/Public/SynthBenchmark.h
--- a/Engine/Source/Developer/SynthBenchmark/Public/SynthBenchmark.h
+++ b/Engine/Source/Developer/SynthBenchmark/Public/SynthBenchmark.h
@@ -17,6 +17,9 @@ public:
 	// @param >0, WorkScale 10 for normal precision and runtime of less than a second
 	virtual void Run(FSynthBenchmarkResults& Out, bool bGPUBenchmark = true, float WorkScale = 10.0f) const = 0;
 
+	// could be moved out of SynthBenchmark, it only requires the RHI, only returns valid data after the RHI started up
+	virtual void GetRHIDisplay(FGPUAdpater& Out) const = 0;
+
 	/**
 	 * Singleton-like access to this module's interface.  This is just for convenience!
 	 * Beware of calling this during the shutdown phase, though.  Your module might have been unloaded already.
diff --git a/Engine/Source/Editor/MaterialEditor/Private/MaterialEditor.cpp b/Engine/Source/Editor/MaterialEditor/Private/MaterialEditor.cpp
--- a/Engine/Source/Editor/MaterialEditor/Private/MaterialEditor.cpp
+++ b/Engine/Source/Editor/MaterialEditor/Private/MaterialEditor.cpp
@@ -1392,18 +1392,26 @@ void FMaterialEditor::UpdateOriginalMaterial()
 		for (int32 ExpressionIndex = 0; ExpressionIndex < MaterialFunction->ParentFunction->FunctionExpressions.Num(); ExpressionIndex++)
 		{
 			UMaterialExpression* CurrentExpression = MaterialFunction->ParentFunction->FunctionExpressions[ExpressionIndex];
+			ensureMsgf(CurrentExpression, TEXT("Invalid expression at index [%i] whilst saving material function."), ExpressionIndex);
 
 			// Link the expressions back to their function
-			CurrentExpression->Material = NULL;
-			CurrentExpression->Function = MaterialFunction->ParentFunction;
+			if (CurrentExpression)
+			{
+				CurrentExpression->Material = NULL;
+				CurrentExpression->Function = MaterialFunction->ParentFunction;
+			}	
 		}
 		for (int32 ExpressionIndex = 0; ExpressionIndex < MaterialFunction->ParentFunction->FunctionEditorComments.Num(); ExpressionIndex++)
 		{
 			UMaterialExpressionComment* CurrentExpression = MaterialFunction->ParentFunction->FunctionEditorComments[ExpressionIndex];
+			ensureMsgf(CurrentExpression, TEXT("Invalid comment at index [%i] whilst saving material function."), ExpressionIndex);
 
 			// Link the expressions back to their function
-			CurrentExpression->Material = NULL;
-			CurrentExpression->Function = MaterialFunction->ParentFunction;
+			if (CurrentExpression)
+			{
+				CurrentExpression->Material = NULL;
+				CurrentExpression->Function = MaterialFunction->ParentFunction;
+			}
 		}
 
 		// mark the parent function as changed
diff --git a/Engine/Source/Editor/PropertyEditor/Private/SDetailsView.cpp b/Engine/Source/Editor/PropertyEditor/Private/SDetailsView.cpp
--- a/Engine/Source/Editor/PropertyEditor/Private/SDetailsView.cpp
+++ b/Engine/Source/Editor/PropertyEditor/Private/SDetailsView.cpp
@@ -402,9 +402,12 @@ bool SDetailsView::ShouldSetNewObjects( const TArray< TWeakObjectPtr< UObject >
 
 void SDetailsView::BeginSelectingObjects(const TArray< TWeakObjectPtr< UObject > >& InObjects)
 {
-	RegisterActiveTimer(0.f, FWidgetActiveTimerDelegate::CreateLambda([=] (double, float) {
-		return FinishSelectingObjects(InObjects);
-	}));
+	if( FSlateApplication::IsInitialized() )
+	{
+		RegisterActiveTimer(0.f, FWidgetActiveTimerDelegate::CreateLambda([=](double, float) {
+			return FinishSelectingObjects(InObjects);
+		}));
+	}
 }
 
 EActiveTimerReturnType SDetailsView::FinishSelectingObjects(const TArray< TWeakObjectPtr< UObject > > InObjects)
diff --git a/Engine/Source/Editor/UMGEditor/Private/WidgetBlueprintCompiler.cpp b/Engine/Source/Editor/UMGEditor/Private/WidgetBlueprintCompiler.cpp
--- a/Engine/Source/Editor/UMGEditor/Private/WidgetBlueprintCompiler.cpp
+++ b/Engine/Source/Editor/UMGEditor/Private/WidgetBlueprintCompiler.cpp
@@ -194,6 +194,7 @@ void FWidgetBlueprintCompiler::CreateClassVariablesFromBlueprint()
 	Super::CreateClassVariablesFromBlueprint();
 
 	UWidgetBlueprint* Blueprint = WidgetBlueprint();
+	UClass* ParentClass = Blueprint->ParentClass;
 
 	ValidateWidgetNames();
 
@@ -202,16 +203,7 @@ void FWidgetBlueprintCompiler::CreateClassVariablesFromBlueprint()
 	Blueprint->WidgetTree->GetAllWidgets(Widgets);
 
 	// Sort the widgets alphabetically
-	{
-		struct FWidgetSorter
-		{
-			bool operator()(const UWidget& A, const UWidget& B) const
-			{
-				return B.GetFName() < A.GetFName();
-			}
-		};
-		Widgets.Sort(FWidgetSorter());
-	}
+	Widgets.Sort( []( const UWidget& Lhs, const UWidget& Rhs ) { return Rhs.GetFName() < Lhs.GetFName(); } );
 
 	// Add widget variables
 	for ( UWidget* Widget : Widgets )
@@ -230,6 +222,13 @@ void FWidgetBlueprintCompiler::CreateClassVariablesFromBlueprint()
 			WidgetClass = BPWidgetClass->GetAuthoritativeClass();
 		}
 
+		UProperty* ExistingProperty = ParentClass->FindPropertyByName( Widget->GetFName() );
+		if ( ExistingProperty && ExistingProperty->HasMetaData( "BindWidget" ) )
+		{
+			WidgetToMemberVariableMap.Add( Widget, ExistingProperty );
+			continue;
+		}
+
 		FEdGraphPinType WidgetPinType(Schema->PC_Object, TEXT(""), WidgetClass, false, false);
 		
 		// Always name the variable according to the underlying FName of the widget object
@@ -309,6 +308,25 @@ void FWidgetBlueprintCompiler::FinishCompilingClass(UClass* Class)
 		});
 	}
 
+	UClass* ParentClass = Blueprint->ParentClass;
+	for ( TUObjectPropertyBase<UWidget*>* WidgetProperty : TFieldRange<TUObjectPropertyBase<UWidget*>>( ParentClass ) )
+	{
+		if ( WidgetProperty->HasMetaData( "BindWidget" ) && ( !WidgetProperty->HasMetaData( "OptionalWidget" ) || !WidgetProperty->GetBoolMetaData( "OptionalWidget" ) ) )
+		{
+			UWidget* const* Widget = WidgetToMemberVariableMap.FindKey( WidgetProperty );
+			if ( !Widget )
+			{
+				MessageLog.Error( *LOCTEXT( "RequiredWidget_NotBound", "Non-optional widget binding @@ not found." ).ToString(),
+				                    WidgetProperty );
+			}
+			else if ( !( *Widget )->IsA( WidgetProperty->PropertyClass ) )
+			{
+				MessageLog.Error( *LOCTEXT( "IncorrectWidgetTypes", "@@ is of type @@ property is of type @@." ).ToString(), *Widget,
+				                    ( *Widget )->GetClass(), WidgetProperty->PropertyClass );
+			}
+		}
+	}
+
 	Super::FinishCompilingClass(Class);
 }
 
diff --git a/Engine/Source/Editor/UMGEditor/Private/WidgetBlueprintEditorUtils.cpp b/Engine/Source/Editor/UMGEditor/Private/WidgetBlueprintEditorUtils.cpp
--- a/Engine/Source/Editor/UMGEditor/Private/WidgetBlueprintEditorUtils.cpp
+++ b/Engine/Source/Editor/UMGEditor/Private/WidgetBlueprintEditorUtils.cpp
@@ -129,6 +129,12 @@ bool FWidgetBlueprintEditorUtils::VerifyWidgetRename(TSharedRef<class FWidgetBlu
 		}
 	}
 
+	UProperty* Property = Blueprint->ParentClass->FindPropertyByName( NewNameSlug );
+	if ( Property && Property->HasMetaData( "BindWidget" ) )
+	{
+		return true;
+	}
+
 	FKismetNameValidator Validator(Blueprint);
 
 	// For variable comparison, use the slug
@@ -151,6 +157,9 @@ bool FWidgetBlueprintEditorUtils::RenameWidget(TSharedRef<FWidgetBlueprintEditor
 	UWidget* Widget = Blueprint->WidgetTree->FindWidget(OldObjectName);
 	check(Widget);
 
+	UClass* ParentClass = Blueprint->ParentClass;
+	check( ParentClass );
+
 	bool bRenamed = false;
 
 	TSharedPtr<INameValidatorInterface> NameValidator = MakeShareable(new FKismetNameValidator(Blueprint));
@@ -158,8 +167,11 @@ bool FWidgetBlueprintEditorUtils::RenameWidget(TSharedRef<FWidgetBlueprintEditor
 	// Get the new FName slug from the given display name
 	const FName NewFName = MakeObjectNameFromDisplayLabel(NewDisplayName, Widget->GetFName());
 
+	UProperty* ExistingProperty = ParentClass->FindPropertyByName( NewFName );
+	const bool bBindWidget = ExistingProperty && ExistingProperty->HasMetaData( "BindWidget" );
+
 	// NewName should be already validated. But one must make sure that NewTemplateName is also unique.
-	const bool bUniqueNameForTemplate = ( EValidatorResult::Ok == NameValidator->IsValid(NewFName) );
+	const bool bUniqueNameForTemplate = ( EValidatorResult::Ok == NameValidator->IsValid( NewFName ) || bBindWidget );
 	if ( Widget && bUniqueNameForTemplate )
 	{
 		// Stringify the FNames
diff --git a/Engine/Source/Editor/UnrealEd/Classes/Commandlets/CookCommandlet.h b/Engine/Source/Editor/UnrealEd/Classes/Commandlets/CookCommandlet.h
--- a/Engine/Source/Editor/UnrealEd/Classes/Commandlets/CookCommandlet.h
+++ b/Engine/Source/Editor/UnrealEd/Classes/Commandlets/CookCommandlet.h
@@ -40,6 +40,8 @@ class UCookCommandlet
 	bool bUseSerializationForGeneratingPackageDependencies;
 	/** Only cook packages specified on commandline options (for debugging)*/
 	bool bCookSinglePackage;
+	/** Should we output additional verbose cooking warnings */
+	bool bVerboseCookerWarnings;
 	/** All commandline tokens */
 	TArray<FString> Tokens;
 	/** All commandline switches */
diff --git a/Engine/Source/Editor/UnrealEd/Classes/Commandlets/GatherTextCommandletBase.h b/Engine/Source/Editor/UnrealEd/Classes/Commandlets/GatherTextCommandletBase.h
--- a/Engine/Source/Editor/UnrealEd/Classes/Commandlets/GatherTextCommandletBase.h
+++ b/Engine/Source/Editor/UnrealEd/Classes/Commandlets/GatherTextCommandletBase.h
@@ -119,10 +119,25 @@ public:
 	TMultiMap<FString, FLocItem> EntriesBySourceLocation;
 };
 
-typedef TMultiMap< FString, TSharedRef< FLocConflict > > TLocConflictContainer;
-
 struct FConflictReportInfo
 {
+	struct FLocConflictContainerKeyFuncs : BaseKeyFuncs<TSharedRef<FLocConflict>, FString, true>
+	{
+		static FORCEINLINE const FString& GetSetKey(const TPair<FString, TSharedRef<FLocConflict>>& Element)
+		{
+			return Element.Key;
+		}
+		static FORCEINLINE bool Matches(const FString& A, const FString& B)
+		{
+			return A.Equals(B, ESearchCase::CaseSensitive);
+		}
+		static FORCEINLINE uint32 GetKeyHash(const FString& Key)
+		{
+			return FCrc::StrCrc32<TCHAR>(*Key);
+		}
+	};
+	typedef TMultiMap<FString, TSharedRef<FLocConflict>, FDefaultSetAllocator, FLocConflictContainerKeyFuncs> TLocConflictContainer;
+
 	/**
 	 * Return the singleton instance of the report info.
 	 * @return The singleton instance.
diff --git a/Engine/Source/Editor/UnrealEd/Classes/CookOnTheSide/CookOnTheFlyServer.h b/Engine/Source/Editor/UnrealEd/Classes/CookOnTheSide/CookOnTheFlyServer.h
--- a/Engine/Source/Editor/UnrealEd/Classes/CookOnTheSide/CookOnTheFlyServer.h
+++ b/Engine/Source/Editor/UnrealEd/Classes/CookOnTheSide/CookOnTheFlyServer.h
@@ -25,6 +25,7 @@ enum class ECookInitializationFlags
 	UseSerializationForPackageDependencies =	0x100, // should we use the serialization code path for generating package dependencies (old method will be deprecated)
 	BuildDDCInBackground =						0x200, // build ddc content in background while the editor is running (only valid for modes which are in editor IsCookingInEditor())
 	GeneratedAssetRegistry =					0x400, // have we generated asset registry yet
+	OutputVerboseCookerWarnings =				0x800, // output additional cooker warnings about content issues
 };
 ENUM_CLASS_FLAGS(ECookInitializationFlags);
 
diff --git a/Engine/Source/Editor/UnrealEd/Private/Commandlets/CookCommandlet.cpp b/Engine/Source/Editor/UnrealEd/Private/Commandlets/CookCommandlet.cpp
--- a/Engine/Source/Editor/UnrealEd/Private/Commandlets/CookCommandlet.cpp
+++ b/Engine/Source/Editor/UnrealEd/Private/Commandlets/CookCommandlet.cpp
@@ -514,6 +514,7 @@ int32 UCookCommandlet::Main(const FString& CmdLineParams)
 	bErrorOnEngineContentUse = Switches.Contains(TEXT("ERRORONENGINECONTENTUSE"));
 	bUseSerializationForGeneratingPackageDependencies = Switches.Contains(TEXT("UseSerializationForGeneratingPackageDependencies"));
 	bCookSinglePackage = Switches.Contains(TEXT("cooksinglepackage"));
+	bVerboseCookerWarnings = Switches.Contains(TEXT("verbosecookerwarnings"));
 	if (bLeakTest)
 	{
 		for (FObjectIterator It; It; ++It)
@@ -1019,6 +1020,7 @@ bool UCookCommandlet::NewCook( const TArray<ITargetPlatform*>& Platforms, TArray
 	CookFlags |= bSkipEditorContent ? ECookInitializationFlags::SkipEditorContent : ECookInitializationFlags::None;	
 	CookFlags |= bUseSerializationForGeneratingPackageDependencies ? ECookInitializationFlags::UseSerializationForPackageDependencies : ECookInitializationFlags::None;
 	CookFlags |= bUnversioned ? ECookInitializationFlags::Unversioned : ECookInitializationFlags::None;
+	CookFlags |= bVerboseCookerWarnings ? ECookInitializationFlags::OutputVerboseCookerWarnings : ECookInitializationFlags::None;
 
 	TArray<UClass*> FullGCAssetClasses;
 	if (FullGCAssetClassNames.Num())
diff --git a/Engine/Source/Editor/UnrealEd/Private/Commandlets/DiffFilesCommandlet.cpp b/Engine/Source/Editor/UnrealEd/Private/Commandlets/DiffFilesCommandlet.cpp
--- a/Engine/Source/Editor/UnrealEd/Private/Commandlets/DiffFilesCommandlet.cpp
+++ b/Engine/Source/Editor/UnrealEd/Private/Commandlets/DiffFilesCommandlet.cpp
@@ -68,7 +68,13 @@ bool UDiffFilesCommandlet::Initialize(const TCHAR* Parms)
 			{
 				for (int32 i = 0; i < Paths.Num(); i++)
 				{
-					IFileManager::Get().FindFiles((TArray<FString>&)FilesInPath, *(Paths[i] / PackageWildcard), 1, 0);
+					TArray<FString> Temp;
+					IFileManager::Get().FindFiles(Temp, *(Paths[i] / PackageWildcard), 1, 0);
+					FString Path = FPaths::GetPath(Paths[i] / PackageWildcard);
+					for (const auto& T : Temp)
+					{
+						FilesInPath.Add(Path / T);
+					}
 				}
 			}
 		}
diff --git a/Engine/Source/Editor/UnrealEd/Private/Commandlets/GatherTextCommandletBase.cpp b/Engine/Source/Editor/UnrealEd/Private/Commandlets/GatherTextCommandletBase.cpp
--- a/Engine/Source/Editor/UnrealEd/Private/Commandlets/GatherTextCommandletBase.cpp
+++ b/Engine/Source/Editor/UnrealEd/Private/Commandlets/GatherTextCommandletBase.cpp
@@ -92,6 +92,7 @@ bool UGatherTextCommandletBase::WriteJSONToTextFile(TSharedPtr<FJsonObject> Outp
 
 		if (!DidFileExist)
 		{
+			// Checkout on a new file will cause it to be added
 			if( SourceControl.IsValid() )
 			{
 				FText SCCErrorText;
diff --git a/Engine/Source/Editor/UnrealEd/Private/Commandlets/GatherTextFromAssetsCommandlet.cpp b/Engine/Source/Editor/UnrealEd/Private/Commandlets/GatherTextFromAssetsCommandlet.cpp
--- a/Engine/Source/Editor/UnrealEd/Private/Commandlets/GatherTextFromAssetsCommandlet.cpp
+++ b/Engine/Source/Editor/UnrealEd/Private/Commandlets/GatherTextFromAssetsCommandlet.cpp
@@ -69,7 +69,7 @@ void UGatherTextFromAssetsCommandlet::ProcessGatherableTextDataArray(const FStri
 			}
 
 			// Entry already exists, check for conflict.
-			if (ExistingEntry.IsValid() && ExistingEntry->Source.Text != ( NewEntry.SourceString.IsValid() ? **(NewEntry.SourceString) : TEXT("") ))
+			if (ExistingEntry.IsValid() && !ExistingEntry->Source.Text.Equals(*NewEntry.SourceString, ESearchCase::CaseSensitive))
 			{
 				NewEntry.Status = EAssetTextGatherStatus::IdentityConflict;
 			}
@@ -79,8 +79,8 @@ void UGatherTextFromAssetsCommandlet::ProcessGatherableTextDataArray(const FStri
 			{
 				if (!TextSourceSiteContext.IsEditorOnly || ShouldGatherFromEditorOnlyData)
 				{
-				ManifestInfo->AddEntry(TextSourceSiteContext.SiteDescription, GatherableTextData.NamespaceName, NewEntry.SourceString.Get() ? *(NewEntry.SourceString) : TEXT(""), Context );
-			}
+					ManifestInfo->AddEntry(TextSourceSiteContext.SiteDescription, GatherableTextData.NamespaceName, *NewEntry.SourceString, Context);
+				}
 			}
 
 			// Add to conflict tracker.
@@ -339,6 +339,8 @@ int32 UGatherTextFromAssetsCommandlet::Main(const FString& Params)
 		UE_LOG(LogGatherTextFromAssetsCommandlet, Warning, TEXT("No files found or none passed the include/exclude criteria."));
 	}
 
+	const bool bSkipGatherCache = FParse::Param(FCommandLine::Get(), TEXT("SkipGatherCache"));
+
 	TArray< FString > PackageFileNamesToLoad;
 	for (FString& PackageFile : PackageFileNamesToProcess)
 	{
@@ -349,7 +351,7 @@ int32 UGatherTextFromAssetsCommandlet::Main(const FString& Params)
 			FPackageFileSummary PackageFileSummary;
 			(*FileReader) << PackageFileSummary;
 
-			bool MustLoadForGather = false;
+			bool MustLoadForGather = bSkipGatherCache;
 
 			// Packages not resaved since localization gathering flagging was added to packages must be loaded.
 			if (PackageFileSummary.GetFileVersionUE4() < VER_UE4_PACKAGE_REQUIRES_LOCALIZATION_GATHER_FLAGGING)
@@ -431,6 +433,8 @@ int32 UGatherTextFromAssetsCommandlet::Main(const FString& Params)
 		{
 			FString PackageFileName = PackageFileNamesToLoad[PackageIndex];
 
+			UE_LOG(LogGatherTextFromAssetsCommandlet, Verbose, TEXT("Loading package: '%s'."), *PackageFileName);
+
 			UPackage *Package = LoadPackage( NULL, *PackageFileName, LOAD_None );
 			if( Package )
 			{
diff --git a/Engine/Source/Editor/UnrealEd/Private/Commandlets/GenerateTextLocalizationReportCommandlet.cpp b/Engine/Source/Editor/UnrealEd/Private/Commandlets/GenerateTextLocalizationReportCommandlet.cpp
--- a/Engine/Source/Editor/UnrealEd/Private/Commandlets/GenerateTextLocalizationReportCommandlet.cpp
+++ b/Engine/Source/Editor/UnrealEd/Private/Commandlets/GenerateTextLocalizationReportCommandlet.cpp
@@ -334,12 +334,16 @@ bool UGenerateTextLocalizationReportCommandlet::ProcessWordCountReport(const FSt
 		ReportData.SetEntry(NewRowIndex, CultureStr, FString::FromInt( TranslatedWords ) );
 	}
 
-	if( SourceControlInfo.IsValid() )
+	const bool DidFileExist = FPaths::FileExists(ReportFilePath);
+	if (DidFileExist)
 	{
-		FText SCCErrorText;
-		if (!SourceControlInfo->CheckOutFile(ReportFilePath, SCCErrorText))
+		if( SourceControlInfo.IsValid() )
 		{
-			UE_LOG(LogGenerateTextLocalizationReportCommandlet, Warning, TEXT("Check out of file %s failed: %s"), *ReportFilePath, *SCCErrorText.ToString());
+			FText SCCErrorText;
+			if (!SourceControlInfo->CheckOutFile(ReportFilePath, SCCErrorText))
+			{
+				UE_LOG(LogGenerateTextLocalizationReportCommandlet, Warning, TEXT("Check out of file %s failed: %s"), *ReportFilePath, *SCCErrorText.ToString());
+			}
 		}
 	}
 
@@ -348,6 +352,20 @@ bool UGenerateTextLocalizationReportCommandlet::ProcessWordCountReport(const FSt
 		UE_LOG(LogGenerateTextLocalizationReportCommandlet, Error, TEXT("Unable to save report at %s."), *ReportFilePath);
 		return false;
 	}
+
+	if (!DidFileExist)
+	{
+		// Checkout on a new file will cause it to be added
+		if( SourceControlInfo.IsValid() )
+		{
+			FText SCCErrorText;
+			if (!SourceControlInfo->CheckOutFile(ReportFilePath, SCCErrorText))
+			{
+				UE_LOG(LogGenerateTextLocalizationReportCommandlet, Warning, TEXT("Check out of file %s failed: %s"), *ReportFilePath, *SCCErrorText.ToString());
+			}
+		}
+	}
+
 	return true;
 }
 
@@ -371,12 +389,16 @@ bool UGenerateTextLocalizationReportCommandlet::ProcessConflictReport(const FStr
 
 	ReportStr = FConflictReportInfo::GetInstance().ToString();
 
-	if( SourceControlInfo.IsValid() )
+	const bool DidFileExist = FPaths::FileExists(ReportFilePath);
+	if (DidFileExist)
 	{
-		FText SCCErrorText;
-		if (!SourceControlInfo->CheckOutFile(ReportFilePath, SCCErrorText))
+		if( SourceControlInfo.IsValid() )
 		{
-			UE_LOG(LogGenerateTextLocalizationReportCommandlet, Warning, TEXT("Check out of file %s failed: %s"), *ReportFilePath, *SCCErrorText.ToString());
+			FText SCCErrorText;
+			if (!SourceControlInfo->CheckOutFile(ReportFilePath, SCCErrorText))
+			{
+				UE_LOG(LogGenerateTextLocalizationReportCommandlet, Warning, TEXT("Check out of file %s failed: %s"), *ReportFilePath, *SCCErrorText.ToString());
+			}
 		}
 	}
 
@@ -385,5 +407,19 @@ bool UGenerateTextLocalizationReportCommandlet::ProcessConflictReport(const FStr
 		UE_LOG(LogGenerateTextLocalizationReportCommandlet, Error, TEXT("Unable to save report at %s."), *ReportFilePath);
 		return false;
 	}
+
+	if (!DidFileExist)
+	{
+		// Checkout on a new file will cause it to be added
+		if( SourceControlInfo.IsValid() )
+		{
+			FText SCCErrorText;
+			if (!SourceControlInfo->CheckOutFile(ReportFilePath, SCCErrorText))
+			{
+				UE_LOG(LogGenerateTextLocalizationReportCommandlet, Warning, TEXT("Check out of file %s failed: %s"), *ReportFilePath, *SCCErrorText.ToString());
+			}
+		}
+	}
+
 	return true;
 }
diff --git a/Engine/Source/Editor/UnrealEd/Private/Commandlets/GenerateTextLocalizationResourceCommandlet.cpp b/Engine/Source/Editor/UnrealEd/Private/Commandlets/GenerateTextLocalizationResourceCommandlet.cpp
--- a/Engine/Source/Editor/UnrealEd/Private/Commandlets/GenerateTextLocalizationResourceCommandlet.cpp
+++ b/Engine/Source/Editor/UnrealEd/Private/Commandlets/GenerateTextLocalizationResourceCommandlet.cpp
@@ -133,13 +133,17 @@ int32 UGenerateTextLocalizationResourceCommandlet::Main(const FString& Params)
 		// Write resource.
 		const FString TextLocalizationResourcePath = DestinationPath / CultureName / ResourceName;
 
-		if( SourceControlInfo.IsValid() )
+		const bool DidFileExist = FPaths::FileExists(TextLocalizationResourcePath);
+		if (DidFileExist)
 		{
-			FText SCCErrorText;
-			if (!SourceControlInfo->CheckOutFile(TextLocalizationResourcePath, SCCErrorText))
+			if( SourceControlInfo.IsValid() )
 			{
-				UE_LOG(LogGenerateTextLocalizationResourceCommandlet, Error, TEXT("Check out of file %s failed: %s"), *TextLocalizationResourcePath, *SCCErrorText.ToString());
-				return -1;
+				FText SCCErrorText;
+				if (!SourceControlInfo->CheckOutFile(TextLocalizationResourcePath, SCCErrorText))
+				{
+					UE_LOG(LogGenerateTextLocalizationResourceCommandlet, Error, TEXT("Check out of file %s failed: %s"), *TextLocalizationResourcePath, *SCCErrorText.ToString());
+					return -1;
+				}
 			}
 		}
 
@@ -154,6 +158,20 @@ int32 UGenerateTextLocalizationResourceCommandlet::Main(const FString& Params)
 			}
 			TextLocalizationResourceArchive->Close();
 		}
+
+		if (!DidFileExist)
+		{
+			// Checkout on a new file will cause it to be added
+			if( SourceControlInfo.IsValid() )
+			{
+				FText SCCErrorText;
+				if (!SourceControlInfo->CheckOutFile(TextLocalizationResourcePath, SCCErrorText))
+				{
+					UE_LOG(LogGenerateTextLocalizationResourceCommandlet, Error, TEXT("Check out of file %s failed: %s"), *TextLocalizationResourcePath, *SCCErrorText.ToString());
+					return -1;
+				}
+			}
+		}
 	}
 
 	return 0;
diff --git a/Engine/Source/Editor/UnrealEd/Private/Commandlets/InternationalizationExportCommandlet.cpp b/Engine/Source/Editor/UnrealEd/Private/Commandlets/InternationalizationExportCommandlet.cpp
--- a/Engine/Source/Editor/UnrealEd/Private/Commandlets/InternationalizationExportCommandlet.cpp
+++ b/Engine/Source/Editor/UnrealEd/Private/Commandlets/InternationalizationExportCommandlet.cpp
@@ -527,13 +527,17 @@ bool UInternationalizationExportCommandlet::DoExport( const FString& SourcePath,
 
 					NewPortableObject.SortEntries();
 
-					if( SourceControlInfo.IsValid() )
+					const bool DidFileExist = FPaths::FileExists(OutputFileName);
+					if (DidFileExist)
 					{
-						FText SCCErrorText;
-						if (!SourceControlInfo->CheckOutFile(OutputFileName, SCCErrorText))
+						if( SourceControlInfo.IsValid() )
 						{
-							UE_LOG(LogInternationalizationExportCommandlet, Error, TEXT("Check out of file %s failed: %s"), *OutputFileName, *SCCErrorText.ToString());
-							return false;
+							FText SCCErrorText;
+							if (!SourceControlInfo->CheckOutFile(OutputFileName, SCCErrorText))
+							{
+								UE_LOG(LogInternationalizationExportCommandlet, Error, TEXT("Check out of file %s failed: %s"), *OutputFileName, *SCCErrorText.ToString());
+								return false;
+							}
 						}
 					}
 
@@ -544,6 +548,20 @@ bool UInternationalizationExportCommandlet::DoExport( const FString& SourcePath,
 						UE_LOG( LogInternationalizationExportCommandlet, Error, TEXT("Could not write file %s"), *OutputFileName );
 						return false;
 					}
+
+					if (!DidFileExist)
+					{
+						// Checkout on a new file will cause it to be added
+						if( SourceControlInfo.IsValid() )
+						{
+							FText SCCErrorText;
+							if (!SourceControlInfo->CheckOutFile(OutputFileName, SCCErrorText))
+							{
+								UE_LOG(LogInternationalizationExportCommandlet, Error, TEXT("Check out of file %s failed: %s"), *OutputFileName, *SCCErrorText.ToString());
+								return false;
+							}
+						}
+					}
 				}
 			}
 		}
diff --git a/Engine/Source/Editor/UnrealEd/Private/CookOnTheFlyServer.cpp b/Engine/Source/Editor/UnrealEd/Private/CookOnTheFlyServer.cpp
--- a/Engine/Source/Editor/UnrealEd/Private/CookOnTheFlyServer.cpp
+++ b/Engine/Source/Editor/UnrealEd/Private/CookOnTheFlyServer.cpp
@@ -1764,7 +1764,7 @@ uint32 UCookOnTheFlyServer::TickCookOnTheSide( const float TimeSlice, uint32 &Co
 			if ( CookWarningsList.Contains(BuildFilename) == false )
 			{
 				CookWarningsList.Add( BuildFilename );
-				GOutputCookingWarnings = true;
+				GOutputCookingWarnings = IsCookFlagSet(ECookInitializationFlags::OutputVerboseCookerWarnings);
 			}
 			
 			//  if the package is already loaded then try to avoid reloading it :)
@@ -2240,12 +2240,13 @@ uint32 UCookOnTheFlyServer::TickCookOnTheSide( const float TimeSlice, uint32 &Co
 				bool bSucceededSavePackage = false;
 				{
 					SCOPE_PACKAGE_STAT(SaveCookedPackage);
-
+					
 					SCOPE_TIMER(SaveCookedPackage);
 					uint32 SaveFlags = SAVE_KeepGUID | (bShouldSaveAsync ? SAVE_Async : SAVE_None) | (IsCookFlagSet(ECookInitializationFlags::Unversioned) ? SAVE_Unversioned : 0);
-					GOutputCookingWarnings = true;
+					GOutputCookingWarnings = IsCookFlagSet(ECookInitializationFlags::OutputVerboseCookerWarnings);
 					SavePackageResult = SaveCookedPackage(Package, SaveFlags, bWasUpToDate, AllTargetPlatformNames);
 					GOutputCookingWarnings = false;
+
 					if (SavePackageResult == ESavePackageResult::Success)
 					{
 						bSucceededSavePackage = true;
diff --git a/Engine/Source/Programs/AutomationTool/AutomationUtils/ProjectParams.cs b/Engine/Source/Programs/AutomationTool/AutomationUtils/ProjectParams.cs
--- a/Engine/Source/Programs/AutomationTool/AutomationUtils/ProjectParams.cs
+++ b/Engine/Source/Programs/AutomationTool/AutomationUtils/ProjectParams.cs
@@ -1190,7 +1190,7 @@ namespace AutomationTool
         /// <summary>
         /// Cook: Do not include a version number in the cooked content
         /// </summary>
-        public bool UnversionedCookedContent;
+        public bool UnversionedCookedContent = true;
 
 
 		/// <summary>
diff --git a/Engine/Source/Programs/AutomationTool/GUBP/LegacyNodes.cs b/Engine/Source/Programs/AutomationTool/GUBP/LegacyNodes.cs
--- a/Engine/Source/Programs/AutomationTool/GUBP/LegacyNodes.cs
+++ b/Engine/Source/Programs/AutomationTool/GUBP/LegacyNodes.cs
@@ -2651,6 +2651,7 @@ partial class GUBP
 			var LocalizationProjectNames = GetLocalizationProjectNames();
 			var LocalizationProviderName = GetLocalizationProviderName();
 			var LocalizationProviderArgs = GetLocalizationProviderArgs();
+			var AdditionalCommandletArguments = GetAdditionalCommandletArguments();
 
 			// Build the correct command line arguments.
 			var CommandLineArguments = "";
@@ -2685,6 +2686,11 @@ partial class GUBP
 				CommandLineArguments += " " + LocalizationProviderArgs;
 			}
 
+			if (!String.IsNullOrEmpty(AdditionalCommandletArguments))
+			{
+				CommandLineArguments += " -AdditionalCommandletArguments=\"" + AdditionalCommandletArguments + "\"";
+			}
+
 			// Run the localise script.
 			CommandUtils.RunUAT(CommandUtils.CmdEnv, "Localise" + CommandLineArguments);
 
@@ -2723,6 +2729,11 @@ partial class GUBP
 			return "";
 		}
 
+		protected virtual string GetAdditionalCommandletArguments()
+		{
+			return "";
+		}
+
 		protected string LocalizationBranchSuffix;
 	}
 
diff --git a/Engine/Source/Programs/AutomationTool/Scripts/CookCommand.Automation.cs b/Engine/Source/Programs/AutomationTool/Scripts/CookCommand.Automation.cs
--- a/Engine/Source/Programs/AutomationTool/Scripts/CookCommand.Automation.cs
+++ b/Engine/Source/Programs/AutomationTool/Scripts/CookCommand.Automation.cs
@@ -400,7 +400,7 @@ public partial class Project : CommandUtils
 
                 // diff the content
                 List<string> AllFiles = Directory.EnumerateFiles(FullCookPath, "*.uasset", System.IO.SearchOption.AllDirectories).ToList();
-                AllFiles.AddRange(Directory.EnumerateFiles(FullCookPath, "*.map", System.IO.SearchOption.AllDirectories).ToList());
+                AllFiles.AddRange(Directory.EnumerateFiles(FullCookPath, "*.umap", System.IO.SearchOption.AllDirectories).ToList());
                 foreach (string SourceFilename in AllFiles)
                 {
                     // Filename.StartsWith( CookedSandboxesPath );
@@ -408,6 +408,7 @@ public partial class Project : CommandUtils
 
                     string DestFilename = TemporaryFilesPath + RelativeFilename;
 
+                    Log("Comparing file "+ RelativeFilename);
 
                     byte[] SourceFile = null;
                     try
@@ -464,8 +465,8 @@ public partial class Project : CommandUtils
                                 {
                                     Log("Failed to create directory " + Path.GetDirectoryName(SavedDestFilename) + " Exception " + E.ToString());
                                 }
-                                File.Copy(SourceFilename, SavedSourceFilename);
-                                File.Copy(DestFilename, SavedDestFilename);
+                                File.Copy(SourceFilename, SavedSourceFilename, true);
+                                File.Copy(DestFilename, SavedDestFilename, true);
                                 Log("Content temporarily saved to " + SavedSourceFilename + " and " + SavedDestFilename + " at offset " + Index.ToString());
                                 break;
                             }
diff --git a/Engine/Source/Programs/AutomationTool/Scripts/Localisation.Automation.cs b/Engine/Source/Programs/AutomationTool/Scripts/Localisation.Automation.cs
--- a/Engine/Source/Programs/AutomationTool/Scripts/Localisation.Automation.cs
+++ b/Engine/Source/Programs/AutomationTool/Scripts/Localisation.Automation.cs
@@ -15,6 +15,7 @@ using EpicGames.Localization;
 [Help("LocalizationProjectNames", "Comma separated list of the projects to gather text from.")]
 [Help("LocalizationBranch", "Optional suffix to use when uploading the new data to the localization provider.")]
 [Help("LocalizationProvider", "Optional localization provide override (default is OneSky).")]
+[Help("AdditionalCommandletArguments", "Optional arguments to pass to the gather process.")]
 class Localise : BuildCommand
 {
 	public override void ExecuteBuild()
@@ -53,6 +54,12 @@ class Localise : BuildCommand
 			LocalizationProviderName = "OneSky";
 		}
 
+		var AdditionalCommandletArguments = ParseParamValue("AdditionalCommandletArguments");
+		if (AdditionalCommandletArguments == null)
+		{
+			AdditionalCommandletArguments = "";
+		}
+
 		var RootWorkingDirectory = CombinePaths(CmdEnv.LocalRoot, UEProjectDirectory);
 
 		// Try and find our localization provider
@@ -108,6 +115,11 @@ class Localise : BuildCommand
 			{
 				var CommandletArguments = String.Format("-config={0}", LocalizationConfigFile) + (String.IsNullOrEmpty(CommandletSCCArguments) ? "" : " " + CommandletSCCArguments);
 
+				if (!String.IsNullOrEmpty(AdditionalCommandletArguments))
+				{
+					CommandletArguments += " " + AdditionalCommandletArguments;
+				}
+
 				Log("Localization for {0} {1}", EditorArguments, CommandletArguments);
 
 				Log("Running UE4Editor to generate localization data");
diff --git a/Engine/Source/Programs/CrashReporter/CrashReportClient/Private/Linux/CrashReportClientMainLinux.cpp b/Engine/Source/Programs/CrashReporter/CrashReportClient/Private/Linux/CrashReportClientMainLinux.cpp
--- a/Engine/Source/Programs/CrashReporter/CrashReportClient/Private/Linux/CrashReportClientMainLinux.cpp
+++ b/Engine/Source/Programs/CrashReporter/CrashReportClient/Private/Linux/CrashReportClientMainLinux.cpp
@@ -11,6 +11,8 @@ extern int32 ReportCrash(const FLinuxCrashContext& Context);	// FIXME: handle ex
  */
 void CrashReporterCrashHandler(const FGenericCrashContext& GenericContext)
 {
+	// at this point we should already be using malloc crash handler (see PlatformCrashHandler)
+
 	const FLinuxCrashContext& Context = static_cast< const FLinuxCrashContext& >( GenericContext );
 
 	printf("CrashHandler: Signal=%d\n", Context.Signal);
diff --git a/Engine/Source/Runtime/AnimGraphRuntime/Private/AnimNodes/AnimNode_BlendSpacePlayer.cpp b/Engine/Source/Runtime/AnimGraphRuntime/Private/AnimNodes/AnimNode_BlendSpacePlayer.cpp
--- a/Engine/Source/Runtime/AnimGraphRuntime/Private/AnimNodes/AnimNode_BlendSpacePlayer.cpp
+++ b/Engine/Source/Runtime/AnimGraphRuntime/Private/AnimNodes/AnimNode_BlendSpacePlayer.cpp
@@ -15,7 +15,8 @@ FAnimNode_BlendSpacePlayer::FAnimNode_BlendSpacePlayer()
 	, PlayRate(1.0f)
 	, bLoop(true)
 	, StartPosition(0.f)
-	, BlendSpace(NULL)
+	, BlendSpace(nullptr)
+	, PreviousBlendSpace(nullptr)
 {
 }
 
@@ -23,20 +24,11 @@ void FAnimNode_BlendSpacePlayer::Initialize(const FAnimationInitializeContext& C
 {
 	FAnimNode_AssetPlayerBase::Initialize(Context);
 
-	BlendSampleDataCache.Empty();
-	
 	EvaluateGraphExposedInputs.Execute(Context);
-	InternalTimeAccumulator = FMath::Clamp(StartPosition, 0.f, 1.0f);
-	if(StartPosition == 0.f && PlayRate < 0.0f)
-	{
-		// Blend spaces run between 0 and 1
-		InternalTimeAccumulator = 1.0f;
-	}
 
-	if (BlendSpace != NULL)
-	{
-		BlendSpace->InitializeFilter(&BlendFilter);
-	}
+	Reinitialize();
+
+	PreviousBlendSpace = BlendSpace;
 }
 
 void FAnimNode_BlendSpacePlayer::CacheBones(const FAnimationCacheBonesContext& Context) 
@@ -60,6 +52,11 @@ void FAnimNode_BlendSpacePlayer::UpdateInternal(const FAnimationUpdateContext& C
 
 		const FVector BlendInput(X, Y, Z);
 	
+		if (PreviousBlendSpace != BlendSpace)
+		{
+			Reinitialize();
+		}
+
 		Context.AnimInstanceProxy->MakeBlendSpaceTickRecord(TickRecord, BlendSpace, BlendInput, BlendSampleDataCache, BlendFilter, bLoop, PlayRate, Context.GetFinalBlendWeight(), /*inout*/ InternalTimeAccumulator, MarkerTickRecord);
 
 		// Update the sync group if it exists
@@ -67,6 +64,8 @@ void FAnimNode_BlendSpacePlayer::UpdateInternal(const FAnimationUpdateContext& C
 		{
 			SyncGroup->TestTickRecordForLeadership(GroupRole);
 		}
+
+		PreviousBlendSpace = BlendSpace;
 	}
 }
 
@@ -112,3 +111,19 @@ UAnimationAsset* FAnimNode_BlendSpacePlayer::GetAnimAsset()
 {
 	return BlendSpace;
 }
+
+void FAnimNode_BlendSpacePlayer::Reinitialize()
+{
+	BlendSampleDataCache.Empty();
+	InternalTimeAccumulator = FMath::Clamp(StartPosition, 0.f, 1.0f);
+	if (StartPosition == 0.f && PlayRate < 0.0f)
+	{
+		// Blend spaces run between 0 and 1
+		InternalTimeAccumulator = 1.0f;
+	}
+
+	if (BlendSpace != NULL)
+	{
+		BlendSpace->InitializeFilter(&BlendFilter);
+	}
+}
\ No newline at end of file
diff --git a/Engine/Source/Runtime/AnimGraphRuntime/Public/AnimNodes/AnimNode_BlendSpacePlayer.h b/Engine/Source/Runtime/AnimGraphRuntime/Public/AnimNodes/AnimNode_BlendSpacePlayer.h
--- a/Engine/Source/Runtime/AnimGraphRuntime/Public/AnimNodes/AnimNode_BlendSpacePlayer.h
+++ b/Engine/Source/Runtime/AnimGraphRuntime/Public/AnimNodes/AnimNode_BlendSpacePlayer.h
@@ -50,6 +50,9 @@ protected:
 	UPROPERTY()
 	TArray<FBlendSampleData> BlendSampleDataCache;
 
+	UPROPERTY(transient)
+	UBlendSpaceBase* PreviousBlendSpace;
+
 public:	
 	FAnimNode_BlendSpacePlayer();
 
@@ -68,4 +71,7 @@ public:
 
 protected:
 	void UpdateInternal(const FAnimationUpdateContext& Context);
+
+private:
+	void Reinitialize();
 };
diff --git a/Engine/Source/Runtime/Core/Private/GenericPlatform/GenericPlatformMemory.cpp b/Engine/Source/Runtime/Core/Private/GenericPlatform/GenericPlatformMemory.cpp
--- a/Engine/Source/Runtime/Core/Private/GenericPlatform/GenericPlatformMemory.cpp
+++ b/Engine/Source/Runtime/Core/Private/GenericPlatform/GenericPlatformMemory.cpp
@@ -185,9 +185,9 @@ uint32 FGenericPlatformMemory::GetPhysicalGBRam()
 	return FPlatformMemory::GetConstants().TotalPhysicalGB;
 }
 
-bool FGenericPlatformMemory::PageProtect( void* const Ptr, const SIZE_T Size, const uint32 ProtectMode )
+bool FGenericPlatformMemory::PageProtect(void* const Ptr, const SIZE_T Size, const bool bCanRead, const bool bCanWrite)
 {
-	UE_LOG(LogMemory, Error, TEXT("FGenericPlatformMemory::PageProtect not implemented on this platform"));
+	UE_LOG(LogMemory, Verbose, TEXT("FGenericPlatformMemory::PageProtect not implemented on this platform"));
 	return false;
 }
 
diff --git a/Engine/Source/Runtime/Core/Private/HAL/MallocStomp.cpp b/Engine/Source/Runtime/Core/Private/HAL/MallocStomp.cpp
--- a/Engine/Source/Runtime/Core/Private/HAL/MallocStomp.cpp
+++ b/Engine/Source/Runtime/Core/Private/HAL/MallocStomp.cpp
@@ -2,9 +2,10 @@
 
 #include "CorePrivatePCH.h"
 
+#include "MallocStomp.h"
+
 #if USE_MALLOC_STOMP
 
-#include "MallocStomp.h"
 
 static void MallocStompOverrunTest()
 {
@@ -52,7 +53,7 @@ void* FMallocStomp::Malloc(SIZE_T Size, uint32 Alignment)
 		*AllocDataPtr = AllocData;
 
 		// Page protect the first page, this will cause the exception in case the is an underrun.
-		FPlatformMemory::PageProtect(FullAllocationPointer, PageSize, NoAccessProtectMode);
+		FPlatformMemory::PageProtect(FullAllocationPointer, PageSize, false, false);
 	}
 	else
 	{
@@ -63,7 +64,7 @@ void* FMallocStomp::Malloc(SIZE_T Size, uint32 Alignment)
 		*AllocDataPtr = AllocData;
 
 		// Page protect the last page, this will cause the exception in case the is an overrun.
-		FPlatformMemory::PageProtect(reinterpret_cast<void*>(reinterpret_cast<uint8*>(FullAllocationPointer) + AllocFullPageSize), PageSize, NoAccessProtectMode);
+		FPlatformMemory::PageProtect(reinterpret_cast<void*>(reinterpret_cast<uint8*>(FullAllocationPointer) + AllocFullPageSize), PageSize, false, false);
 	}
 
 	return ReturnedPointer;
diff --git a/Engine/Source/Runtime/Core/Private/HAL/UnrealMemory.cpp b/Engine/Source/Runtime/Core/Private/HAL/UnrealMemory.cpp
--- a/Engine/Source/Runtime/Core/Private/HAL/UnrealMemory.cpp
+++ b/Engine/Source/Runtime/Core/Private/HAL/UnrealMemory.cpp
@@ -16,6 +16,205 @@
 #include "PlatformMallocCrash.h"
 
 
+class FMallocPurgatoryProxy : public FMalloc
+{
+	/** Malloc we're based on, aka using under the hood							*/
+	FMalloc*			UsedMalloc;
+	enum
+	{
+		PURGATORY_STOMP_CHECKS_FRAMES = 4, // we want to allow several frames since it is theoretically possible for something to be blocked mid-allocation for many frames.
+		PURGATORY_STOMP_MAX_PURGATORY_MEM = 100000000, // for loading and whatnot, we really can't keep this stuff forever
+		PURGATORY_STOMP_CHECKS_CANARYBYTE = 0xdc,
+	};
+	uint32 LastCheckFrame; // unitialized; won't matter
+	FThreadSafeCounter OutstandingSizeInKB;
+	FThreadSafeCounter NextOversizeClear;
+	TLockFreePointerListUnordered<void> Purgatory[PURGATORY_STOMP_CHECKS_FRAMES];
+
+public:
+	/**
+	* Constructor for thread safe proxy malloc that takes a malloc to be used and a
+	* synchronization object used via FScopeLock as a parameter.
+	*
+	* @param	InMalloc					FMalloc that is going to be used for actual allocations
+	*/
+	FMallocPurgatoryProxy(FMalloc* InMalloc)
+		: UsedMalloc(InMalloc)
+	{}
+
+	virtual void InitializeStatsMetadata() override
+	{
+		UsedMalloc->InitializeStatsMetadata();
+	}
+
+	/**
+	* Malloc
+	*/
+	virtual void* Malloc(SIZE_T Size, uint32 Alignment) override
+	{
+		return UsedMalloc->Malloc(Size, Alignment);
+	}
+
+	/**
+	* Realloc
+	*/
+	virtual void* Realloc(void* Ptr, SIZE_T NewSize, uint32 Alignment) override
+	{
+		return UsedMalloc->Realloc(Ptr, NewSize, Alignment);
+	}
+
+	/**
+	* Free
+	*/
+	virtual void Free(void* Ptr) override
+	{
+		if (!Ptr)// || !GIsGameThreadIdInitialized)
+		{
+			UsedMalloc->Free(Ptr);
+			return;
+		}
+		{
+			SIZE_T Size = 0;
+			verify(GetAllocationSize(Ptr, Size) && Size);
+			FMemory::Memset(Ptr, uint8(PURGATORY_STOMP_CHECKS_CANARYBYTE), Size);
+			Purgatory[GFrameNumber % PURGATORY_STOMP_CHECKS_FRAMES].Push(Ptr);
+			OutstandingSizeInKB.Add((Size + 1023) / 1024);
+		}
+		FPlatformMisc::MemoryBarrier();
+		uint32 LocalLastCheckFrame = LastCheckFrame;
+		uint32 LocalGFrameNumber = GFrameNumber;
+
+		bool bFlushAnyway = OutstandingSizeInKB.GetValue() > PURGATORY_STOMP_MAX_PURGATORY_MEM / 1024;
+
+		if (bFlushAnyway || LocalLastCheckFrame != LocalGFrameNumber)
+		{
+			if (bFlushAnyway || FPlatformAtomics::InterlockedCompareExchange((volatile int32*)&LastCheckFrame, LocalGFrameNumber, LocalLastCheckFrame) == LocalLastCheckFrame)
+			{
+				int32 FrameToPop = ((bFlushAnyway ? NextOversizeClear.Increment() : LocalGFrameNumber) + PURGATORY_STOMP_CHECKS_FRAMES - 1) % PURGATORY_STOMP_CHECKS_FRAMES;
+				//FPlatformMisc::LowLevelOutputDebugStringf(TEXT("Pop %d at %d %d\r\n"), FrameToPop, LocalGFrameNumber, LocalGFrameNumber % PURGATORY_STOMP_CHECKS_FRAMES);
+				while (true)
+				{
+					uint8* Pop = (uint8*)Purgatory[FrameToPop].Pop();
+					if (!Pop)
+					{
+						break;
+					}
+					SIZE_T Size = 0;
+					verify(GetAllocationSize(Pop, Size) && Size);
+					for (SIZE_T At = 0; At < Size; At++)
+					{
+						if (Pop[At] != PURGATORY_STOMP_CHECKS_CANARYBYTE)
+						{
+							FPlatformMisc::LowLevelOutputDebugStringf(TEXT("Freed memory at %p + %d == %x (should be %x)\r\n"), Pop, At, (int32)Pop[At], (int32)PURGATORY_STOMP_CHECKS_CANARYBYTE);
+							UE_LOG(LogMemory, Fatal, TEXT("Freed memory at %p + %d == %x (should be %x)"), Pop, At, (int32)Pop[At], (int32)PURGATORY_STOMP_CHECKS_CANARYBYTE);
+						}
+					}
+					UsedMalloc->Free(Pop);
+					OutstandingSizeInKB.Subtract((Size + 1023) / 1024);
+				}
+			}
+		}
+	}
+
+	/** Writes allocator stats from the last update into the specified destination. */
+	virtual void GetAllocatorStats(FGenericMemoryStats& out_Stats) override
+	{
+		UsedMalloc->GetAllocatorStats(out_Stats);
+	}
+
+	/** Dumps allocator stats to an output device. */
+	virtual void DumpAllocatorStats(class FOutputDevice& Ar) override
+	{
+		UsedMalloc->DumpAllocatorStats(Ar);
+	}
+
+	/**
+	* Validates the allocator's heap
+	*/
+	virtual bool ValidateHeap() override
+	{
+		return(UsedMalloc->ValidateHeap());
+	}
+
+	virtual bool Exec(UWorld* InWorld, const TCHAR* Cmd, FOutputDevice& Ar) override
+	{
+		return UsedMalloc->Exec(InWorld, Cmd, Ar);
+	}
+
+	/**
+	* If possible determine the size of the memory allocated at the given address
+	*
+	* @param Original - Pointer to memory we are checking the size of
+	* @param SizeOut - If possible, this value is set to the size of the passed in pointer
+	* @return true if succeeded
+	*/
+	virtual bool GetAllocationSize(void *Original, SIZE_T &SizeOut) override
+	{
+		return UsedMalloc->GetAllocationSize(Original, SizeOut);
+	}
+
+	virtual const TCHAR* GetDescriptiveName() override
+	{
+		return UsedMalloc->GetDescriptiveName();
+	}
+};
+
+void FMemory::EnablePurgatoryTests()
+{
+	static bool bOnce = false;
+	if (bOnce)
+	{
+		UE_LOG(LogMemory, Error, TEXT("Purgatory proxy was already turned on."));
+		return;
+	}
+	bOnce = true;
+	while (true)
+	{
+		FMalloc* LocalGMalloc = GMalloc;
+		FMalloc* Proxy = new FMallocPurgatoryProxy(LocalGMalloc);
+		if (FPlatformAtomics::InterlockedCompareExchangePointer((void**)&GMalloc, Proxy, LocalGMalloc) == LocalGMalloc)
+		{
+			UE_LOG(LogConsoleResponse, Display, TEXT("Purgatory proxy is now on."));
+			return;
+		}
+		delete Proxy;
+	}
+}
+
+#if !UE_BUILD_SHIPPING
+#include "TaskGraphInterfaces.h"
+static void FMallocBinnedOverrunTest()
+{
+	const uint32 ArraySize = 64;
+	uint8* Pointer = new uint8[ArraySize];
+	delete[] Pointer;
+	FFunctionGraphTask::CreateAndDispatchWhenReady(
+		[Pointer, ArraySize]()
+	{
+		//FPlatformProcess::Sleep(.01f);
+		Pointer[ArraySize / 2] = 0xcc;
+	},
+		TStatId()
+		);
+}
+
+FAutoConsoleCommand FMallocBinnedTestCommand
+(
+TEXT("Memory.StaleTest"),
+TEXT("Test for Memory.UsePurgatory. *** Will crash the game!"),
+FConsoleCommandDelegate::CreateStatic(&FMallocBinnedOverrunTest)
+);
+
+FAutoConsoleCommand FMallocUsePurgatoryCommand
+(
+TEXT("Memory.UsePurgatory"),
+TEXT("Uses the purgatory malloc proxy to check if things are writing to stale pointers."),
+FConsoleCommandDelegate::CreateStatic(&FMemory::EnablePurgatoryTests)
+);
+#endif
+
+
+
 /** Helper function called on first allocation to create and initialize GMalloc */
 void GCreateMalloc()
 {
@@ -40,6 +239,8 @@ void GCreateMalloc()
 		GMalloc = new FMallocThreadSafeProxy( GMalloc );
 	}
 }
+
+
 #if STATS
 	#define MALLOC_GT_HOOKS 1
 #else
diff --git a/Engine/Source/Runtime/Core/Private/Internationalization/ICUText.cpp b/Engine/Source/Runtime/Core/Private/Internationalization/ICUText.cpp
--- a/Engine/Source/Runtime/Core/Private/Internationalization/ICUText.cpp
+++ b/Engine/Source/Runtime/Core/Private/Internationalization/ICUText.cpp
@@ -4,9 +4,7 @@
 
 #if UE_ENABLE_ICU
 #include "Text.h"
-#include "TextData.h"
-#include "TextHistory.h"
-#include "ICUCulture.h"
+#include "ICUText.h"
 
 PRAGMA_DISABLE_SHADOW_VARIABLE_WARNINGS
 #include <unicode/coll.h>
@@ -17,7 +15,6 @@ PRAGMA_ENABLE_SHADOW_VARIABLE_WARNINGS
 #include <unicode/uniset.h>
 #include <unicode/ubidi.h>
 
-#include "ICUUtilities.h"
 #include "ICUTextCharacterIterator.h"
 
 bool FText::IsWhitespace( const TCHAR Char )
@@ -28,6 +25,21 @@ bool FText::IsWhitespace( const TCHAR Char )
 	return u_isWhitespace(ICUChar) != 0;
 }
 
+// static
+FText FText::AsCurrencyBase(int64 BaseVal, const FString& CurrencyCode, const FCulturePtr& TargetCulture)
+{
+	FInternationalization& I18N = FInternationalization::Get();
+	checkf(I18N.IsInitialized() == true, TEXT("FInternationalization is not initialized. An FText formatting method was likely used in static object initialization - this is not supported."));
+	const FCulture& Culture = TargetCulture.IsValid() ? *TargetCulture : *I18N.GetCurrentCulture();
+
+	const FDecimalNumberFormattingRules& FormattingRules = Culture.Implementation->GetCurrencyFormattingRules(CurrencyCode);
+	const FNumberFormattingOptions& FormattingOptions = FormattingRules.CultureDefaultFormattingOptions;
+	double Val = static_cast<double>(BaseVal) / FMath::Pow(10.0f, FormattingOptions.MaximumFractionalDigits);
+	FString NativeString = FastDecimalFormat::NumberToString(Val, FormattingRules, FormattingOptions);
+
+	return FText::CreateNumericalText(MakeShareable(new TGeneratedTextData<FTextHistory_AsCurrency>(MoveTemp(NativeString), FTextHistory_AsCurrency(Val, CurrencyCode, nullptr, TargetCulture))));
+}
+
 FText FText::AsDate(const FDateTime& DateTime, const EDateTimeStyle::Type DateStyle, const FString& TimeZone, const FCulturePtr& TargetCulture)
 {
 	FInternationalization& I18N = FInternationalization::Get();
diff --git a/Engine/Source/Runtime/Core/Private/Internationalization/LegacyText.cpp b/Engine/Source/Runtime/Core/Private/Internationalization/LegacyText.cpp
--- a/Engine/Source/Runtime/Core/Private/Internationalization/LegacyText.cpp
+++ b/Engine/Source/Runtime/Core/Private/Internationalization/LegacyText.cpp
@@ -11,6 +11,15 @@ bool FText::IsWhitespace( const TCHAR Char )
 	return FChar::IsWhitespace(Char);
 }
 
+FText FText::AsCurrencyBase(int64 BaseVal, const FString& CurrencyCode, const FCulturePtr& TargetCulture)
+{
+	FInternationalization& I18N = FInternationalization::Get();
+	checkf(I18N.IsInitialized() == true, TEXT("FInternationalization is not initialized. An FText formatting method was likely used in static object initialization - this is not supported."));
+	const FCulture& Culture = TargetCulture.IsValid() ? *TargetCulture : *I18N.GetCurrentCulture();
+	double Val = static_cast<double>(BaseVal) / FMath::Pow(10.0f, Culture.NumberFormattingRule.CurrencyDecimalDigits);
+	return FText::CreateNumericalText(MakeShareable(new TGeneratedTextData<FTextHistory_AsCurrency>(Culture.NumberFormattingRule.AsCurrency(Val), FTextHistory_AsCurrency(Val, CurrencyCode, nullptr, TargetCulture))));
+}
+
 FText FText::AsDate( const FDateTime& DateTime, const EDateTimeStyle::Type DateStyle, const FString& TimeZone, const FCulturePtr& TargetCulture )
 {
 	checkf(FInternationalization::Get().IsInitialized() == true, TEXT("FInternationalization is not initialized. An FText formatting method was likely used in static object initialization - this is not supported."));
diff --git a/Engine/Source/Runtime/Core/Private/Internationalization/TextHistory.cpp b/Engine/Source/Runtime/Core/Private/Internationalization/TextHistory.cpp
--- a/Engine/Source/Runtime/Core/Private/Internationalization/TextHistory.cpp
+++ b/Engine/Source/Runtime/Core/Private/Internationalization/TextHistory.cpp
@@ -576,6 +576,8 @@ FText FTextHistory_AsCurrency::ToText(bool bInAsSource) const
 {
 	TSharedPtr< FCulture, ESPMode::ThreadSafe > CurrentCulture = bInAsSource? FInternationalization::Get().GetInvariantCulture() : TargetCulture;
 
+PRAGMA_DISABLE_DEPRECATION_WARNINGS
+	// when we remove AsCurrency should be easy to switch these to AsCurrencyBase and change SourceValue to be BaseVal in AsCurrencyBase (currently is the pre-divided value)
 	const FNumberFormattingOptions* FormatOptionsPtr = (FormatOptions.IsSet()) ? &FormatOptions.GetValue() : nullptr;
 	switch(SourceValue.GetType())
 	{
@@ -601,6 +603,7 @@ FText FTextHistory_AsCurrency::ToText(bool bInAsSource) const
 			check(0);
 		}
 	}
+PRAGMA_ENABLE_DEPRECATION_WARNINGS 
 	return FText();
 }
 
diff --git a/Engine/Source/Runtime/Core/Private/Linux/LinuxPlatformCrashContext.cpp b/Engine/Source/Runtime/Core/Private/Linux/LinuxPlatformCrashContext.cpp
--- a/Engine/Source/Runtime/Core/Private/Linux/LinuxPlatformCrashContext.cpp
+++ b/Engine/Source/Runtime/Core/Private/Linux/LinuxPlatformCrashContext.cpp
@@ -4,6 +4,7 @@
 #include "Linux/LinuxPlatformCrashContext.h"
 #include "Misc/App.h"
 #include "EngineVersion.h"
+#include "PlatformMallocCrash.h"
 
 #include <sys/utsname.h>	// for uname()
 #include <signal.h>
@@ -73,22 +74,7 @@ void FLinuxCrashContext::InitFromSignal(int32 InSignal, siginfo_t* InInfo, void*
  */
 void GracefulTerminationHandler(int32 Signal, siginfo_t* Info, void* Context)
 {
-	printf("CtrlCHandler: Signal=%d\n", Signal);
-
-	// make sure as much data is written to disk as possible
-	if (GLog)
-	{
-		GLog->Flush();
-	}
-	if (GWarn)
-	{
-		GWarn->Flush();
-	}
-	if (GError)
-	{
-		GError->Flush();
-	}
-
+	// do not flush logs at this point; this can result in a deadlock if the signal was received while we were holding lock in the malloc (flushing allocates memory)
 	if( !GIsRequestingExit )
 	{
 		GIsRequestingExit = 1;
@@ -417,11 +403,10 @@ void DLLEXPORT GenerateCrashInfoAndLaunchReporter(const FLinuxCrashContext & Con
  */
 void DefaultCrashHandler(const FLinuxCrashContext & Context)
 {
-	// Switch to malloc crash.
-	//FMallocCrash::Get().SetAsGMalloc(); 
-
 	printf("DefaultCrashHandler: Signal=%d\n", Context.Signal);
 
+	// at this point we should already be using malloc crash handler (see PlatformCrashHandler)
+
 	ReportCrash(Context);
 	if (GLog)
 	{
@@ -448,6 +433,9 @@ void PlatformCrashHandler(int32 Signal, siginfo_t* Info, void* Context)
 {
 	fprintf(stderr, "Signal %d caught.\n", Signal);
 
+	// Switch to malloc crash.
+	FPlatformMallocCrash::Get().SetAsGMalloc();
+
 	FLinuxCrashContext CrashContext;
 	CrashContext.InitFromSignal(Signal, Info, Context);
 
diff --git a/Engine/Source/Runtime/Core/Private/Linux/LinuxPlatformMemory.cpp b/Engine/Source/Runtime/Core/Private/Linux/LinuxPlatformMemory.cpp
--- a/Engine/Source/Runtime/Core/Private/Linux/LinuxPlatformMemory.cpp
+++ b/Engine/Source/Runtime/Core/Private/Linux/LinuxPlatformMemory.cpp
@@ -114,9 +114,26 @@ class FMalloc* FLinuxPlatformMemory::BaseAllocator()
 	return Allocator;
 }
 
-bool FLinuxPlatformMemory::PageProtect( void* const Ptr, const SIZE_T Size, const uint32 ProtectMode )
+bool FLinuxPlatformMemory::PageProtect(void* const Ptr, const SIZE_T Size, const bool bCanRead, const bool bCanWrite)
 {
-	return mprotect(Ptr, Size, static_cast<int32>(ProtectMode)) == 0;
+	int32 ProtectMode;
+	if (bCanRead && bCanWrite)
+	{
+		ProtectMode = PROT_READ | PROT_WRITE;
+	}
+	else if (bCanRead)
+	{
+		ProtectMode = PROT_READ;
+	}
+	else if (bCanWrite)
+	{
+		ProtectMode = PROT_WRITE;
+	}
+	else
+	{
+		ProtectMode = PROT_NONE;
+	}
+	return mprotect(Ptr, Size, ProtectMode) == 0;
 }
 
 void* FLinuxPlatformMemory::BinnedAllocFromOS( SIZE_T Size )
diff --git a/Engine/Source/Runtime/Core/Private/Mac/MacPlatformMemory.cpp b/Engine/Source/Runtime/Core/Private/Mac/MacPlatformMemory.cpp
--- a/Engine/Source/Runtime/Core/Private/Mac/MacPlatformMemory.cpp
+++ b/Engine/Source/Runtime/Core/Private/Mac/MacPlatformMemory.cpp
@@ -8,9 +8,7 @@
 #include "MallocTBB.h"
 #include "MallocAnsi.h"
 #include "MallocBinned2.h"
-#if USE_MALLOC_STOMP
-	#include "MallocStomp.h"
-#endif // USE_MALLOC_STOMP
+#include "MallocStomp.h"
 
 #include <sys/param.h>
 #include <sys/mount.h>
@@ -131,8 +129,25 @@ const FPlatformMemoryConstants& FMacPlatformMemory::GetConstants()
 	return MemoryConstants;	
 }
 
-bool FMacPlatformMemory::PageProtect( void* const Ptr, const SIZE_T Size, const uint32 ProtectMode )
+bool FMacPlatformMemory::PageProtect(void* const Ptr, const SIZE_T Size, const bool bCanRead, const bool bCanWrite)
 {
+	int32 ProtectMode;
+	if (bCanRead && bCanWrite)
+	{
+		ProtectMode = PROT_READ | PROT_WRITE;
+	}
+	else if (bCanRead)
+	{
+		ProtectMode = PROT_READ;
+	}
+	else if (bCanWrite)
+	{
+		ProtectMode = PROT_WRITE;
+	}
+	else
+	{
+		ProtectMode = PROT_NONE;
+	}
 	return mprotect(Ptr, Size, static_cast<int32>(ProtectMode)) == 0;
 }
 
diff --git a/Engine/Source/Runtime/Core/Private/Misc/AutomationTest.cpp b/Engine/Source/Runtime/Core/Private/Misc/AutomationTest.cpp
--- a/Engine/Source/Runtime/Core/Private/Misc/AutomationTest.cpp
+++ b/Engine/Source/Runtime/Core/Private/Misc/AutomationTest.cpp
@@ -175,10 +175,10 @@ bool FAutomationTestFramework::RunSmokeTests()
 
 			double EndTime = FPlatformTime::Seconds();
 			double TimeForTest = static_cast<float>(EndTime - SmokeTestStartTime);
-			if (TimeForTest > 1.0f)
+			if (TimeForTest > 2.0f)
 			{
 				//force a failure if a smoke test takes too long
-				UE_LOG(LogAutomationTest, Warning, TEXT("Smoke tests took > 1s to run: %.2fs"), (float)TimeForTest);
+				UE_LOG(LogAutomationTest, Warning, TEXT("Smoke tests took > 2s to run: %.2fs"), (float)TimeForTest);
 			}
 
 			FAutomationTestFramework::DumpAutomationTestExecutionInfo( OutExecutionInfoMap );
diff --git a/Engine/Source/Runtime/Core/Private/Tests/Internationalization/TextTest.cpp b/Engine/Source/Runtime/Core/Private/Tests/Internationalization/TextTest.cpp
--- a/Engine/Source/Runtime/Core/Private/Tests/Internationalization/TextTest.cpp
+++ b/Engine/Source/Runtime/Core/Private/Tests/Internationalization/TextTest.cpp
@@ -445,7 +445,7 @@ bool FTextTest::RunTest (const FString& Parameters)
 			FText AsNumberTest1 = FText::AsNumber(5.5421);
 
 			FText AsPercentTest1 = FText::AsPercent(0.925);
-			FText AsCurrencyTest1 = FText::AsCurrency(100.25);
+			FText AsCurrencyTest1 = FText::AsCurrencyBase(10025, TEXT("USD"));
 
 			FDateTime DateTimeInfo(2080, 8, 20, 9, 33, 22);
 			FText AsDateTimeTest1 = FText::AsDateTime(DateTimeInfo, EDateTimeStyle::Default, EDateTimeStyle::Default, TEXT("UTC"));
@@ -508,7 +508,7 @@ bool FTextTest::RunTest (const FString& Parameters)
 					AddError( TEXT("Percent Output=") + AsPercentTest1.ToString() );
 				}
 
-				if(AsCurrencyTest1.CompareTo(FText::AsCurrency(100.25)) != 0)
+				if(AsCurrencyTest1.CompareTo(FText::AsCurrencyBase(10025, TEXT("USD"))) != 0)
 				{
 					AddError( TEXT("AsCurrencyTest1 did not rebuild correctly in French-Canadian") );
 					AddError( TEXT("Currency Output=") + AsCurrencyTest1.ToString() );
diff --git a/Engine/Source/Runtime/Core/Private/Windows/WindowsPlatformMemory.cpp b/Engine/Source/Runtime/Core/Private/Windows/WindowsPlatformMemory.cpp
--- a/Engine/Source/Runtime/Core/Private/Windows/WindowsPlatformMemory.cpp
+++ b/Engine/Source/Runtime/Core/Private/Windows/WindowsPlatformMemory.cpp
@@ -4,9 +4,7 @@
 
 #include "MallocTBB.h"
 #include "MallocAnsi.h"
-#if USE_MALLOC_STOMP
-	#include "MallocStomp.h"
-#endif // USE_MALLOC_STOMP
+#include "MallocStomp.h"
 #include "GenericPlatformMemoryPoolStats.h"
 #include "MemoryMisc.h"
 
@@ -170,9 +168,26 @@ const FPlatformMemoryConstants& FWindowsPlatformMemory::GetConstants()
 	return MemoryConstants;	
 }
 
-bool FWindowsPlatformMemory::PageProtect( void* const Ptr, const SIZE_T Size, const uint32 ProtectMode )
+bool FWindowsPlatformMemory::PageProtect(void* const Ptr, const SIZE_T Size, const bool bCanRead, const bool bCanWrite)
 {
 	DWORD flOldProtect;
+	uint32 ProtectMode = 0;
+	if (bCanRead && bCanWrite)
+	{
+		ProtectMode = PAGE_READWRITE;
+	}
+	else if (bCanWrite)
+	{
+		ProtectMode = PAGE_READWRITE;
+	}
+	else if (bCanRead)
+	{
+		ProtectMode = PAGE_READONLY;
+	}
+	else
+	{
+		ProtectMode = PAGE_NOACCESS;
+	}
 	return VirtualProtect(Ptr, Size, ProtectMode, &flOldProtect) != 0;
 }
 void* FWindowsPlatformMemory::BinnedAllocFromOS( SIZE_T Size )
diff --git a/Engine/Source/Runtime/Core/Private/Windows/WindowsPlatformMisc.cpp b/Engine/Source/Runtime/Core/Private/Windows/WindowsPlatformMisc.cpp
--- a/Engine/Source/Runtime/Core/Private/Windows/WindowsPlatformMisc.cpp
+++ b/Engine/Source/Runtime/Core/Private/Windows/WindowsPlatformMisc.cpp
@@ -2395,35 +2395,71 @@ FString FWindowsPlatformMisc::GetPrimaryGPUBrand()
 
 void FWindowsPlatformMisc::GetGPUDriverInfo(const FString DeviceDescription, FString& InternalDriverVersion, FString& UserDriverVersion, FString& DriverDate)
 {
-	// to prevent bugs
-	InternalDriverVersion.Empty();
-	UserDriverVersion.Empty();
-	DriverDate.Empty();
+	// to distinguish failed GetGPUDriverInfo() from call to GetGPUDriverInfo()
+	InternalDriverVersion = TEXT("Unknown");
+	UserDriverVersion = TEXT("Unknown");
+	DriverDate = TEXT("Unknown");
 
 	for(uint32 i = 0;; ++i)
 	{
 		// Iterate all installed display adapters
-		FString Key = FString::Printf(TEXT("SYSTEM\\CurrentControlSet\\Control\\Class\\{4D36E968-E325-11CE-BFC1-08002BE10318}\\%04d"), i);
+		const FString Key = FString::Printf(TEXT("SYSTEM\\CurrentControlSet\\Control\\Class\\{4D36E968-E325-11CE-BFC1-08002BE10318}\\%04d"), i);
+		const TCHAR* DeviceDescriptionValueName = TEXT("Device Description");
 
 		FString LocalDeviceDescription; // e.g. "NVIDIA GeForce GTX 680" or "AMD Radeon R9 200 / HD 7900 Series"
-		bool bDevice = FWindowsPlatformMisc::QueryRegKey(HKEY_LOCAL_MACHINE, *Key, TEXT("Device Description"), LocalDeviceDescription);	// AMD and NVIDIA
+		bool bDevice = FWindowsPlatformMisc::QueryRegKey(HKEY_LOCAL_MACHINE, *Key, DeviceDescriptionValueName, /*out*/ LocalDeviceDescription); // AMD and NVIDIA
 
-		if(!bDevice)
+		if (!bDevice)
 		{
-			break;
+			// Try again in Settings subfolder
+			const FString SettingsSubKey = Key + TEXT("\\Settings");
+			bDevice = FWindowsPlatformMisc::QueryRegKey(HKEY_LOCAL_MACHINE, *SettingsSubKey, DeviceDescriptionValueName, /*out*/ LocalDeviceDescription); // AMD and NVIDIA
+
+			if (!bDevice)
+			{
+				// Neither root nor Settings subfolder contained a "Device Description" value so this is probably not a device
+				break;
+			}
 		}
 
+		// e.g. "NVIDIA" or "Advanced Micro Devices, Inc."
+		FString LocalProviderName;
+		FWindowsPlatformMisc::QueryRegKey(HKEY_LOCAL_MACHINE, *Key, TEXT("ProviderName"), LocalProviderName);
+
 		// AMD (not the Catalyst one) and NVIDIA
-		// e.g. "15.200.1062.1004"(AMD) "9.18.13.4788" (NVIDIA)
+		// e.g. "15.200.1062.1004"(AMD) "9.18.13.4788"(NVIDIA)
 		FString LocalInternalDriverVersion;
 		FWindowsPlatformMisc::QueryRegKey(HKEY_LOCAL_MACHINE, *Key, TEXT("DriverVersion"), LocalInternalDriverVersion);
 
-		// if there is no Catalyst_Version we use the internal version
+		// e.g. "Catalyst 15.7.1"(AMD) or "347.88"(NVIDIA)
 		FString LocalUserDriverVersion = LocalInternalDriverVersion;
 
-		// AMD only
-		// e.g. 15.7.1
-		FWindowsPlatformMisc::QueryRegKey(HKEY_LOCAL_MACHINE, *Key, TEXT("Catalyst_Version"), LocalUserDriverVersion);
+		if(LocalProviderName == TEXT("NVIDIA"))
+		{
+			// "9.18.13.4788" -> "347.88"
+			// the following code works with the current numbering scheme, if needd we have to update that
+			if(LocalUserDriverVersion.Left(6) == TEXT("9.18.1"))
+			{
+				// e.g. 3.4788
+				FString RightPart = LocalUserDriverVersion.RightChop(6);
+
+				if(RightPart.Len() == 6 && RightPart[1] == (TCHAR)'.')
+				{
+					// e.g. 347.88
+					LocalUserDriverVersion = RightPart.Left(1) + RightPart[2] + RightPart[3] + TEXT(".") + RightPart[4] + RightPart[5];
+				}
+			}
+		}
+		else 
+		{
+			// we assume LocalProviderName == TEXT("Advanced Micro Devices, Inc.")
+
+			// e.g. 15.7.1
+			if(FWindowsPlatformMisc::QueryRegKey(HKEY_LOCAL_MACHINE, *Key, TEXT("Catalyst_Version"), LocalUserDriverVersion))
+			{
+				LocalUserDriverVersion = FString(TEXT("Catalyst ")) + LocalUserDriverVersion;
+			}
+		}
 
 		// AMD and NVIDIA
 		// e.g. 3-13-2015
diff --git a/Engine/Source/Runtime/Core/Private/Windows/WindowsPlatformSurvey.cpp b/Engine/Source/Runtime/Core/Private/Windows/WindowsPlatformSurvey.cpp
--- a/Engine/Source/Runtime/Core/Private/Windows/WindowsPlatformSurvey.cpp
+++ b/Engine/Source/Runtime/Core/Private/Windows/WindowsPlatformSurvey.cpp
@@ -386,6 +386,8 @@ void FWindowsPlatformSurvey::TickSurveyHardware( FHardwareSurveyResults& OutResu
 
 	ISynthBenchmark::Get().Run(OutResults.SynthBenchmark, true, 5.f);
 
+	ISynthBenchmark::Get().GetRHIDisplay(OutResults.RHIAdpater);
+
 	// Get CPU speed
 	if (OutResults.CPUCount > 0)
 	{
diff --git a/Engine/Source/Runtime/Core/Private/Windows/WindowsWindow.cpp b/Engine/Source/Runtime/Core/Private/Windows/WindowsWindow.cpp
--- a/Engine/Source/Runtime/Core/Private/Windows/WindowsWindow.cpp
+++ b/Engine/Source/Runtime/Core/Private/Windows/WindowsWindow.cpp
@@ -530,17 +530,17 @@ void FWindowsWindow::Hide()
 /** Toggle native window between fullscreen and normal mode */
 void FWindowsWindow::SetWindowMode( EWindowMode::Type NewWindowMode )
 {
-	EWindowMode::Type InNewWindowMode = NewWindowMode;
-
-	if (InNewWindowMode == EWindowMode::WindowedMirror)
+	if (NewWindowMode == EWindowMode::WindowedMirror)
 	{
 		// treat WindowedMirror as a regular Windowed mode here
-		InNewWindowMode = EWindowMode::Windowed;
+		NewWindowMode = EWindowMode::Windowed;
 	}
 
-	if( InNewWindowMode != WindowMode )
+	if (NewWindowMode != WindowMode)
 	{
-		bool bTrueFullscreen = NewWindowMode == EWindowMode::Fullscreen;
+		WindowMode = NewWindowMode;
+
+		const bool bTrueFullscreen = NewWindowMode == EWindowMode::Fullscreen;
 
 		// Setup Win32 Flags to be used for Fullscreen mode
 		LONG WindowStyle = GetWindowLong(HWnd, GWL_STYLE);
@@ -574,7 +574,7 @@ void FWindowsWindow::SetWindowMode( EWindowMode::Type NewWindowMode )
 		}
 
 		// If we're not in fullscreen, make it so
-		if( NewWindowMode == EWindowMode::WindowedFullscreen || NewWindowMode == EWindowMode::Fullscreen)
+		if (NewWindowMode == EWindowMode::WindowedFullscreen || NewWindowMode == EWindowMode::Fullscreen)
 		{
 			::GetWindowPlacement(HWnd, &PreFullscreenWindowPlacement);
 
@@ -635,8 +635,6 @@ void FWindowsWindow::SetWindowMode( EWindowMode::Type NewWindowMode )
 
 			::SetWindowPlacement(HWnd, &PreFullscreenWindowPlacement);
 		}
-
-		WindowMode = NewWindowMode;
 	}
 }
 
diff --git a/Engine/Source/Runtime/Core/Public/GenericPlatform/GenericPlatformMemory.h b/Engine/Source/Runtime/Core/Public/GenericPlatform/GenericPlatformMemory.h
--- a/Engine/Source/Runtime/Core/Public/GenericPlatform/GenericPlatformMemory.h
+++ b/Engine/Source/Runtime/Core/Public/GenericPlatform/GenericPlatformMemory.h
@@ -229,10 +229,11 @@ struct CORE_API FGenericPlatformMemory
 	 *
 	 * @param Ptr Address to the starting page of the region of pages whose access protection attributes are to be changed.
 	 * @param Size The size of the region whose access protection attributes are to be changed, in bytes.
-	 * @param ProtectMode The OS dependent memory protection option.
+	 * @param bCanRead Can the memory be read.
+	 * @param bCanWrite Can the memory be written to.
 	 * @return True if the specified pages' protection mode was changed.
 	 */
-	static bool PageProtect( void* const Ptr, const SIZE_T Size, const uint32 ProtectMode );
+	static bool PageProtect(void* const Ptr, const SIZE_T Size, const bool bCanRead, const bool bCanWrite);
 
 	/**
 	 * Allocates pages from the OS.
diff --git a/Engine/Source/Runtime/Core/Public/GenericPlatform/GenericPlatformSurvey.h b/Engine/Source/Runtime/Core/Public/GenericPlatform/GenericPlatformSurvey.h
--- a/Engine/Source/Runtime/Core/Public/GenericPlatform/GenericPlatformSurvey.h
+++ b/Engine/Source/Runtime/Core/Public/GenericPlatform/GenericPlatformSurvey.h
@@ -185,6 +185,20 @@ struct FHardwareDisplay
 };
 
 
+struct FGPUAdpater	
+{
+	static const uint32 MaxStringLength = 260;
+
+	TCHAR AdapterName[MaxStringLength];
+	//
+	TCHAR AdapterInternalDriverVersion[MaxStringLength];
+	// extra data when RHI starts up, AMD has the Catalyst Driver version which is more useful for the user, otherwise it's the same as GRHIAdapterInternalDriverVersion
+	TCHAR AdapterUserDriverVersion[MaxStringLength];
+	// extra data when RHI starts up
+	TCHAR AdapterDriverDate[MaxStringLength];
+};
+
+
 struct FHardwareSurveyResults
 {
 	static const int32 MaxDisplayCount = 8; 
@@ -217,6 +231,8 @@ struct FHardwareSurveyResults
 
 	uint32 DisplayCount;
 	FHardwareDisplay Displays[MaxDisplayCount];
+	// RHI chosen adapter, only if RHI started on Windows DX11
+	FGPUAdpater RHIAdpater;
 
 	uint32 ErrorCount;
 	TCHAR LastSurveyError[MaxStringLength];
diff --git a/Engine/Source/Runtime/Core/Public/HAL/IPlatformFileCachedWrapper.h b/Engine/Source/Runtime/Core/Public/HAL/IPlatformFileCachedWrapper.h
--- a/Engine/Source/Runtime/Core/Public/HAL/IPlatformFileCachedWrapper.h
+++ b/Engine/Source/Runtime/Core/Public/HAL/IPlatformFileCachedWrapper.h
@@ -229,8 +229,8 @@ public:
 	}
 	virtual bool ShouldBeUsed(IPlatformFile* Inner, const TCHAR* CmdLine) const override
 	{
-		// default to false on Windows since FAsyncBufferedFileReaderWindows already buffers the data
-		bool bResult = !PLATFORM_WINDOWS && FPlatformProperties::RequiresCookedData();
+		// Default to false on platforms that already do platform file level caching
+		bool bResult = !PLATFORM_PS4 && !PLATFORM_WINDOWS && FPlatformProperties::RequiresCookedData();
 
 		// Allow a choice between shorter load times or less memory on desktop platforms.
 		// Note: this cannot be in config since they aren't read at that point.
diff --git a/Engine/Source/Runtime/Core/Public/HAL/MallocStomp.h b/Engine/Source/Runtime/Core/Public/HAL/MallocStomp.h
--- a/Engine/Source/Runtime/Core/Public/HAL/MallocStomp.h
+++ b/Engine/Source/Runtime/Core/Public/HAL/MallocStomp.h
@@ -2,10 +2,15 @@
 
 #pragma once
 
-//#define USE_MALLOC_STOMP 1
+#if !defined(WITH_EDITOR)
+#error "WITH_EDITOR must be defined"
+#endif
 
-#if USE_MALLOC_STOMP
+#if !WITH_EDITOR && PLATFORM_DESKTOP
+	#define USE_MALLOC_STOMP 0
+#endif
 
+#if USE_MALLOC_STOMP
 /**
  * Stomp memory allocator. It helps find the following errors:
  * - Read or writes off the end of an allocation.
diff --git a/Engine/Source/Runtime/Core/Public/HAL/Platform.h b/Engine/Source/Runtime/Core/Public/HAL/Platform.h
--- a/Engine/Source/Runtime/Core/Public/HAL/Platform.h
+++ b/Engine/Source/Runtime/Core/Public/HAL/Platform.h
@@ -294,10 +294,6 @@
 	#define PLATFORM_RHITHREAD_DEFAULT_BYPASS					1
 #endif
 
-#ifndef PLATFORM_CAN_TOGGLE_RHITHREAD_IN_SHIPPING
-	#define PLATFORM_CAN_TOGGLE_RHITHREAD_IN_SHIPPING			0
-#endif
-
 // deprecated, do not use
 #define PLATFORM_HAS_THREADSAFE_RHIGetRenderQueryResult	#
 #define PLATFORM_SUPPORTS_RHI_THREAD #
diff --git a/Engine/Source/Runtime/Core/Public/HAL/UnrealMemory.h b/Engine/Source/Runtime/Core/Public/HAL/UnrealMemory.h
--- a/Engine/Source/Runtime/Core/Public/HAL/UnrealMemory.h
+++ b/Engine/Source/Runtime/Core/Public/HAL/UnrealMemory.h
@@ -139,5 +139,10 @@ struct CORE_API FMemory
 	 * and hopefully freeing some pointers will trigger a crash.
 	 */
 	static void TestMemory();
+	/**
+	* Called once main is started and we have -purgatorymallocproxy.
+	* This uses the purgatory malloc proxy to check if things are writing to stale pointers.
+	*/
+	static void EnablePurgatoryTests();
 };
 
diff --git a/Engine/Source/Runtime/Core/Public/Internationalization/Text.h b/Engine/Source/Runtime/Core/Public/Internationalization/Text.h
--- a/Engine/Source/Runtime/Core/Public/Internationalization/Text.h
+++ b/Engine/Source/Runtime/Core/Public/Internationalization/Text.h
@@ -171,21 +171,37 @@ public:
 	static FText AsNumber(uint64 Val,	const FNumberFormattingOptions* const Options = NULL, const FCulturePtr& TargetCulture = NULL);
 	static FText AsNumber(long Val,		const FNumberFormattingOptions* const Options = NULL, const FCulturePtr& TargetCulture = NULL);
 
-	/**
-	 * Generate an FText that represents the passed number as currency in the current culture
-	 */
+	DEPRECATED(4.11, "Use FText::AsCurrencyBase instead")
 	static FText AsCurrency(float Val,  const FString& CurrencyCode = FString(), const FNumberFormattingOptions* const Options = NULL, const FCulturePtr& TargetCulture = NULL);
+	DEPRECATED(4.11, "Use FText::AsCurrencyBase instead")
 	static FText AsCurrency(double Val, const FString& CurrencyCode = FString(), const FNumberFormattingOptions* const Options = NULL, const FCulturePtr& TargetCulture = NULL);
+	DEPRECATED(4.11, "Use FText::AsCurrencyBase instead")
 	static FText AsCurrency(int8 Val,   const FString& CurrencyCode = FString(), const FNumberFormattingOptions* const Options = NULL, const FCulturePtr& TargetCulture = NULL);
+	DEPRECATED(4.11, "Use FText::AsCurrencyBase instead")
 	static FText AsCurrency(int16 Val,  const FString& CurrencyCode = FString(), const FNumberFormattingOptions* const Options = NULL, const FCulturePtr& TargetCulture = NULL);
+	DEPRECATED(4.11, "Use FText::AsCurrencyBase instead")
 	static FText AsCurrency(int32 Val,  const FString& CurrencyCode = FString(), const FNumberFormattingOptions* const Options = NULL, const FCulturePtr& TargetCulture = NULL);
+	DEPRECATED(4.11, "Use FText::AsCurrencyBase instead")
 	static FText AsCurrency(int64 Val,  const FString& CurrencyCode = FString(), const FNumberFormattingOptions* const Options = NULL, const FCulturePtr& TargetCulture = NULL);
+	DEPRECATED(4.11, "Use FText::AsCurrencyBase instead")
 	static FText AsCurrency(uint8 Val,  const FString& CurrencyCode = FString(), const FNumberFormattingOptions* const Options = NULL, const FCulturePtr& TargetCulture = NULL);
+	DEPRECATED(4.11, "Use FText::AsCurrencyBase instead")
 	static FText AsCurrency(uint16 Val, const FString& CurrencyCode = FString(), const FNumberFormattingOptions* const Options = NULL, const FCulturePtr& TargetCulture = NULL);
+	DEPRECATED(4.11, "Use FText::AsCurrencyBase instead")
 	static FText AsCurrency(uint32 Val, const FString& CurrencyCode = FString(), const FNumberFormattingOptions* const Options = NULL, const FCulturePtr& TargetCulture = NULL);
+	DEPRECATED(4.11, "Use FText::AsCurrencyBase instead")
 	static FText AsCurrency(uint64 Val, const FString& CurrencyCode = FString(), const FNumberFormattingOptions* const Options = NULL, const FCulturePtr& TargetCulture = NULL);
+	DEPRECATED(4.11, "Use FText::AsCurrencyBase instead")
 	static FText AsCurrency(long Val,   const FString& CurrencyCode = FString(), const FNumberFormattingOptions* const Options = NULL, const FCulturePtr& TargetCulture = NULL);
 
+	/**
+	 * Generate an FText that represents the passed number as currency in the current culture.
+	 * BaseVal is specified in the smallest fractional value of the currency and will be converted for formatting according to the selected culture.
+	 * Keep in mind the CurrencyCode is completely independent of the culture it's displayed in (and they do not imply one another).
+	 * For example: FText::AsCurrencyBase(650, TEXT("EUR")); would return an FText of "€6.50" in most English cultures (en_US/en_UK) and "6,50€" in Spanish (es_ES).
+	 */
+	static FText AsCurrencyBase(int64 BaseVal, const FString& CurrencyCode, const FCulturePtr& TargetCulture = NULL);
+
 	/**
 	 * Generate an FText that represents the passed number as a percentage in the current culture
 	 */
diff --git a/Engine/Source/Runtime/Core/Public/Linux/LinuxPlatformMemory.h b/Engine/Source/Runtime/Core/Public/Linux/LinuxPlatformMemory.h
--- a/Engine/Source/Runtime/Core/Public/Linux/LinuxPlatformMemory.h
+++ b/Engine/Source/Runtime/Core/Public/Linux/LinuxPlatformMemory.h
@@ -49,7 +49,7 @@ struct CORE_API FLinuxPlatformMemory : public FGenericPlatformMemory
 	static class FMalloc* BaseAllocator();
 	static FPlatformMemoryStats GetStats();
 	static const FPlatformMemoryConstants& GetConstants();
-	static bool PageProtect( void* const Ptr, const SIZE_T Size, const uint32 ProtectMode );
+	static bool PageProtect(void* const Ptr, const SIZE_T Size, const bool bCanRead, const bool bCanWrite);
 	static void* BinnedAllocFromOS( SIZE_T Size );
 	static void BinnedFreeToOS( void* Ptr );
 	static FSharedMemoryRegion * MapNamedSharedMemoryRegion(const FString& InName, bool bCreate, uint32 AccessMode, SIZE_T Size);
diff --git a/Engine/Source/Runtime/Core/Public/Mac/MacPlatformMemory.h b/Engine/Source/Runtime/Core/Public/Mac/MacPlatformMemory.h
--- a/Engine/Source/Runtime/Core/Public/Mac/MacPlatformMemory.h
+++ b/Engine/Source/Runtime/Core/Public/Mac/MacPlatformMemory.h
@@ -24,7 +24,7 @@ struct CORE_API FMacPlatformMemory : public FGenericPlatformMemory
 	static FPlatformMemoryStats GetStats();
 	static const FPlatformMemoryConstants& GetConstants();
 	static FMalloc* BaseAllocator();
-	static bool PageProtect( void* const Ptr, const SIZE_T Size, const uint32 ProtectMode );
+	static bool PageProtect(void* const Ptr, const SIZE_T Size, const bool bCanRead, const bool bCanWrite);
 	static void* BinnedAllocFromOS( SIZE_T Size );
 	static void BinnedFreeToOS( void* Ptr );
 	//~ End FGenericPlatformMemory Interface
diff --git a/Engine/Source/Runtime/Core/Public/Windows/WIndowsPlatform.h b/Engine/Source/Runtime/Core/Public/Windows/WIndowsPlatform.h
--- a/Engine/Source/Runtime/Core/Public/Windows/WIndowsPlatform.h
+++ b/Engine/Source/Runtime/Core/Public/Windows/WIndowsPlatform.h
@@ -74,8 +74,7 @@ typedef FWindowsPlatformTypes FPlatformTypes;
 #endif
 #define PLATFORM_COMPILER_HAS_TCHAR_WMAIN					1
 
-#define PLATFORM_RHITHREAD_DEFAULT_BYPASS					1
-#define PLATFORM_CAN_TOGGLE_RHITHREAD_IN_SHIPPING			1
+#define PLATFORM_RHITHREAD_DEFAULT_BYPASS					WITH_EDITOR
 
 #define PLATFORM_SUPPORTS_STACK_SYMBOLS						1
 
diff --git a/Engine/Source/Runtime/CoreUObject/Private/Misc/RedirectCollector.cpp b/Engine/Source/Runtime/CoreUObject/Private/Misc/RedirectCollector.cpp
--- a/Engine/Source/Runtime/CoreUObject/Private/Misc/RedirectCollector.cpp
+++ b/Engine/Source/Runtime/CoreUObject/Private/Misc/RedirectCollector.cpp
@@ -142,7 +142,8 @@ void FRedirectCollector::ResolveStringAssetReference(FString FilterPackage)
 			}
 			else
 			{
-				UE_LOG(LogRedirectors, Warning, TEXT("String Asset Reference '%s' was not found!"), *ToLoad);
+				const FString Referencer = RefFilenameAndProperty.Property.Len() ? RefFilenameAndProperty.Property : TEXT("Unknown");
+				UE_LOG(LogRedirectors, Warning, TEXT("String Asset Reference '%s' was not found! (Referencer '%s')"), *ToLoad, *Referencer);
 			}
 		}
 	}
diff --git a/Engine/Source/Runtime/CoreUObject/Private/Serialization/AsyncLoading.cpp b/Engine/Source/Runtime/CoreUObject/Private/Serialization/AsyncLoading.cpp
--- a/Engine/Source/Runtime/CoreUObject/Private/Serialization/AsyncLoading.cpp
+++ b/Engine/Source/Runtime/CoreUObject/Private/Serialization/AsyncLoading.cpp
@@ -15,6 +15,7 @@
 #include "AssetRegistryInterface.h"
 #include "BlueprintSupport.h"
 
+#define FIND_MEMORY_STOMPS (1 && PLATFORM_WINDOWS && !WITH_EDITORONLY_DATA)
 
 /*-----------------------------------------------------------------------------
 	Async loading stats.
@@ -135,7 +136,9 @@ public:
 		const EInternalObjectFlags AsyncFlags = EInternalObjectFlags::Async | EInternalObjectFlags::AsyncLoading;
 		for (UObject* Obj : ReferencedObjects)
 		{
-			check(Obj);
+			// Temporary fatal messages instead of checks to find the cause for a one-time crash in shipping config
+			UE_CLOG(Obj == nullptr, LogStreaming, Fatal, TEXT("NULL object in Async Objects Referencer"));
+			UE_CLOG(!Obj->IsValidLowLevelFast(), LogStreaming, Fatal, TEXT("Invalid object in Async Objects Referencer"));
 			Obj->AtomicallyClearInternalFlags(AsyncFlags);
 			check(!Obj->HasAnyInternalFlags(AsyncFlags))
 		}
@@ -2285,6 +2288,44 @@ void ResumeAsyncLoadingInternal()
 	FArchiveAsync.
 ----------------------------------------------------------------------------*/
 
+static uint8* MallocAsyncBuffer(const SIZE_T Size, SIZE_T& OutAllocatedSize)
+{
+	uint8* Result = nullptr;
+#if FIND_MEMORY_STOMPS
+	const SIZE_T PageSize = FPlatformMemory::GetConstants().PageSize;
+	const SIZE_T Alignment = PageSize;
+	const SIZE_T AlignedSize = (Size + Alignment - 1U) & -static_cast<int32>(Alignment);
+	const SIZE_T AllocFullPageSize = AlignedSize + (PageSize - 1) & ~(PageSize - 1U);
+	check(AllocFullPageSize >= Size);
+	OutAllocatedSize = AllocFullPageSize;
+#if PLATFORM_LINUX || PLATFORM_MAC
+	Result = (uint8*)mmap(nullptr, AllocFullPageSize, PROT_READ | PROT_WRITE, MAP_PRIVATE | MAP_ANON, -1, 0);
+#else
+	Result = (uint8*)FPlatformMemory::BinnedAllocFromOS(AllocFullPageSize);
+#endif // PLATFORM_LINUX || PLATFORM_MAC
+#else
+	OutAllocatedSize = Size;
+	Result = (uint8*)FMemory::Malloc(Size);
+#endif // FIND_MEMORY_STOMPS
+	return Result;
+}
+
+static void FreeAsyncBuffer(uint8* Buffer, const SIZE_T AllocatedSize)
+{
+	if (Buffer)
+	{
+#if FIND_MEMORY_STOMPS
+#if PLATFORM_LINUX || PLATFORM_MAC
+		munmap(Buffer, AllocatedSize);
+#else
+		FPlatformMemory::BinnedFreeToOS(Buffer);
+#endif // PLATFORM_LINUX || PLATFORM_MAC
+#else
+		FMemory::Free(Buffer);
+#endif // FIND_MEMORY_STOMPS
+	}
+}
+
 /**
  * Constructor, initializing all member variables.
  */
@@ -2308,10 +2349,14 @@ FArchiveAsync::FArchiveAsync( const TCHAR* InFileName )
 	PrecacheStartPos[CURRENT]	= 0;
 	PrecacheEndPos[CURRENT]		= 0;
 	PrecacheBuffer[CURRENT]		= nullptr;
+	PrecacheBufferSize[CURRENT] = 0;
+	PrecacheBufferProtected[CURRENT] = false;
 
 	PrecacheStartPos[NEXT]		= 0;
 	PrecacheEndPos[NEXT]		= 0;
 	PrecacheBuffer[NEXT]		= nullptr;
+	PrecacheBufferSize[NEXT] = 0;
+	PrecacheBufferProtected[NEXT] = false;
 
 	// Relies on default constructor initializing to 0.
 	check( PrecacheReadStatus[CURRENT].GetValue() == 0 );
@@ -2364,16 +2409,20 @@ void FArchiveAsync::FlushCache()
 
 	// Invalidate any precached data and free memory for current buffer.
 	Delta += PrecacheEndPos[CURRENT] - PrecacheStartPos[CURRENT];
-	FMemory::Free( PrecacheBuffer[CURRENT] );
+	FreeAsyncBuffer(PrecacheBuffer[CURRENT], PrecacheBufferSize[CURRENT]);
 	PrecacheBuffer[CURRENT]		= nullptr;
 	PrecacheStartPos[CURRENT]	= 0;
 	PrecacheEndPos[CURRENT]		= 0;
-	
+	PrecacheBufferSize[CURRENT] = 0;
+	PrecacheBufferProtected[CURRENT] = false;
+
 	// Invalidate any precached data and free memory for next buffer.
-	FMemory::Free( PrecacheBuffer[NEXT] );
+	FreeAsyncBuffer(PrecacheBuffer[NEXT], PrecacheBufferSize[NEXT]);
 	PrecacheBuffer[NEXT]		= nullptr;
 	PrecacheStartPos[NEXT]		= 0;
 	PrecacheEndPos[NEXT]		= 0;
+	PrecacheBufferSize[NEXT] = 0;
+	PrecacheBufferProtected[NEXT] = false;
 
 	Delta += PrecacheEndPos[NEXT] - PrecacheStartPos[NEXT];
 	DEC_DWORD_STAT_BY(STAT_StreamingAllocSize, Delta);
@@ -2440,20 +2489,24 @@ bool FArchiveAsync::SetCompressionMap( TArray<FCompressedChunk>* InCompressedChu
  */
 void FArchiveAsync::BufferSwitcheroo()
 {
-	check( PrecacheReadStatus[CURRENT].GetValue() == 0 );
-	check( PrecacheReadStatus[NEXT].GetValue() == 0 );
+	check(PrecacheReadStatus[CURRENT].GetValue() == 0);
+	check(PrecacheReadStatus[NEXT].GetValue() == 0);
 
 	// Switcheroo.
 	DEC_DWORD_STAT_BY(STAT_StreamingAllocSize, PrecacheEndPos[CURRENT] - PrecacheStartPos[CURRENT]);
-	FMemory::Free( PrecacheBuffer[CURRENT] );
-	PrecacheBuffer[CURRENT]		= PrecacheBuffer[NEXT];
-	PrecacheStartPos[CURRENT]	= PrecacheStartPos[NEXT];
-	PrecacheEndPos[CURRENT]		= PrecacheEndPos[NEXT];
+	FreeAsyncBuffer(PrecacheBuffer[CURRENT], PrecacheBufferSize[CURRENT]);
+	PrecacheBuffer[CURRENT] = PrecacheBuffer[NEXT];
+	PrecacheStartPos[CURRENT] = PrecacheStartPos[NEXT];
+	PrecacheEndPos[CURRENT] = PrecacheEndPos[NEXT];
+	PrecacheBufferSize[CURRENT] = PrecacheBufferSize[NEXT];
+	PrecacheBufferProtected[CURRENT] = PrecacheBufferProtected[NEXT];
 
 	// Next buffer is unused/ free.
 	PrecacheBuffer[NEXT]		= nullptr;
 	PrecacheStartPos[NEXT]		= 0;
 	PrecacheEndPos[NEXT]		= 0;
+	PrecacheBufferSize[NEXT] = 0;
+	PrecacheBufferProtected[NEXT] = 0;
 }
 
 /**
@@ -2525,8 +2578,9 @@ void FArchiveAsync::PrecacheCompressedChunk( int64 ChunkIndex, int64 BufferIndex
 
 	// In theory we could use FMemory::Realloc if it had a way to signal that we don't want to copy
 	// the data (implicit realloc behavior).
-	FMemory::Free( PrecacheBuffer[BufferIndex] );
-	PrecacheBuffer[BufferIndex]		= (uint8*) FMemory::Malloc( PrecacheEndPos[BufferIndex] - PrecacheStartPos[BufferIndex] );
+	FreeAsyncBuffer(PrecacheBuffer[BufferIndex], PrecacheBufferSize[BufferIndex]);
+	PrecacheBufferProtected[BufferIndex] = false;
+	PrecacheBuffer[BufferIndex] = MallocAsyncBuffer(PrecacheEndPos[BufferIndex] - PrecacheStartPos[BufferIndex], PrecacheBufferSize[BufferIndex]);
 	{
 		INC_DWORD_STAT_BY(STAT_StreamingAllocSize, PrecacheEndPos[BufferIndex] - PrecacheStartPos[BufferIndex]);
 	}
@@ -2636,9 +2690,9 @@ bool FArchiveAsync::Precache( int64 RequestOffset, int64 RequestSize )
 			PrecacheEndPos[CURRENT]		= FMath::Min( PrecacheEndPos[CURRENT], FileSize );
 			// In theory we could use FMemory::Realloc if it had a way to signal that we don't want to copy
 			// the data (implicit realloc behavior).
-			FMemory::Free( PrecacheBuffer[CURRENT] );
-
-			PrecacheBuffer[CURRENT]		= (uint8*) FMemory::Malloc( PrecacheEndPos[CURRENT] - PrecacheStartPos[CURRENT] );
+			FreeAsyncBuffer(PrecacheBuffer[CURRENT], PrecacheBufferSize[CURRENT]);
+			PrecacheBufferProtected[CURRENT] = false;
+			PrecacheBuffer[CURRENT] = MallocAsyncBuffer(PrecacheEndPos[CURRENT] - PrecacheStartPos[CURRENT], PrecacheBufferSize[CURRENT]);
 			{
 				INC_DWORD_STAT_BY(STAT_StreamingAllocSize, PrecacheEndPos[CURRENT] - PrecacheStartPos[CURRENT]);
 			}
@@ -2666,9 +2720,26 @@ bool FArchiveAsync::Precache( int64 RequestOffset, int64 RequestSize )
  * @param	Count	Number of bytes to read
  */
 void FArchiveAsync::Serialize(void* Data, int64 Count)
-{
-	// Ensure we aren't reading beyond the end of the file
-	checkf( CurrentPos + Count <= TotalSize(), TEXT("Seeked past end of file %s (%lld / %lld)"), *FileName, CurrentPos + Count, TotalSize() );
+{	
+#if PLATFORM_DESKTOP
+	// Show a message box indicating, possible, corrupt data (desktop platforms only)
+	if (CurrentPos + Count > TotalSize())
+	{
+		FText ErrorMessage, ErrorCaption;
+		GConfig->GetText(TEXT("/Script/Engine.Engine"),
+			  			 TEXT("SerializationOutOfBoundsErrorMessage"),
+						 ErrorMessage,
+						 GEngineIni);
+		GConfig->GetText(TEXT("/Script/Engine.Engine"),
+			TEXT("SerializationOutOfBoundsErrorMessageCaption"),
+			ErrorCaption,
+			GEngineIni);
+
+		FPlatformMisc::MessageBoxExt(EAppMsgType::Ok, *ErrorMessage.ToString(), *ErrorCaption.ToString());
+	}
+#endif
+	// Ensure we aren't reading beyond the end of the file	
+	checkf( CurrentPos + Count <= TotalSize(), TEXT("Seeked past end of file %s (%lld / %lld)"), *FileName, CurrentPos + Count, TotalSize() );	
 
 #if LOOKING_FOR_PERF_ISSUES
 	uint32 StartCycles = 0;
@@ -2735,6 +2806,16 @@ void FArchiveAsync::Serialize(void* Data, int64 Count)
 			FPlatformProcess::SleepNoStats(0.0f);
 		} while (PrecacheReadStatus[CURRENT].GetValue());
 	}
+#if FIND_MEMORY_STOMPS
+	if (!PrecacheBufferProtected[CURRENT])
+	{
+		if (!FPlatformMemory::PageProtect(PrecacheBuffer[CURRENT], PrecacheBufferSize[CURRENT], true, false))
+		{
+			UE_LOG(LogStreaming, Warning, TEXT("Unable to write protect async buffer %d for %s"), int32(CURRENT), *FileName);
+		}
+		PrecacheBufferProtected[CURRENT] = true;
+	}
+#endif // FIND_MEMORY_STOMPS
 
 	// Update stats if we were blocked.
 #if LOOKING_FOR_PERF_ISSUES
diff --git a/Engine/Source/Runtime/CoreUObject/Private/Serialization/BulkData.cpp b/Engine/Source/Runtime/CoreUObject/Private/Serialization/BulkData.cpp
--- a/Engine/Source/Runtime/CoreUObject/Private/Serialization/BulkData.cpp
+++ b/Engine/Source/Runtime/CoreUObject/Private/Serialization/BulkData.cpp
@@ -828,7 +828,6 @@ void FUntypedBulkData::Serialize( FArchive& Ar, UObject* Owner, int32 Idx )
 				else if (BulkDataFlags & BULKDATA_PayloadInSeperateFile)
 				{
 					Filename = FPaths::ChangeExtension(Filename, TEXT(".ubulk"));
-					checkSlow(FPaths::FileExists(Filename));
 				}
 			}
 			// Serialize the bulk data right away.
diff --git a/Engine/Source/Runtime/CoreUObject/Private/Serialization/PropertyLocalizationDataGathering.cpp b/Engine/Source/Runtime/CoreUObject/Private/Serialization/PropertyLocalizationDataGathering.cpp
--- a/Engine/Source/Runtime/CoreUObject/Private/Serialization/PropertyLocalizationDataGathering.cpp
+++ b/Engine/Source/Runtime/CoreUObject/Private/Serialization/PropertyLocalizationDataGathering.cpp
@@ -27,32 +27,43 @@ FPropertyLocalizationDataGatherer::FPropertyLocalizationDataGatherer(TArray<FGat
 	// Iterate over each root object in the package
 	for (const UObject* Object : RootObjectsInPackage)
 	{
-		// See if we have a custom handler for this type
-		FLocalizationDataGatheringCallback* CustomCallback = nullptr;
-		for (const UClass* Class = Object->GetClass(); Class != nullptr; Class = Class->GetSuperClass())
-		{
-			CustomCallback = GetTypeSpecificLocalizationDataGatheringCallbacks().Find(Class);
-			if (CustomCallback)
-			{
-				break;
-			}
-		}
+		GatherLocalizationDataFromObjectWithCallbacks(Object, EPropertyLocalizationGathererTextFlags::None);
+	}
+}
 
+void FPropertyLocalizationDataGatherer::GatherLocalizationDataFromObjectWithCallbacks(const UObject* Object, const EPropertyLocalizationGathererTextFlags GatherTextFlags)
+{
+	// See if we have a custom handler for this type
+	FLocalizationDataGatheringCallback* CustomCallback = nullptr;
+	for (const UClass* Class = Object->GetClass(); Class != nullptr; Class = Class->GetSuperClass())
+	{
+		CustomCallback = GetTypeSpecificLocalizationDataGatheringCallbacks().Find(Class);
 		if (CustomCallback)
 		{
-			(*CustomCallback)(Object, *this, EPropertyLocalizationGathererTextFlags::None);
-		}
-		else
-		{
-			GatherLocalizationDataFromObject(Object, EPropertyLocalizationGathererTextFlags::None);
+			break;
 		}
 	}
+
+	if (CustomCallback)
+	{
+		(*CustomCallback)(Object, *this, GatherTextFlags);
+	}
+	else
+	{
+		GatherLocalizationDataFromObject(Object, GatherTextFlags);
+	}
 }
 
 void FPropertyLocalizationDataGatherer::GatherLocalizationDataFromObject(const UObject* Object, const EPropertyLocalizationGathererTextFlags GatherTextFlags)
 {
 	checkf(IsObjectValidForGather(Object), TEXT("Cannot gather for objects outside of the current package! Package: '%s'. Object: '%s'."), *Package->GetFullName(), *Object->GetFullName());
 
+	if (Object->HasAnyFlags(RF_Transient))
+	{
+		// Transient objects aren't saved, so skip them as part of the gather
+		return;
+	}
+
 	// Skip objects that we've already processed to avoid repeated work and cyclic chains
 	{
 		FObjectAndGatherFlags Key(Object, GatherTextFlags);
@@ -87,7 +98,7 @@ void FPropertyLocalizationDataGatherer::GatherLocalizationDataFromObject(const U
 
 		for (const UObject* ChildObject : ChildObjects)
 		{
-			GatherLocalizationDataFromObject(ChildObject, GatherTextFlags);
+			GatherLocalizationDataFromObjectWithCallbacks(ChildObject, GatherTextFlags);
 		}
 	}
 }
@@ -142,6 +153,12 @@ void FPropertyLocalizationDataGatherer::GatherLocalizationDataFromStructFields(c
 
 void FPropertyLocalizationDataGatherer::GatherLocalizationDataFromChildTextProperties(const FString& PathToParent, const UProperty* const Property, const void* const ValueAddress, const EPropertyLocalizationGathererTextFlags GatherTextFlags)
 {
+	if (Property->HasAnyPropertyFlags(CPF_Transient))
+	{
+		// Transient properties aren't saved, so skip them as part of the gather
+		return;
+	}
+
 	const UTextProperty* const TextProperty = Cast<const UTextProperty>(Property);
 	const UArrayProperty* const ArrayProperty = Cast<const UArrayProperty>(Property);
 	const UMapProperty* const MapProperty = Cast<const UMapProperty>(Property);
@@ -204,7 +221,7 @@ void FPropertyLocalizationDataGatherer::GatherLocalizationDataFromChildTextPrope
 				const UObject* InnerObject = ObjectProperty->GetObjectPropertyValue(ElementValueAddress);
 				if (InnerObject && IsObjectValidForGather(InnerObject))
 				{
-					GatherLocalizationDataFromObject(InnerObject, ChildPropertyGatherTextFlags);
+					GatherLocalizationDataFromObjectWithCallbacks(InnerObject, ChildPropertyGatherTextFlags);
 				}
 			}
 		}
