Copying //UE4/Dev-Editor to Dev-Main (//UE4/Dev-Main)

diff --git a/Engine/Source/Runtime/Core/Private/Misc/ConfigCacheIni.cpp b/Engine/Source/Runtime/Core/Private/Misc/ConfigCacheIni.cpp
--- a/Engine/Source/Runtime/Core/Private/Misc/ConfigCacheIni.cpp
+++ b/Engine/Source/Runtime/Core/Private/Misc/ConfigCacheIni.cpp
@@ -50,85 +50,6 @@ namespace
 #endif
 }
 
-class FTextFriendHelper
-{
-	friend class FConfigFile;
-	friend class FConfigCacheIni;
-
-	static FString AsString( const FText& Text )
-	{
-		if( Text.IsTransient() )
-		{
-			UE_LOG( LogConfig, Warning, TEXT( "FTextFriendHelper::AsString() Transient FText") );
-			return FString(TEXT("Error: Transient FText"));
-		}
-
-		FString Str;
-		if( FTextInspector::GetSourceString(Text) )
-		{
-			FString SourceString( *FTextInspector::GetSourceString(Text) );
-			for( auto Iter( SourceString.CreateConstIterator()); Iter; ++Iter )
-			{
-				const TCHAR Ch = *Iter;
-				if( Ch == '\"' )
-				{
-					Str += TEXT("\\\"");
-				}
-				else if( Ch == '\t' )
-				{
-					Str += TEXT("\\t");
-				}
-				else if( Ch == '\r' )
-				{
-					Str += TEXT("\\r");
-				}
-				else if( Ch == '\n' )
-				{
-					Str += TEXT("\\n");
-				}
-				else if( Ch == '\\' )
-				{
-					Str += TEXT("\\\\");
-				}
-				else
-				{
-					Str += Ch;
-				}
-			}
-		}
-
-		//this prevents our source code text gatherer from trying to gather the following messages
-#define LOC_DEFINE_REGION
-		if( Text.IsCultureInvariant() )
-		{
-			return FString::Printf( TEXT( "NSLOCTEXT(\"\",\"\",\"%s\")" ), *Str );
-		}
-		else
-		{
-			FTextDisplayStringRef DisplayString = FTextInspector::GetSharedDisplayString(Text);
-
-			FString Namespace;
-			FString Key;
-			const bool FoundNamespaceAndKey = FTextLocalizationManager::Get().FindNamespaceAndKeyFromDisplayString(DisplayString, Namespace, Key);
-
-			// If this has no namespace or key, attempt to give it a GUID for a key and register it.
-			if (!FoundNamespaceAndKey && GIsEditor)
-			{
-				Key = FGuid::NewGuid().ToString();
-				if (!FTextLocalizationManager::Get().AddDisplayString(DisplayString, Namespace, Key))
-				{
-					// Could not add display string, reset namespace and key.
-					Namespace.Empty();
-					Key.Empty();
-				}
-			}
-
-			return FString::Printf( TEXT( "NSLOCTEXT(\"%s\",\"%s\",\"%s\")" ), *Namespace, *Key, *Str );
-		}
-#undef LOC_DEFINE_REGION
-	}
-};
-
 /*-----------------------------------------------------------------------------
 	FConfigSection
 -----------------------------------------------------------------------------*/
@@ -1132,20 +1053,8 @@ bool FConfigFile::GetString( const TCHAR* Section, const TCHAR* Key, FString& Va
 	{
 		return false;
 	}
-
-	//this prevents our source code text gatherer from trying to gather the following messages
-#define LOC_DEFINE_REGION
-	if( FCString::Strstr( **PairString, TEXT("LOCTEXT") ) )
-	{
-		UE_LOG( LogConfig, Warning, TEXT( "FConfigFile::GetString( %s, %s ) contains LOCTEXT"), Section, Key );
-		return false;
-	}
-	else
-	{
-		Value = **PairString;
-		return true;
-	}
-#undef LOC_DEFINE_REGION
+	Value = **PairString;
+	return true;
 }
 
 bool FConfigFile::GetText( const TCHAR* Section, const TCHAR* Key, FText& Value ) const
@@ -1160,11 +1069,7 @@ bool FConfigFile::GetText( const TCHAR* Section, const TCHAR* Key, FText& Value
 	{
 		return false;
 	}
-	if( !FParse::Text( **PairString, Value, Section ) )
-	{
-		Value = FText::FromString( **PairString );
-	}
-	return true;
+	return FTextStringHelper::ReadFromString( **PairString, Value, Section );
 }
 
 bool FConfigFile::GetInt64( const TCHAR* Section, const TCHAR* Key, int64& Value ) const
@@ -1209,9 +1114,10 @@ void FConfigFile::SetText( const TCHAR* Section, const TCHAR* Key, const FText&
 		Sec = &Add( Section, FConfigSection() );
 	}
 
-	FString* Str = Sec->Find( Key );
-	const FString StrValue = FTextFriendHelper::AsString( Value );
+	FString StrValue;
+	FTextStringHelper::WriteToString(StrValue, Value);
 
+	FString* Str = Sec->Find( Key );
 	if( Str == NULL )
 	{
 		Sec->Add( Key, StrValue );
@@ -1691,20 +1597,8 @@ bool FConfigCacheIni::GetString( const TCHAR* Section, const TCHAR* Key, FString
 	{
 		return false;
 	}
-
-	//this prevents our source code text gatherer from trying to gather the following messages
-#define LOC_DEFINE_REGION
-	if( FCString::Strstr( **PairString, TEXT("LOCTEXT") ) )
-	{
-		UE_LOG( LogConfig, Warning, TEXT( "FConfigCacheIni::GetString( %s, %s, %s ) contains LOCTEXT"), Section, Key, *Filename );
-		return false;
-	}
-	else
-	{
-		Value = **PairString;
-		return true;
-	}
-#undef LOC_DEFINE_REGION
+	Value = **PairString;
+	return true;
 }
 
 bool FConfigCacheIni::GetText( const TCHAR* Section, const TCHAR* Key, FText& Value, const FString& Filename )
@@ -1728,11 +1622,7 @@ bool FConfigCacheIni::GetText( const TCHAR* Section, const TCHAR* Key, FText& Va
 	{
 		return false;
 	}
-	if( !FParse::Text( **PairString, Value, Section ) )
-	{
-		Value = FText::FromString( **PairString );
-	}
-	return true;
+	return FTextStringHelper::ReadFromString( **PairString, Value, Section );
 }
 
 bool FConfigCacheIni::GetSection( const TCHAR* Section, TArray<FString>& Result, const FString& Filename )
@@ -1816,9 +1706,10 @@ void FConfigCacheIni::SetText( const TCHAR* Section, const TCHAR* Key, const FTe
 		Sec = &File->Add( Section, FConfigSection() );
 	}
 
-	FString* Str = Sec->Find( Key );
-	const FString StrValue = FTextFriendHelper::AsString( Value );
+	FString StrValue;
+	FTextStringHelper::WriteToString(StrValue, Value);
 
+	FString* Str = Sec->Find( Key );
 	if( !Str )
 	{
 		Sec->Add( Key, StrValue );
diff --git a/Engine/Source/Runtime/Core/Private/Misc/Parse.cpp b/Engine/Source/Runtime/Core/Private/Misc/Parse.cpp
--- a/Engine/Source/Runtime/Core/Private/Misc/Parse.cpp
+++ b/Engine/Source/Runtime/Core/Private/Misc/Parse.cpp
@@ -308,170 +308,82 @@ bool FParse::Value( const TCHAR* Stream, const TCHAR* Match, FString& Value, boo
 }
 
 // 
-// Parse an Text token
-// This is expected to in the form NSLOCTEXT("Namespace","Key","SourceString") or LOCTEXT("Key","SourceString")
+// Parse a quoted string.
 //
-bool FParse::Text( const TCHAR* Str, FText& Value, const TCHAR* Namespace )
+bool FParse::QuotedString( const TCHAR* Buffer, FString& Value, int32* OutNumCharsRead )
 {
-	while( FChar::IsWhitespace( *Str ) )
+	if (OutNumCharsRead)
 	{
-		++Str;
+		*OutNumCharsRead = 0;
 	}
 
-	//this prevents our source code text gatherer from trying to gather the following messages
-#define LOC_DEFINE_REGION
-	const bool bFoundNSLocText = FCString::Strncmp( Str, TEXT("NSLOCTEXT"), 9 ) == 0;
-	const bool bFoundLocText = !bFoundNSLocText && FCString::Strncmp( Str, TEXT("LOCTEXT"), 7 ) == 0;
-#undef LOC_DEFINE_REGION
+	const TCHAR* Start = Buffer;
+
+	// Require opening quote
+	if (*Buffer++ != TCHAR('"'))
+	{
+		return false;
+	}
 
-	if( bFoundNSLocText || (Namespace && bFoundLocText) )
+	while (*Buffer && *Buffer != TCHAR('"') && *Buffer != TCHAR('\n') && *Buffer != TCHAR('\r'))
 	{
-		Str += bFoundNSLocText ? 9 : 7;
-		while( FChar::IsWhitespace( *Str ) )
+		if (*Buffer != TCHAR('\\')) // unescaped character
 		{
-			++Str;
+			Value += *Buffer++;
 		}
-
-		if( *Str != '(' )
+		else if (*++Buffer == TCHAR('\\')) // escaped backslash "\\"
 		{
-			return false;
+			Value += TEXT("\\");
+			++Buffer;
 		}
-		++Str;
-
-		while( FChar::IsWhitespace( *Str ) )
+		else if (*Buffer == TCHAR('\"')) // escaped double quote "\""
 		{
-			++Str;
+			Value += TCHAR('"');
+			++Buffer;
 		}
-
-		if( *Str )
+		else if (*Buffer == TCHAR('\'')) // escaped single quote "\'"
 		{
-			bool bInQuotes = false;
-			bool bHaveEscapeCode = false;
-			bool bError = false;
-			int32 StringCount = 0;
-			int32 CommaCount = 0;
-			FString ParsedString;
-			FString NamespaceString;
-			FString KeyString;
-			FString SourceString;
-			const int32 ExpectedStringCount = bFoundNSLocText ? 3 : 2;
-
-			while( *Str && (*Str != ')' || bInQuotes) && !bError )
-			{
-				const TCHAR c = *Str;
-				if( bInQuotes )
-				{
-					// Parsing string
-					if( bHaveEscapeCode )
-					{
-						// Handle escape code
-						if( c == 'r' )
-						{
-							ParsedString += '\r';
-						}
-						else if( c == 'n' )
-						{
-							ParsedString += '\n';
-						}
-						else if( c == '"' )
-						{
-							ParsedString += '"';
-						}
-						else if( c == 't' )
-						{
-							ParsedString += '\t';
-						}
-						else if( c == '\\' )
-						{
-							ParsedString += '\\';
-						}
-						else
-						{
-							bError = true;
-						}
-						bHaveEscapeCode = false;
-					}
-					else
-					{
-						if( c == '"' )
-						{
-							bInQuotes = false;
-							++StringCount;
-							if (StringCount == ExpectedStringCount - 2)
-							{
-								NamespaceString = ParsedString;
-							}
-							else if (StringCount == ExpectedStringCount - 1)
-							{
-								KeyString = ParsedString;
-							}
-						}
-						else if( c == '\\' )
-						{
-							bHaveEscapeCode = true;
-						}
-						else
-						{
-							ParsedString += c;
-						}
-					}
-				}
-				else
-				{
-					// In between strings
-					if( c == '"' )
-					{
-						// Opening quote of parsed string
-						if( StringCount < ExpectedStringCount && StringCount == CommaCount )
-						{
-							ParsedString = TEXT("");
-							bInQuotes = true;
-						}
-						else
-						{
-							bError = true;
-						}
-					}
-					else if( c == ',' )
-					{
-						if( CommaCount >= StringCount )
-						{
-							bError = true;
-						}
-						else
-						{
-							++CommaCount;
-						}
-					}
-					else if( !FChar::IsWhitespace( c ) )
-					{
-						bError = true;
-					}
-				}
-				++Str;
-			}
+			Value += TCHAR('\'');
+			++Buffer;
+		}
+		else if (*Buffer == TCHAR('n')) // escaped newline
+		{
+			Value += TCHAR('\n');
+			++Buffer;
+		}
+		else if (*Buffer == TCHAR('r')) // escaped carriage return
+		{
+			Value += TCHAR('\r');
+			++Buffer;
+		}
+		else // some other escape sequence, assume it's a hex character value
+		{
+			Value += FString::Printf(TEXT("%c"), (HexDigit(Buffer[0]) * 16) + HexDigit(Buffer[1]));
+			Buffer += 2;
+		}
+	}
 
-			if (KeyString.Len() == 0 && (bFoundNSLocText ? NamespaceString : Namespace).Len() > 0)
-			{
-				bError = true;
-			}
+	// Require closing quote
+	if (*Buffer++ != TCHAR('"'))
+	{
+		return false;
+	}
 
-			if( *Str == ')' && !bError && StringCount == ExpectedStringCount )
-			{
-				if (KeyString.IsEmpty())
-				{
-					Value = FText::AsCultureInvariant( ParsedString );
-				}
-				else
-				{
-					Value = FInternationalization::Get().ForUseOnlyByLocMacroAndGraphNodeTextLiterals_CreateText( *ParsedString, *(bFoundNSLocText ? NamespaceString : Namespace), *KeyString );
-				}
-				return true;
-			}
-		}
+	if (OutNumCharsRead)
+	{
+		*OutNumCharsRead = (Buffer - Start);
 	}
 
-	return false;
+	return true;
+}
+
+// 
+// Parse an Text token
+// This is expected to in the form NSLOCTEXT("Namespace","Key","SourceString") or LOCTEXT("Key","SourceString")
+//
+bool FParse::Text( const TCHAR* Buffer, FText& Value, const TCHAR* Namespace )
+{
+	return FTextStringHelper::ReadFromString(Buffer, Value, Namespace);
 }
 
 // 
diff --git a/Engine/Source/Runtime/Core/Private/Tests/Async/AsyncTest.cpp b/Engine/Source/Runtime/Core/Private/Tests/Async/AsyncTest.cpp
--- a/Engine/Source/Runtime/Core/Private/Tests/Async/AsyncTest.cpp
+++ b/Engine/Source/Runtime/Core/Private/Tests/Async/AsyncTest.cpp
@@ -6,6 +6,7 @@
 #include "AutomationTest.h"
 #include "Future.h"
 
+#if WITH_DEV_AUTOMATION_TESTS
 
 IMPLEMENT_SIMPLE_AUTOMATION_TEST(FAsyncGraphTest, "System.Core.Async.Async (Task Graph)", EAutomationTestFlags::EditorContext | EAutomationTestFlags::EngineFilter)
 IMPLEMENT_SIMPLE_AUTOMATION_TEST(FAsyncThreadedTaskTest, "System.Core.Async.Async (Thread)", EAutomationTestFlags::EditorContext | EAutomationTestFlags::EngineFilter)
@@ -94,3 +95,5 @@ bool FAsyncCompletionCallbackTest::RunTest(const FString& Parameters)
 
 	return true;
 }
+
+#endif //WITH_DEV_AUTOMATION_TESTS
\ No newline at end of file
diff --git a/Engine/Source/Runtime/Core/Private/Tests/HAL/PlatformTest.cpp b/Engine/Source/Runtime/Core/Private/Tests/HAL/PlatformTest.cpp
--- a/Engine/Source/Runtime/Core/Private/Tests/HAL/PlatformTest.cpp
+++ b/Engine/Source/Runtime/Core/Private/Tests/HAL/PlatformTest.cpp
@@ -2,10 +2,10 @@
 
 #include "CorePrivatePCH.h"
 
+#if WITH_DEV_AUTOMATION_TESTS
 
 float TheCompilerDoesntKnowThisIsAlwaysZero = 0.0f;
 
-
 struct TestA
 {
 	virtual ~TestA() {}
@@ -100,3 +100,5 @@ bool FPlatformVerificationTest::RunTest (const FString& Parameters)
 
 	return true;
 }
+
+#endif //WITH_DEV_AUTOMATION_TESTS
\ No newline at end of file
diff --git a/Engine/Source/Runtime/Core/Private/Tests/Internationalization/DateTimeFormattingRulesTest.cpp b/Engine/Source/Runtime/Core/Private/Tests/Internationalization/DateTimeFormattingRulesTest.cpp
--- a/Engine/Source/Runtime/Core/Private/Tests/Internationalization/DateTimeFormattingRulesTest.cpp
+++ b/Engine/Source/Runtime/Core/Private/Tests/Internationalization/DateTimeFormattingRulesTest.cpp
@@ -3,6 +3,7 @@
 #include "CorePrivatePCH.h"
 #include "AutomationTest.h"
 
+#if WITH_DEV_AUTOMATION_TESTS
 
 // disable optimization for DateTimeFormattingRulesTest as it compiles very slowly in development builds
 PRAGMA_DISABLE_OPTIMIZATION
@@ -111,4 +112,6 @@ bool FDateTimeFormattingRulesTest::RunTest (const FString& Parameters)
 }
 
 
-PRAGMA_ENABLE_OPTIMIZATION
\ No newline at end of file
+PRAGMA_ENABLE_OPTIMIZATION
+
+#endif //WITH_DEV_AUTOMATION_TESTS
\ No newline at end of file
diff --git a/Engine/Source/Runtime/Core/Private/Tests/Math/UnitConversionTests.cpp b/Engine/Source/Runtime/Core/Private/Tests/Math/UnitConversionTests.cpp
--- a/Engine/Source/Runtime/Core/Private/Tests/Math/UnitConversionTests.cpp
+++ b/Engine/Source/Runtime/Core/Private/Tests/Math/UnitConversionTests.cpp
@@ -3,6 +3,7 @@
 #include "CorePrivatePCH.h"
 #include "UnitConversion.h"
 
+#if WITH_DEV_AUTOMATION_TESTS
 
 IMPLEMENT_SIMPLE_AUTOMATION_TEST(FUnitUnitTests, "System.Core.Math.Unit Conversion", EAutomationTestFlags::ApplicationContextMask | EAutomationTestFlags::SmokeFilter)
 
@@ -100,4 +101,6 @@ bool FUnitUnitTests::RunTest(const FString& Parameters)
 		}
 	}
 	return bSuccess;
-}
\ No newline at end of file
+}
+
+#endif //WITH_DEV_AUTOMATION_TESTS
\ No newline at end of file
diff --git a/Engine/Source/Runtime/Core/Private/Tests/Misc/PathsTest.cpp b/Engine/Source/Runtime/Core/Private/Tests/Misc/PathsTest.cpp
--- a/Engine/Source/Runtime/Core/Private/Tests/Misc/PathsTest.cpp
+++ b/Engine/Source/Runtime/Core/Private/Tests/Misc/PathsTest.cpp
@@ -3,9 +3,9 @@
 #include "CorePrivatePCH.h"
 #include "AutomationTest.h"
 
-
 FString FPaths::GameProjectFilePath;
 
+#if WITH_DEV_AUTOMATION_TESTS
 
 IMPLEMENT_SIMPLE_AUTOMATION_TEST(FPathTests, "System.Core.Misc.Paths", EAutomationTestFlags::ApplicationContextMask | EAutomationTestFlags::SmokeFilter)
 
@@ -86,3 +86,6 @@ bool FPathTests::RunTest( const FString& Parameters )
 
 	return true;
 }
+
+
+#endif //WITH_DEV_AUTOMATION_TESTS
\ No newline at end of file
diff --git a/Engine/Source/Runtime/Core/Private/Tests/Misc/TypeContainerTest.cpp b/Engine/Source/Runtime/Core/Private/Tests/Misc/TypeContainerTest.cpp
--- a/Engine/Source/Runtime/Core/Private/Tests/Misc/TypeContainerTest.cpp
+++ b/Engine/Source/Runtime/Core/Private/Tests/Misc/TypeContainerTest.cpp
@@ -10,6 +10,7 @@
 #include "TypeContainer.h"
 #include "AutomationTest.h"
 
+#if WITH_DEV_AUTOMATION_TESTS
 
 IMPLEMENT_SIMPLE_AUTOMATION_TEST(FTypeContainerTest, "System.Core.Misc.TypeContainer", EAutomationTestFlags::ApplicationContextMask | EAutomationTestFlags::EngineFilter)
 
@@ -248,5 +249,5 @@ bool FTypeContainerTest::RunTest(const FString& Parameters)
 	return true;
 }
 
-
+#endif //WITH_DEV_AUTOMATION_TESTS
 #endif //PLATFORM_COMPILER_HAS_VARIADIC_TEMPLATES
diff --git a/Engine/Source/Runtime/Core/Private/Windows/WindowsPlatformSplash.cpp b/Engine/Source/Runtime/Core/Private/Windows/WindowsPlatformSplash.cpp
--- a/Engine/Source/Runtime/Core/Private/Windows/WindowsPlatformSplash.cpp
+++ b/Engine/Source/Runtime/Core/Private/Windows/WindowsPlatformSplash.cpp
@@ -235,7 +235,7 @@ HBITMAP LoadSplashBitmap()
 			}
 		}
 
-		DeleteObject(Converter);
+		Converter->Release();
 	}
 
 	// Create a DIB from the converted IWICBitmapSource
@@ -286,7 +286,7 @@ HBITMAP LoadSplashBitmap()
 			{
 				if (hDIBBitmap)
 				{
-					DeleteObject(hDIBBitmap);
+					ensure(DeleteObject(hDIBBitmap));
 				}
 
 				hDIBBitmap = CreateDIBSection(hdcScreen, &bminfo, DIB_RGB_COLORS, &ImageBits, NULL, 0);
@@ -324,34 +324,34 @@ HBITMAP LoadSplashBitmap()
 		// Image Extraction failed, clear allocated memory
 		if (FAILED(hr) && hDIBBitmap)
 		{
-			DeleteObject(hDIBBitmap);
+			ensure(DeleteObject(hDIBBitmap));
 			hDIBBitmap = NULL;
 		}
 	}
 
 	if ( OriginalBitmapSource )
 	{
-		DeleteObject( OriginalBitmapSource );
+		OriginalBitmapSource->Release();
 	}
 
 	if ( ToRenderBitmapSource )
 	{
-		DeleteObject( ToRenderBitmapSource );
+		ToRenderBitmapSource->Release();
 	}
 
 	if ( Decoder )
 	{
-		DeleteObject( Decoder );
+		Decoder->Release();
 	}
 
 	if ( Frame )
 	{
-		DeleteObject( Frame );
+		Frame->Release();
 	}
 
 	if ( Factory )
 	{
-		DeleteObject( Factory );
+		Factory->Release();
 	}
 
 	return hDIBBitmap;
@@ -563,7 +563,7 @@ uint32 WINAPI StartSplashScreenThread( LPVOID unused )
 			}
 		}
 
-		DeleteObject(GSplashScreenBitmap);
+		ensure(DeleteObject(GSplashScreenBitmap));
 		GSplashScreenBitmap = NULL;
 	}
 
diff --git a/Engine/Source/Runtime/Core/Public/Internationalization/Text.h b/Engine/Source/Runtime/Core/Public/Internationalization/Text.h
--- a/Engine/Source/Runtime/Core/Public/Internationalization/Text.h
+++ b/Engine/Source/Runtime/Core/Public/Internationalization/Text.h
@@ -353,6 +353,13 @@ public:
 
 	bool ShouldGatherForLocalization() const;
 
+#if WITH_EDITOR
+	/**
+	 * Constructs a new FText with the SourceString of the specified text but with the specified namespace and key
+	 */
+	static FText ChangeKey( FString Namespace, FString Key, const FText& Text );
+#endif
+
 private:
 
 	/** Special constructor used to create StaticEmptyText without also allocating a history object */
@@ -367,13 +374,6 @@ private:
 
 	friend CORE_API FArchive& operator<<( FArchive& Ar, FText& Value );
 
-#if WITH_EDITOR
-	/**
-	 * Constructs a new FText with the SourceString of the specified text but with the specified namespace and key
-	 */
-	static FText ChangeKey( FString Namespace, FString Key, const FText& Text );
-#endif
-
 	/**
 	 * Generate an FText for a string formatted numerically.
 	 */
@@ -430,8 +430,6 @@ public:
 	friend class FTextInspector;
 	friend class UStruct;
 	friend class UGatherTextFromAssetsCommandlet;
-	friend class UEdGraphSchema;
-	friend class UEdGraphSchema_K2;
 	friend class FTextHistory_NamedFormat;
 	friend class FTextHistory_ArgumentDataFormat;
 	friend class FTextHistory_OrderedFormat;
@@ -647,6 +645,37 @@ public:
 	static uint32 GetFlags(const FText& Text);
 };
 
+class CORE_API FTextStringHelper
+{
+public:
+	/**
+	 * Attempt to extract an FText instance from the given stream of text.
+	 *
+	 * @param Buffer			The buffer of text to read from.
+	 * @param OutValue			The text value to fill with the read text.
+	 * @param Namespace			An optional namespace to use when parsing texts that use LOCTEXT (default is an empty namespace).
+	 * @param OutNumCharsRead	An optional output parameter to fill with the number of characters we read from the given buffer.
+	 * @param bRequiresQuotes	True if the read text literal must be surrounded by quotes (eg, when loading from a delimited list)
+	 *
+	 * @return True if we read a valid FText instance into OutValue, false otherwise
+	 */
+	static bool ReadFromString(const TCHAR* Buffer, FText& OutValue, const TCHAR* Namespace = nullptr, int32* OutNumCharsRead = nullptr, const bool bRequiresQuotes = false);
+
+	/**
+	 * Write the given FText instance to a stream of text
+	 *
+	 * @param Buffer			The buffer of text to write to.
+	 * @param Value				The text value to write into the buffer.
+	 * @param bRequiresQuotes	True if the written text literal must be surrounded by quotes (eg, when saving as a delimited list)
+	 *
+	 * @return True if we read a valid FText instance into OutValue, false otherwise
+	 */
+	static bool WriteToString(FString& Buffer, const FText& Value, const bool bRequiresQuotes = false);
+
+private:
+	static bool ReadFromString_ComplexText(const TCHAR* Buffer, FText& OutValue, const TCHAR* Namespace, int32* OutNumCharsRead);
+};
+
 class CORE_API FTextBuilder
 {
 public:
diff --git a/Engine/Source/Runtime/Core/Public/Math/Quat.h b/Engine/Source/Runtime/Core/Public/Math/Quat.h
--- a/Engine/Source/Runtime/Core/Public/Math/Quat.h
+++ b/Engine/Source/Runtime/Core/Public/Math/Quat.h
@@ -485,10 +485,19 @@ public:
 		return SlerpFullPath_NotNormalized(quat1, quat2, Alpha).GetNormalized();
 	}
 	
-	/** Given start and end quaternions of quat1 and quat2, and tangents at those points tang1 and tang2, calculate the point at Alpha (between 0 and 1) between them. Result is normalized. */
+	/**
+	 * Given start and end quaternions of quat1 and quat2, and tangents at those points tang1 and tang2, calculate the point at Alpha (between 0 and 1) between them. Result is normalized.
+	 * This will correct alignment by ensuring that the shortest path is taken.
+	 */
 	static CORE_API FQuat Squad( const FQuat& quat1, const FQuat& tang1, const FQuat& quat2, const FQuat& tang2, float Alpha );
 
-	/** 
+	/**
+	 * Simpler Squad that doesn't do any checks for 'shortest distance' etc.
+	 * Given start and end quaternions of quat1 and quat2, and tangents at those points tang1 and tang2, calculate the point at Alpha (between 0 and 1) between them. Result is normalized.
+	 */
+	static CORE_API FQuat SquadFullPath(const FQuat& quat1, const FQuat& tang1, const FQuat& quat2, const FQuat& tang2, float Alpha);
+
+	/**
 	 * Calculate tangents between given points
 	 *
 	 * @param PrevP quaternion at P-1
diff --git a/Engine/Source/Runtime/Core/Public/Misc/AssertionMacros.h b/Engine/Source/Runtime/Core/Public/Misc/AssertionMacros.h
--- a/Engine/Source/Runtime/Core/Public/Misc/AssertionMacros.h
+++ b/Engine/Source/Runtime/Core/Public/Misc/AssertionMacros.h
@@ -161,29 +161,28 @@ namespace UE4Asserts_Private
 
 #define ensureMsg(InExpression, InMsg) UE4Asserts_Private::DeprecatedEnsure(ensureMsgf(InExpression, InMsg))
 
-// These are only for use as part of the GET_MEMBER_NAME_CHECKED macro
-template <typename T>
-bool JunkFunc(const T&);
-
-template <typename T>
-bool JunkFunc(const volatile T&);
+namespace UE4Asserts_Private
+{
+	// A junk function to allow us to use sizeof on a member variable which is potentially a bitfield
+	template <typename T>
+	bool GetMemberNameCheckedJunk(const T&);
+	template <typename T>
+	bool GetMemberNameCheckedJunk(const volatile T&);
+}
 
 // Returns FName(TEXT("MemberName")), while statically verifying that the member exists in ClassName
 #define GET_MEMBER_NAME_CHECKED(ClassName, MemberName) \
-	((void)sizeof(JunkFunc(((ClassName*)0)->MemberName)), FName(TEXT(#MemberName)))
+	((void)sizeof(UE4Asserts_Private::GetMemberNameCheckedJunk(((ClassName*)0)->MemberName)), FName(TEXT(#MemberName)))
 
 #define GET_MEMBER_NAME_STRING_CHECKED(ClassName, MemberName) \
-	((void)sizeof(JunkFunc(((ClassName*)0)->MemberName)), TEXT(#MemberName))
+	((void)sizeof(UE4Asserts_Private::GetMemberNameCheckedJunk(((ClassName*)0)->MemberName)), TEXT(#MemberName))
 
 // Returns FName(TEXT("FunctionName")), while statically verifying that the function exists in ClassName
-// &((ClassName*)0)->FunctionName will generate 'error: cannot create a non-constant pointer to member function' on non-MSVC compilers
-#if PLATFORM_WINDOWS && !defined(__clang__)
-	#define GET_FUNCTION_NAME_CHECKED(ClassName, FunctionName) \
-		((void)sizeof(&((ClassName*)0)->FunctionName), FName(TEXT(#FunctionName)))
-#else
-	#define GET_FUNCTION_NAME_CHECKED(ClassName, FunctionName) \
-		FName(TEXT(#FunctionName))
-#endif
+#define GET_FUNCTION_NAME_CHECKED(ClassName, FunctionName) \
+	((void)sizeof(&ClassName::FunctionName), FName(TEXT(#FunctionName)))
+
+#define GET_FUNCTION_NAME_STRING_CHECKED(ClassName, FunctionName) \
+	((void)sizeof(&ClassName::FunctionName), TEXT(#FunctionName))
 
 /*----------------------------------------------------------------------------
 	Stats helpers
diff --git a/Engine/Source/Runtime/Core/Public/Misc/Parse.h b/Engine/Source/Runtime/Core/Public/Misc/Parse.h
--- a/Engine/Source/Runtime/Core/Public/Misc/Parse.h
+++ b/Engine/Source/Runtime/Core/Public/Misc/Parse.h
@@ -69,7 +69,9 @@ struct CORE_API FParse
 	/** Checks if a command-line parameter exists in the stream. */
 	static bool Param( const TCHAR* Stream, const TCHAR* Param );
 	/** Parse an Text token. */
-	static bool Text( const TCHAR* Stream, FText& Value, const TCHAR* Namespace = NULL );
+	static bool Text( const TCHAR* Stream, FText& Value, const TCHAR* Namespace = nullptr );
+	/** Parse a quoted string token. */
+	static bool QuotedString( const TCHAR* Stream, FString& Value, int32* OutNumCharsRead = nullptr );
 
 	//
 	// Parse a hex digit.
diff --git a/Engine/Source/Runtime/Core/Public/UObject/EditorObjectVersion.h b/Engine/Source/Runtime/Core/Public/UObject/EditorObjectVersion.h
--- a/Engine/Source/Runtime/Core/Public/UObject/EditorObjectVersion.h
+++ b/Engine/Source/Runtime/Core/Public/UObject/EditorObjectVersion.h
@@ -8,8 +8,8 @@ struct CORE_API FEditorObjectVersion
 	{
 		// Before any version changes were made
 		BeforeCustomVersionWasAdded = 0,
-
-
+		// Localizble text gathered and stored in packages is now flagged with a localizable text gathering process version
+		GatheredTextProcessVersionFlagging,
 		// -----<new versions can be added above this line>-------------------------------------------------
 		VersionPlusOne,
 		LatestVersion = VersionPlusOne - 1
diff --git a/Engine/Source/Runtime/CoreUObject/Private/Misc/PackageName.cpp b/Engine/Source/Runtime/CoreUObject/Private/Misc/PackageName.cpp
--- a/Engine/Source/Runtime/CoreUObject/Private/Misc/PackageName.cpp
+++ b/Engine/Source/Runtime/CoreUObject/Private/Misc/PackageName.cpp
@@ -519,6 +519,61 @@ bool FPackageName::IsValidLongPackageName(const FString& InLongPackageName, bool
 	return bValidRoot;
 }
 
+bool FPackageName::IsEnginePackageName(const FString& InLongPackageName)
+{
+	const auto& Paths = FLongPackagePathsSingleton::Get();
+
+	UPackage* TransientPackage = GetTransientPackage();
+	if (TransientPackage && InLongPackageName.StartsWith(TransientPackage->GetName()))
+	{
+		// Don't consider the transient package to be an engine package; it is as yet undefined.
+		// This can lead to issues later if it is saved as an engine package, but we can't do anything about that.
+		return false;
+	}
+	else if (InLongPackageName.StartsWith(Paths.ScriptRootPath))
+	{
+		// Query the module name following the /Script/ prefix with the module manager to determine whether it is an engine or a game module
+		int32 StartModuleNameIndex;
+		if (InLongPackageName.FindLastChar(TEXT('/'), StartModuleNameIndex))
+		{
+			int32 EndModuleNameIndex;
+			if (InLongPackageName.FindLastChar(TEXT('.'), EndModuleNameIndex))
+			{
+				const FString ModuleNameString = InLongPackageName.Mid(StartModuleNameIndex + 1, EndModuleNameIndex - StartModuleNameIndex - 1);
+				const FName ModuleName = FName(*ModuleNameString);
+				if (FModuleManager::Get().IsModuleLoaded(ModuleName))
+				{
+					FModuleStatus ModuleStatus;
+					if (FModuleManager::Get().QueryModule(ModuleName, ModuleStatus))
+					{
+						return !ModuleStatus.bIsGameModule;
+					}
+				}
+			}
+		}
+
+		return false;
+	}
+	else
+	{
+		// Query content roots to see if they lie within the engine path or not
+		for (const auto& Pair : Paths.ContentRootToPath)
+		{
+			const FString EnginePath = FPaths::EngineDir();
+
+			if (InLongPackageName.StartsWith(Pair.RootPath))
+			{
+				if (Pair.ContentPath.StartsWith(EnginePath))
+				{
+					return true;
+				}
+			}
+		}
+	}
+
+	return false;
+}
+
 void FPackageName::RegisterMountPoint(const FString& RootPath, const FString& ContentPath)
 {
 	FLongPackagePathsSingleton& Paths = FLongPackagePathsSingleton::Get();
diff --git a/Engine/Source/Runtime/CoreUObject/Private/UObject/Obj.cpp b/Engine/Source/Runtime/CoreUObject/Private/UObject/Obj.cpp
--- a/Engine/Source/Runtime/CoreUObject/Private/UObject/Obj.cpp
+++ b/Engine/Source/Runtime/CoreUObject/Private/UObject/Obj.cpp
@@ -1615,9 +1615,6 @@ void UObject::LoadConfig( UClass* ConfigClass/*=NULL*/, const TCHAR* InFilename/
 		UArrayProperty* Array = dynamic_cast<UArrayProperty*>( Property );
 		if( Array == NULL )
 		{
-			// todo: UE-22989 - This special case handling for FText properties can be removed when UE-22989 is fixed so that Import/ExportText on an FText property doesn't lose the namespace/key data
-			UTextProperty* TextProp = dynamic_cast<UTextProperty*>( Property );
-
 			for( int32 i=0; i<Property->ArrayDim; i++ )
 			{
 				if( Property->ArrayDim!=1 )
@@ -1625,27 +1622,14 @@ void UObject::LoadConfig( UClass* ConfigClass/*=NULL*/, const TCHAR* InFilename/
 					Key = FString::Printf(TEXT("%s[%i]"), *Property->GetName(), i);
 				}
 
-				bool bFoundValue = false;
-				if (TextProp)
+				FString Value;
+				const bool bFoundValue = GConfig->GetString( *ClassSection, *Key, Value, *PropFileName );
+				if (bFoundValue)
 				{
-					FText Value;
-					bFoundValue = GConfig->GetText( *ClassSection, *Key, Value, *PropFileName );
-					if (bFoundValue)
+					if (Property->ImportText(*Value, Property->ContainerPtrToValuePtr<uint8>(this, i), PortFlags, this) == NULL)
 					{
-						TextProp->SetPropertyValue_InContainer(this, Value, i);
-					}
-				}
-				else
-				{
-					FString Value;
-					bFoundValue = GConfig->GetString( *ClassSection, *Key, Value, *PropFileName );
-					if (bFoundValue)
-					{
-						if (Property->ImportText(*Value, Property->ContainerPtrToValuePtr<uint8>(this, i), PortFlags, this) == NULL)
-						{
-							// this should be an error as the properties from the .ini / .int file are not correctly being read in and probably are affecting things in subtle ways
-							UE_LOG(LogObj, Error, TEXT("LoadConfig (%s): import failed for %s in: %s"), *GetPathName(), *Property->GetName(), *Value);
-						}
+						// this should be an error as the properties from the .ini / .int file are not correctly being read in and probably are affecting things in subtle ways
+						UE_LOG(LogObj, Error, TEXT("LoadConfig (%s): import failed for %s in: %s"), *GetPathName(), *Property->GetName(), *Value);
 					}
 				}
 
@@ -1678,30 +1662,10 @@ void UObject::LoadConfig( UClass* ConfigClass/*=NULL*/, const TCHAR* InFilename/
 				// Only override default properties if there is something to override them with.
 				if ( List.Num() > 0 )
 				{
-					// todo: UE-22989 - This special case handling for FText properties can be removed when UE-22989 is fixed so that Import/ExportText on an FText property doesn't lose the namespace/key data
-					UTextProperty* TextProp = dynamic_cast<UTextProperty*>( Array->Inner );
-
 					ArrayHelper.EmptyAndAddValues(List.Num());
-
-					if (TextProp)
+					for( int32 i=List.Num()-1,c=0; i>=0; i--,c++ )
 					{
-						for( int32 i=List.Num()-1,c=0; i>=0; i--,c++ )
-						{
-							FText Value;
-							if (!FParse::Text(*List[i], Value, *ClassSection))
-							{
-								Value = FText::FromString(*List[i]);
-							}
-
-							TextProp->SetPropertyValue(ArrayHelper.GetRawPtr(c), Value);
-						}
-					}
-					else
-					{
-						for( int32 i=List.Num()-1,c=0; i>=0; i--,c++ )
-						{
-							Array->Inner->ImportText( *List[i], ArrayHelper.GetRawPtr(c), PortFlags, this );
-						}
+						Array->Inner->ImportText( *List[i], ArrayHelper.GetRawPtr(c), PortFlags, this );
 					}
 				}
 				else
@@ -1854,78 +1818,12 @@ void UObject::SaveConfig( uint64 Flags, const TCHAR* InFilename, FConfigCacheIni
 					// Default ini's require the array syntax to be applied to the property name
 					FString CompleteKey = FString::Printf(TEXT("%s%s"), bIsADefaultIniWrite ? TEXT("+") : TEXT(""), *Key);
 
-					// todo: UE-22989 - This special case handling for FText properties can be removed when UE-22989 is fixed so that Import/ExportText on an FText property doesn't lose the namespace/key data
-					UTextProperty* TextProp = dynamic_cast<UTextProperty*>( Array->Inner );
-
 					FScriptArrayHelper_InContainer ArrayHelper(Array, this);
-
-					if (TextProp)
-					{
-						auto TextToConfigString = [](const FText& InText) -> FString
-						{
-							if (InText.IsTransient())
-							{
-								UE_LOG(LogConfig, Warning, TEXT( "UObject::SaveConfig() Transient FText"));
-								return FString(TEXT("Error: Transient FText"));
-							}
-
-							if (InText.IsCultureInvariant())
-							{
-								// todo: UE-22989 - This isn't really the correct way to handle culture invariant strings, but it mirrors the logic in FTextFriendHelper
-								return FString::Printf(
-									TEXT("NSLOCTEXT(\"\",\"\",\"%s\")" ), 
-									*InText.ToString().ReplaceCharWithEscapedChar().ReplaceQuotesWithEscapedQuotes()
-									);
-							}
-							else
-							{
-								const FString* SourceString = FTextInspector::GetSourceString(InText);
-								if (SourceString)
-								{
-									FTextDisplayStringRef DisplayString = FTextInspector::GetSharedDisplayString(InText);
-
-									FString LocNamespace;
-									FString LocKey;
-									const bool FoundNamespaceAndKey = FTextLocalizationManager::Get().FindNamespaceAndKeyFromDisplayString(DisplayString, LocNamespace, LocKey);
-
-									// If this has no namespace or key, attempt to give it a GUID for a key and register it.
-									if (!FoundNamespaceAndKey && GIsEditor)
-									{
-										LocKey = FGuid::NewGuid().ToString();
-										if (!FTextLocalizationManager::Get().AddDisplayString(DisplayString, LocNamespace, LocKey))
-										{
-											// Could not add display string, reset namespace and key.
-											LocNamespace.Empty();
-											LocKey.Empty();
-										}
-									}
-
-									return FString::Printf(
-										TEXT("NSLOCTEXT(\"%s\",\"%s\",\"%s\")" ), 
-										*LocNamespace, 
-										*LocKey, 
-										*InText.ToString().ReplaceCharWithEscapedChar().ReplaceQuotesWithEscapedQuotes()
-										);
-								}
-							}
-
-							return InText.ToString();
-						};
-
-						for( int32 i=0; i<ArrayHelper.Num(); i++ )
-						{
-							const FText* TextValuePtr = TextProp->GetPropertyValuePtr(ArrayHelper.GetRawPtr(i));
-							Sec->Add(*CompleteKey, TextToConfigString(*TextValuePtr));
-						}
-					}
-					else
+					for( int32 i=0; i<ArrayHelper.Num(); i++ )
 					{
-						for( int32 i=0; i<ArrayHelper.Num(); i++ )
-						{
-							FString	Buffer;
-							Array->Inner->ExportTextItem( Buffer, ArrayHelper.GetRawPtr(i), ArrayHelper.GetRawPtr(i), this, PortFlags );
-							Sec->Add(*CompleteKey, *Buffer);
-						}
+						FString	Buffer;
+						Array->Inner->ExportTextItem( Buffer, ArrayHelper.GetRawPtr(i), ArrayHelper.GetRawPtr(i), this, PortFlags );
+						Sec->Add(*CompleteKey, *Buffer);
 					}
 				}
 				else if( Property->Identical_InContainer(this, SuperClassDefaultObject) )
@@ -1952,20 +1850,9 @@ void UObject::SaveConfig( uint64 Flags, const TCHAR* InFilename, FConfigCacheIni
 							Key = TempKey;
 						}
 
-						// todo: UE-22989 - This special case handling for FText properties can be removed when UE-22989 is fixed so that Import/ExportText on an FText property doesn't lose the namespace/key data
-						UTextProperty* TextProp = dynamic_cast<UTextProperty*>( Property );
-
-						if (TextProp)
-						{
-							const FText* TextValuePtr = TextProp->GetPropertyValuePtr_InContainer(this, Index);
-							Config->SetText( *Section, *Key, *TextValuePtr, *PropFileName );
-						}
-						else
-						{
-							FString	Value;
-							Property->ExportText_InContainer( Index, Value, this, this, this, PortFlags );
-							Config->SetString( *Section, *Key, *Value, *PropFileName );
-						}
+						FString	Value;
+						Property->ExportText_InContainer( Index, Value, this, this, this, PortFlags );
+						Config->SetString( *Section, *Key, *Value, *PropFileName );
 					}
 					else if( Property->Identical_InContainer(this, SuperClassDefaultObject, Index) )
 					{
diff --git a/Engine/Source/Runtime/CoreUObject/Private/UObject/PackageFileSummary.cpp b/Engine/Source/Runtime/CoreUObject/Private/UObject/PackageFileSummary.cpp
--- a/Engine/Source/Runtime/CoreUObject/Private/UObject/PackageFileSummary.cpp
+++ b/Engine/Source/Runtime/CoreUObject/Private/UObject/PackageFileSummary.cpp
@@ -2,6 +2,7 @@
 
 #include "CoreUObjectPrivate.h"
 #include "EngineVersion.h"
+#include "EditorObjectVersion.h"
 
 FPackageFileSummary::FPackageFileSummary()
 {
diff --git a/Engine/Source/Runtime/CoreUObject/Private/UObject/Property.cpp b/Engine/Source/Runtime/CoreUObject/Private/UObject/Property.cpp
--- a/Engine/Source/Runtime/CoreUObject/Private/UObject/Property.cpp
+++ b/Engine/Source/Runtime/CoreUObject/Private/UObject/Property.cpp
@@ -254,50 +254,13 @@ const TCHAR* UPropertyHelpers::ReadToken( const TCHAR* Buffer, FString& String,
 {
 	if( *Buffer == TCHAR('"') )
 	{
-		// Get quoted string.
-		Buffer++;
-		while( *Buffer && *Buffer!=TCHAR('"') && *Buffer!=TCHAR('\n') && *Buffer!=TCHAR('\r') )
-		{
-			if( *Buffer != TCHAR('\\') ) // unescaped character
-			{
-				String += *Buffer++;
-			}
-			else if( *++Buffer==TCHAR('\\') ) // escaped backslash "\\"
-			{
-				String += TEXT("\\");
-				Buffer++;
-			}
-			else if ( *Buffer == TCHAR('\"') ) // escaped double quote "\""
-			{
-				String += TCHAR('"');
-				Buffer++;
-			}
-			else if (*Buffer == TCHAR('\'')) // escaped single quote "\'"
-			{
-				String += TCHAR('\'');
-				Buffer++;
-			}
-			else if ( *Buffer == TCHAR('n') ) // escaped newline
-			{
-				String += TCHAR('\n');
-				Buffer++;
-			}
-			else if ( *Buffer == TCHAR('r') ) // escaped carriage return
-			{
-				String += TCHAR('\r');
-				Buffer++;
-			}
-			else // some other escape sequence, assume it's a hex character value
-			{
-				String = FString::Printf(TEXT("%s%c"), *String, FParse::HexDigit(Buffer[0])*16 + FParse::HexDigit(Buffer[1]));
-				Buffer += 2;
-			}
-		}
-		if( *Buffer++ != TCHAR('"') )
+		int32 NumCharsRead = 0;
+		if (!FParse::QuotedString(Buffer, String, &NumCharsRead))
 		{
 			UE_LOG(LogProperty, Warning, TEXT("ReadToken: Bad quoted string") );
-			return NULL;
+			return nullptr;
 		}
+		Buffer += NumCharsRead;
 	}
 	else if( FChar::IsAlnum( *Buffer ) || (DottedNames && (*Buffer==TCHAR('/'))) || (*Buffer > 255) )
 	{
diff --git a/Engine/Source/Runtime/CoreUObject/Private/UObject/UTextProperty.cpp b/Engine/Source/Runtime/CoreUObject/Private/UObject/UTextProperty.cpp
--- a/Engine/Source/Runtime/CoreUObject/Private/UObject/UTextProperty.cpp
+++ b/Engine/Source/Runtime/CoreUObject/Private/UObject/UTextProperty.cpp
@@ -48,66 +48,96 @@ void UTextProperty::SerializeItem( FArchive& Ar, void* Value, void const* Defaul
 void UTextProperty::ExportTextItem( FString& ValueStr, const void* PropertyValue, const void* DefaultValue, UObject* Parent, int32 PortFlags, UObject* ExportRootScope ) const
 {
 	const FText& TextValue = GetPropertyValue(PropertyValue);
-	const FString& StringValue = FTextInspector::GetDisplayString(TextValue);
-	if (0 != (PortFlags & PPF_ExportCpp))
-	{
-		ValueStr += FString::Printf(TEXT("FText::FromString(%s)"), *UStrProperty::ExportCppHardcodedText(StringValue, FString()));
-	}
-	else if( !(PortFlags & PPF_Delimited) )
+
+	if (PortFlags & PPF_ExportCpp)
 	{
-		ValueStr += StringValue;
+		ValueStr += GenerateCppCodeForTextValue(TextValue, FString());
 	}
-	else if ( StringValue.Len() > 0 )
+	else
 	{
-		ValueStr += FString::Printf( TEXT("\"%s\""), *(StringValue.ReplaceCharWithEscapedChar()) );
+		FTextStringHelper::WriteToString(ValueStr, TextValue, !!(PortFlags & PPF_Delimited));
 	}
 }
 
 const TCHAR* UTextProperty::ImportText_Internal( const TCHAR* Buffer, void* Data, int32 PortFlags, UObject* Parent, FOutputDevice* ErrorText ) const
 {
-	FText* Text = GetPropertyValuePtr(Data);
+	FText* TextPtr = GetPropertyValuePtr(Data);
 
-	FString BufferAsString;
-	if( !(PortFlags & PPF_Delimited) )
+	int32 NumCharsRead = 0;
+	if (FTextStringHelper::ReadFromString(Buffer, *TextPtr, nullptr, &NumCharsRead, !!(PortFlags & PPF_Delimited)))
 	{
-		BufferAsString = FString(Buffer);
+		Buffer += NumCharsRead;
+		return Buffer;
+	}
+	
+	return nullptr;
+}
+
+FString UTextProperty::GenerateCppCodeForTextValue(const FText& InValue, const FString& Indent)
+{
+	FString CppCode;
+
+	const FString& StringValue = FTextInspector::GetDisplayString(InValue);
 
-		// in order to indicate that the value was successfully imported, advance the buffer past the last character that was imported
-		Buffer += FCString::Strlen(Buffer);
+	if (InValue.IsEmpty())
+	{
+		CppCode += TEXT("FText::GetEmpty()");
+	}
+	else if (InValue.IsCultureInvariant())
+	{
+		// Produces FText::AsCultureInvariant(TEXT("..."))
+		CppCode += TEXT("FText::AsCultureInvariant(");
+		CppCode += UStrProperty::ExportCppHardcodedText(StringValue, Indent + TEXT("\t"));
+		CppCode += TEXT(")");
 	}
 	else
 	{
-		// require quoted string here
-		if (*Buffer != TCHAR('"'))
+		bool bIsLocalized = false;
+		FString Namespace;
+		FString Key;
+		const FString* SourceString = FTextInspector::GetSourceString(InValue);
+
+		if (SourceString && InValue.ShouldGatherForLocalization())
 		{
-			ErrorText->Logf(TEXT("Missing opening '\"' in string property value: %s"), Buffer);
-			return NULL;
+			bIsLocalized = FTextLocalizationManager::Get().FindNamespaceAndKeyFromDisplayString(FTextInspector::GetSharedDisplayString(InValue), Namespace, Key);
 		}
 
-		const TCHAR* Start = Buffer;
-		FString Temp;
-		Buffer = UPropertyHelpers::ReadToken(Buffer, Temp);
-
-		if (Buffer == NULL)
+		if (bIsLocalized)
 		{
-			return NULL;
+			// Produces FInternationalization::ForUseOnlyByLocMacroAndGraphNodeTextLiterals_CreateText(TEXT("..."), TEXT("..."), TEXT("..."))
+			CppCode += TEXT("FInternationalization::ForUseOnlyByLocMacroAndGraphNodeTextLiterals_CreateText(\n");
+
+			CppCode += Indent;
+			CppCode += TEXT("\t");
+			CppCode += UStrProperty::ExportCppHardcodedText(*SourceString, Indent + TEXT("\t\t"));
+			CppCode += TEXT(", /* Literal Text */\n");
+
+			CppCode += Indent;
+			CppCode += TEXT("\t");
+			CppCode += UStrProperty::ExportCppHardcodedText(Namespace, Indent + TEXT("\t\t"));
+			CppCode += TEXT(", /* Namespace */\n");
+
+			CppCode += Indent;
+			CppCode += TEXT("\t");
+			CppCode += UStrProperty::ExportCppHardcodedText(Key, Indent + TEXT("\t\t"));
+			CppCode += TEXT(" /* Key */\n");
+
+			CppCode += Indent;
+			CppCode += TEXT("\t)");
 		}
-
-		if (Buffer > Start && Buffer[-1] != TCHAR('"'))
+		else
 		{
-			ErrorText->Logf(TEXT("Missing terminating '\"' in string property value: %s"), Start);
-			return NULL;
+			// Produces FText::FromString(TEXT("..."))
+			CppCode += TEXT("FText::FromString(");
+			CppCode += UStrProperty::ExportCppHardcodedText(StringValue, Indent + TEXT("\t"));
+			CppCode += TEXT(")");
 		}
-
-		BufferAsString = FString(Temp);
 	}
 
-	*Text = FText::FromString(FString(BufferAsString));
-
-	return Buffer;
+	return CppCode;
 }
 
 IMPLEMENT_CORE_INTRINSIC_CLASS(UTextProperty, UProperty,
 {
 }
-);
\ No newline at end of file
+);
diff --git a/Engine/Source/Runtime/CoreUObject/Public/Misc/PackageName.h b/Engine/Source/Runtime/CoreUObject/Public/Misc/PackageName.h
--- a/Engine/Source/Runtime/CoreUObject/Public/Misc/PackageName.h
+++ b/Engine/Source/Runtime/CoreUObject/Public/Misc/PackageName.h
@@ -99,6 +99,13 @@ public:
 	 * @return							true if a valid long package name
 	 */
 	static bool IsValidLongPackageName(const FString& InLongPackageName, bool bIncludeReadOnlyRoots = false, FText* OutReason = NULL);
+	/**
+	 * Returns whether the specified long package name is an engine package.
+	 *
+	 * @param InLongPackageName Long Package Name
+	 * @return true if it is an Engine package.
+	 */
+	static bool IsEnginePackageName(const FString& InLongPackageName);
 	/**
 	 * Checks if the given string is a long package name or not.
 	 *
diff --git a/Engine/Source/Runtime/CoreUObject/Public/UObject/UTextProperty.h b/Engine/Source/Runtime/CoreUObject/Public/UObject/UTextProperty.h
--- a/Engine/Source/Runtime/CoreUObject/Public/UObject/UTextProperty.h
+++ b/Engine/Source/Runtime/CoreUObject/Public/UObject/UTextProperty.h
@@ -35,4 +35,7 @@ public:
 		return FString();
 	}
 	// End of UProperty interface
+
+	/** Generate the correct C++ code for the given text value */
+	static FString GenerateCppCodeForTextValue(const FText& InValue, const FString& Indent);
 };
\ No newline at end of file
diff --git a/Engine/Source/Runtime/Engine/Classes/Components/SplineMeshComponent.h b/Engine/Source/Runtime/Engine/Classes/Components/SplineMeshComponent.h
--- a/Engine/Source/Runtime/Engine/Classes/Components/SplineMeshComponent.h
+++ b/Engine/Source/Runtime/Engine/Classes/Components/SplineMeshComponent.h
@@ -152,6 +152,7 @@ class ENGINE_API USplineMeshComponent : public UStaticMeshComponent, public IInt
 	//Begin USceneComponent Interface
 	virtual FPrimitiveSceneProxy* CreateSceneProxy() override;
 	virtual FBoxSphereBounds CalcBounds(const FTransform& LocalToWorld) const override;
+	virtual FTransform GetSocketTransform(FName InSocketName, ERelativeTransformSpace TransformSpace = RTS_World) const override;
 	//End USceneComponent Interface
 
 	//Begin UPrimitiveComponent Interface
diff --git a/Engine/Source/Runtime/Engine/Private/AI/Navigation/NavLinkProxy.cpp b/Engine/Source/Runtime/Engine/Private/AI/Navigation/NavLinkProxy.cpp
--- a/Engine/Source/Runtime/Engine/Private/AI/Navigation/NavLinkProxy.cpp
+++ b/Engine/Source/Runtime/Engine/Private/AI/Navigation/NavLinkProxy.cpp
@@ -29,13 +29,26 @@ ANavLinkProxy::ANavLinkProxy(const FObjectInitializer& ObjectInitializer) : Supe
 	SpriteComponent = CreateEditorOnlyDefaultSubobject<UBillboardComponent>(TEXT("Sprite"));
 	if (!IsRunningCommandlet() && (SpriteComponent != NULL))
 	{
-		static ConstructorHelpers::FObjectFinderOptional<UTexture2D> SpriteTexture(TEXT("/Engine/EditorResources/AI/S_NavLink"));
+		struct FConstructorStatics
+		{
+			ConstructorHelpers::FObjectFinderOptional<UTexture2D> SpriteTexture;
+			FName ID_Decals;
+			FText NAME_Decals;
+			FConstructorStatics()
+				: SpriteTexture(TEXT("/Engine/EditorResources/AI/S_NavLink"))
+				, ID_Decals(TEXT("Navigation"))
+				, NAME_Decals(NSLOCTEXT("SpriteCategory", "Navigation", "Navigation"))
+			{
+			}
+		};
+		static FConstructorStatics ConstructorStatics;
 
-		SpriteComponent->Sprite = SpriteTexture.Get();
+		SpriteComponent->Sprite = ConstructorStatics.SpriteTexture.Get();
 		SpriteComponent->RelativeScale3D = FVector(0.5f, 0.5f, 0.5f);
 		SpriteComponent->bHiddenInGame = true;
 		SpriteComponent->bVisible = true;
-
+		SpriteComponent->SpriteInfo.Category = ConstructorStatics.ID_Decals;
+		SpriteComponent->SpriteInfo.DisplayName = ConstructorStatics.NAME_Decals;
 		SpriteComponent->AttachParent = RootComponent;
 		SpriteComponent->SetAbsolute(false, false, true);
 		SpriteComponent->bIsScreenSizeScaled = true;
diff --git a/Engine/Source/Runtime/Engine/Private/Components/AudioComponent.cpp b/Engine/Source/Runtime/Engine/Private/Components/AudioComponent.cpp
--- a/Engine/Source/Runtime/Engine/Private/Components/AudioComponent.cpp
+++ b/Engine/Source/Runtime/Engine/Private/Components/AudioComponent.cpp
@@ -380,6 +380,9 @@ void UAudioComponent::UpdateSpriteTexture()
 {
 	if (SpriteComponent)
 	{
+		SpriteComponent->SpriteInfo.Category = TEXT("Sounds");
+		SpriteComponent->SpriteInfo.DisplayName = NSLOCTEXT("SpriteCategory", "Sounds", "Sounds");
+
 		if (bAutoActivate)
 		{
 			SpriteComponent->SetSprite(LoadObject<UTexture2D>(nullptr, TEXT("/Engine/EditorResources/AudioIcons/S_AudioComponent_AutoActivate.S_AudioComponent_AutoActivate")));
diff --git a/Engine/Source/Runtime/Engine/Private/Components/SplineMeshComponent.cpp b/Engine/Source/Runtime/Engine/Private/Components/SplineMeshComponent.cpp
--- a/Engine/Source/Runtime/Engine/Private/Components/SplineMeshComponent.cpp
+++ b/Engine/Source/Runtime/Engine/Private/Components/SplineMeshComponent.cpp
@@ -6,6 +6,7 @@
 #include "ShaderParameterUtils.h"
 #include "NavigationSystemHelpers.h"
 #include "AI/Navigation/NavCollision.h"
+#include "Engine/StaticMeshSocket.h"
 
 
 //////////////////////////////////////////////////////////////////////////
@@ -696,6 +697,42 @@ FBoxSphereBounds USplineMeshComponent::CalcBounds(const FTransform& LocalToWorld
 	return FBoxSphereBounds(BoundingBox.TransformBy(LocalToWorld));
 }
 
+FTransform USplineMeshComponent::GetSocketTransform(FName InSocketName, ERelativeTransformSpace TransformSpace) const
+{
+	if (InSocketName != NAME_None)
+	{
+		UStaticMeshSocket const* const Socket = GetSocketByName(InSocketName);
+		if (Socket)
+		{
+			FTransform SocketTransform;
+			SocketTransform = FTransform(Socket->RelativeRotation, Socket->RelativeLocation * GetAxisMask(ForwardAxis), Socket->RelativeScale);
+			SocketTransform = SocketTransform * CalcSliceTransform(GetAxisValue(Socket->RelativeLocation, ForwardAxis));
+
+			switch (TransformSpace)
+			{
+			case RTS_World:
+			{
+				return SocketTransform * GetComponentToWorld();
+			}
+			case RTS_Actor:
+			{
+				if (const AActor* Actor = GetOwner())
+				{
+					return (SocketTransform * GetComponentToWorld()).GetRelativeTransform(GetOwner()->GetTransform());
+				}
+				break;
+			}
+			case RTS_Component:
+			{
+				return SocketTransform;
+			}
+			}
+		}
+	}
+
+	return Super::GetSocketTransform(InSocketName, TransformSpace);
+}
+
 
 FTransform USplineMeshComponent::CalcSliceTransform(const float DistanceAlong) const
 {
@@ -982,7 +1019,7 @@ void USplineMeshComponent::RecreateCollision()
 			{
 				FKConvexElem& ConvexElem = *new(BodySetup->AggGeom.ConvexElems) FKConvexElem();
 
-				const FVector Radii = FVector(BoxElem.X / 2, BoxElem.Y / 2, BoxElem.Z / 2);
+				const FVector Radii = FVector(BoxElem.X / 2, BoxElem.Y / 2, BoxElem.Z / 2).ComponentMax(FVector(1.0f));
 				const FTransform ElementTM = BoxElem.GetTransform();
 				ConvexElem.VertexData.Empty(8);
 				ConvexElem.VertexData.Add(ElementTM.TransformPosition(Radii * FVector(-1,-1,-1)));
diff --git a/Engine/Source/Runtime/Engine/Private/DecalActor.cpp b/Engine/Source/Runtime/Engine/Private/DecalActor.cpp
--- a/Engine/Source/Runtime/Engine/Private/DecalActor.cpp
+++ b/Engine/Source/Runtime/Engine/Private/DecalActor.cpp
@@ -63,6 +63,8 @@ ADecalActor::ADecalActor(const FObjectInitializer& ObjectInitializer)
 		{
 			SpriteComponent->Sprite = ConstructorStatics.DecalTexture.Get();
 			SpriteComponent->RelativeScale3D = FVector(0.5f, 0.5f, 0.5f);
+			SpriteComponent->SpriteInfo.Category = ConstructorStatics.ID_Decals;
+			SpriteComponent->SpriteInfo.DisplayName = ConstructorStatics.NAME_Decals;
 			SpriteComponent->AttachParent = Decal;
 			SpriteComponent->bIsScreenSizeScaled = true;
 			SpriteComponent->bAbsoluteScale = true;
diff --git a/Engine/Source/Runtime/Engine/Private/EdGraph/EdGraphSchema.cpp b/Engine/Source/Runtime/Engine/Private/EdGraph/EdGraphSchema.cpp
--- a/Engine/Source/Runtime/Engine/Private/EdGraph/EdGraphSchema.cpp
+++ b/Engine/Source/Runtime/Engine/Private/EdGraph/EdGraphSchema.cpp
@@ -465,23 +465,7 @@ void UEdGraphSchema::TrySetDefaultObject(UEdGraphPin& Pin, UObject* NewDefaultOb
 
 void UEdGraphSchema::TrySetDefaultText(UEdGraphPin& InPin, const FText& InNewDefaultText) const
 {
-	if(InNewDefaultText.IsEmpty())
-	{
-		InPin.DefaultTextValue = InNewDefaultText;
-	}
-	else
-	{
-#if WITH_EDITOR
-		if(InNewDefaultText.IsCultureInvariant())
-		{
-			InPin.DefaultTextValue = InNewDefaultText;
-		}
-		else
-		{
-			InPin.DefaultTextValue = FText::ChangeKey(TEXT(""), FGuid::NewGuid().ToString(), InNewDefaultText);
-		}
-#endif
-	}
+	InPin.DefaultTextValue = InNewDefaultText;
 
 #if WITH_EDITOR
 	UEdGraphNode* Node = InPin.GetOwningNode();
diff --git a/Engine/Source/Runtime/Engine/Private/GameEngine.cpp b/Engine/Source/Runtime/Engine/Private/GameEngine.cpp
--- a/Engine/Source/Runtime/Engine/Private/GameEngine.cpp
+++ b/Engine/Source/Runtime/Engine/Private/GameEngine.cpp
@@ -380,7 +380,10 @@ void UGameEngine::SwitchGameWindowToUseGameViewport()
 		GameViewportWindowPtr->SetContent(GameViewportWidgetRef);
 		GameViewportWindowPtr->SlatePrepass();
 		
-		SceneViewport->ResizeFrame((uint32)GSystemResolution.ResX, (uint32)GSystemResolution.ResY, GSystemResolution.WindowMode, 0, 0);
+		if ( SceneViewport.IsValid() )
+		{
+			SceneViewport->ResizeFrame((uint32)GSystemResolution.ResX, (uint32)GSystemResolution.ResY, GSystemResolution.WindowMode, 0, 0);
+		}
 
 		// Move the registration of the game viewport to that messages are correctly received.
 		if (!FPlatformProperties::SupportsWindowedMode())
diff --git a/Engine/Source/Runtime/Engine/Private/GameViewportClient.cpp b/Engine/Source/Runtime/Engine/Private/GameViewportClient.cpp
--- a/Engine/Source/Runtime/Engine/Private/GameViewportClient.cpp
+++ b/Engine/Source/Runtime/Engine/Private/GameViewportClient.cpp
@@ -1269,6 +1269,11 @@ void UGameViewportClient::ProcessScreenShots(FViewport* InViewport)
 			bScreenshotSuccessful = FSlateApplication::Get().TakeScreenshot( WindowRef, Bitmap, Size);
 			GScreenshotResolutionX = Size.X;
 			GScreenshotResolutionY = Size.Y;
+
+			if (Size.X == 0 || Size.Y == 0)
+			{
+				bScreenshotSuccessful = false;
+			}
 		}
 		else
 		{
@@ -1311,8 +1316,6 @@ void UGameViewportClient::ProcessScreenShots(FViewport* InViewport)
 				TArray<uint8> CompressedBitmap;
 				FImageUtils::CompressImageArray(Size.X, Size.Y, Bitmap, CompressedBitmap);
 				FFileHelper::SaveArrayToFile(CompressedBitmap, *ScreenShotName);
-
-//				FFileHelper::CreateBitmap(*ScreenShotName, InViewport->GetSizeXY().X, InViewport->GetSizeXY().Y, Bitmap.GetData(), &SourceRect, &IFileManager::Get(), NULL, bWriteAlpha);
 			}
 		}
 
diff --git a/Engine/Source/Runtime/Engine/Private/ImageUtils.cpp b/Engine/Source/Runtime/Engine/Private/ImageUtils.cpp
--- a/Engine/Source/Runtime/Engine/Private/ImageUtils.cpp
+++ b/Engine/Source/Runtime/Engine/Private/ImageUtils.cpp
@@ -249,7 +249,7 @@ void FImageUtils::CompressImageArray( int32 ImageWidth, int32 ImageHeight, TArra
 	// Copy scaled image into destination thumb
 	int32 MemorySize = ImageWidth*ImageHeight*sizeof(FColor);
 	ThumbnailByteArray.AddUninitialized(MemorySize);
-	FMemory::Memcpy(&(ThumbnailByteArray[0]), &(SrcData[0]), MemorySize);
+	FMemory::Memcpy(ThumbnailByteArray.GetData(), SrcData.GetData(), MemorySize);
 
 	// Compress data - convert into a .png
 	TempThumbnail.CompressImageData();
diff --git a/Engine/Source/Runtime/Engine/Private/PrimitiveSceneProxy.cpp b/Engine/Source/Runtime/Engine/Private/PrimitiveSceneProxy.cpp
--- a/Engine/Source/Runtime/Engine/Private/PrimitiveSceneProxy.cpp
+++ b/Engine/Source/Runtime/Engine/Private/PrimitiveSceneProxy.cpp
@@ -60,6 +60,7 @@ FPrimitiveSceneProxy::FPrimitiveSceneProxy(const UPrimitiveComponent* InComponen
 ,	bSupportsDistanceFieldRepresentation(false)
 ,	bSupportsHeightfieldRepresentation(false)
 ,	bNeedsLevelAddedToWorldNotification(false)
+,	bWantsSelectionOutline(true)
 ,	bUseAsOccluder(InComponent->bUseAsOccluder)
 ,	bAllowApproximateOcclusion(InComponent->Mobility != EComponentMobility::Movable)
 ,	bSelectable(InComponent->bSelectable)
diff --git a/Engine/Source/Runtime/Engine/Private/StaticMeshBuild.cpp b/Engine/Source/Runtime/Engine/Private/StaticMeshBuild.cpp
--- a/Engine/Source/Runtime/Engine/Private/StaticMeshBuild.cpp
+++ b/Engine/Source/Runtime/Engine/Private/StaticMeshBuild.cpp
@@ -132,15 +132,65 @@ void UStaticMesh::Build(bool bSilent, TArray<FText>* OutErrors)
 
 		// Find any static mesh components that use this mesh and fixup their override colors if necessary.
 		// Also invalidate lighting. *** WARNING components may be reattached here! ***
-		for( TObjectIterator<UStaticMeshComponent> It; It; ++It )
+		const uint32 NumLODs = RenderData->LODResources.Num();
+		for (TObjectIterator<UStaticMeshComponent> It; It; ++It)
 		{
 			if ( It->StaticMesh == this )
 			{
-				It->FixupOverrideColorsIfNecessary( true );
+				// Initialize override vertex colors on any new LODs which have just been created
+				It->SetLODDataCount(NumLODs, It->LODData.Num());
+
+				FStaticMeshComponentLODInfo& LOD0Info = It->LODData[0];
+				if (LOD0Info.OverrideVertexColors)
+				{
+					for (uint32 LODIndex = 1; LODIndex < NumLODs; ++LODIndex)
+					{
+						FStaticMeshComponentLODInfo& LODInfo = It->LODData[LODIndex];
+
+						if (LODInfo.OverrideVertexColors == nullptr && LODInfo.PaintedVertices.Num() == 0)
+						{
+							LODInfo.OverrideVertexColors = new FColorVertexBuffer;
+
+							FStaticMeshLODResources& CurRenderData = RenderData->LODResources[LODIndex];
+
+							TArray<FColor> NewOverrideColors;
+
+							if (LOD0Info.PaintedVertices.Num() > 0)
+							{
+								// Build override colors for LOD, based on LOD0
+								RemapPaintedVertexColors(
+									LOD0Info.PaintedVertices,
+									*LOD0Info.OverrideVertexColors,
+									CurRenderData.PositionVertexBuffer,
+									&CurRenderData.VertexBuffer,
+									NewOverrideColors
+									);
+							}
+							if (NewOverrideColors.Num())
+							{
+								LODInfo.OverrideVertexColors->InitFromColorArray(NewOverrideColors);
+
+								// Update the PaintedVertices array
+								const int32 NumVerts = CurRenderData.GetNumVertices();
+								check(NumVerts == NewOverrideColors.Num());
+
+								LODInfo.PaintedVertices.Reserve(NumVerts);
+								for (int32 VertIndex = 0; VertIndex < NumVerts; ++VertIndex)
+								{
+									FPaintedVertex* Vertex = new(LODInfo.PaintedVertices) FPaintedVertex;
+									Vertex->Position = CurRenderData.PositionVertexBuffer.VertexPosition(VertIndex);
+									Vertex->Normal = CurRenderData.VertexBuffer.VertexTangentZ(VertIndex);
+									Vertex->Color = LODInfo.OverrideVertexColors->VertexColor(VertIndex);
+								}
+							}
+						}
+					}
+				}
+
+				It->FixupOverrideColorsIfNecessary(true);
 				It->InvalidateLightingCache();
 			}
 		}
-
 	}
 
 	// Calculate extended bounds
diff --git a/Engine/Source/Runtime/Engine/Private/Tests/AnalyticsTest.cpp b/Engine/Source/Runtime/Engine/Private/Tests/AnalyticsTest.cpp
--- a/Engine/Source/Runtime/Engine/Private/Tests/AnalyticsTest.cpp
+++ b/Engine/Source/Runtime/Engine/Private/Tests/AnalyticsTest.cpp
@@ -7,6 +7,8 @@
 #include "Runtime/Analytics/Analytics/Public/AnalyticsEventAttribute.h"
 #include "UnrealEngine.h"
 
+#if WITH_DEV_AUTOMATION_TESTS
+
 /**
 * Artificial Record Event for analytics - Simulates the engine startup simulation.  
 */
@@ -137,4 +139,6 @@ bool FAnalyticsEventAttributeUnitTest::RunTest(const FString& Parameters)
 
 	ExecutionInfo.LogItems.Add(TEXT("SKIPPED 'FAnalyticsEventAttributeUnitTest' test.  EngineAnalytics are not currently available."));
 	return true;
-}
\ No newline at end of file
+}
+
+#endif //WITH_DEV_AUTOMATION_TESTS
\ No newline at end of file
diff --git a/Engine/Source/Runtime/Engine/Private/Tests/EngineAutomationTests.cpp b/Engine/Source/Runtime/Engine/Private/Tests/EngineAutomationTests.cpp
--- a/Engine/Source/Runtime/Engine/Private/Tests/EngineAutomationTests.cpp
+++ b/Engine/Source/Runtime/Engine/Private/Tests/EngineAutomationTests.cpp
@@ -13,6 +13,7 @@
 #include "PlatformFeatures.h"
 #include "SaveGameSystem.h"
 
+#if WITH_DEV_AUTOMATION_TESTS
 
 namespace
 {
@@ -479,10 +480,12 @@ bool FAutomationLogAddError::RunTest(const FString& Parameters)
 	return true;
 }
 
+#endif //WITH_DEV_AUTOMATION_TESTS
+
 /* UAutomationTestSettings interface
- *****************************************************************************/
+*****************************************************************************/
 
-UAutomationTestSettings::UAutomationTestSettings( const FObjectInitializer& ObjectInitializer )
+UAutomationTestSettings::UAutomationTestSettings(const FObjectInitializer& ObjectInitializer)
 	: Super(ObjectInitializer)
 {
 }
diff --git a/Engine/Source/Runtime/Engine/Private/Tests/Internationalization/TextPropertyTest.cpp b/Engine/Source/Runtime/Engine/Private/Tests/Internationalization/TextPropertyTest.cpp
--- a/Engine/Source/Runtime/Engine/Private/Tests/Internationalization/TextPropertyTest.cpp
+++ b/Engine/Source/Runtime/Engine/Private/Tests/Internationalization/TextPropertyTest.cpp
@@ -7,10 +7,12 @@
 #define LOCTEXT_NAMESPACE "TextPropertyTest"
 
 UTextPropertyTestObject::UTextPropertyTestObject(const FObjectInitializer& ObjectInitializer)
-	: Super(ObjectInitializer), DefaultedText( LOCTEXT("DefaultedText", "DefaultValue") )
+	: Super(ObjectInitializer), DefaultedText(LOCTEXT("DefaultedText", "DefaultValue"))
 {
 }
 
+#if WITH_DEV_AUTOMATION_TESTS
+
 IMPLEMENT_SIMPLE_AUTOMATION_TEST(FTextPropertyTest, "System.Engine.Internationalization.Text Property Test", EAutomationTestFlags::ApplicationContextMask | EAutomationTestFlags::SmokeFilter)
 
 bool FTextPropertyTest::RunTest (const FString& Parameters)
@@ -35,7 +37,9 @@ bool FTextPropertyTest::RunTest (const FString& Parameters)
 		{
 			FString ExportedStringValue;
 			DefaultedTextProperty->ExportTextItem(ExportedStringValue, &(NewUObject->DefaultedText), NULL, NULL, 0, NULL);
-			if (ExportedStringValue != NewUObject->DefaultedText.ToString())
+			FText ExportedTextValue;
+			FTextStringHelper::ReadFromString(*ExportedStringValue, ExportedTextValue);
+			if (ExportedTextValue.ToString() != NewUObject->DefaultedText.ToString())
 			{
 				AddError(TEXT("UTextProperty::ExportTextItem failed to provide the display string."));
 			}
@@ -131,4 +135,6 @@ bool FTextPropertyTest::RunTest (const FString& Parameters)
 	return true;
 }
 
-#undef LOCTEXT_NAMESPACE
\ No newline at end of file
+#undef LOCTEXT_NAMESPACE
+
+#endif //WITH_DEV_AUTOMATION_TESTS
\ No newline at end of file
diff --git a/Engine/Source/Runtime/Engine/Private/Tests/OSSTests.cpp b/Engine/Source/Runtime/Engine/Private/Tests/OSSTests.cpp
--- a/Engine/Source/Runtime/Engine/Private/Tests/OSSTests.cpp
+++ b/Engine/Source/Runtime/Engine/Private/Tests/OSSTests.cpp
@@ -6,6 +6,8 @@
 
 #include "AutomationCommon.h"
 
+#if WITH_DEV_AUTOMATION_TESTS
+
 DEFINE_LOG_CATEGORY_STATIC(LogHackAutomationTests, Log, All);
 
 
@@ -286,4 +288,6 @@ bool FNullAchievementsTest::RunTest(const FString& Parameters)
 }
 */
 
-#endif
+#endif //WITH_DEV_AUTOMATION_TESTS
+
+#endif
\ No newline at end of file
diff --git a/Engine/Source/Runtime/JsonUtilities/Private/JsonObjectConverter.cpp b/Engine/Source/Runtime/JsonUtilities/Private/JsonObjectConverter.cpp
--- a/Engine/Source/Runtime/JsonUtilities/Private/JsonObjectConverter.cpp
+++ b/Engine/Source/Runtime/JsonUtilities/Private/JsonObjectConverter.cpp
@@ -49,6 +49,10 @@ TSharedPtr<FJsonValue> ConvertScalarUPropertyToJsonValue(UProperty* Property, co
 	{
 		return MakeShareable(new FJsonValueString(StringProperty->GetPropertyValue(Value)));
 	}
+	else if (UTextProperty *TextProperty = Cast<UTextProperty>(Property))
+	{
+		return MakeShareable(new FJsonValueString(TextProperty->GetPropertyValue(Value).ToString()));
+	}
 	else if (UArrayProperty *ArrayProperty = Cast<UArrayProperty>(Property))
 	{
 		TArray< TSharedPtr<FJsonValue> > Out;
diff --git a/Engine/Source/Runtime/Online/OnlineSubsystem/Private/OnlineKeyValuePair.cpp b/Engine/Source/Runtime/Online/OnlineSubsystem/Private/OnlineKeyValuePair.cpp
--- a/Engine/Source/Runtime/Online/OnlineSubsystem/Private/OnlineKeyValuePair.cpp
+++ b/Engine/Source/Runtime/Online/OnlineSubsystem/Private/OnlineKeyValuePair.cpp
@@ -1109,6 +1109,11 @@ bool FVariantDataConverter::ConvertScalarUPropertyToVariant(UProperty* Property,
 		FString StringValue = StringProperty->GetPropertyValue(Value);
 		OutVariantData.SetValue(StringValue);
 	}
+	else if (UTextProperty *TextProperty = Cast<UTextProperty>(Property))
+	{
+		FText TextValue = TextProperty->GetPropertyValue(Value);
+		OutVariantData.SetValue(TextValue.ToString());
+	}
 	else if (UArrayProperty *ArrayProperty = Cast<UArrayProperty>(Property))
 	{
 		TSharedPtr<FJsonValue> Json = FJsonObjectConverter::UPropertyToJsonValue(Property, Value, 0, 0);
diff --git a/Engine/Source/Runtime/Online/OnlineSubsystemUtils/Private/OnlineSubsystemUtils.cpp b/Engine/Source/Runtime/Online/OnlineSubsystemUtils/Private/OnlineSubsystemUtils.cpp
--- a/Engine/Source/Runtime/Online/OnlineSubsystemUtils/Private/OnlineSubsystemUtils.cpp
+++ b/Engine/Source/Runtime/Online/OnlineSubsystemUtils/Private/OnlineSubsystemUtils.cpp
@@ -162,6 +162,7 @@ static bool OnlineExec( UWorld* InWorld, const TCHAR* Cmd, FOutputDevice& Ar )
 			{
 				if (FParse::Command(&Cmd, TEXT("TEST")))
 				{
+#if WITH_DEV_AUTOMATION_TESTS
 					if (FParse::Command(&Cmd, TEXT("FRIENDS")))
 					{
 						TArray<FString> Invites;
@@ -305,6 +306,7 @@ static bool OnlineExec( UWorld* InWorld, const TCHAR* Cmd, FOutputDevice& Ar )
 						(new FTestExternalUIInterface(SubName, bTestLoginUI, bTestFriendsUI, bTestInviteUI, bTestAchievementsUI, bTestWebURL, bTestProfileUI))->Test();
 						bWasHandled = true;
 					}
+#endif //WITH_DEV_AUTOMATION_TESTS
 				}
 			}
 		}
diff --git a/Engine/Source/Runtime/Online/OnlineSubsystemUtils/Private/Tests/OnlineSubsystemUtilitiesTests.cpp b/Engine/Source/Runtime/Online/OnlineSubsystemUtils/Private/Tests/OnlineSubsystemUtilitiesTests.cpp
--- a/Engine/Source/Runtime/Online/OnlineSubsystemUtils/Private/Tests/OnlineSubsystemUtilitiesTests.cpp
+++ b/Engine/Source/Runtime/Online/OnlineSubsystemUtils/Private/Tests/OnlineSubsystemUtilitiesTests.cpp
@@ -3,6 +3,8 @@
 #include "OnlineSubsystemUtilsPrivatePCH.h"
 #include "TestIdentityInterface.h"
 
+#if WITH_DEV_AUTOMATION_TESTS
+
 DEFINE_LOG_CATEGORY_STATIC(OSSUtilsTestLog, All, All);
 
 namespace OSSUtilsTestHelper
@@ -194,4 +196,6 @@ bool FTestIdentityInterfaceLogoutWhileNotLoggedIn::RunTest(const FString& Parame
 	ADD_LATENT_AUTOMATION_COMMAND( FTestIdentityInterfaceLoginOut( LogoutTestStruct ) );
 
 	return true;
-}
\ No newline at end of file
+}
+
+#endif //#if WITH_DEV_AUTOMATION_TESTS
\ No newline at end of file
diff --git a/Engine/Source/Runtime/Online/OnlineSubsystemUtils/Private/Tests/TestAchievementsInterface.cpp b/Engine/Source/Runtime/Online/OnlineSubsystemUtils/Private/Tests/TestAchievementsInterface.cpp
--- a/Engine/Source/Runtime/Online/OnlineSubsystemUtils/Private/Tests/TestAchievementsInterface.cpp
+++ b/Engine/Source/Runtime/Online/OnlineSubsystemUtils/Private/Tests/TestAchievementsInterface.cpp
@@ -4,6 +4,7 @@
 #include "ModuleManager.h"
 #include "TestAchievementsInterface.h"
 
+#if WITH_DEV_AUTOMATION_TESTS
 
 /**
  * The process of this test is:
@@ -201,4 +202,6 @@ void FTestAchievementsInterface::OnAchievementsWritten(const FUniqueNetId& Playe
 void FTestAchievementsInterface::OnAchievementsUnlocked(const FUniqueNetId& PlayerId, const FString& AchievementId)
 {
 	UE_LOG(LogOnline, Display, TEXT("Achievement Unlocked - %s"), *AchievementId);
-}
\ No newline at end of file
+}
+
+#endif //WITH_DEV_AUTOMATION_TESTS
\ No newline at end of file
diff --git a/Engine/Source/Runtime/Online/OnlineSubsystemUtils/Private/Tests/TestTitleFileInterface.h b/Engine/Source/Runtime/Online/OnlineSubsystemUtils/Private/Tests/TestTitleFileInterface.h
--- a/Engine/Source/Runtime/Online/OnlineSubsystemUtils/Private/Tests/TestTitleFileInterface.h
+++ b/Engine/Source/Runtime/Online/OnlineSubsystemUtils/Private/Tests/TestTitleFileInterface.h
@@ -4,6 +4,8 @@
 
 #include "OnlineTitleFileInterface.h"
 
+#if WITH_DEV_AUTOMATION_TESTS
+
 class FTestTitleFileInterface
 {
 public:
@@ -40,4 +42,6 @@ private:
 
 	/** Async file reads still pending completion */
 	int32 NumPendingFileReads;
-};
\ No newline at end of file
+};
+
+#endif //WITH_DEV_AUTOMATION_TESTS
\ No newline at end of file
diff --git a/Engine/Source/Runtime/Online/OnlineSubsystemUtils/Private/Tests/TestVoiceData.h b/Engine/Source/Runtime/Online/OnlineSubsystemUtils/Private/Tests/TestVoiceData.h
--- a/Engine/Source/Runtime/Online/OnlineSubsystemUtils/Private/Tests/TestVoiceData.h
+++ b/Engine/Source/Runtime/Online/OnlineSubsystemUtils/Private/Tests/TestVoiceData.h
@@ -2,7 +2,7 @@
 
 #pragma once
 
-#if UE_BUILD_SHIPPING
+#if WITH_DEV_AUTOMATION_TESTS
 
 uint8 RawVoiceTestData[] = { 0x0, 0x0 };
 
diff --git a/Engine/Source/Runtime/PhysXFormats/Private/PhysXFormats.cpp b/Engine/Source/Runtime/PhysXFormats/Private/PhysXFormats.cpp
--- a/Engine/Source/Runtime/PhysXFormats/Private/PhysXFormats.cpp
+++ b/Engine/Source/Runtime/PhysXFormats/Private/PhysXFormats.cpp
@@ -105,7 +105,7 @@ public:
 		}
 
 		// Set up cooking
-		const PxCookingParams& Params = PhysXCooking->getParams();
+		const PxCookingParams Params = PhysXCooking->getParams();
 		PxCookingParams NewParams = Params;
 		NewParams.targetPlatform = PhysXFormat;
 
@@ -119,11 +119,15 @@ public:
 			// Meshes which can be deformed need different cooking parameters to inhibit vertex welding and add an extra skin around the collision mesh for safety.
 			// We need to set the meshWeldTolerance to zero, even when disabling 'clean mesh' as PhysX will attempt to perform mesh cleaning anyway according to this meshWeldTolerance
 			// if the convex hull is not well formed.
-			NewParams.meshPreprocessParams = PxMeshPreprocessingFlags(PxMeshPreprocessingFlag::eDISABLE_CLEAN_MESH | PxMeshPreprocessingFlag::eREMOVE_UNREFERENCED_VERTICES | PxMeshPreprocessingFlag::eREMOVE_DUPLICATED_TRIANGLES);
+			// Set the skin thickness as a proportion of the overall size of the mesh as PhysX's internal tolerances also use the overall size to calculate the epsilon used.
+			const FBox Bounds(SrcBuffer);
+			const float MaxExtent = (Bounds.Max - Bounds.Min).Size();
+			NewParams.skinWidth = MaxExtent / 512.0f;
+			NewParams.meshPreprocessParams = PxMeshPreprocessingFlags(PxMeshPreprocessingFlag::eDISABLE_CLEAN_MESH);
 			NewParams.areaTestEpsilon = 0.0f;
 			NewParams.meshWeldTolerance = 0.0f;
+			PhysXCooking->setParams(NewParams);
 		}
-		PhysXCooking->setParams(NewParams);
 
 		// Cook the convex mesh to a temp buffer
 		TArray<uint8> CookedMeshBuffer;
@@ -143,9 +147,7 @@ public:
 		// Return default cooking params to normal
 		if (bDeformableMesh)
 		{
-			NewParams.meshPreprocessParams = Params.meshPreprocessParams;
-			NewParams.areaTestEpsilon = Params.areaTestEpsilon;
-			NewParams.meshWeldTolerance = Params.meshWeldTolerance;
+			PhysXCooking->setParams(Params);
 		}
 
 		if( Result && CookedMeshBuffer.Num() > 0 )
diff --git a/Engine/Source/Runtime/Serialization/Private/Backends/JsonStructDeserializerBackend.cpp b/Engine/Source/Runtime/Serialization/Private/Backends/JsonStructDeserializerBackend.cpp
--- a/Engine/Source/Runtime/Serialization/Private/Backends/JsonStructDeserializerBackend.cpp
+++ b/Engine/Source/Runtime/Serialization/Private/Backends/JsonStructDeserializerBackend.cpp
@@ -263,7 +263,12 @@ bool FJsonStructDeserializerBackend::ReadProperty( UProperty* Property, UPropert
 			
 			if (Property->GetClass() == UTextProperty::StaticClass())
 			{
-				return SetPropertyValue<UTextProperty, FText>(Property, Outer, Data, ArrayIndex, FText::FromString(StringValue));
+				FText TextValue;
+				if (!FTextStringHelper::ReadFromString(*StringValue, TextValue))
+				{
+					TextValue = FText::FromString(StringValue);
+				}
+				return SetPropertyValue<UTextProperty, FText>(Property, Outer, Data, ArrayIndex, TextValue);
 			}
 
 			if (Property->GetClass() == UByteProperty::StaticClass())
diff --git a/Engine/Source/Runtime/Serialization/Private/Backends/JsonStructSerializerBackend.cpp b/Engine/Source/Runtime/Serialization/Private/Backends/JsonStructSerializerBackend.cpp
--- a/Engine/Source/Runtime/Serialization/Private/Backends/JsonStructSerializerBackend.cpp
+++ b/Engine/Source/Runtime/Serialization/Private/Backends/JsonStructSerializerBackend.cpp
@@ -198,7 +198,9 @@ void FJsonStructSerializerBackend::WriteProperty(const FStructSerializerState& S
 	}
 	else if (State.ValueType == UTextProperty::StaticClass())
 	{
-		WritePropertyValue(JsonWriter, State, Cast<UTextProperty>(State.ValueProperty)->GetPropertyValue_InContainer(State.ValueData, ArrayIndex).ToString());
+		FString StringValue;
+		FTextStringHelper::WriteToString(StringValue, Cast<UTextProperty>(State.ValueProperty)->GetPropertyValue_InContainer(State.ValueData, ArrayIndex));
+		WritePropertyValue(JsonWriter, State, StringValue);
 	}
 
 	// classes & objects
diff --git a/Engine/Source/Runtime/Slate/Private/Framework/Docking/TabManager.cpp b/Engine/Source/Runtime/Slate/Private/Framework/Docking/TabManager.cpp
--- a/Engine/Source/Runtime/Slate/Private/Framework/Docking/TabManager.cpp
+++ b/Engine/Source/Runtime/Slate/Private/Framework/Docking/TabManager.cpp
@@ -392,7 +392,7 @@ void FTabManager::FPrivateApi::OnTabManagerClosing()
 
 bool FTabManager::FPrivateApi::CanTabLeaveTabWell(const TSharedRef<const SDockTab>& TabToTest) const
 {
-	return TabToTest != TabManager.MainNonCloseableTab.Pin();
+	return TabManager.bCanDoDragOperation && TabToTest != TabManager.MainNonCloseableTab.Pin();
 }
 
 const TArray< TWeakPtr<SDockingArea> >& FTabManager::FPrivateApi::GetLiveDockAreas() const
@@ -1037,6 +1037,7 @@ FTabManager::FTabManager( const TSharedPtr<SDockTab>& InOwnerTab, const TSharedR
 , PrivateApi( MakeShareable(new FPrivateApi(*this)) )
 , LastDocumentUID( 0 )
 , bIsSavingVisualState( false )
+, bCanDoDragOperation( true )
 {
 	LocalWorkspaceMenuRoot = FWorkspaceItem::NewGroup(LOCTEXT("LocalWorkspaceRoot", "Local Workspace Root"));
 }
diff --git a/Engine/Source/Runtime/Slate/Private/Framework/Text/Tests/RichTextMarkupProcessingTest.cpp b/Engine/Source/Runtime/Slate/Private/Framework/Text/Tests/RichTextMarkupProcessingTest.cpp
--- a/Engine/Source/Runtime/Slate/Private/Framework/Text/Tests/RichTextMarkupProcessingTest.cpp
+++ b/Engine/Source/Runtime/Slate/Private/Framework/Text/Tests/RichTextMarkupProcessingTest.cpp
@@ -2,7 +2,7 @@
 
 #include "SlatePrivatePCH.h"
 
-#if WITH_FANCY_TEXT
+#if WITH_FANCY_TEXT & WITH_DEV_AUTOMATION_TESTS
 
 #include "AutomationTest.h"
 
@@ -385,4 +385,4 @@ bool FRTFProcessingTest::RunTest (const FString& Parameters)
 
 PRAGMA_ENABLE_OPTIMIZATION
 
-#endif //WITH_FANCY_TEXT
\ No newline at end of file
+#endif //WITH_FANCY_TEXT & WITH_DEV_AUTOMATION_TESTS
\ No newline at end of file
diff --git a/Engine/Source/Runtime/Slate/Private/Widgets/Docking/SDockTab.cpp b/Engine/Source/Runtime/Slate/Private/Widgets/Docking/SDockTab.cpp
--- a/Engine/Source/Runtime/Slate/Private/Widgets/Docking/SDockTab.cpp
+++ b/Engine/Source/Runtime/Slate/Private/Widgets/Docking/SDockTab.cpp
@@ -70,9 +70,24 @@ FReply SDockTab::OnDragDetected( const FGeometry& MyGeometry, const FPointerEven
 		FMath::Clamp(TabGrabOffset.Y / TabSize.Y, 0.0f, 1.0f) );
 	
 	auto PinnedParent = ParentPtr.Pin();
-	return ensure(PinnedParent.IsValid())
-		? PinnedParent->StartDraggingTab(SharedThis(this), TabGrabOffsetFraction, MouseEvent)
-		: FReply::Unhandled();
+	if (PinnedParent.IsValid())
+	{
+		//See if we can drag tabs contain in this manager
+		if (GetTabManager()->GetCanDoDragOperation())
+		{
+			return PinnedParent->StartDraggingTab(SharedThis(this), TabGrabOffsetFraction, MouseEvent);
+		}
+		else
+		{
+			return FReply::Handled();
+		}
+	}
+	else
+	{
+		// Should never get here (but sometimes does, under unknown circumstances)
+		// TODO: investigate how the parent pointer can become invalid
+		return FReply::Unhandled();
+	}
 }
 
 FReply SDockTab::OnMouseButtonUp( const FGeometry& MyGeometry, const FPointerEvent& MouseEvent )
diff --git a/Engine/Source/Runtime/Slate/Private/Widgets/Images/SThrobber.cpp b/Engine/Source/Runtime/Slate/Private/Widgets/Images/SThrobber.cpp
--- a/Engine/Source/Runtime/Slate/Private/Widgets/Images/SThrobber.cpp
+++ b/Engine/Source/Runtime/Slate/Private/Widgets/Images/SThrobber.cpp
@@ -152,7 +152,7 @@ int32 SCircularThrobber::OnPaint( const FPaintArgs& Args, const FGeometry& Allot
 			(PieceIdx + 1) / (float)NumPieces,
 			LocalOffset + LocalOffset * FVector2D(FMath::Sin(Angle), FMath::Cos(Angle)));
 		FPaintGeometry PaintGeom = AllottedGeometry.ToPaintGeometry(PieceImage->ImageSize, PieceLocalTransform);
-		FSlateDrawElement::MakeBox(OutDrawElements, LayerId, PaintGeom, PieceImage, MyClippingRect, ESlateDrawEffect::None, InWidgetStyle.GetColorAndOpacityTint() );
+		FSlateDrawElement::MakeBox(OutDrawElements, LayerId, PaintGeom, PieceImage, MyClippingRect, ESlateDrawEffect::None, FinalColorAndOpacity);
 	}
 	
 	return LayerId;
diff --git a/Engine/Source/Runtime/Slate/Private/Widgets/Layout/SScissorRectBox.cpp b/Engine/Source/Runtime/Slate/Private/Widgets/Layout/SScissorRectBox.cpp
--- a/Engine/Source/Runtime/Slate/Private/Widgets/Layout/SScissorRectBox.cpp
+++ b/Engine/Source/Runtime/Slate/Private/Widgets/Layout/SScissorRectBox.cpp
@@ -9,13 +9,14 @@ SScissorRectBox::SScissorRectBox()
 
 void SScissorRectBox::Construct(const FArguments& InArgs)
 {
+	bCanTick = false;
+
 	ChildSlot
 	[
 		InArgs._Content.Widget
 	];
 }
 
-
 FVector2D SScissorRectBox::ComputeDesiredSize( float ) const
 {
 	EVisibility ChildVisibility = ChildSlot.GetWidget()->GetVisibility();
diff --git a/Engine/Source/Runtime/Slate/Public/Framework/Docking/TabManager.h b/Engine/Source/Runtime/Slate/Public/Framework/Docking/TabManager.h
--- a/Engine/Source/Runtime/Slate/Public/Framework/Docking/TabManager.h
+++ b/Engine/Source/Runtime/Slate/Public/Framework/Docking/TabManager.h
@@ -619,6 +619,12 @@ class SLATE_API FTabManager : public TSharedFromThis<FTabManager>
 		/** Provide a tab that will be the main tab and cannot be closed. */
 		void SetMainTab(const TSharedRef<const SDockTab>& InTab);
 
+		/* Prevent or allow all tabs to be drag */
+		void SetCanDoDragOperation(bool CanDoDragOperation) { bCanDoDragOperation = CanDoDragOperation; }
+		
+		/* Return true if we can do drag operation */
+		bool GetCanDoDragOperation() { return bCanDoDragOperation; }
+
 		/** @return if the provided tab can be closed. */
 		bool IsTabCloseable(const TSharedRef<const SDockTab>& InTab) const;
 
@@ -751,6 +757,9 @@ class SLATE_API FTabManager : public TSharedFromThis<FTabManager>
 
 		/** The main tab, this tab cannot be closed. */
 		TWeakPtr<const SDockTab> MainNonCloseableTab;
+
+		/* Prevent or allow Drag operation. */
+		bool bCanDoDragOperation;
 };
 
 
diff --git a/Engine/Source/Runtime/SlateCore/Private/Styling/CoreStyle.cpp b/Engine/Source/Runtime/SlateCore/Private/Styling/CoreStyle.cpp
--- a/Engine/Source/Runtime/SlateCore/Private/Styling/CoreStyle.cpp
+++ b/Engine/Source/Runtime/SlateCore/Private/Styling/CoreStyle.cpp
@@ -633,6 +633,7 @@ TSharedRef<ISlateStyle> FCoreStyle::Create( const FName& InStyleSetName )
 		Style->Set( "ColorPicker.Font", TTF_FONT( "Fonts/Roboto-Regular", 10 ) );
 		Style->Set( "ColorPicker.Mode", new IMAGE_BRUSH( "Common/ColorPicker_Mode_16x", Icon16x16) );
 		Style->Set( "ColorPicker.Separator", new IMAGE_BRUSH( "Common/ColorPicker_Separator", FVector2D(2.0f, 2.0f) ) );
+		Style->Set( "ColorPicker.Selector", new IMAGE_BRUSH( "Common/Circle", FVector2D(8, 8) ) );
 		Style->Set( "ColorPicker.Slider", FSliderStyle()
 			.SetDisabledThumbImage( IMAGE_BRUSH( "Common/ColorPicker_SliderHandle", FVector2D(8.0f, 32.0f)) )
 			.SetNormalThumbImage( IMAGE_BRUSH( "Common/ColorPicker_SliderHandle", FVector2D(8.0f, 32.0f)) )
diff --git a/Engine/Source/Runtime/SlateCore/Private/Widgets/SWindow.cpp b/Engine/Source/Runtime/SlateCore/Private/Widgets/SWindow.cpp
--- a/Engine/Source/Runtime/SlateCore/Private/Widgets/SWindow.cpp
+++ b/Engine/Source/Runtime/SlateCore/Private/Widgets/SWindow.cpp
@@ -231,6 +231,9 @@ void SWindow::Construct(const FArguments& InArgs)
 	// If the window has no OS border, simulate it ourselves, enlarging window by the size that OS border would have.
 	FVector2D WindowSize = GetWindowSizeFromClientSize(InArgs._ClientSize);
 
+	// Get change in size resulting from the above call
+	const FVector2D DeltaSize = WindowSize - InArgs._ClientSize;
+
 	// calculate initial window position
 	FVector2D WindowPosition = InArgs._ScreenPosition;
 
@@ -302,9 +305,12 @@ void SWindow::Construct(const FArguments& InArgs)
 			break;
 		}
 
-		// Clamp window size to be no greater than the work area size
-		WindowSize.X = FMath::Min(WindowSize.X, AutoCenterRect.GetSize().X);
-		WindowSize.Y = FMath::Min(WindowSize.Y, AutoCenterRect.GetSize().Y);
+		if (InArgs._SaneWindowPlacement)
+		{
+			// Clamp window size to be no greater than the work area size
+			WindowSize.X = FMath::Min(WindowSize.X, AutoCenterRect.GetSize().X);
+			WindowSize.Y = FMath::Min(WindowSize.Y, AutoCenterRect.GetSize().Y);
+		}
 
 		// Setup a position and size for the main frame window that's centered in the desktop work area
 		const FVector2D DisplayTopLeft( AutoCenterRect.Left, AutoCenterRect.Top );
@@ -320,7 +326,8 @@ void SWindow::Construct(const FArguments& InArgs)
 	this->InitialDesiredScreenPosition = WindowPosition;
 	this->InitialDesiredSize = WindowSize;
 
-	Resize(WindowSize);
+	// Resize adds extra borders / title bar if necessary, but this is already taken into account in WindowSize, so subtract them again first
+	Resize(WindowSize -DeltaSize);
 
 	// Window visibility is currently driven by whether the window is interactive.
 	this->Visibility = TAttribute<EVisibility>::Create( TAttribute<EVisibility>::FGetter::CreateRaw(this, &SWindow::GetWindowVisibility) );
diff --git a/Engine/Source/Runtime/UMG/Private/Components/CircularThrobber.cpp b/Engine/Source/Runtime/UMG/Private/Components/CircularThrobber.cpp
--- a/Engine/Source/Runtime/UMG/Private/Components/CircularThrobber.cpp
+++ b/Engine/Source/Runtime/UMG/Private/Components/CircularThrobber.cpp
@@ -10,7 +10,8 @@
 // UCircularThrobber
 
 UCircularThrobber::UCircularThrobber(const FObjectInitializer& ObjectInitializer)
-	: Super(ObjectInitializer)
+	: Super(ObjectInitializer),
+	  bEnableRadius(true)
 {
 	SCircularThrobber::FArguments DefaultArgs;
 	Image = *DefaultArgs._PieceImage;
@@ -47,6 +48,16 @@ void UCircularThrobber::SynchronizeProperties()
 	MyCircularThrobber->SetNumPieces(FMath::Clamp(NumberOfPieces, 1, 25));
 	MyCircularThrobber->SetPeriod(Period);
 	MyCircularThrobber->SetRadius(Radius);
+
+	// If widget is child of Canvas Panel and 'Size to Content' is enabled, we allow user to modify radius.
+	bEnableRadius = true;
+	if (UCanvasPanelSlot* Panel = Cast<UCanvasPanelSlot>(Slot))
+	{
+		if (!Panel->GetAutoSize())
+		{
+			bEnableRadius = false;
+		}
+	}
 }
 
 void UCircularThrobber::SetNumberOfPieces(int32 InNumberOfPieces)
diff --git a/Engine/Source/Runtime/UMG/Public/Components/CircularThrobber.h b/Engine/Source/Runtime/UMG/Public/Components/CircularThrobber.h
--- a/Engine/Source/Runtime/UMG/Public/Components/CircularThrobber.h
+++ b/Engine/Source/Runtime/UMG/Public/Components/CircularThrobber.h
@@ -27,8 +27,8 @@ public:
 	UPROPERTY(EditAnywhere, BlueprintReadOnly, Category=Appearance)
 	float Period;
 
-	/** The radius of the circle */
-	UPROPERTY(EditAnywhere, BlueprintReadOnly, Category=Appearance)
+	/** The radius of the circle. If the throbber is a child of Canvas Panel, the 'Size to Content' option must be enabled in order to set Radius. */
+	UPROPERTY(EditAnywhere, BlueprintReadOnly, Category=Appearance, meta=(EditCondition="bEnableRadius"))
 	float Radius;
 
 	/** Image to use for each segment of the throbber */
@@ -75,6 +75,10 @@ protected:
 	//~ End UWidget Interface
 
 private:
+
+	UPROPERTY(Transient)
+	bool bEnableRadius;
+
 	/** The CircularThrobber widget managed by this object. */
 	TSharedPtr<SCircularThrobber> MyCircularThrobber;
 };
diff --git a/Engine/Source/Runtime/UnrealAudio/Private/Tests/UnrealAudioDeviceTests.cpp b/Engine/Source/Runtime/UnrealAudio/Private/Tests/UnrealAudioDeviceTests.cpp
--- a/Engine/Source/Runtime/UnrealAudio/Private/Tests/UnrealAudioDeviceTests.cpp
+++ b/Engine/Source/Runtime/UnrealAudio/Private/Tests/UnrealAudioDeviceTests.cpp
@@ -7,7 +7,7 @@
 #include "UnrealAudioTestGenerators.h"
 #include "UnrealAudioSoundFile.h"
 
-#if ENABLE_UNREAL_AUDIO
+#if ENABLE_UNREAL_AUDIO & WITH_DEV_AUTOMATION_TESTS
 
 namespace UAudio
 {
@@ -247,4 +247,4 @@ namespace UAudio
 
 }
 
-#endif // #if ENABLE_UNREAL_AUDIO
+#endif // #if ENABLE_UNREAL_AUDIO & WITH_DEV_AUTOMATION_TESTS
diff --git a/Engine/Source/Runtime/UnrealAudio/Private/Tests/UnrealAudioSystemTests.cpp b/Engine/Source/Runtime/UnrealAudio/Private/Tests/UnrealAudioSystemTests.cpp
--- a/Engine/Source/Runtime/UnrealAudio/Private/Tests/UnrealAudioSystemTests.cpp
+++ b/Engine/Source/Runtime/UnrealAudio/Private/Tests/UnrealAudioSystemTests.cpp
@@ -4,7 +4,7 @@
 #include "UnrealAudioModule.h"
 #include "UnrealAudioEmitterInternal.h"
 
-#if ENABLE_UNREAL_AUDIO
+#if ENABLE_UNREAL_AUDIO & WITH_DEV_AUTOMATION_TESTS
 
 namespace UAudio
 {
@@ -378,4 +378,4 @@ namespace UAudio
 
 }
 
-#endif // #if ENABLE_UNREAL_AUDIO
+#endif // #if ENABLE_UNREAL_AUDIO & WITH_DEV_AUTOMATION_TESTS
