Fixed shadow variable warning caused by bad merge

diff --git a/Engine/Source/Runtime/CoreUObject/Private/UObject/LinkerLoad.cpp b/Engine/Source/Runtime/CoreUObject/Private/UObject/LinkerLoad.cpp
--- a/Engine/Source/Runtime/CoreUObject/Private/UObject/LinkerLoad.cpp
+++ b/Engine/Source/Runtime/CoreUObject/Private/UObject/LinkerLoad.cpp
@@ -2309,7 +2309,6 @@ FLinkerLoad::EVerifyResult FLinkerLoad::VerifyImport(int32 ImportIndex)
 #if WITH_EDITOR
 				if( GIsEditor && !IsRunningCommandlet())
 				{
-					bool bSupressLinkerError = false;
 					bSupressLinkerError = IsSuppressableBlueprintImportError(Import);
 					if (!bSupressLinkerError)
 					{
