Preventing async loading resumed log spam in the editor.

diff --git a/Engine/Source/Runtime/CoreUObject/Private/Serialization/AsyncLoading.cpp b/Engine/Source/Runtime/CoreUObject/Private/Serialization/AsyncLoading.cpp
--- a/Engine/Source/Runtime/CoreUObject/Private/Serialization/AsyncLoading.cpp
+++ b/Engine/Source/Runtime/CoreUObject/Private/Serialization/AsyncLoading.cpp
@@ -962,7 +962,9 @@ void FAsyncLoadingThread::ResumeLoading()
 {
 	check(IsInGameThread());
 	const int32 SuspendCount = IsLoadingSuspended.Decrement();
+#if !WITH_EDITORONLY_DATA
 	UE_LOG(LogStreaming, Display, TEXT("Resuming async loading (%d)"), SuspendCount);
+#endif
 	UE_CLOG(SuspendCount < 0, LogStreaming, Fatal, TEXT("ResumeAsyncLoadingThread: Async loading was resumed more times than it was suspended."));
 	if (IsMultithreaded() && SuspendCount == 0)
 	{
