Copying //UE4/Dev-Blueprints to Dev-Main (//UE4/Dev-Main)

==========================
MAJOR FEATURES + CHANGES
==========================

Change 2865817 on 2016/02/12 by [REDACTED_PERSON]

	Handling deferred (cyclic) dependency issues that arrise from imports not found within the linker's ImportMap (imports loaded from config files and UProperty::ImportText).

Change 2866282 on 2016/02/13 by [REDACTED_PERSON]

	Blueprint C++ Conversion:
	- Added CustomDynamicClassInitialization meta data and function
	- References to other converted fields are gathered before CDO is created (it solves many dependencies)

Change 2867921 on 2016/02/15 by [REDACTED_PERSON]

	Fixing (CIS error) typo - Accidently assigning uninitialized variable to itself.

Change 2867970 on 2016/02/15 by [REDACTED_PERSON]

	Removing "static" keyword on template specializations (CIS error).

Change 2868401 on 2016/02/16 by [REDACTED_PERSON]

	TBaseStructure for more noexport structures

Change 2868404 on 2016/02/16 by [REDACTED_PERSON]

	Blueprint C++ Conversion:
	- Fixed some issues related to NOEXPORT structures
	- Added  FEmitDefaultValueHelper::SpecialStructureConstructor

Change 2868461 on 2016/02/16 by [REDACTED_PERSON]

	Blueprint C++ Conversion:
	- Fixed component initialization in actors, based on a DynamicClass

Change 2868481 on 2016/02/16 by [REDACTED_PERSON]

	Blueprint C++ Conversion: (Work in progress.)
	In BPGC templates cor SCS components are owned by the class. In DynamicClass the templates are owned by CDO. It requires to update the import path, when a component template is referenced by an extern object.

Change 2868769 on 2016/02/16 by [REDACTED_PERSON]

	Improved parsing for multi-parameter AutoCreateRefTerm meta data

Change 2870310 on 2016/02/17 by [REDACTED_PERSON]

	Incremental refactor and update for the blueprint profiler.

	- Shuffled files around and renamed to make the layout more informational
	- Added profiler capture and playback contexts
	- Added blueprint and function contexts, this enables rapid stat discards in situations such as blueprint compilation.
	- Moved the blueprint execution mapping out of the profiler and into the blueprint contexts
	- Refactored the event playback so it processes with one event at a time.

Change 2870386 on 2016/02/17 by [REDACTED_PERSON]

	Improved  UProperty::ExportCppDeclaration - Const reference to a pointer should be: "Type* const &" , not "const Type* &".

Change 2870686 on 2016/02/17 by [REDACTED_PERSON]

	Engine changes needed for Bullet Train compatibility:
	- Adding BP-exposed GetClosestBone function, with ability to optionally filter only bones associated with physics assets
	- Exposing GetRemoteRole to BPs

Change 2871419 on 2016/02/17 by [REDACTED_PERSON]

	Fixing CIS incude error.

Change 2872190 on 2016/02/18 by [REDACTED_PERSON]

	Another CIS fix - forward declaring a class that wasn't included.

Change 2872285 on 2016/02/18 by [REDACTED_PERSON]

	added bDontNativizeDataOnlyBP in Editor.ini

Change 2872826 on 2016/02/18 by [REDACTED_PERSON]

	CIS fixes for mac builds

Change 2874284 on 2016/02/19 by [REDACTED_PERSON]

	Supporting nested struct properties that are filled out through text imports during deferred dependency loading.

Change 2874299 on 2016/02/19 by [REDACTED_PERSON]

	CIS fix - macro typo in non-debug build.

Change 2875571 on 2016/02/22 by [REDACTED_PERSON]

	Blueprint C++ Conversion:
	Fixed generated FCompiledInDeferStruct instance for nativized structs. It has a proper overriden name.

Change 2875574 on 2016/02/22 by [REDACTED_PERSON]

	Blueprint C++ Conversion:
	Class type of unconverted asset is not replaced in import map while cooking.

Change 2875741 on 2016/02/22 by [REDACTED_PERSON]

	Array Item Get nodes now return by-ref.

	All use cases of the node, as well as use cases of ForEachLoop and ForEachLoopWithBreaks, have been updated to use a Copy node on the output to maintain any existing functionality.

diff --git a/Engine/Config/BaseEditor.ini b/Engine/Config/BaseEditor.ini
--- a/Engine/Config/BaseEditor.ini
+++ b/Engine/Config/BaseEditor.ini
@@ -473,12 +473,13 @@ bUseStaticVariablesInWrappers=true
 bUseStaticVariablesInClasses=true
 bSkipUFUNCTION=false
 bUsePRAGMA_DISABLE_OPTIMIZATION=false
-
-; this flag can change BP behaviour, see UE-23629. It doesn't help in image size reduction.
+bDontNativizeDataOnlyBP=false
+; this flag can change BP behaviour, see UE-23629. Moreover it doesn't help in image size reduction.
 bUseLocalEventGraphVariables=false
 
-;+ExcludedBlueprintTypes=/Script/UMGEditor.WidgetBlueprint
-;+ExcludedBlueprintTypes=/Script/Engine.AnimBlueprint
+; +ExcludedBlueprintTypes=/Script/UMGEditor.WidgetBlueprint
+; +ExcludedBlueprintTypes=/Script/Engine.AnimBlueprint
+; +ExcludedBlueprintTypes=/Script/GameplayAbilities.GameplayAbilityBlueprint
 
 +ExcludedAssets=/Engine/Tutorial/BlueprintTutorials/TutorialAssets/Tutorial_BP_Interface
 +ExcludedAssets=/Engine/Tutorial/SubEditors/TutorialAssets/Character/TutorialCharacter
diff --git a/Engine/Source/Developer/BlueprintCompilerCppBackend/Private/BlueprintCompilerCppBackend.cpp b/Engine/Source/Developer/BlueprintCompilerCppBackend/Private/BlueprintCompilerCppBackend.cpp
--- a/Engine/Source/Developer/BlueprintCompilerCppBackend/Private/BlueprintCompilerCppBackend.cpp
+++ b/Engine/Source/Developer/BlueprintCompilerCppBackend/Private/BlueprintCompilerCppBackend.cpp
@@ -589,6 +589,18 @@ FString FBlueprintCompilerCppBackend::EmitCallStatmentInner(FEmitterLocalContext
 	return Result;
 }
 
+FString FBlueprintCompilerCppBackend::EmitArrayGetByRef(FEmitterLocalContext& EmitterContext, FBlueprintCompiledStatement& Statement)
+{
+	check(Statement.RHS.Num() == 2);
+
+	FString Result;
+	Result += TermToText(EmitterContext, Statement.RHS[0]);
+	Result += TEXT("[");
+	Result += TermToText(EmitterContext, Statement.RHS[1]);
+	Result += TEXT("]");
+	return Result;
+}
+
 FString FBlueprintCompilerCppBackend::TermToText(FEmitterLocalContext& EmitterContext, const FBPTerminal* Term, bool bUseSafeContext, bool bGetter)
 {
 	const FString PSC_Self(TEXT("self"));
@@ -606,6 +618,10 @@ FString FBlueprintCompilerCppBackend::TermToText(FEmitterLocalContext& EmitterCo
 		{
 			return EmitCallStatmentInner(EmitterContext, *Term->InlineGeneratedParameter, true, FString());
 		}
+		else if (KCST_ArrayGetByRef == Term->InlineGeneratedParameter->Type)
+		{
+			return EmitArrayGetByRef(EmitterContext, *Term->InlineGeneratedParameter);
+		}
 		else
 		{
 			ensureMsgf(false, TEXT("KCST %d is not accepted as inline statement."), Term->InlineGeneratedParameter->Type);
diff --git a/Engine/Source/Developer/BlueprintCompilerCppBackend/Private/BlueprintCompilerCppBackend.h b/Engine/Source/Developer/BlueprintCompilerCppBackend/Private/BlueprintCompilerCppBackend.h
--- a/Engine/Source/Developer/BlueprintCompilerCppBackend/Private/BlueprintCompilerCppBackend.h
+++ b/Engine/Source/Developer/BlueprintCompilerCppBackend/Private/BlueprintCompilerCppBackend.h
@@ -54,6 +54,7 @@ protected:
 	FString EmitMethodInputParameterList(FEmitterLocalContext& EmitterContext, FBlueprintCompiledStatement& Statement);
 	FString EmitSwitchValueStatmentInner(FEmitterLocalContext& EmitterContext, FBlueprintCompiledStatement& Statement);
 	FString EmitCallStatmentInner(FEmitterLocalContext& EmitterContext, FBlueprintCompiledStatement& Statement, bool bInline, FString PostFix);
+	FString EmitArrayGetByRef(FEmitterLocalContext& EmitterContext, FBlueprintCompiledStatement& Statement);
 
 public:
 	FString TermToText(FEmitterLocalContext& EmitterContext, const FBPTerminal* Term, bool bUseSafeContext = true, bool bGetter = true);
diff --git a/Engine/Source/Developer/BlueprintCompilerCppBackend/Private/BlueprintCompilerCppBackendBase.cpp b/Engine/Source/Developer/BlueprintCompilerCppBackend/Private/BlueprintCompilerCppBackendBase.cpp
--- a/Engine/Source/Developer/BlueprintCompilerCppBackend/Private/BlueprintCompilerCppBackendBase.cpp
+++ b/Engine/Source/Developer/BlueprintCompilerCppBackend/Private/BlueprintCompilerCppBackendBase.cpp
@@ -226,9 +226,18 @@ FString FBlueprintCompilerCppBackendBase::GenerateCodeFromClass(UClass* SourceCl
 	}
 	else
 	{
+		TArray<FString> AdditionalMD;
+		const FString ReplaceConvertedMD = FEmitHelper::GenerateReplaceConvertedMD(OriginalSourceClass);
+		if (!ReplaceConvertedMD.IsEmpty())
+		{
+			AdditionalMD.Add(ReplaceConvertedMD);
+		}
+
+		AdditionalMD.Add(FString::Printf(TEXT("CustomDynamicClassInitialization=\"%s::__GatherReferencedConvertedFields\""), *CppClassName));
+
 		EmitterContext.Header.AddLine(FString::Printf(TEXT("UCLASS(%s%s)")
 			, (!SourceClass->IsChildOf<UBlueprintFunctionLibrary>()) ? TEXT("Blueprintable, BlueprintType, ") : TEXT("")
-			, *FEmitHelper::ReplaceConvertedMetaData(OriginalSourceClass)));
+			, *FEmitHelper::HandleMetaData(nullptr, false, &AdditionalMD)));
 
 		UClass* SuperClass = SourceClass->GetSuperClass();
 		FString ClassDefinition = FString::Printf(TEXT("class %s : public %s"), *CppClassName, *FEmitHelper::GetCppName(SuperClass));
@@ -259,6 +268,7 @@ FString FBlueprintCompilerCppBackendBase::GenerateCodeFromClass(UClass* SourceCl
 		EmitterContext.Header.AddLine(FString::Printf(TEXT("%s(const FObjectInitializer& ObjectInitializer = FObjectInitializer::Get());"), *CppClassName));
 		EmitterContext.Header.AddLine(TEXT("virtual void PostLoadSubobjects(FObjectInstancingGraph* OuterInstanceGraph) override;"));
 		EmitterContext.Header.AddLine(TEXT("static void __StaticDependenciesAssets(TArray<FBlueprintDependencyData>& AssetsToLoad);"));
+		EmitterContext.Header.AddLine(TEXT("static void __GatherReferencedConvertedFields(UDynamicClass* InDynamicClass);"));
 
 		FEmitDefaultValueHelper::GenerateConstructor(EmitterContext);
 	}
diff --git a/Engine/Source/Developer/BlueprintCompilerCppBackend/Private/BlueprintCompilerCppBackendUtils.cpp b/Engine/Source/Developer/BlueprintCompilerCppBackend/Private/BlueprintCompilerCppBackendUtils.cpp
--- a/Engine/Source/Developer/BlueprintCompilerCppBackend/Private/BlueprintCompilerCppBackendUtils.cpp
+++ b/Engine/Source/Developer/BlueprintCompilerCppBackend/Private/BlueprintCompilerCppBackendUtils.cpp
@@ -870,6 +870,7 @@ FString FEmitHelper::LiteralTerm(FEmitterLocalContext& EmitterContext, const FEd
 	{
 		auto StructType = Cast<UScriptStruct>(Type.PinSubCategoryObject.Get());
 		ensure(StructType);
+
 		if (StructType == TBaseStructure<FVector>::Get())
 		{
 			FVector Vect = FVector::ZeroVector;
@@ -923,7 +924,6 @@ FString FEmitHelper::LiteralTerm(FEmitterLocalContext& EmitterContext, const FEd
 			}
 			else
 			{
-				EmitterContext.AddLine(FString::Printf(TEXT("auto %s = %s{};"), *LocalStructNativeName, *StructName)); // TODO: ?? should "::GetDefaultValue()" be called here?
 				FStructOnScope StructOnScope(StructType);
 
 				class FImportTextErrorContext : public FStringOutputDevice
@@ -950,6 +950,13 @@ FString FEmitHelper::LiteralTerm(FEmitterLocalContext& EmitterContext, const FEd
 					UE_LOG(LogK2Compiler, Error, TEXT("FEmitHelper::LiteralTerm cannot parse struct \"%s\" error: %s class: %s"), *CustomValue, *ImportError, *GetPathNameSafe(EmitterContext.GetCurrentlyGeneratedClass()));
 				}
 
+				FString CustomConstructor;
+				if (FEmitDefaultValueHelper::SpecialStructureConstructor(StructType, StructOnScope.GetStructMemory(), &CustomConstructor))
+				{
+					return CustomConstructor;
+				}
+
+				EmitterContext.AddLine(FString::Printf(TEXT("auto %s = %s{};"), *LocalStructNativeName, *StructName)); // TODO: ?? should "::GetDefaultValue()" be called here?
 				for (auto LocalProperty : TFieldRange<const UProperty>(StructType))
 				{
 					FEmitDefaultValueHelper::OuterGenerate(EmitterContext, LocalProperty, LocalStructNativeName, StructOnScope.GetStructMemory(), nullptr, FEmitDefaultValueHelper::EPropertyAccessOperator::Dot);
@@ -1408,6 +1415,8 @@ FString FEmitHelper::AccessInaccessiblePropertyUsingOffset(FEmitterLocalContext&
 {
 	check(Property);
 
+	UE_LOG(LogK2Compiler, Warning, TEXT("AccessInaccessiblePropertyUsingOffset - NOEXPORT structure should be handled in a custom way: %s"), *GetPathNameSafe(Property->GetOwnerStruct()));
+
 	const int32 PropertyOffsetWithoutEditorOnlyMembers = HelperWithoutEditorOnlyMembers::OffsetWithoutEditorOnlyMembers(Property);
 	const int32 PropertySizeWithoutEditorOnlyMembers = HelperWithoutEditorOnlyMembers::SizeWithoutEditorOnlyMembers(Property);
 	const int32 ElementSizeWithoutEditorOnlyMembers = PropertySizeWithoutEditorOnlyMembers / Property->ArrayDim;
@@ -1423,7 +1432,7 @@ FString FEmitHelper::AccessInaccessiblePropertyUsingOffset(FEmitterLocalContext&
 	const uint32 CppTemplateTypeFlags = EPropertyExportCPPFlags::CPPF_CustomTypeName
 		| EPropertyExportCPPFlags::CPPF_NoConst | EPropertyExportCPPFlags::CPPF_NoRef | EPropertyExportCPPFlags::CPPF_NoStaticArray
 		| EPropertyExportCPPFlags::CPPF_BlueprintCppBackend;
-	const FString TypeDeclaration = EmitterContext.ExportCppDeclaration(Property, EExportedDeclaration::Parameter, CppTemplateTypeFlags, true);
+	const FString TypeDeclaration = EmitterContext.ExportCppDeclaration(Property, EExportedDeclaration::Member, CppTemplateTypeFlags, true);
 	return FString::Printf(TEXT("(*(AccessPrivateProperty<%s>(%s(%s), 0x%08X, 0x%08X, %d)))")
 		, *TypeDeclaration
 		, *ContextAdressOp
diff --git a/Engine/Source/Developer/BlueprintCompilerCppBackend/Private/BlueprintCompilerCppBackendUtils.h b/Engine/Source/Developer/BlueprintCompilerCppBackend/Private/BlueprintCompilerCppBackendUtils.h
--- a/Engine/Source/Developer/BlueprintCompilerCppBackend/Private/BlueprintCompilerCppBackendUtils.h
+++ b/Engine/Source/Developer/BlueprintCompilerCppBackend/Private/BlueprintCompilerCppBackendUtils.h
@@ -290,6 +290,9 @@ struct FEmitDefaultValueHelper
 	// returns empty string if cannot handle
 	static FString HandleClassSubobject(FEmitterLocalContext& Context, UObject* Object, FEmitterLocalContext::EClassSubobjectList ListOfSubobjectsTyp, bool bCreate, bool bInitilize);
 
+	// returns true, and fill OutResult, when the structure is handled in a custom way.
+	static bool SpecialStructureConstructor(const UScriptStruct* Struct, const uint8* ValuePtr, FString* OutResult);
+
 private:
 	// Returns native term, 
 	// returns empty string if cannot handle
diff --git a/Engine/Source/Developer/BlueprintCompilerCppBackend/Private/BlueprintCompilerCppBackendValueHelper.cpp b/Engine/Source/Developer/BlueprintCompilerCppBackend/Private/BlueprintCompilerCppBackendValueHelper.cpp
--- a/Engine/Source/Developer/BlueprintCompilerCppBackend/Private/BlueprintCompilerCppBackendValueHelper.cpp
+++ b/Engine/Source/Developer/BlueprintCompilerCppBackend/Private/BlueprintCompilerCppBackendValueHelper.cpp
@@ -175,38 +175,29 @@ void FEmitDefaultValueHelper::InnerGenerate(FEmitterLocalContext& Context, const
 		const bool bInitializeWithoutScriptStruct = false;
 		FScriptArrayHelper ScriptArrayHelper(ArrayProperty, ValuePtr);
 		UStructProperty* InnerStructProperty = Cast<UStructProperty>(ArrayProperty->Inner);
-		auto IsSpecialStructType = [](UScriptStruct* InScriptStruct) ->bool
-		{
-			return (InScriptStruct == TBaseStructure<FRotator>::Get())
-				|| (InScriptStruct == TBaseStructure<FTransform>::Get())
-				|| (InScriptStruct == TBaseStructure<FLinearColor>::Get())
-				|| (InScriptStruct == TBaseStructure<FColor>::Get())
-				|| (InScriptStruct == TBaseStructure<FVector>::Get())
-				|| (InScriptStruct == TBaseStructure<FVector2D>::Get())
-				|| (InScriptStruct == TBaseStructure<FRandomStream>::Get())
-				|| (InScriptStruct == TBaseStructure<FGuid>::Get());
-		};
-		UScriptStruct* InnerStruct = ((!bInitializeWithoutScriptStruct) && InnerStructProperty && !IsSpecialStructType(InnerStructProperty->Struct)) ? InnerStructProperty->Struct : nullptr;
+		UScriptStruct* RegularInnerStruct = ((!bInitializeWithoutScriptStruct) && InnerStructProperty && !FEmitDefaultValueHelper::SpecialStructureConstructor(InnerStructProperty->Struct, nullptr, nullptr))
+			? InnerStructProperty->Struct
+			: nullptr;
 		if (ScriptArrayHelper.Num())
 		{
-			const TCHAR* ArrayReserveFunctionName = InnerStruct ? TEXT("AddUninitialized") : TEXT("Reserve");
+			const TCHAR* ArrayReserveFunctionName = RegularInnerStruct ? TEXT("AddUninitialized") : TEXT("Reserve");
 			Context.AddLine(FString::Printf(TEXT("%s.%s(%d);"), *PathToMember, ArrayReserveFunctionName, ScriptArrayHelper.Num()));
 
-			if (InnerStruct)
+			if (RegularInnerStruct)
 			{
-				const FString InnerStructStr = Context.FindGloballyMappedObject(InnerStruct);
+				const FString InnerStructStr = Context.FindGloballyMappedObject(RegularInnerStruct);
 				Context.AddLine(FString::Printf(TEXT("%s->InitializeStruct(%s.GetData(), %d);"), *InnerStructStr, *PathToMember, ScriptArrayHelper.Num()));
 			}
 		}
 
-		const FStructOnScope DefaultStruct(InnerStruct);
+		const FStructOnScope DefaultStruct(RegularInnerStruct);
 
 		for (int32 Index = 0; Index < ScriptArrayHelper.Num(); ++Index)
 		{
 			const uint8* LocalValuePtr = ScriptArrayHelper.GetRawPtr(Index);
 
 			bool bComplete = false;
-			if (!InnerStruct)
+			if (!RegularInnerStruct)
 			{
 				FString ValueStr;
 				bComplete = OneLineConstruction(Context, ArrayProperty->Inner, LocalValuePtr, ValueStr, bInitializeWithoutScriptStruct);
@@ -214,7 +205,7 @@ void FEmitDefaultValueHelper::InnerGenerate(FEmitterLocalContext& Context, const
 				Context.AddLine(FString::Printf(TEXT("%s.Add(%s);"), *PathToMember, *ValueStr));
 			}
 			
-			if (InnerStruct || (bInitializeWithoutScriptStruct && !bComplete))
+			if (RegularInnerStruct || (bInitializeWithoutScriptStruct && !bComplete))
 			{
 				const FString LocalPathToMember = FString::Printf(TEXT("%s[%d]"), *PathToMember, Index);
 				InnerGenerate(Context, ArrayProperty->Inner, LocalPathToMember, LocalValuePtr, DefaultStruct.GetStructMemory(), true);
@@ -223,6 +214,204 @@ void FEmitDefaultValueHelper::InnerGenerate(FEmitterLocalContext& Context, const
 	}
 }
 
+bool FEmitDefaultValueHelper::SpecialStructureConstructor(const UScriptStruct* Struct, const uint8* ValuePtr, /*out*/ FString* OutResult)
+{
+	check(ValuePtr || !OutResult);
+
+	if (TBaseStructure<FTransform>::Get() == Struct)
+	{
+		if (OutResult)
+		{
+			const FTransform* Transform = reinterpret_cast<const FTransform*>(ValuePtr);
+			const auto Rotation = Transform->GetRotation();
+			const auto Translation = Transform->GetTranslation();
+			const auto Scale = Transform->GetScale3D();
+			*OutResult = FString::Printf(TEXT("FTransform(FQuat(%f, %f, %f, %f), FVector(%f, %f, %f), FVector(%f, %f, %f))")
+				, Rotation.X, Rotation.Y, Rotation.Z, Rotation.W
+				, Translation.X, Translation.Y, Translation.Z
+				, Scale.X, Scale.Y, Scale.Z);
+		}
+		return true;
+	}
+
+	if (TBaseStructure<FVector>::Get() == Struct)
+	{
+		if (OutResult)
+		{
+			const FVector* Vector = reinterpret_cast<const FVector*>(ValuePtr);
+			*OutResult = FString::Printf(TEXT("FVector(%f, %f, %f)"), Vector->X, Vector->Y, Vector->Z);
+		}
+		return true;
+	}
+
+	if (TBaseStructure<FGuid>::Get() == Struct)
+	{
+		if (OutResult)
+		{
+			const FGuid* Guid = reinterpret_cast<const FGuid*>(ValuePtr);
+			*OutResult = FString::Printf(TEXT("FGuid(0x%08X, 0x%08X, 0x%08X, 0x%08X)"), Guid->A, Guid->B, Guid->C, Guid->D);
+		}
+		return true;
+	}
+
+	if (TBaseStructure<FRotator>::Get() == Struct)
+	{
+		if (OutResult)
+		{
+			const FRotator* Rotator = reinterpret_cast<const FRotator*>(ValuePtr);
+			*OutResult = FString::Printf(TEXT("FRotator(%f, %f, %f)"), Rotator->Pitch, Rotator->Yaw, Rotator->Roll);
+		}
+		return true;
+	}
+
+	if (TBaseStructure<FLinearColor>::Get() == Struct)
+	{
+		if (OutResult)
+		{
+			const FLinearColor* LinearColor = reinterpret_cast<const FLinearColor*>(ValuePtr);
+			*OutResult = FString::Printf(TEXT("FLinearColor(%f, %f, %f, %f)"), LinearColor->R, LinearColor->B, LinearColor->G, LinearColor->A);
+		}
+		return true;
+	}
+
+	if (TBaseStructure<FColor>::Get() == Struct)
+	{
+		if (OutResult)
+		{
+			const FColor* Color = reinterpret_cast<const FColor*>(ValuePtr);
+			*OutResult = FString::Printf(TEXT("FColor(%d, %d, %d, %d)"), Color->R, Color->B, Color->G, Color->A);
+		}
+		return true;
+	}
+
+	if (TBaseStructure<FVector2D>::Get() == Struct)
+	{
+		if (OutResult)
+		{
+			const FVector2D* Vector2D = reinterpret_cast<const FVector2D*>(ValuePtr);
+			*OutResult = FString::Printf(TEXT("FVector2D(%f, %f)"), Vector2D->X, Vector2D->Y);
+		}
+		return true;
+	}
+
+	if (TBaseStructure<FBox2D>::Get() == Struct)
+	{
+		if (OutResult)
+		{
+			const FBox2D* Box2D = reinterpret_cast<const FBox2D*>(ValuePtr);
+			*OutResult = FString::Printf(TEXT("CreateFBox2D(FVector2D(%f, %f), FVector2D(%f, %f), %s)")
+				, Box2D->Min.X
+				, Box2D->Min.Y
+				, Box2D->Max.X
+				, Box2D->Max.Y
+				, Box2D->bIsValid ? TEXT("true") : TEXT("false"));
+		}
+		return true;
+	}
+
+	if (TBaseStructure<FFloatRangeBound>::Get() == Struct)
+	{
+		if (OutResult)
+		{
+			const FFloatRangeBound* FloatRangeBound = reinterpret_cast<const FFloatRangeBound*>(ValuePtr);
+			if (FloatRangeBound->IsExclusive())
+			{
+				*OutResult = FString::Printf(TEXT("FFloatRangeBound::Exclusive(%f)"), FloatRangeBound->GetValue());
+			}
+			if (FloatRangeBound->IsInclusive())
+			{
+				*OutResult = FString::Printf(TEXT("FFloatRangeBound::Inclusive(%f)"), FloatRangeBound->GetValue());
+			}
+			if (FloatRangeBound->IsOpen())
+			{
+				*OutResult = TEXT("FFloatRangeBound::Open()");
+			}
+		}
+		return true;
+	}
+
+	if (TBaseStructure<FFloatRange>::Get() == Struct)
+	{
+		if (OutResult)
+		{
+			const FFloatRange* FloatRangeBound = reinterpret_cast<const FFloatRange*>(ValuePtr);
+
+			FString LowerBoundStr;
+			FFloatRangeBound LowerBound = FloatRangeBound->GetLowerBound();
+			SpecialStructureConstructor(TBaseStructure<FFloatRangeBound>::Get(), (uint8*)&LowerBound, &LowerBoundStr);
+
+			FString UpperBoundStr;
+			FFloatRangeBound UpperBound = FloatRangeBound->GetUpperBound();
+			SpecialStructureConstructor(TBaseStructure<FFloatRangeBound>::Get(), (uint8*)&UpperBound, &UpperBoundStr);
+
+			*OutResult = FString::Printf(TEXT("FFloatRange(%s, %s)"), *LowerBoundStr, *UpperBoundStr);
+		}
+		return true;
+	}
+
+	if (TBaseStructure<FInt32RangeBound>::Get() == Struct)
+	{
+		if (OutResult)
+		{
+			const FInt32RangeBound* RangeBound = reinterpret_cast<const FInt32RangeBound*>(ValuePtr);
+			if (RangeBound->IsExclusive())
+			{
+				*OutResult = FString::Printf(TEXT("FInt32RangeBound::Exclusive(%d)"), RangeBound->GetValue());
+			}
+			if (RangeBound->IsInclusive())
+			{
+				*OutResult = FString::Printf(TEXT("FInt32RangeBound::Inclusive(%d)"), RangeBound->GetValue());
+			}
+			if (RangeBound->IsOpen())
+			{
+				*OutResult = TEXT("FInt32RangeBound::Open()");
+			}
+		}
+		return true;
+	}
+
+	if (TBaseStructure<FInt32Range>::Get() == Struct)
+	{
+		if (OutResult)
+		{
+			const FInt32Range* RangeBound = reinterpret_cast<const FInt32Range*>(ValuePtr);
+
+			FString LowerBoundStr;
+			FInt32RangeBound LowerBound = RangeBound->GetLowerBound();
+			SpecialStructureConstructor(TBaseStructure<FInt32RangeBound>::Get(), (uint8*)&LowerBound, &LowerBoundStr);
+
+			FString UpperBoundStr;
+			FInt32RangeBound UpperBound = RangeBound->GetUpperBound();
+			SpecialStructureConstructor(TBaseStructure<FInt32RangeBound>::Get(), (uint8*)&UpperBound, &UpperBoundStr);
+
+			*OutResult = FString::Printf(TEXT("FInt32Range(%s, %s)"), *LowerBoundStr, *UpperBoundStr);
+		}
+		return true;
+	}
+
+	if (TBaseStructure<FFloatInterval>::Get() == Struct)
+	{
+		if (OutResult)
+		{
+			const FFloatInterval* Interval = reinterpret_cast<const FFloatInterval*>(ValuePtr);
+			*OutResult = FString::Printf(TEXT("FFloatInterval(%f, %f)"), Interval->Min, Interval->Max);
+		}
+		return true;
+	}
+
+	if (TBaseStructure<FInt32Interval>::Get() == Struct)
+	{
+		if (OutResult)
+		{
+			const FInt32Interval* Interval = reinterpret_cast<const FInt32Interval*>(ValuePtr);
+			*OutResult = FString::Printf(TEXT("FFloatInterval(%d, %d)"), Interval->Min, Interval->Max);
+		}
+		return true;
+	}
+
+	return false;
+}
+
 FString FEmitDefaultValueHelper::HandleSpecialTypes(FEmitterLocalContext& Context, const UProperty* Property, const uint8* ValuePtr)
 {
 	//TODO: Use Path maps for Objects
@@ -267,64 +456,10 @@ FString FEmitDefaultValueHelper::HandleSpecialTypes(FEmitterLocalContext& Contex
 
 	if (auto StructProperty = Cast<UStructProperty>(Property))
 	{
-		if (TBaseStructure<FTransform>::Get() == StructProperty->Struct)
-		{
-			check(ValuePtr);
-			const FTransform* Transform = reinterpret_cast<const FTransform*>(ValuePtr);
-			const auto Rotation = Transform->GetRotation();
-			const auto Translation = Transform->GetTranslation();
-			const auto Scale = Transform->GetScale3D();
-			return FString::Printf(TEXT("FTransform(FQuat(%f, %f, %f, %f), FVector(%f, %f, %f), FVector(%f, %f, %f))")
-				, Rotation.X, Rotation.Y, Rotation.Z, Rotation.W
-				, Translation.X, Translation.Y, Translation.Z
-				, Scale.X, Scale.Y, Scale.Z);
-		}
-
-		if (TBaseStructure<FVector>::Get() == StructProperty->Struct)
+		FString StructConstructor;
+		if (SpecialStructureConstructor(StructProperty->Struct, ValuePtr, &StructConstructor))
 		{
-			const FVector* Vector = reinterpret_cast<const FVector*>(ValuePtr);
-			return FString::Printf(TEXT("FVector(%f, %f, %f)"), Vector->X, Vector->Y, Vector->Z);
-		}
-
-		if (TBaseStructure<FGuid>::Get() == StructProperty->Struct)
-		{
-			const FGuid* Guid = reinterpret_cast<const FGuid*>(ValuePtr);
-			return FString::Printf(TEXT("FGuid(0x%08X, 0x%08X, 0x%08X, 0x%08X)"), Guid->A, Guid->B, Guid->C, Guid->D);
-		}
-
-		if (TBaseStructure<FRotator>::Get() == StructProperty->Struct)
-		{
-			const FRotator* Rotator = reinterpret_cast<const FRotator*>(ValuePtr);
-			return FString::Printf(TEXT("FRotator(%f, %f, %f)"), Rotator->Pitch, Rotator->Yaw, Rotator->Roll);
-		}
-
-		if (TBaseStructure<FLinearColor>::Get() == StructProperty->Struct)
-		{
-			const FLinearColor* LinearColor = reinterpret_cast<const FLinearColor*>(ValuePtr);
-			return FString::Printf(TEXT("FLinearColor(%f, %f, %f, %f)"), LinearColor->R, LinearColor->B, LinearColor->G, LinearColor->A);
-		}
-
-		if (TBaseStructure<FColor>::Get() == StructProperty->Struct)
-		{
-			const FColor* Color = reinterpret_cast<const FColor*>(ValuePtr);
-			return FString::Printf(TEXT("FColor(%d, %d, %d, %d)"), Color->R, Color->B, Color->G, Color->A);
-		}
-
-		if (TBaseStructure<FVector2D>::Get() == StructProperty->Struct)
-		{
-			const FVector2D* Vector2D = reinterpret_cast<const FVector2D*>(ValuePtr);
-			return FString::Printf(TEXT("FVector2D(%f, %f)"), Vector2D->X, Vector2D->Y);
-		}
-
-		if (TBaseStructure<FBox2D>::Get() == StructProperty->Struct)
-		{
-			const FBox2D* Box2D = reinterpret_cast<const FBox2D*>(ValuePtr);
-			return FString::Printf(TEXT("CreateFBox2D(FVector2D(%f, %f), FVector2D(%f, %f), %s)")
-				, Box2D->Min.X
-				, Box2D->Min.Y
-				, Box2D->Max.X
-				, Box2D->Max.Y
-				, Box2D->bIsValid ? TEXT("true") : TEXT("false"));
+			return StructConstructor;
 		}
 	}
 
@@ -457,7 +592,7 @@ FString FEmitDefaultValueHelper::HandleNonNativeComponent(FEmitterLocalContext&
 
 struct FDependenciesHelper
 {
-private:
+public:
 	// Keep sync with FTypeSingletonCache::GenerateSingletonName
 	static FString GenerateZConstructor(UField* Item)
 	{
@@ -498,51 +633,11 @@ private:
 		}
 
 		const FString ClassString = Item->IsA<UClass>() ? TEXT("UClass") : TEXT("UScriptStruct");
-		const FString PostFix = Item->IsA<UClass>() ? TEXT("_NoRegister") : TEXT("");
-		return FString(TEXT("Z_Construct_")) + ClassString + TEXT("_") + Result + PostFix + TEXT("()");
+		return FString(TEXT("Z_Construct_")) + ClassString + TEXT("_") + Result + TEXT("()");
 	}
-public:
+
 	static void AddDependenciesInConstructor(FEmitterLocalContext& Context)
 	{
-		const bool bUseZConstructorInGeneratedCode = false;
-		if (Context.Dependencies.ConvertedClasses.Num())
-		{
-			Context.AddLine(TEXT("// List of all referenced converted classes"));
-		}
-		for (auto LocStruct : Context.Dependencies.ConvertedClasses)
-		{
-			FString ClassConstructor;
-			if (bUseZConstructorInGeneratedCode)
-			{
-				ClassConstructor = GenerateZConstructor(Context.Dependencies.FindOriginalClass(LocStruct));
-				Context.AddLine(FString::Printf(TEXT("extern UClass* %s;"), *ClassConstructor));
-			}
-			else
-			{
-				ClassConstructor = FEmitHelper::GetCppName(LocStruct, true) + TEXT("::StaticClass()");
-			}
-			Context.AddLine(FString::Printf(TEXT("CastChecked<UDynamicClass>(GetClass())->ReferencedConvertedFields.Add(%s);"), *ClassConstructor));
-		}
-
-		if (Context.Dependencies.ConvertedStructs.Num())
-		{
-			Context.AddLine(TEXT("// List of all referenced converted structures"));
-		}
-		for (auto LocStruct : Context.Dependencies.ConvertedStructs)
-		{
-			FString StructConstructor;
-			if (bUseZConstructorInGeneratedCode)
-			{
-				StructConstructor = GenerateZConstructor(LocStruct);
-				Context.AddLine(FString::Printf(TEXT("extern UScriptStruct* %s;"), *StructConstructor));
-			}
-			else
-			{
-				StructConstructor = FEmitHelper::GetCppName(LocStruct, true) + TEXT("::StaticStruct()");
-			}
-			Context.AddLine(FString::Printf(TEXT("CastChecked<UDynamicClass>(GetClass())->ReferencedConvertedFields.Add(%s);"), *StructConstructor));
-		}
-
 		if (Context.Dependencies.Assets.Num())
 		{
 			Context.AddLine(TEXT("// List of all referenced assets"));
@@ -662,7 +757,6 @@ void FEmitDefaultValueHelper::GenerateConstructor(FEmitterLocalContext& Context)
 		Context.AddLine(TEXT("{"));
 		Context.IncreaseIndent();
 		Context.AddLine(TEXT("ensure(0 == CastChecked<UDynamicClass>(GetClass())->MiscConvertedSubobjects.Num());"));
-		Context.AddLine(TEXT("ensure(0 == CastChecked<UDynamicClass>(GetClass())->ReferencedConvertedFields.Num());"));
 		Context.AddLine(TEXT("ensure(0 == CastChecked<UDynamicClass>(GetClass())->UsedAssets.Num());"));
 		Context.AddLine(TEXT("ensure(0 == CastChecked<UDynamicClass>(GetClass())->DynamicBindingObjects.Num());"));
 		Context.AddLine(TEXT("ensure(0 == CastChecked<UDynamicClass>(GetClass())->ComponentTemplates.Num());"));
@@ -860,6 +954,37 @@ void FEmitDefaultValueHelper::GenerateConstructor(FEmitterLocalContext& Context)
 		Context.AddLine(TEXT("}"));
 	}
 
+	{
+		Context.AddLine(FString::Printf(TEXT("void %s::__GatherReferencedConvertedFields(UDynamicClass* InDynamicClass)"), *CppClassName));
+		Context.AddLine(TEXT("{"));
+		Context.IncreaseIndent();
+		Context.AddLine(TEXT("ensure(0 == InDynamicClass->ReferencedConvertedFields.Num());"));
+		if (Context.Dependencies.ConvertedClasses.Num())
+		{
+			Context.AddLine(TEXT("// List of all referenced converted classes"));
+		}
+		for (auto LocStruct : Context.Dependencies.ConvertedClasses)
+		{
+			const FString ClassConstructor = FDependenciesHelper::GenerateZConstructor(Context.Dependencies.FindOriginalClass(LocStruct));
+			Context.AddLine(FString::Printf(TEXT("extern UClass* %s;"), *ClassConstructor));
+			Context.AddLine(FString::Printf(TEXT("InDynamicClass->ReferencedConvertedFields.Add(%s);"), *ClassConstructor));
+		}
+
+		if (Context.Dependencies.ConvertedStructs.Num())
+		{
+			Context.AddLine(TEXT("// List of all referenced converted structures"));
+		}
+		for (auto LocStruct : Context.Dependencies.ConvertedStructs)
+		{
+			const FString StructConstructor = FDependenciesHelper::GenerateZConstructor(LocStruct);
+			Context.AddLine(FString::Printf(TEXT("extern UScriptStruct* %s;"), *StructConstructor));
+			Context.AddLine(FString::Printf(TEXT("InDynamicClass->ReferencedConvertedFields.Add(%s);"), *StructConstructor));
+		}
+
+		Context.DecreaseIndent();
+		Context.AddLine(TEXT("}"));
+	}
+
 	FDependenciesHelper::AddStaticFunctionsForDependencies(Context);
 
 	FBackendHelperUMG::EmitWidgetInitializationFunctions(Context);
diff --git a/Engine/Source/Developer/BlueprintNativeCodeGen/Private/BlueprintNativeCodeGenModule.cpp b/Engine/Source/Developer/BlueprintNativeCodeGen/Private/BlueprintNativeCodeGenModule.cpp
--- a/Engine/Source/Developer/BlueprintNativeCodeGen/Private/BlueprintNativeCodeGenModule.cpp
+++ b/Engine/Source/Developer/BlueprintNativeCodeGen/Private/BlueprintNativeCodeGenModule.cpp
@@ -9,6 +9,7 @@
 #include "IBlueprintCompilerCppBackendModule.h"
 #include "BlueprintEditorUtils.h"
 #include "Engine/Blueprint.h" // for EBlueprintType
+#include "Engine/InheritableComponentHandler.h"
 
 /*******************************************************************************
 * NativizationCookControllerImpl
@@ -49,7 +50,8 @@ protected:
 	//~ Begin FScriptCookReplacmentCoordinator interface
 	virtual EReplacementResult IsTargetedForReplacement(const UPackage* Package) const override;
 	virtual EReplacementResult IsTargetedForReplacement(const UObject* Object) const override;
-	virtual UClass* FindReplacedClass(const UClass* Class) const override;
+	virtual UClass* FindReplacedClassForObject(const UObject* Object) const override;
+	virtual UObject* FindReplacedNameAndOuter(UObject* Object, FName& OutName) const override;
 	//~ End FScriptCookReplacmentCoordinator interface
 private:
 	void ReadConfig();
@@ -501,25 +503,85 @@ void FBlueprintNativeCodeGenModule::FinalizeManifest()
 	}
 }
 
-UClass* FBlueprintNativeCodeGenModule::FindReplacedClass(const UClass* Class) const
+UClass* FBlueprintNativeCodeGenModule::FindReplacedClassForObject(const UObject* Object) const
 {
 	// we're only looking to replace class types:
-	while (Class)
+	if (Object && Object->IsA<UField>() 
+		&& (EReplacementResult::ReplaceCompletely == IsTargetedForReplacement(Object)))
 	{
-		if (Class == UUserDefinedEnum::StaticClass())
+		for (const UClass* Class = Object->GetClass(); Class; Class = Class->GetSuperClass())
 		{
-			return UEnum::StaticClass();
+			if (Class == UUserDefinedEnum::StaticClass())
+			{
+				return UEnum::StaticClass();
+			}
+			if (Class == UUserDefinedStruct::StaticClass())
+			{
+				return UScriptStruct::StaticClass();
+			}
+			if (Class == UBlueprintGeneratedClass::StaticClass())
+			{
+				return UDynamicClass::StaticClass();
+			}
 		}
-		if (Class == UUserDefinedStruct::StaticClass())
+	}
+	ensure(!Object || !(Object->IsA<UUserDefinedStruct>() || Object->IsA<UUserDefinedEnum>()));
+	return nullptr;
+}
+
+UObject* FBlueprintNativeCodeGenModule::FindReplacedNameAndOuter(UObject* Object, FName& OutName) const
+{
+	OutName = NAME_None;
+	UObject* Outer = nullptr;
+
+	UActorComponent* ActorComponent = Cast<UActorComponent>(Object);
+	if (ActorComponent)
+	{
+		//if is child of a BPGC and not child of a CDO
+		UBlueprintGeneratedClass* BPGC = nullptr;
+		for (UObject* OuterObject = ActorComponent->GetOuter(); OuterObject && !BPGC; OuterObject = OuterObject->GetOuter())
 		{
-			return UScriptStruct::StaticClass();
+			if (OuterObject->HasAnyFlags(RF_ClassDefaultObject))
+			{
+				return Outer;
+			}
+			BPGC = Cast<UBlueprintGeneratedClass>(OuterObject);
 		}
-		if (Class == UBlueprintGeneratedClass::StaticClass())
+
+		for (UBlueprintGeneratedClass* SuperBPGC = BPGC; SuperBPGC && (OutName == NAME_None); SuperBPGC = Cast<UBlueprintGeneratedClass>(SuperBPGC->GetSuperClass()))
 		{
-			return UDynamicClass::StaticClass();
+			if (SuperBPGC->InheritableComponentHandler)
+			{
+				FComponentKey FoundKey = SuperBPGC->InheritableComponentHandler->FindKey(ActorComponent);
+				if (FoundKey.IsValid())
+				{
+					OutName = FoundKey.IsSCSKey() ? FoundKey.GetSCSVariableName() : ActorComponent->GetFName();
+					Outer = BPGC->GetDefaultObject(false);
+					break;
+				}
+			}
+			if (SuperBPGC->SimpleConstructionScript)
+			{
+				for (auto Node : SuperBPGC->SimpleConstructionScript->GetAllNodes())
+				{
+					if (Node->ComponentTemplate == ActorComponent)
+					{
+						OutName = Node->VariableName;
+						if (OutName != NAME_None)
+						{
+							Outer = BPGC->GetDefaultObject(false);
+							break;
+						}
+					}
+				}
+			}
 		}
+	}
 
-		Class = Class->GetSuperClass();
+	if (Outer && (EReplacementResult::ReplaceCompletely == IsTargetedForReplacement(Object->GetClass())))
+	{
+		UE_LOG(LogBlueprintCodeGen, Log, TEXT("Object '%s' has replaced name '%s' and outer: '%s'"), *GetPathNameSafe(Object), *OutName.ToString(), *GetPathNameSafe(Outer));
+		return Outer;
 	}
 
 	return nullptr;
diff --git a/Engine/Source/Developer/BlueprintProfiler/BlueprintProfiler.Build.cs b/Engine/Source/Developer/BlueprintProfiler/BlueprintProfiler.Build.cs
--- a/Engine/Source/Developer/BlueprintProfiler/BlueprintProfiler.Build.cs
+++ b/Engine/Source/Developer/BlueprintProfiler/BlueprintProfiler.Build.cs
@@ -24,12 +24,6 @@ public class BlueprintProfiler : ModuleRules
             PrivateDependencyModuleNames.Add("BlueprintGraph");
             PrivateDependencyModuleNames.Add("EditorStyle");
         }
-
-        if (UnrealBuildTool.BuildHostPlatform.Current.Platform == UnrealTargetPlatform.Linux)
-        {
-            // work around for bad naming of "classes" directory, proper resolution is tracked as UE-23934
-            bFasterWithoutUnity = true;
-        }
     }
 }
 
diff --git a/Engine/Source/Developer/BlueprintProfiler/Private/BlueprintProfiler.cpp b/Engine/Source/Developer/BlueprintProfiler/Private/BlueprintProfiler.cpp
new file mode 100644
--- /dev/null
+++ b/Engine/Source/Developer/BlueprintProfiler/Private/BlueprintProfiler.cpp
@@ -0,0 +1,309 @@
+// Copyright 1998-2016 Epic Games, Inc. All Rights Reserved.
+
+#include "BlueprintProfilerPCH.h"
+#include "EditorStyleSet.h"
+
+#define LOCTEXT_NAMESPACE "BlueprintProfiler"
+
+#pragma optimize( "g", off )
+
+//////////////////////////////////////////////////////////////////////////
+// FBlueprintProfiler
+
+IMPLEMENT_MODULE(FBlueprintProfiler, BlueprintProfiler);
+
+DECLARE_CYCLE_STAT(TEXT("Tick"), STAT_ProfilerTickCost, STATGROUP_BlueprintProfiler);
+DECLARE_CYCLE_STAT(TEXT("Processing Events"), STAT_ProfilerInstrumentationCost, STATGROUP_BlueprintProfiler);
+DECLARE_CYCLE_STAT(TEXT("Blueprint Lookup"), STAT_BlueprintLookupCost, STATGROUP_BlueprintProfiler);
+
+FBlueprintProfiler::FBlueprintProfiler()
+	: bProfilerActive(false)
+	, bProfilingCaptureActive(false)
+	, bPIEActive(false)
+{
+}
+
+FBlueprintProfiler::~FBlueprintProfiler()
+{
+	if (bProfilingCaptureActive)
+	{
+		ToggleProfilingCapture();
+	}
+}
+
+void FBlueprintProfiler::StartupModule()
+{
+	FBlueprintCoreDelegates::OnToggleScriptProfiler.AddRaw(this, &FBlueprintProfiler::RegisterDelegates);
+}
+
+void FBlueprintProfiler::ShutdownModule()
+{
+	FBlueprintCoreDelegates::OnToggleScriptProfiler.RemoveAll(this);
+}
+
+void FBlueprintProfiler::ToggleProfilingCapture()
+{
+	// Toggle profiler state
+	bProfilerActive = !bProfilerActive;
+#if WITH_EDITOR
+	bProfilingCaptureActive = bProfilerActive && bPIEActive;
+#else
+	bProfilingCaptureActive = bProfilerActive;
+#endif // WITH_EDITOR
+	// Broadcast capture state change so delegate consumers can update their state.
+	FBlueprintCoreDelegates::OnToggleScriptProfiler.Broadcast(bProfilerActive);
+}
+
+void FBlueprintProfiler::InstrumentEvent(const EScriptInstrumentationEvent& Event)
+{
+	#if WITH_EDITOR
+	SCOPE_CYCLE_COUNTER(STAT_ProfilerInstrumentationCost);
+
+	const UObject* ContextObject = Event.GetContextObject();
+	if (ContextObject && !ContextObject->HasAnyFlags(RF_Transient))
+	#endif
+	{
+		// Handle context switching events
+		CaptureContext.UpdateContext(Event.GetContextObject(), InstrumentationEventQueue);
+		// Add instrumented event
+		int32 ScriptCodeOffset = -1;
+		FName CurrentFunctionName = Event.GetEventName();
+		if (Event.IsStackFrameValid())
+		{
+			const FFrame& StackFrame = Event.GetStackFrame();
+			ScriptCodeOffset = StackFrame.Code - StackFrame.Node->Script.GetData() - 1;
+			CurrentFunctionName = StackFrame.Node->GetFName();
+		}
+		InstrumentationEventQueue.Add(FScriptInstrumentedEvent(Event.GetType(), CurrentFunctionName, ScriptCodeOffset));
+		// Reset context on event end
+		if (Event.GetType() == EScriptInstrumentation::Stop)
+		{
+			CaptureContext.ResetContext();
+		}
+	}
+}
+
+#if WITH_EDITOR
+
+TSharedPtr<FBlueprintExecutionContext> FBlueprintProfiler::GetBlueprintContext(const FString& BlueprintClassPath)
+{
+	SCOPE_CYCLE_COUNTER(STAT_BlueprintLookupCost);
+	TSharedPtr<FBlueprintExecutionContext>& Result = PathToBlueprintContext.FindOrAdd(BlueprintClassPath);
+	if (!Result.IsValid())
+	{
+		Result = MakeShareable<FBlueprintExecutionContext>(new FBlueprintExecutionContext);
+		// Map the blueprint and initialise the context.
+		if (Result->InitialiseContext(BlueprintClassPath))
+		{
+			// Register delegate to handle updates
+			Result->GetBlueprint().Get()->OnCompiled().AddRaw(this, &FBlueprintProfiler::RemoveAllBlueprintReferences);
+		}
+		else
+		{
+			Result.Reset();
+			PathToBlueprintContext.Remove(BlueprintClassPath);
+		}
+	}
+	return Result;
+}
+
+TSharedPtr<FScriptExecutionNode> FBlueprintProfiler::GetProfilerDataForNode(const UEdGraphNode* GraphNode)
+{
+	TSharedPtr<FScriptExecutionNode> Result;
+	if (GraphNode)
+	{
+		if (UBlueprint* Blueprint = GraphNode->GetTypedOuter<UBlueprint>())
+		{
+			if (TSharedPtr<FBlueprintExecutionContext>* BlueprintContext = PathToBlueprintContext.Find(Blueprint->GeneratedClass->GetPathName()))
+			{
+				Result = (*BlueprintContext)->GetProfilerDataForNode(GraphNode);
+			}
+		}
+	}
+	return Result;
+}
+
+bool FBlueprintProfiler::HasDataForInstance(const UObject* Instance) const
+{
+	bool bHasData = false;
+	if (Instance)
+	{
+		if (const UBlueprintGeneratedClass* BlueprintClass = Cast<UBlueprintGeneratedClass>(Instance->GetClass()))
+		{
+			if (const TSharedPtr<FBlueprintExecutionContext>* Result = PathToBlueprintContext.Find(BlueprintClass->GetPathName()))
+			{
+				FName InstanceName(*Instance->GetPathName());
+				bHasData = (*Result)->HasProfilerDataForInstance(InstanceName);
+			}
+		}
+	}
+	return bHasData;
+}
+
+void FBlueprintProfiler::Tick(float DeltaSeconds)
+{
+	SCOPE_CYCLE_COUNTER(STAT_ProfilerTickCost);
+	ProcessEventProfilingData();
+}
+
+bool FBlueprintProfiler::IsTickable() const
+{
+	return InstrumentationEventQueue.Num() > 0;
+}
+
+void FBlueprintProfiler::ProcessEventProfilingData()
+{
+	int32 StartEvent = 0;
+	int32 StopEvent = 0;
+	TSharedPtr<FBlueprintExecutionContext> BlueprintContext;
+	TSet<TSharedPtr<FBlueprintExecutionContext>> DirtyContexts;
+	// Iterate through the events and batch into single script executions
+	for (auto EventIter = InstrumentationEventQueue.CreateConstIterator(); EventIter; ++EventIter)
+	{
+		switch (EventIter->GetType())
+		{
+			case EScriptInstrumentation::Class:
+			{
+				StartEvent = EventIter.GetIndex();
+				BlueprintContext = GetBlueprintContext(EventIter->GetObjectPath());
+				break;
+			}
+			case EScriptInstrumentation::Instance:
+			{
+				MapBlueprintInstance(BlueprintContext, EventIter->GetObjectPath());
+				break;
+			}
+			case EScriptInstrumentation::Stop:
+			{
+				StopEvent = EventIter.GetIndex();
+				// Create and validate new event.
+				FScriptEventPlayback NewEvent(BlueprintContext);
+				if (NewEvent.Process(InstrumentationEventQueue, StartEvent, StopEvent))
+				{
+					DirtyContexts.Add(BlueprintContext);
+				}
+				break;
+			}
+		}
+	}
+	// Flush processed event data
+	if (StopEvent > 0)
+	{
+		InstrumentationEventQueue.RemoveAt(0, StopEvent + 1, false);
+	}
+	// Update dirty contexts
+	for (auto Context : DirtyContexts)
+	{
+		Context->UpdateConnectedStats();
+	}
+}
+
+void FBlueprintProfiler::RemoveAllBlueprintReferences(UBlueprint* Blueprint)
+{
+	if (Blueprint)
+	{
+		const FString BlueprintClassPath = Blueprint->GeneratedClass->GetPathName();
+		if (PathToBlueprintContext.Contains(BlueprintClassPath))
+		{
+			// Tear down the blueprint context.
+			PathToBlueprintContext.Remove(BlueprintClassPath);
+			// Remove compilation hook, this will get added again if the blueprint is profiled
+			Blueprint->OnCompiled().RemoveAll(this);
+			// Broadcast changes to stats structure
+			GraphLayoutChangedDelegate.Broadcast(Blueprint);
+		}
+	}
+}
+
+#endif // WITH_EDITOR
+
+void FBlueprintProfiler::ResetProfilingData()
+{
+	InstrumentationEventQueue.Reset();
+	CaptureContext.ResetContext();
+#if WITH_EDITOR
+	PathToBlueprintContext.Reset();
+#endif // WITH_EDITOR
+}
+
+void FBlueprintProfiler::RegisterDelegates(bool bEnabled)
+{
+	if (bEnabled)
+	{
+#if WITH_EDITOR
+		// Register for PIE begin and end events in the editor
+		FEditorDelegates::BeginPIE.AddRaw(this, &FBlueprintProfiler::BeginPIE);
+		FEditorDelegates::EndPIE.AddRaw(this, &FBlueprintProfiler::EndPIE);
+#endif // WITH_EDITOR
+		ResetProfilingData();
+		// Start consuming profiling events for capture
+		FBlueprintCoreDelegates::OnScriptProfilingEvent.AddRaw(this, &FBlueprintProfiler::InstrumentEvent);
+	}
+	else
+	{
+#if WITH_EDITOR
+		// Unregister for PIE begin and end events in the editor
+		FEditorDelegates::BeginPIE.RemoveAll(this);
+		FEditorDelegates::EndPIE.RemoveAll(this);
+#endif // WITH_EDITOR
+		ResetProfilingData();
+		// Stop consuming profiling events for capture
+		FBlueprintCoreDelegates::OnScriptProfilingEvent.RemoveAll(this);
+	}
+}
+
+void FBlueprintProfiler::BeginPIE(bool bIsSimulating)
+{
+#if WITH_EDITOR
+	bPIEActive = true;
+	bProfilingCaptureActive = bProfilerActive && bPIEActive;
+#endif // WITH_EDITOR
+}
+
+void FBlueprintProfiler::EndPIE(bool bIsSimulating)
+{
+#if WITH_EDITOR
+	bPIEActive = false;
+	bProfilingCaptureActive = bProfilerActive && bPIEActive;
+#endif // WITH_EDITOR
+}
+
+void FBlueprintProfiler::MapBlueprintInstance(TSharedPtr<FBlueprintExecutionContext> BlueprintContext, const FString& InstancePath)
+{
+	FName InstanceName(*InstancePath);
+	TWeakObjectPtr<const UObject> Instance = BlueprintContext->GetInstance(InstanceName);
+
+	if (!BlueprintContext->HasProfilerDataForInstance(InstanceName))
+	{
+		// Create new instance node
+		TSharedPtr<FScriptExecutionInstance> InstanceNode = MakeShareable<FScriptExecutionInstance>(new FScriptExecutionInstance);
+		InstanceNode->SetNameByString(InstanceName.ToString());
+		InstanceNode->SetObservedObject(Instance.Get());
+		InstanceNode->SetGraphName(NAME_None);
+		InstanceNode->SetFlags(EScriptExecutionNodeFlags::Instance);
+		const FSlateBrush* Icon = FEditorStyle::GetBrush(TEXT("BlueprintProfiler.Actor"));
+		InstanceNode->SetIcon(const_cast<FSlateBrush*>(Icon));
+		InstanceNode->SetIconColor(FLinearColor(1.f, 1.f, 1.f, 0.8f));
+		InstanceNode->SetToolTipText(LOCTEXT("NavigateToInstanceHyperlink_ToolTip", "Navigate to the Instance"));
+		// Update widget data
+		if (const AActor* Actor = Cast<AActor>(Instance.Get()))
+		{
+			InstanceNode->SetDisplayName(FText::FromString(Actor->GetActorLabel()));
+		}
+		else
+		{
+			InstanceNode->SetDisplayName(FText::FromString(Instance.Get()->GetName()));
+		}
+		// Link to parent blueprint entry
+		TSharedPtr<FScriptExecutionBlueprint> BlueprintNode = BlueprintContext->GetBlueprintExecNode();
+		BlueprintNode->AddInstance(InstanceNode);
+		InstanceNode->SetParentNode(BlueprintNode);
+		// Fill out events from the blueprint root node
+		for (int32 NodeIdx = 0; NodeIdx < BlueprintNode->GetNumChildren(); ++NodeIdx)
+		{
+			InstanceNode->AddChildNode(BlueprintNode->GetChildByIndex(NodeIdx));
+		}
+	}
+}
+
+#undef LOCTEXT_NAMESPACE
diff --git a/Engine/Source/Developer/BlueprintProfiler/Private/BlueprintProfilerPCH.h b/Engine/Source/Developer/BlueprintProfiler/Private/BlueprintProfilerPCH.h
--- a/Engine/Source/Developer/BlueprintProfiler/Private/BlueprintProfilerPCH.h
+++ b/Engine/Source/Developer/BlueprintProfiler/Private/BlueprintProfilerPCH.h
@@ -3,16 +3,22 @@
 #ifndef _INC_BLUEPRINTPROFILER
 #define _INC_BLUEPRINTPROFILER
 
-#include "BlueprintProfilerModule.h"
-#include "Public/BlueprintProfiler.h"
+#include "CoreUObject.h"
+#include "Core.h"
+#include "Script.h"
 
 #if WITH_EDITOR
+#include "Editor.h"
 #include "Editor/Kismet/Public/Profiler/TracePath.h"
 #include "Editor/Kismet/Public/Profiler/EventExecution.h"
 #include "Editor/Kismet/Public/Profiler/ScriptPerfData.h"
 #include "Editor/BlueprintGraph/Public/BlueprintGraphDefinitions.h"
 #include "Editor/UnrealEd/Public/Kismet2/KismetEditorUtilities.h"
-#include "Editor.h"
 #endif // WITH_EDITOR
 
+#include "BlueprintProfilerModule.h"
+#include "Public/BlueprintProfiler.h"
+
+DECLARE_STATS_GROUP(TEXT("BlueprintProfiler"), STATGROUP_BlueprintProfiler, STATCAT_Advanced);
+
 #endif // _INC_BLUEPRINTPROFILER
diff --git a/Engine/Source/Developer/BlueprintProfiler/Private/ScriptInstrumentationCapture.cpp b/Engine/Source/Developer/BlueprintProfiler/Private/ScriptInstrumentationCapture.cpp
new file mode 100644
--- /dev/null
+++ b/Engine/Source/Developer/BlueprintProfiler/Private/ScriptInstrumentationCapture.cpp
@@ -0,0 +1,43 @@
+// Copyright 1998-2016 Epic Games, Inc. All Rights Reserved.
+
+#include "BlueprintProfilerPCH.h"
+
+//////////////////////////////////////////////////////////////////////////
+// FScriptInstrumentedEvent
+
+void FScriptInstrumentedEvent::SetData(EScriptInstrumentation::Type InEventType, const FString& InPathData, const int32 InCodeOffset)
+{
+	EventType = InEventType;
+	PathData = InPathData;
+	CodeOffset = InCodeOffset;
+	Time = FPlatformTime::Seconds();
+}
+
+//////////////////////////////////////////////////////////////////////////
+// FInstrumentationCaptureContext
+
+void FInstrumentationCaptureContext::UpdateContext(const UObject* InContextObject, TArray<FScriptInstrumentedEvent>& InstrumentationQueue)
+{
+	// Handle instance context switch
+	if (ContextObject != InContextObject && InContextObject)
+	{
+		// Handle class context switch
+		if (const UClass* NewClass = InContextObject->GetClass())
+		{
+			if (NewClass != ContextClass)
+			{
+				InstrumentationQueue.Add(FScriptInstrumentedEvent(EScriptInstrumentation::Class, NewClass->GetPathName()));
+				ContextClass = NewClass;
+			}
+		}
+		// Handle instance change
+		InstrumentationQueue.Add(FScriptInstrumentedEvent(EScriptInstrumentation::Instance, InContextObject->GetPathName()));
+		ContextObject = InContextObject;
+	}
+}
+
+void FInstrumentationCaptureContext::ResetContext()
+{
+	ContextClass.Reset();
+	ContextObject.Reset();
+}
\ No newline at end of file
diff --git a/Engine/Source/Developer/BlueprintProfiler/Private/ScriptInstrumentationPlayback.cpp b/Engine/Source/Developer/BlueprintProfiler/Private/ScriptInstrumentationPlayback.cpp
new file mode 100644
--- /dev/null
+++ b/Engine/Source/Developer/BlueprintProfiler/Private/ScriptInstrumentationPlayback.cpp
@@ -0,0 +1,636 @@
+// Copyright 1998-2016 Epic Games, Inc. All Rights Reserved.
+
+#include "BlueprintProfilerPCH.h"
+
+#define LOCTEXT_NAMESPACE "ScriptInstrumentationPlayback"
+
+DECLARE_CYCLE_STAT(TEXT("Statistic Update"), STAT_StatUpdateCost, STATGROUP_BlueprintProfiler);
+DECLARE_CYCLE_STAT(TEXT("Node Lookup"), STAT_NodeLookupCost, STATGROUP_BlueprintProfiler);
+
+//////////////////////////////////////////////////////////////////////////
+// FBlueprintExecutionContext
+
+bool FBlueprintExecutionContext::InitialiseContext(const FString& BlueprintPath)
+{
+	// Locate the blueprint from the path
+	if (UObject* ObjectPtr = FindObject<UObject>(nullptr, *BlueprintPath))
+	{
+		UBlueprintGeneratedClass* BPClass = Cast<UBlueprintGeneratedClass>(ObjectPtr);
+		if (BPClass && BPClass->bHasInstrumentation)
+		{
+			BlueprintClass = BPClass;
+			Blueprint = Cast<UBlueprint>(BPClass->ClassGeneratedBy);
+			UbergraphFunctionName = BPClass->UberGraphFunction ? BPClass->UberGraphFunction->GetFName() : NAME_None;
+		}
+	}
+	if (Blueprint.IsValid() && BlueprintClass.IsValid())
+	{
+		// Create new blueprint exec node
+		BlueprintNode = MakeShareable<FScriptExecutionBlueprint>(new FScriptExecutionBlueprint);
+		BlueprintNode->SetNameByString(BlueprintPath);
+		BlueprintNode->SetObservedObject(Blueprint.Get());
+		BlueprintNode->SetGraphName(NAME_None);
+		BlueprintNode->SetDisplayName(FText::FromName(Blueprint.Get()->GetFName()));
+		BlueprintNode->SetFlags(EScriptExecutionNodeFlags::Class);
+		const FSlateBrush* Icon_Normal = FEditorStyle::GetBrush(TEXT("BlueprintProfiler.BPIcon_Normal"));
+		BlueprintNode->SetIcon(const_cast<FSlateBrush*>(Icon_Normal));
+		BlueprintNode->SetIconColor(FLinearColor(0.46f, 0.54f, 0.81f));
+		BlueprintNode->SetToolTipText(LOCTEXT("NavigateToBlueprintHyperlink_ToolTip", "Navigate to the Blueprint"));
+		// Map the blueprint execution
+		if (MapBlueprintExecution())
+		{
+			return true;
+		}
+	}
+	return false;
+}
+
+bool FBlueprintExecutionContext::HasProfilerDataForInstance(const FName InstanceName) const
+{
+	bool bHasInstanceData = false;
+	if (BlueprintNode.IsValid())
+	{
+		FName RemappedInstanceName = RemapInstancePath(InstanceName);
+		TSharedPtr<FScriptExecutionNode> Result = BlueprintNode->GetInstanceByName(RemappedInstanceName);
+		bHasInstanceData = Result.IsValid();
+	}
+	return bHasInstanceData;
+}
+
+FName FBlueprintExecutionContext::RemapInstancePath(const FName InstanceName) const
+{
+	FName Result = InstanceName;
+	if (const FName* RemappedName = PIEInstanceNameMap.Find(InstanceName))
+	{
+		Result = *RemappedName;
+	}
+	return Result;
+}
+
+TWeakObjectPtr<const UObject> FBlueprintExecutionContext::GetInstance(FName& InstanceNameInOut)
+{
+	TWeakObjectPtr<const UObject> Result;
+	FName CorrectedName = InstanceNameInOut;
+	if (const FName* SearchName = PIEInstanceNameMap.Find(InstanceNameInOut)) 
+	{
+		CorrectedName = *SearchName;
+	}
+	if (TWeakObjectPtr<const UObject>* SearchResult = ActorInstances.Find(CorrectedName))
+	{
+		Result = *SearchResult;
+	}
+	else
+	{
+		// Attempt to locate the instance and map PIE objects to editor world objects
+		if (const UObject* ObjectPtr = FindObject<UObject>(nullptr, *InstanceNameInOut.ToString()))
+		{
+			// Get Outer world
+			if (UWorld* ObjectWorld = ObjectPtr->GetTypedOuter<UWorld>())
+			{
+				if (ObjectWorld->WorldType == EWorldType::PIE)
+				{
+					FWorldContext& EditorWorldContext = GEditor->GetEditorWorldContext();
+					if (UWorld* EditorWorld = EditorWorldContext.World())
+					{
+						for (auto LevelIter : EditorWorld->GetLevels())
+						{
+							if (UObject* EditorObject = FindObject<UObject>(LevelIter, *ObjectPtr->GetName()))
+							{
+								CorrectedName = FName(*EditorObject->GetPathName());
+								ActorInstances.Add(CorrectedName) = EditorObject;
+								PIEInstanceNameMap.Add(InstanceNameInOut) = CorrectedName;
+								InstanceNameInOut = CorrectedName;
+								Result = EditorObject;
+								break;
+							}
+						}
+					}
+				}
+				else
+				{
+					ActorInstances.Add(InstanceNameInOut) = ObjectPtr;
+					Result = ObjectPtr;
+				}
+			}
+		}
+
+	}
+	return Result;
+}
+
+TSharedPtr<FBlueprintFunctionContext> FBlueprintExecutionContext::FindFunctionContext(const FName FunctionNameIn) const
+{
+	TSharedPtr<FBlueprintFunctionContext> Result;
+	FName FunctionName = (FunctionNameIn ==  UEdGraphSchema_K2::GN_EventGraph) ? UbergraphFunctionName : FunctionNameIn;
+
+	if (const TSharedPtr<FBlueprintFunctionContext>* SearchResult = FunctionContexts.Find(FunctionName))
+	{
+		Result = *SearchResult;
+	}
+	return Result;
+}
+
+TSharedPtr<FBlueprintFunctionContext> FBlueprintExecutionContext::GetOrCreateFunctionContext(const FName FunctionNameIn)
+{
+	FName FunctionName = (FunctionNameIn ==  UEdGraphSchema_K2::GN_EventGraph) ? UbergraphFunctionName : FunctionNameIn;
+	TSharedPtr<FBlueprintFunctionContext>& Result = FunctionContexts.FindOrAdd(FunctionName);
+	if (!Result.IsValid())
+	{
+		if (UFunction* Function = BlueprintClass->FindFunctionByName(FunctionName))
+		{
+			Result = MakeShareable(new FBlueprintFunctionContext(Function, BlueprintClass));
+		}
+	}
+	check(Result.IsValid());
+	return Result;
+}
+
+bool FBlueprintExecutionContext::HasProfilerDataForNode(const UEdGraphNode* GraphNode) const
+{
+	SCOPE_CYCLE_COUNTER(STAT_NodeLookupCost);
+	const UEdGraph* OuterGraph = GraphNode ? GraphNode->GetTypedOuter<UEdGraph>() : nullptr;
+	if (OuterGraph)
+	{
+		TSharedPtr<FBlueprintFunctionContext> FunctionContext = FindFunctionContext(OuterGraph->GetFName());
+		if (FunctionContext.IsValid())
+		{
+			return FunctionContext->HasProfilerDataForNode(GraphNode->GetFName());
+		}
+	}
+	return false;
+}
+
+TSharedPtr<FScriptExecutionNode> FBlueprintExecutionContext::GetProfilerDataForNode(const UEdGraphNode* GraphNode)
+{
+	SCOPE_CYCLE_COUNTER(STAT_NodeLookupCost);
+	TSharedPtr<FScriptExecutionNode> Result;
+	const UEdGraph* OuterGraph = GraphNode ? GraphNode->GetTypedOuter<UEdGraph>() : nullptr;
+	if (OuterGraph)
+	{
+		TSharedPtr<FBlueprintFunctionContext> FunctionContext = FindFunctionContext(OuterGraph->GetFName());
+		if (FunctionContext.IsValid() && FunctionContext->HasProfilerDataForNode(GraphNode->GetFName()))
+		{
+			Result = FunctionContext->GetProfilerDataForNode(GraphNode->GetFName());
+		}
+	}
+	return Result;
+}
+
+void FBlueprintExecutionContext::UpdateConnectedStats()
+{
+	SCOPE_CYCLE_COUNTER(STAT_StatUpdateCost);
+	if (BlueprintNode.IsValid())
+	{
+		FTracePath InitialTracePath;
+		BlueprintNode->RefreshStats(InitialTracePath);
+	}
+}
+
+bool FBlueprintExecutionContext::MapBlueprintExecution()
+{
+	bool bMappingSuccessful = false;
+
+	if (UBlueprint* BlueprintToMap = Blueprint.Get())
+	{
+		// Find all event nodes for event graph and UCS
+		TArray<UEdGraph*> Graphs;
+		TArray<UK2Node_Event*> EventNodes;
+		TArray<UK2Node_FunctionEntry*> FunctionEntryNodes;
+		BlueprintToMap->GetAllGraphs(Graphs);
+		for (auto GraphIter : Graphs)
+		{
+			if (GraphIter->GetFName() == UEdGraphSchema_K2::GN_EventGraph)
+			{
+				GraphIter->GetNodesOfClass<UK2Node_Event>(EventNodes);
+			}
+			if (GraphIter->GetFName() == UEdGraphSchema_K2::FN_UserConstructionScript)
+			{
+				GraphIter->GetNodesOfClass<UK2Node_FunctionEntry>(FunctionEntryNodes);
+			}
+		}
+		// Map execution paths for each event node
+		for (auto EventIter : EventNodes)
+		{
+			FName EventName = EventIter->GetFunctionName();
+			TSharedPtr<FBlueprintFunctionContext> FunctionContext = GetOrCreateFunctionContext(EventName);
+			TSharedPtr<FScriptExecutionNode> EventExecNode = FunctionContext->GetProfilerDataForNode(EventName);
+			EventExecNode->SetName(EventName);
+			EventExecNode->SetObservedObject(EventIter);
+			EventExecNode->SetGraphName(EventIter->GetTypedOuter<UEdGraph>()->GetFName());
+			EventExecNode->SetFlags(EScriptExecutionNodeFlags::Event);
+			FLinearColor IconColor = FLinearColor(0.91f, 0.16f, 0.16f);
+			FText ToolTipText = LOCTEXT("NavigateToEventLocationHyperlink_ToolTip", "Navigate to the Event");
+			EventExecNode->SetDisplayName(EventIter->GetNodeTitle(ENodeTitleType::ListView));
+			EventExecNode->SetToolTipText(ToolTipText);
+			const FSlateBrush* Icon_Node = EventIter->ShowPaletteIconOnNode() ?	FEditorStyle::GetBrush(EventIter->GetPaletteIcon(IconColor)) :
+																			FEditorStyle::GetBrush(TEXT("BlueprintProfiler.BPNode"));
+			EventExecNode->SetIcon(const_cast<FSlateBrush*>(Icon_Node));
+			EventExecNode->SetIconColor(IconColor);
+			// Add event to mapped blueprint
+			BlueprintNode->AddChildNode(EventExecNode);
+			// Map the event execution flow
+			TSharedPtr<FScriptExecutionNode> MappedEvent = MapNodeExecution(EventIter);
+			if (MappedEvent.IsValid())
+			{
+				EventExecNode->AddChildNode(MappedEvent);
+			}
+		}
+		// Map execution paths for each function entry node
+		for (auto FunctionIter : FunctionEntryNodes)
+		{
+			FName FunctionName = FunctionIter->SignatureName;
+			TSharedPtr<FBlueprintFunctionContext> FunctionContext = GetOrCreateFunctionContext(FunctionName);
+			TSharedPtr<FScriptExecutionNode> FunctionExecNode = FunctionContext->GetProfilerDataForNode(FunctionName);
+			FunctionExecNode->SetName(FunctionName);
+			FunctionExecNode->SetObservedObject(FunctionIter);
+			FunctionExecNode->SetGraphName(FunctionIter->GetTypedOuter<UEdGraph>()->GetFName());
+			FunctionExecNode->SetFlags(EScriptExecutionNodeFlags::Event);
+			FLinearColor IconColor = FLinearColor(0.91f, 0.16f, 0.16f);
+			FText ToolTipText = LOCTEXT("NavigateToFunctionLocationHyperlink_ToolTip", "Navigate to the Function");
+			FunctionExecNode->SetDisplayName(FunctionIter->GetNodeTitle(ENodeTitleType::ListView));
+			FunctionExecNode->SetToolTipText(ToolTipText);
+			const FSlateBrush* Icon = FunctionIter->ShowPaletteIconOnNode() ?	FEditorStyle::GetBrush(FunctionIter->GetPaletteIcon(IconColor)) :
+																				FEditorStyle::GetBrush(TEXT("BlueprintProfiler.BPNode"));
+			FunctionExecNode->SetIcon(const_cast<FSlateBrush*>(Icon));
+			FunctionExecNode->SetIconColor(IconColor);
+			// Add Function to mapped blueprint
+			BlueprintNode->AddChildNode(FunctionExecNode);
+			// Map the Function execution flow
+			TSharedPtr<FScriptExecutionNode> MappedFunction = MapNodeExecution(FunctionIter);
+			if (MappedFunction.IsValid())
+			{
+				FunctionExecNode->AddChildNode(MappedFunction);
+			}
+		}
+		bMappingSuccessful = true;
+	}
+	return bMappingSuccessful;
+}
+
+TSharedPtr<FScriptExecutionNode> FBlueprintExecutionContext::MapNodeExecution(UEdGraphNode* NodeToMap)
+{
+	TSharedPtr<FScriptExecutionNode> MappedNode;
+	if (NodeToMap)
+	{
+		const UEdGraph* OwningGraph = NodeToMap->GetTypedOuter<UEdGraph>();
+		TSharedPtr<FBlueprintFunctionContext> FunctionContext = GetOrCreateFunctionContext(OwningGraph->GetFName());
+		const bool bNodeMapped = FunctionContext->HasProfilerDataForNode(NodeToMap->GetFName());
+		MappedNode = FunctionContext->GetProfilerDataForNode(NodeToMap->GetFName());
+
+		if (!bNodeMapped)
+		{
+			// Construct mapped node values
+			MappedNode->SetName(NodeToMap->GetFName());
+			MappedNode->SetObservedObject(NodeToMap);
+			MappedNode->SetGraphName(NodeToMap->GetTypedOuter<UEdGraph>()->GetFName());
+			MappedNode->SetFlags(EScriptExecutionNodeFlags::Node);
+			FLinearColor IconColor = FLinearColor(1.f, 1.f, 1.f, 0.8f);
+			FText ToolTipText = LOCTEXT("NavigateToNodeLocationHyperlink_ToolTip", "Navigate to the Node");
+			// Evaluate Execution Pins
+			TArray<UEdGraphPin*> ExecPins;
+			for (auto Pin : NodeToMap->Pins)
+			{
+				if (Pin->Direction == EGPD_Output && Pin->PinType.PinCategory == UEdGraphSchema_K2::PC_Exec)
+				{
+					ExecPins.Add(Pin);
+				}
+			}
+			const bool bBranchNode = ExecPins.Num() > 1;
+			if (bBranchNode)
+			{
+				// Determine branch exec type
+				IconColor = FLinearColor(1.f, 1.f, 1.f, 0.8f);
+				if (NodeToMap->IsA<UK2Node_ExecutionSequence>())
+				{
+					MappedNode->AddFlags(EScriptExecutionNodeFlags::SequentialBranch);
+				}
+				else
+				{
+					MappedNode->AddFlags(EScriptExecutionNodeFlags::ConditionalBranch);
+				}
+			}
+			for (auto Pin : ExecPins)
+			{
+				const UEdGraphPin* LinkedPin = Pin->LinkedTo.Num() > 0 ? Pin->LinkedTo[0] : nullptr;
+
+				if (bBranchNode)
+				{
+					const int32 PinScriptCodeOffset = FunctionContext->GetCodeLocationFromPin(Pin);
+					// Check we have a valid script offset or that the pin isn't linked.
+					check((PinScriptCodeOffset != INDEX_NONE)||(LinkedPin == nullptr));
+					TSharedPtr<FScriptExecutionNode> LinkedNode = FunctionContext->GetProfilerDataForNode(Pin->GetFName());
+					LinkedNode->SetName(Pin->GetFName());
+					LinkedNode->SetObservedObject(Pin);
+					LinkedNode->SetGraphName(Pin->GetTypedOuter<UEdGraph>()->GetFName());
+					LinkedNode->SetDisplayName(Pin->GetDisplayName());
+					LinkedNode->SetFlags(EScriptExecutionNodeFlags::ExecPin);
+					LinkedNode->SetToolTipText(LOCTEXT("ExecPin_ToolTip", "Expand execution path"));
+					const FSlateBrush* Icon = LinkedPin ?	FEditorStyle::GetBrush(TEXT("BlueprintProfiler.BPPinConnected")) : 
+															FEditorStyle::GetBrush(TEXT("BlueprintProfiler.BPPinDisconnected"));
+					LinkedNode->SetIcon(const_cast<FSlateBrush*>(Icon));
+					LinkedNode->SetIconColor(FLinearColor(1.f, 1.f, 1.f, 0.8f));
+					MappedNode->AddLinkedNode(PinScriptCodeOffset, LinkedNode);
+					// Map pin exec flow
+					if (LinkedPin)
+					{
+						TSharedPtr<FScriptExecutionNode> ExecPinNode = MapNodeExecution(LinkedPin->GetOwningNode());
+						if (ExecPinNode.IsValid())
+						{
+							LinkedNode->AddChildNode(ExecPinNode);
+						}
+					}
+				}
+				else if (LinkedPin)
+				{
+					TSharedPtr<FScriptExecutionNode> ExecLinkedNode = MapNodeExecution(LinkedPin->GetOwningNode());
+					if (ExecLinkedNode.IsValid())
+					{
+						const int32 PinScriptCodeOffset = FunctionContext->GetCodeLocationFromPin(Pin);
+						check(PinScriptCodeOffset != INDEX_NONE);
+						MappedNode->AddLinkedNode(PinScriptCodeOffset, ExecLinkedNode);
+					}
+				}
+			}
+			// Evaluate Children for call sites
+			if (UK2Node_CallFunction* FunctionCallSite = Cast<UK2Node_CallFunction>(NodeToMap))
+			{
+				const UEdGraphNode* FunctionNode = nullptr; 
+				if (UEdGraph* CalledGraph = FunctionCallSite->GetFunctionGraph(FunctionNode))
+				{
+					// Update Exec node
+					MappedNode->AddFlags(EScriptExecutionNodeFlags::FunctionCall);
+					IconColor = FLinearColor(0.46f, 0.54f, 0.95f);
+					ToolTipText = LOCTEXT("NavigateToFunctionLocationHyperlink_ToolTip", "Navigate to the Function Callsite");
+					// Grab function entries
+					TArray<UK2Node_FunctionEntry*> FunctionEntryNodes;
+					CalledGraph->GetNodesOfClass<UK2Node_FunctionEntry>(FunctionEntryNodes);
+	
+					for (auto EntryNodeIter : FunctionEntryNodes)
+					{
+						TSharedPtr<FScriptExecutionNode> FunctionEntry = MapNodeExecution(EntryNodeIter);
+						MappedNode->AddChildNode(FunctionEntry);
+					}
+				}
+			}
+			//else if (UK2Node_MacroInstance* MacroCallSite = Cast<UK2Node_MacroInstance>(NodeToMap))
+			//{
+			//	if (UEdGraph* MacroGraph = MacroCallSite->GetMacroGraph())
+			//	{
+			//		// Update Exec node
+			//		MappedNode->AddFlags(EScriptExecutionNodeFlags::MacroCall);
+			//		IconColor = FLinearColor(0.36f, 0.34f, 0.71f);
+			//		ToolTipText = LOCTEXT("NavigateToMacroLocationHyperlink_ToolTip", "Navigate to the Macro Callsite");
+			//		// Grab the macro blueprint context
+			//		UBlueprint* MacroBP = MacroGraph->GetTypedOuter<UBlueprint>();
+			//		UBlueprintGeneratedClass* MacroBPGC = Cast<UBlueprintGeneratedClass>(MacroBP->GeneratedClass);
+			//		TSharedPtr<FBlueprintExecutionContext> MacroBlueprintContext = GetBlueprintContext(MacroBPGC->GetPathName());
+			//		// Grab function entries
+			//		TArray<UK2Node_FunctionEntry*> FunctionEntryNodes;
+			//		MacroGraph->GetNodesOfClass<UK2Node_FunctionEntry>(FunctionEntryNodes);
+	
+			//		for (auto EntryNodeIter : FunctionEntryNodes)
+			//		{
+			//			TSharedPtr<FScriptExecutionNode> MacroEntry = MapNodeExecution(MacroBlueprintContext, EntryNodeIter);
+			//			MappedNode->AddChildNode(MacroEntry);
+			//		}
+			//	}
+			//}
+			// Finalise exec node data.
+			MappedNode->SetDisplayName(NodeToMap->GetNodeTitle(ENodeTitleType::ListView));
+			MappedNode->SetToolTipText(ToolTipText);
+			const FSlateBrush* Icon = NodeToMap->ShowPaletteIconOnNode() ?	FEditorStyle::GetBrush(NodeToMap->GetPaletteIcon(IconColor)) :
+																			FEditorStyle::GetBrush(TEXT("BlueprintProfiler.BPNode"));
+			MappedNode->SetIcon(const_cast<FSlateBrush*>(Icon));
+			MappedNode->SetIconColor(IconColor);
+		}
+	}
+	return MappedNode;
+}
+//////////////////////////////////////////////////////////////////////////
+// FBlueprintFunctionContext
+
+bool FBlueprintFunctionContext::HasProfilerDataForNode(const FName NodeName) const
+{
+	return ExecutionNodes.Contains(NodeName);
+}
+
+TSharedPtr<FScriptExecutionNode> FBlueprintFunctionContext::GetProfilerDataForNode(const FName NodeName)
+{
+	TSharedPtr<FScriptExecutionNode>& Result = ExecutionNodes.FindOrAdd(NodeName);
+	if (!Result.IsValid())
+	{
+		Result = MakeShareable(new FScriptExecutionNode);
+	}
+	return Result;
+}
+
+const UEdGraphNode* FBlueprintFunctionContext::GetNodeFromCodeLocation(const int32 ScriptOffset)
+{
+	TWeakObjectPtr<const UEdGraphNode>& Result = ScriptOffsetToNodes.FindOrAdd(ScriptOffset);
+	if (!Result.IsValid() && BlueprintClass.IsValid())
+	{
+		if (const UEdGraphNode* GraphNode = BlueprintClass->GetDebugData().FindSourceNodeFromCodeLocation(Function.Get(), ScriptOffset, true))
+		{
+			Result = GraphNode;
+		}
+	}
+	return Result.Get();
+}
+
+const int32 FBlueprintFunctionContext::GetCodeLocationFromPin(const UEdGraphPin* Pin) const
+{
+	int32 Result = INDEX_NONE;
+	if (BlueprintClass.IsValid())
+	{
+		Result = BlueprintClass.Get()->GetDebugData().FindPinToCodeAssociation(Pin, Function.Get());
+	}
+	return Result;
+}
+
+//////////////////////////////////////////////////////////////////////////
+// FScriptEventPlayback
+
+bool FScriptEventPlayback::Process(const TArray<FScriptInstrumentedEvent>& Events, const int32 Start, const int32 Stop)
+{
+	struct PerNodeEventHelper
+	{
+		TArray<FTracePath> EntryTracePaths;
+		TSharedPtr<FScriptExecutionNode> ExecNode;
+		TSharedPtr<FBlueprintFunctionContext> FunctionContext;
+		TArray<FScriptInstrumentedEvent> UnproccessedEvents;
+	};
+
+	const int32 NumEvents = (Stop + 1) - Start;
+	if (BlueprintContext.IsValid() && NumEvents > 2)
+	{
+		check(Events[Start].GetType() == EScriptInstrumentation::Class);
+		check(Events[Start+1].GetType() == EScriptInstrumentation::Instance);
+		check(Events[Start+2].GetType() == EScriptInstrumentation::Event);
+		// Update event and instance names
+		EventName = Events[Start+2].GetFunctionName();
+		InstanceName = BlueprintContext->RemapInstancePath(FName(*Events[Start+1].GetObjectPath()));
+		// Process event
+		FName CurrentFunctionName = EventName;
+		TSharedPtr<FBlueprintFunctionContext> FunctionContext = BlueprintContext->GetOrCreateFunctionContext(CurrentFunctionName);
+		TSharedPtr<FScriptExecutionNode> EventNode = FunctionContext->GetProfilerDataForNode(EventName);
+		FTracePath EventTrace;
+		TSharedPtr<FScriptPerfData> PerfData = EventNode->GetPerfDataByInstanceAndTracePath(InstanceName, EventTrace);
+		PerfData->AddEventTiming(0.0, Events[Stop].GetTime() - Events[Start+2].GetTime());
+		// Find Associated graph nodes and submit into a node map for later processing.
+		FTracePath TracePath;
+		TArray<FTracePath> TraceStack;
+		TMap<const UEdGraphNode*, PerNodeEventHelper> NodeMap;
+
+		for (int32 EventIdx = Start + 3; EventIdx < Stop; ++EventIdx)
+		{
+			const FScriptInstrumentedEvent& CurrEvent = Events[EventIdx];
+			bool bFunctionChanged = false;
+			// Update script function.
+			if (CurrentFunctionName != CurrEvent.GetFunctionName())
+			{
+				bFunctionChanged = true;
+				CurrentFunctionName = CurrEvent.GetFunctionName();
+				FunctionContext = BlueprintContext->FindFunctionContext(CurrentFunctionName);
+				check(FunctionContext.IsValid());
+			}
+			if (const UEdGraphNode* GraphNode = FunctionContext->GetNodeFromCodeLocation(CurrEvent.GetScriptCodeOffset()))
+			{
+				PerNodeEventHelper& NodeEntry = NodeMap.FindOrAdd(GraphNode);
+				if (!NodeEntry.FunctionContext.IsValid())
+				{
+					NodeEntry.FunctionContext = FunctionContext;
+					NodeEntry.ExecNode = FunctionContext->GetProfilerDataForNode(GraphNode->GetFName());
+				}
+				switch (CurrEvent.GetType())
+				{
+					case EScriptInstrumentation::PureNodeEntry:
+					{
+						NodeEntry.UnproccessedEvents.Add(CurrEvent);
+						break;
+					}
+					case EScriptInstrumentation::NodeEntry:
+					{
+						NodeEntry.EntryTracePaths.Push(TracePath);
+						if (bFunctionChanged && GraphNode->IsA<UK2Node_FunctionEntry>())
+						{
+							TraceStack.Push(TracePath);
+						}
+						NodeEntry.UnproccessedEvents.Add(CurrEvent);
+						break;
+					}
+					case EScriptInstrumentation::PushState:
+					{
+						TraceStack.Push(TracePath);
+						NodeEntry.UnproccessedEvents.Add(CurrEvent);
+						break;
+					}
+					case EScriptInstrumentation::RestoreState:
+					{
+						check(TraceStack.Num());
+						TracePath = TraceStack.Last();
+						break;
+					}
+					case EScriptInstrumentation::PopState:
+					{
+						if (TraceStack.Num())
+						{
+							TracePath = TraceStack.Pop();
+						}
+						FScriptInstrumentedEvent OverrideEvent(CurrEvent);
+						OverrideEvent.OverrideType(EScriptInstrumentation::NodeExit);
+						NodeEntry.UnproccessedEvents.Add(OverrideEvent);
+						// Consolidate enclosed timings
+						int32 NodeEntryIdx = INDEX_NONE;
+						for (int32 NodeEventIdx = NodeEntry.UnproccessedEvents.Num()-2; NodeEventIdx >= 0; --NodeEventIdx)
+						{
+							const FScriptInstrumentedEvent& NodeEvent = NodeEntry.UnproccessedEvents[NodeEventIdx];
+
+							if (NodeEvent.GetType() == EScriptInstrumentation::NodeEntry)
+							{
+								if (NodeEntryIdx != INDEX_NONE)
+								{
+									// Keep only the first node entry.
+									NodeEntry.UnproccessedEvents.RemoveAt(NodeEventIdx, 1);
+								}
+								NodeEntryIdx = NodeEventIdx;
+							}
+							else if (NodeEvent.GetType() == EScriptInstrumentation::NodeExit)
+							{
+								// Strip out all node exits apart from the PopExecution overridden event.
+								NodeEntry.UnproccessedEvents.RemoveAt(NodeEventIdx, 1);
+							}
+						}
+						break;
+					}
+					case EScriptInstrumentation::NodeExit:
+					{
+						if (TraceStack.Num() && bFunctionChanged)
+						{
+							TracePath = TraceStack.Pop();
+						}
+						// Cleanup branching multiple exits and correct the tracepath
+						if (NodeEntry.UnproccessedEvents.Num() && NodeEntry.UnproccessedEvents.Last().GetType() == EScriptInstrumentation::NodeExit)
+						{
+							NodeEntry.UnproccessedEvents.Pop();
+							if (NodeEntry.EntryTracePaths.Num())
+							{
+								TracePath = NodeEntry.EntryTracePaths.Last();
+							}
+						}
+						TracePath.AddExitPin(CurrEvent.GetScriptCodeOffset());
+						NodeEntry.UnproccessedEvents.Add(CurrEvent);
+						break;
+					}
+					default:
+					{
+						NodeEntry.UnproccessedEvents.Add(CurrEvent);
+						break;
+					}
+				}
+			}
+		}
+		// Process Node map timings
+		for (auto NodeEntry : NodeMap)
+		{
+			TSharedPtr<FScriptExecutionNode> ExecNode = NodeEntry.Value.ExecNode;
+			double PureNodeEntryTime = 0.0;
+			double NodeEntryTime = 0.0;
+			for (auto EventIter = NodeEntry.Value.UnproccessedEvents.CreateIterator(); EventIter; ++EventIter)
+			{
+				switch(EventIter->GetType())
+				{
+					case EScriptInstrumentation::PureNodeEntry:
+					{
+						if (PureNodeEntryTime == 0.0)
+						{
+							PureNodeEntryTime = EventIter->GetTime();
+						}
+						break;
+					}
+					case EScriptInstrumentation::NodeEntry:
+					{
+						if (NodeEntryTime == 0.0)
+						{
+							NodeEntryTime = EventIter->GetTime();
+						}
+						break;
+					}
+					case EScriptInstrumentation::NodeExit:
+					{
+						check(NodeEntryTime != 0.0 );
+						PerfData = ExecNode->GetPerfDataByInstanceAndTracePath(InstanceName, NodeEntry.Value.EntryTracePaths.Pop());
+						double PureNodeDuration = PureNodeEntryTime != 0.0 ? (NodeEntryTime - PureNodeEntryTime) : 0.0;
+						double NodeDuration = EventIter->GetTime() - NodeEntryTime;
+						PerfData->AddEventTiming(PureNodeDuration, NodeDuration);
+						PureNodeEntryTime = NodeEntryTime = 0.0;
+						break;
+					}
+				}
+			}
+		}
+		return true;
+	}
+	return false;
+}
+
+#undef LOCTEXT_NAMESPACE
diff --git a/Engine/Source/Developer/BlueprintProfiler/Public/BlueprintInstrumentationSupport.h b/Engine/Source/Developer/BlueprintProfiler/Public/BlueprintInstrumentationSupport.h
deleted file mode 100644
diff --git a/Engine/Source/Developer/BlueprintProfiler/Public/BlueprintProfiler.h b/Engine/Source/Developer/BlueprintProfiler/Public/BlueprintProfiler.h
--- a/Engine/Source/Developer/BlueprintProfiler/Public/BlueprintProfiler.h
+++ b/Engine/Source/Developer/BlueprintProfiler/Public/BlueprintProfiler.h
@@ -3,10 +3,11 @@
 #pragma once
 
 #include "BlueprintProfilerModule.h"
-#include "BlueprintInstrumentationSupport.h"
-#include "TickableEditorObject.h"
+#include "ScriptInstrumentationCapture.h"
 
-class UEdGraphNode;
+#if WITH_EDITOR
+#include "ScriptInstrumentationPlayback.h"
+#endif
 
 //////////////////////////////////////////////////////////////////////////
 // FBlueprintProfiler
@@ -31,15 +32,11 @@ public:
 
 #if WITH_EDITOR
 	virtual FOnBPStatGraphLayoutChanged& GetGraphLayoutChangedDelegate() { return GraphLayoutChangedDelegate; }
-	virtual TSharedPtr<FScriptExecutionContext> GetCachedObjectContext(const FString& ObjectPath) override;
-	virtual TSharedPtr<FScriptExecutionNode> GetProfilerDataForNode(const FName NodeName) override;
-	virtual TSharedPtr<FScriptExecutionNode> GetOrCreateProfilerDataForNode(const FName NodeName) override;
-	virtual bool HasDataForInstance(const UObject* Instance) override;
+	virtual TSharedPtr<FBlueprintExecutionContext> GetBlueprintContext(const FString& BlueprintClassPath) override;
+	virtual TSharedPtr<FScriptExecutionNode> GetProfilerDataForNode(const UEdGraphNode* GraphNode) override;
+	virtual void MapBlueprintInstance(TSharedPtr<FBlueprintExecutionContext> BlueprintContext, const FString& InstancePath) override;
+	virtual bool HasDataForInstance(const UObject* Instance) const override;
 	virtual void ProcessEventProfilingData() override;
-	virtual const FString GetValidInstancePath(const FString PathIn);
-	virtual void NavigateToBlueprint(const FName NameIn) const override;
-	virtual void NavigateToInstance(const FName NameIn) const override;
-	virtual void NavigateToNode(const FName NameIn) const override;
 #endif
 	// End IBlueprintProfilerModule
 
@@ -59,24 +56,11 @@ private:
 	void BeginPIE(bool bIsSimulating);
 	void EndPIE(bool bIsSimulating);
 
-	/** Utility functon to extract script code offset safely */
-	int32 GetScriptCodeOffset(const EScriptInstrumentationEvent& Event) const;
-
-	/** Utility functon to extract object path safely */
-	FString GetCurrentFunctionPath(const EScriptInstrumentationEvent& Event) const;
-
 #if WITH_EDITOR
+
 	/** Removes all blueprint references from data */
 	void RemoveAllBlueprintReferences(UBlueprint* Blueprint);
 
-	/** Map Blueprint Execution flow */
-	TSharedPtr<class FScriptExecutionContext> MapBlueprintExecutionFlow(const FString& BlueprintPath);
-
-	/** Map Blueprint instance */
-	void MapBlueprintInstance(const FString& InstancePath);
-
-	/** Map Node execution flow */
-	TSharedPtr<FScriptExecutionNode> MapNodeExecutionFlow(TSharedPtr<FScriptExecutionContext> BlueprintContext, UEdGraphNode* EventToMap);
 #endif
 
 protected:
@@ -85,9 +69,9 @@ protected:
 	FOnBPStatGraphLayoutChanged GraphLayoutChangedDelegate;
 
 	/** Current instrumentation context */
-	FKismetInstrumentContext ActiveContext;
+	FInstrumentationCaptureContext CaptureContext;
 	/** Raw instrumentation data */
-	TArray<FBlueprintInstrumentedEvent> InstrumentationEvents;
+	TArray<FScriptInstrumentedEvent> InstrumentationEventQueue;
 
 	/** Profiler active state */
 	bool bProfilerActive;
@@ -97,12 +81,8 @@ protected:
 #if WITH_EDITOR
 	/** PIE Active */
 	bool bPIEActive;
-	/** Instance path mappings between PIE and editor */
-	TMap<FString, FString> PIEInstancePathMap;
-	/** Processed profiling data */
-	TMap<FName, TSharedPtr<FScriptExecutionNode>> NodeProfilingData;
 	/** Cached object path and code offset lookup to UObjects */
-	TMap<FString, TSharedPtr<FScriptExecutionContext>> PathToObjectContext;
+	TMap<FString, TSharedPtr<FBlueprintExecutionContext>> PathToBlueprintContext;
 #endif
 	
 };
diff --git a/Engine/Source/Developer/BlueprintProfiler/Public/BlueprintProfilerModule.h b/Engine/Source/Developer/BlueprintProfiler/Public/BlueprintProfilerModule.h
--- a/Engine/Source/Developer/BlueprintProfiler/Public/BlueprintProfilerModule.h
+++ b/Engine/Source/Developer/BlueprintProfiler/Public/BlueprintProfilerModule.h
@@ -4,9 +4,6 @@
 
 #include "ModuleManager.h"
 
-class UBlueprint;
-struct EScriptInstrumentationEvent;
-
 /** Delegate to broadcast structural stats changes */
 DECLARE_MULTICAST_DELEGATE_OneParam(FOnBPStatGraphLayoutChanged, TWeakObjectPtr<UBlueprint>);
 
@@ -34,33 +31,21 @@ public:
 	/** Returns the multicast delegate to signal blueprint layout changes to stats */
 	virtual FOnBPStatGraphLayoutChanged& GetGraphLayoutChangedDelegate() = 0;
 
-	/** Returns the cached object context or creates a new entry */
-	virtual TSharedPtr<class FScriptExecutionContext> GetCachedObjectContext(const FString& ObjectPath) = 0;
+	/** Returns the cached blueprint context or creates and maps a new entry */
+	virtual TSharedPtr<class FBlueprintExecutionContext> GetBlueprintContext(const FString& BlueprintClassPath) = 0;
 
 	/** Returns the current profiling event data for node */
-	virtual TSharedPtr<class FScriptExecutionNode> GetProfilerDataForNode(const FName NodeName) = 0;
+	virtual TSharedPtr<class FScriptExecutionNode> GetProfilerDataForNode(const UEdGraphNode* GraphNode) = 0;
 
-	/** Returns or creates an entry for the node */
-	virtual TSharedPtr<class FScriptExecutionNode> GetOrCreateProfilerDataForNode(const FName NodeName) = 0;
+	/** Maps a blueprint instance and links to blueprint profiling data */
+	virtual void MapBlueprintInstance(TSharedPtr<FBlueprintExecutionContext> BlueprintContext, const FString& InstancePath) = 0;
 
 	/** Returns if the profiler currently retains any data for the specified instance */
-	virtual bool HasDataForInstance(const UObject* Instance) { return false; }
-
-	/** Get Valid Instance Path */
-	virtual const FString GetValidInstancePath(const FString PathIn) = 0;
+	virtual bool HasDataForInstance(const UObject* Instance) const { return false; }
 
 	/** Process profiling event data */
 	virtual void ProcessEventProfilingData() {}
 
-	/** Navigate to blueprint */
-	virtual void NavigateToBlueprint(const FName NameIn) const = 0;
-
-	/** Navigate to instance */
-	virtual void NavigateToInstance(const FName NameIn) const = 0;
-
-	/** Navigate to profiler node */
-	virtual void NavigateToNode(const FName NameIn) const = 0;
-
 #endif // WITH_EDITOR
 
 };
\ No newline at end of file
diff --git a/Engine/Source/Developer/BlueprintProfiler/Public/ScriptInstrumentationCapture.h b/Engine/Source/Developer/BlueprintProfiler/Public/ScriptInstrumentationCapture.h
new file mode 100644
--- /dev/null
+++ b/Engine/Source/Developer/BlueprintProfiler/Public/ScriptInstrumentationCapture.h
@@ -0,0 +1,100 @@
+// Copyright 1998-2016 Epic Games, Inc. All Rights Reserved.
+
+#pragma once
+
+#include "Editor/Kismet/Public/Profiler/TracePath.h"
+
+//////////////////////////////////////////////////////////////////////////
+// FScriptInstrumentedEvent
+
+class FScriptInstrumentedEvent
+{
+public:
+
+	FScriptInstrumentedEvent() {}
+
+	bool operator == (const FScriptInstrumentedEvent& OtherIn) const;
+
+	FScriptInstrumentedEvent(EScriptInstrumentation::Type InEventType, const FString& InPathData)
+		: EventType(InEventType)
+		, PathData(InPathData)
+		, CodeOffset(-1)
+		, Time(FPlatformTime::Seconds())
+	{
+	}
+
+	FScriptInstrumentedEvent(EScriptInstrumentation::Type InEventType, const FName InFunctionName, const int32 InCodeOffset)
+		: EventType(InEventType)
+		, FunctionName(InFunctionName)
+		, CodeOffset(InCodeOffset)
+		, Time(FPlatformTime::Seconds())
+	{
+	}
+
+	/** Set's the event data */
+	void SetData(EScriptInstrumentation::Type InEventType, const FString& InPathData, const int32 InCodeOffset);
+
+	/** Returns the type of script event */
+	const EScriptInstrumentation::Type GetType() const { return EventType; }
+
+	/** Override the event type */
+	void OverrideType(const EScriptInstrumentation::Type NewType) { EventType = NewType; }
+
+	/** Returns the object path for the object that triggered this event */
+	const FString& GetObjectPath() const { return PathData; }
+
+	/** Returns the function name the event occurred in */
+	FName GetFunctionName() const { return FunctionName; }
+
+	/** Returns the script offset at the time this event was generated */
+	int32 GetScriptCodeOffset() const { return CodeOffset; }
+
+	/** Returns the system time when this event was generated */
+	double GetTime() const { return Time; }
+
+	/** Returns if this event represents a change in active class/blueprint */
+	bool IsNewClass() const { return EventType == EScriptInstrumentation::Class; }
+
+	/** Returns if this event represents a change in active instance */
+	bool IsNewInstance() const { return EventType == EScriptInstrumentation::Instance; }
+
+	/** Returns if this event represents a change in active event */
+	bool IsEvent() const { return EventType == EScriptInstrumentation::Event || EventType == EScriptInstrumentation::Stop; }
+
+	/** Returns if this event represents a node execution event */
+	bool IsNodeTiming() const {	return EventType == EScriptInstrumentation::NodeEntry || EventType == EScriptInstrumentation::NodeExit;	}
+
+protected:
+	
+	/** The event type */
+	EScriptInstrumentation::Type EventType;
+	/** The owning function name */
+	FName FunctionName;
+	/** The object path for the object emitting the event */
+	FString PathData;
+	/** The script code offset when this event was emitted */
+	int32 CodeOffset;
+	/** The recorded system time when this event was emitted */
+	double Time;
+};
+
+//////////////////////////////////////////////////////////////////////////
+// FInstrumentationCaptureContext
+
+class FInstrumentationCaptureContext
+{
+public:
+
+	/** Update the current execution context, and emit events to signal changes */
+	void UpdateContext(const UObject* InContextObject, TArray<FScriptInstrumentedEvent>& InstrumentationQueue);
+
+	/** Reset the current event context */
+	void ResetContext();
+
+private:
+
+	/** The current class/blueprint context */
+	TWeakObjectPtr<const UClass> ContextClass;
+	/** The current instance context */
+	TWeakObjectPtr<const UObject> ContextObject;
+};
diff --git a/Engine/Source/Developer/BlueprintProfiler/Public/ScriptInstrumentationPlayback.h b/Engine/Source/Developer/BlueprintProfiler/Public/ScriptInstrumentationPlayback.h
new file mode 100644
--- /dev/null
+++ b/Engine/Source/Developer/BlueprintProfiler/Public/ScriptInstrumentationPlayback.h
@@ -0,0 +1,149 @@
+// Copyright 1998-2016 Epic Games, Inc. All Rights Reserved.
+
+#pragma once
+
+#include "Editor/Kismet/Public/Profiler/TracePath.h"
+
+//////////////////////////////////////////////////////////////////////////
+// FBlueprintExecutionContext
+
+class FBlueprintExecutionContext : public TSharedFromThis<FBlueprintExecutionContext>
+{
+public:
+
+	/** Initialise the context from the blueprint path name */
+	bool InitialiseContext(const FString& BlueprintPath);
+
+	/** Returns if the blueprint context is fully formed */
+	BLUEPRINTPROFILER_API bool IsContextValid() const { return Blueprint.IsValid() && BlueprintClass.IsValid(); }
+
+	/** Returns the blueprint this context represents */
+	TWeakObjectPtr<UBlueprint> GetBlueprint() const { return Blueprint; }
+
+	/** Returns the blueprint class this context represents */
+	TWeakObjectPtr<UBlueprintGeneratedClass> GetBlueprintClass() const { return BlueprintClass; }
+
+	/** Returns the blueprint exec node for this context */
+	TSharedPtr<class FScriptExecutionBlueprint> GetBlueprintExecNode() const { return BlueprintNode; }
+
+	/** Returns true if an execution node exists for the instance */
+	BLUEPRINTPROFILER_API bool HasProfilerDataForInstance(const FName InstanceName) const;
+
+	/** Corrects the instance name if it is a PIE instance and returns the editor actor */
+	TWeakObjectPtr<const UObject> GetInstance(FName& InstanceNameInOut);
+
+	/** Remaps PIE actor instance paths to editor actor instances */
+	BLUEPRINTPROFILER_API FName RemapInstancePath(const FName InstanceName) const ;
+
+	/** Returns the function context matching the supplied name providing it exists */
+	TSharedPtr<class FBlueprintFunctionContext> FindFunctionContext(const FName FunctionNameIn) const;
+
+	/** Returns the function context matching the supplied name, creates a new one if neccesary */
+	TSharedPtr<FBlueprintFunctionContext> GetOrCreateFunctionContext(const FName FunctionNameIn);
+
+	/** Returns if this context has an execution node for the specified graph node */
+	BLUEPRINTPROFILER_API bool HasProfilerDataForNode(const UEdGraphNode* GraphNode) const;
+
+	/** Returns existing or creates a new execution node for the specified graphnode */
+	BLUEPRINTPROFILER_API TSharedPtr<FScriptExecutionNode> GetProfilerDataForNode(const UEdGraphNode* GraphNode);
+
+	/** Updates cascading stats for the blueprint, this merges instance data upwards to the blueprint */
+	void UpdateConnectedStats();
+
+private:
+
+	/** Walks the blueprint and maps events and functions ready for profiling data */
+	bool MapBlueprintExecution();
+
+	/** Maps each blueprint node following execution wires */
+	TSharedPtr<FScriptExecutionNode> MapNodeExecution(UEdGraphNode* NodeToMap);
+
+private:
+
+	/** Blueprint the context represents */
+	TWeakObjectPtr<UBlueprint> Blueprint;
+	/** Blueprint Generated class for the context */
+	TWeakObjectPtr<UBlueprintGeneratedClass> BlueprintClass;
+	/** The ubergraph function name for the blueprint generated class */
+	FName UbergraphFunctionName;
+	/** Root level execution node for this blueprint */
+	TSharedPtr<FScriptExecutionBlueprint> BlueprintNode;
+	/** UFunction contexts for this blueprint */
+	TMap<FName, TSharedPtr<FBlueprintFunctionContext>> FunctionContexts;
+	/** PIE instance remap container */
+	TMap<FName, FName> PIEInstanceNameMap;
+	/** Editor instances */
+	TMap<FName, TWeakObjectPtr<const UObject>> ActorInstances;
+
+};
+
+//////////////////////////////////////////////////////////////////////////
+// FBlueprintFunctionContext
+
+class FBlueprintFunctionContext : public TSharedFromThis<FBlueprintFunctionContext>
+{
+public:
+
+	FBlueprintFunctionContext(UFunction* FunctionIn, TWeakObjectPtr<UBlueprintGeneratedClass> BPGC)
+		: Function(FunctionIn)
+		, BlueprintClass(BPGC)
+	{
+	}
+
+	/** Returns if the function context is fully formed */
+	bool IsContextValid() const { return Function.IsValid() && BlueprintClass.IsValid(); }
+
+	/** Returns the graph node at the specified script offset */
+	const UEdGraphNode* GetNodeFromCodeLocation(const int32 ScriptOffset);
+
+	/** Returns the script offset for the specified pin */
+	const int32 GetCodeLocationFromPin(const UEdGraphPin* Pin) const;
+
+	/** Returns true if this function context contains an execution node matching the node name */
+	bool HasProfilerDataForNode(const FName NodeName) const;
+
+	/** Returns the execution node representing the node name */
+	TSharedPtr<FScriptExecutionNode> GetProfilerDataForNode(const FName NodeName);
+
+private:
+
+	/** The UFunction this context represents */
+	TWeakObjectPtr<UFunction> Function;
+	/** The blueprint generated class for the owning blueprint */
+	TWeakObjectPtr<UBlueprintGeneratedClass> BlueprintClass;
+	/** ScriptCode offset to node cache */
+	TMap<int32, TWeakObjectPtr<const UEdGraphNode>> ScriptOffsetToNodes;
+	/** Execution nodes containing profiling data */
+	TMap<FName, TSharedPtr<FScriptExecutionNode>> ExecutionNodes;
+};
+
+//////////////////////////////////////////////////////////////////////////
+// FScriptEventPlayback
+
+class FScriptEventPlayback
+{
+public:
+
+	FScriptEventPlayback(TSharedPtr<FBlueprintExecutionContext> BlueprintContextIn)
+		: BlueprintContext(BlueprintContextIn)
+	{
+	}
+
+	/** Processes the event and cleans up any errant signals */
+	bool Process(const TArray<class FScriptInstrumentedEvent>& Events, const int32 Start, const int32 Stop);
+
+	/** Sets the blueprint context for this batch */
+	void SetBlueprintContext(TSharedPtr<FBlueprintExecutionContext> BlueprintContextIn) { BlueprintContext = BlueprintContextIn; }
+
+private:
+
+	/** Instance name */
+	FName InstanceName;
+	/** Event name */
+	FName EventName;
+	/** Blueprint exec context */
+	TSharedPtr<FBlueprintExecutionContext> BlueprintContext;
+
+};
+
+
diff --git a/Engine/Source/Developer/BlueprintProfiler/classes/BlueprintInstrumentationSupport.cpp b/Engine/Source/Developer/BlueprintProfiler/classes/BlueprintInstrumentationSupport.cpp
deleted file mode 100644
diff --git a/Engine/Source/Developer/BlueprintProfiler/classes/BlueprintProfiler.cpp b/Engine/Source/Developer/BlueprintProfiler/classes/BlueprintProfiler.cpp
deleted file mode 100644
diff --git a/Engine/Source/Editor/BlueprintGraph/Classes/K2Node_GetArrayItem.h b/Engine/Source/Editor/BlueprintGraph/Classes/K2Node_GetArrayItem.h
new file mode 100644
--- /dev/null
+++ b/Engine/Source/Editor/BlueprintGraph/Classes/K2Node_GetArrayItem.h
@@ -0,0 +1,40 @@
+// Copyright 1998-2016 Epic Games, Inc. All Rights Reserved.
+
+
+#pragma once
+#include "K2Node.h"
+#include "K2Node_GetArrayItem.generated.h"
+
+UCLASS(MinimalAPI, Category = "Utilities|Array", meta=(Keywords = "array"))
+class UK2Node_GetArrayItem : public UK2Node
+{
+	GENERATED_UCLASS_BODY()
+
+public:
+	// UEdGraphNode interface
+	virtual bool IsNodePure() const override { return true; }
+	virtual void AllocateDefaultPins() override;
+	virtual void PostReconstructNode() override;
+	virtual FText GetNodeTitle(ENodeTitleType::Type TitleType) const override;
+	virtual class FNodeHandlingFunctor* CreateNodeHandler(class FKismetCompilerContext& CompilerContext) const override;
+	// End of UEdGraphNode interface
+
+	// UK2Node interface
+	virtual bool ShouldDrawCompact() const { return true; }
+	virtual void NotifyPinConnectionListChanged(UEdGraphPin* Pin) override;
+	virtual void GetMenuActions(FBlueprintActionDatabaseRegistrar& ActionRegistrar) const override;
+	virtual int32 GetNodeRefreshPriority() const override { return EBaseNodeRefreshPriority::Low_UsesDependentWildcard; }
+	// End of UK2Node interface
+
+	/** Helper function to return the array pin */
+	UEdGraphPin* GetTargetArrayPin() { return Pins[0]; }
+
+	/** Helper function to return the index pin */
+	UEdGraphPin* GetIndexPin() { return Pins[1]; }
+
+	/** Helper function to return the result pin */
+	UEdGraphPin* GetResultPin() { return Pins[2]; }
+
+	/** Propagates pin type to the Array pin and the Result pin */
+	void PropagatePinType(FEdGraphPinType& InType);
+};
\ No newline at end of file
diff --git a/Engine/Source/Editor/BlueprintGraph/Classes/K2Node_MakeArray.h b/Engine/Source/Editor/BlueprintGraph/Classes/K2Node_MakeArray.h
--- a/Engine/Source/Editor/BlueprintGraph/Classes/K2Node_MakeArray.h
+++ b/Engine/Source/Editor/BlueprintGraph/Classes/K2Node_MakeArray.h
@@ -33,8 +33,8 @@ public:
 
 	// UK2Node interface
 	virtual bool IsNodePure() const override { return true; }
-	void NotifyPinConnectionListChanged(UEdGraphPin* Pin) override;
-	void GetContextMenuActions(const FGraphNodeContextMenuBuilder& Context) const override;
+	virtual void NotifyPinConnectionListChanged(UEdGraphPin* Pin) override;
+	virtual void GetContextMenuActions(const FGraphNodeContextMenuBuilder& Context) const override;
 	virtual class FNodeHandlingFunctor* CreateNodeHandler(class FKismetCompilerContext& CompilerContext) const override;
 	virtual bool IsConnectionDisallowed(const UEdGraphPin* MyPin, const UEdGraphPin* OtherPin, FString& OutReason) const override;
 	virtual void GetMenuActions(FBlueprintActionDatabaseRegistrar& ActionRegistrar) const override;
@@ -60,5 +60,4 @@ protected:
 
 	/** Returns the function to be used to add a function to the array */
 	BLUEPRINTGRAPH_API UFunction* GetArrayAddFunction() const;
-};
-
+};
\ No newline at end of file
diff --git a/Engine/Source/Editor/BlueprintGraph/Classes/K2Node_SetFieldsInStruct.h b/Engine/Source/Editor/BlueprintGraph/Classes/K2Node_SetFieldsInStruct.h
--- a/Engine/Source/Editor/BlueprintGraph/Classes/K2Node_SetFieldsInStruct.h
+++ b/Engine/Source/Editor/BlueprintGraph/Classes/K2Node_SetFieldsInStruct.h
@@ -17,7 +17,6 @@ class UK2Node_SetFieldsInStruct : public UK2Node_MakeStruct
 	virtual FText GetTooltipText() const override;
 	virtual FName GetPaletteIcon(FLinearColor& OutColor) const override;
 	virtual void ValidateNodeDuringCompilation(class FCompilerResultsLog& MessageLog) const override;
-	virtual void ExpandNode(class FKismetCompilerContext& CompilerContext, UEdGraph* SourceGraph);
 	virtual bool IsConnectionDisallowed(const UEdGraphPin* MyPin, const UEdGraphPin* OtherPin, FString& OutReason) const override;
 	//~ End  UEdGraphNode Interface
 
diff --git a/Engine/Source/Editor/BlueprintGraph/Private/EdGraphSchema_K2.cpp b/Engine/Source/Editor/BlueprintGraph/Private/EdGraphSchema_K2.cpp
--- a/Engine/Source/Editor/BlueprintGraph/Private/EdGraphSchema_K2.cpp
+++ b/Engine/Source/Editor/BlueprintGraph/Private/EdGraphSchema_K2.cpp
@@ -2,6 +2,7 @@
 
 #include "BlueprintGraphPrivatePCH.h"
 
+#include "UObject/DevObjectVersion.h"
 #include "Engine/LevelScriptBlueprint.h"
 #include "Kismet/KismetSystemLibrary.h"
 #include "Kismet/KismetArrayLibrary.h"
@@ -894,6 +895,21 @@ void UEdGraphSchema_K2::GetAutoEmitTermParameters(const UFunction* Function, TAr
 	{
 		FString MetaData = Function->GetMetaData(FBlueprintMetadata::MD_AutoCreateRefTerm);
 		MetaData.ParseIntoArray(AutoEmitParameterNames, TEXT(","), true);
+
+		for (int32 NameIndex = 0; NameIndex < AutoEmitParameterNames.Num();)
+		{
+			FString& ParameterName = AutoEmitParameterNames[NameIndex];
+			ParameterName.Trim();
+			ParameterName.TrimTrailing();
+			if (ParameterName.IsEmpty())
+			{
+				AutoEmitParameterNames.RemoveAtSwap(NameIndex);
+			}
+			else
+			{
+				++NameIndex;
+			}
+		}
 	}
 }
 
@@ -4201,8 +4217,7 @@ bool UEdGraphSchema_K2::IsAutoCreateRefTerm(const UEdGraphPin* Pin) const
 		if (TargetFunction && !Pin->PinName.IsEmpty())
 		{
 			TArray<FString> AutoCreateParameterNames;
-			FString MetaData = TargetFunction->GetMetaData(FBlueprintMetadata::MD_AutoCreateRefTerm);
-			MetaData.ParseIntoArray(AutoCreateParameterNames, TEXT(","), true);
+			GetAutoEmitTermParameters(TargetFunction, AutoCreateParameterNames);
 			bIsAutoCreateRefTerm = AutoCreateParameterNames.Contains(Pin->PinName);
 		}
 	}
@@ -5396,6 +5411,300 @@ void UEdGraphSchema_K2::BackwardCompatibilityNodeConversion(UEdGraph* Graph, boo
 				}
 			}
 		}
+
+		if (Graph->GetLinker() == nullptr)
+		{
+			return;
+		}
+
+		const FCustomVersion* CustomVersion = Graph->GetLinker()->GetCustomVersions().GetVersion(FBlueprintsObjectVersion::GUID);
+		if (CustomVersion == nullptr || CustomVersion->Version < FBlueprintsObjectVersion::ArrayGetByRefUpgrade)
+		{
+			// "ForEachLoop" and "ForEachLoopWithBreak" are special case macros that we do not want to add COPY nodes into, we do however want them outside the macro.
+			auto IsGraphForEachLoop = [](UEdGraph* InGraph) -> bool
+			{
+				check(InGraph);
+				return (InGraph->GetName() == TEXT("ForEachLoop") || InGraph->GetName() == TEXT("ForEachLoopWithBreak")) && InGraph->GetOutermost()->GetName() == TEXT("/Engine/EditorBlueprintResources/StandardMacros");
+			};
+
+			UBlueprint* GraphBlueprint = FBlueprintEditorUtils::FindBlueprintForGraph(Graph);
+			check(GraphBlueprint);
+
+			// For collecting upgrade notes and warnings from the upgrade process.
+			FCompilerResultsLog Results;
+			if (!GraphBlueprint->UpgradeNotesLog.IsValid())
+			{
+				GraphBlueprint->UpgradeNotesLog = TSharedPtr<FCompilerResultsLog>(new FCompilerResultsLog);
+				GraphBlueprint->UpgradeNotesLog->bSilentMode = false;
+				GraphBlueprint->UpgradeNotesLog->bAnnotateMentionedNodes = false;
+			}
+
+			// Finding all function nodes that use the "Array_Get" function and replacing them with the new by-ref node.
+			// We also inject Copy and Break nodes as needed to maintain previous functionality (in case by-value returns
+			// were being taken advantage of).
+			TArray<UK2Node_CallFunction*> CallFuncNodes;
+			Graph->GetNodesOfClass<UK2Node_CallFunction>(CallFuncNodes);
+			for (UK2Node_CallFunction* OldCallFuncNode : CallFuncNodes)
+			{
+				if (UFunction* TargetFunction = OldCallFuncNode->GetTargetFunction())
+				{
+					if (TargetFunction == UKismetArrayLibrary::StaticClass()->FindFunctionByName(GET_MEMBER_NAME_CHECKED(UKismetArrayLibrary, Array_Get)))
+					{
+						UK2Node_GetArrayItem* NewArrayGetNode = NewObject<UK2Node_GetArrayItem>(Graph);
+						check(NewArrayGetNode);
+						NewArrayGetNode->SetFlags(RF_Transactional);
+						Graph->AddNode(NewArrayGetNode, false, false);
+						NewArrayGetNode->CreateNewGuid();
+						NewArrayGetNode->PostPlacedNewNode();
+						NewArrayGetNode->AllocateDefaultPins();
+						NewArrayGetNode->NodePosX = OldCallFuncNode->NodePosX;
+						NewArrayGetNode->NodePosY = OldCallFuncNode->NodePosY;
+
+						bool bIsArrayObjectType = false;
+						for (UEdGraphPin* OldPin : OldCallFuncNode->Pins)
+						{
+							UEdGraphPin* TargetPin = nullptr;
+							if (OldPin->PinType.bIsArray)
+							{
+								TargetPin = NewArrayGetNode->GetTargetArrayPin();
+							}
+							else if (OldPin->Direction == EGPD_Input && OldPin->PinType.bIsArray == false && OldPin->PinType.PinCategory == PC_Int)
+							{
+								TargetPin = NewArrayGetNode->GetIndexPin();
+							}
+							else if (OldPin->Direction == EGPD_Output && OldPin->ParentPin == nullptr)
+							{
+								// The target pin is the Get-by-ref node's result pin
+								TargetPin = NewArrayGetNode->GetResultPin();
+
+								bIsArrayObjectType = TargetPin->PinType.PinCategory == PC_Object || TargetPin->PinType.PinCategory == PC_Class || 
+														TargetPin->PinType.PinCategory == PC_Asset || TargetPin->PinType.PinCategory == PC_AssetClass || 
+														TargetPin->PinType.PinCategory == PC_Interface;
+								// If the graph is from the "ForEachLoop" or "ForEachLoopWithBreak" macro, do not add the COPY node
+								if (!IsGraphForEachLoop(Graph) && !bIsArrayObjectType)
+								{
+									UK2Node_Copy* NewCopyNode = NewObject<UK2Node_Copy>(Graph);
+									check(NewCopyNode);
+									NewCopyNode->SetFlags(RF_Transactional);
+									Graph->AddNode(NewCopyNode, false, false);
+									NewCopyNode->CreateNewGuid();
+									NewCopyNode->PostPlacedNewNode();
+									NewCopyNode->AllocateDefaultPins();
+									NewCopyNode->NodePosX = NewArrayGetNode->NodePosX + 175.0f;
+									NewCopyNode->NodePosY = NewArrayGetNode->NodePosY;
+
+									// Add an upgrade note, informing users that the node was auto-injected
+									{
+										FFormatNamedArguments Args;
+										Args.Add(TEXT("CopyNodeName"), NewCopyNode->GetNodeTitle(ENodeTitleType::ListView));
+										Args.Add(TEXT("ArrayGetNodeName"), NewArrayGetNode->GetNodeTitle(ENodeTitleType::ListView));
+										FText UpgradeNoteText = FText::Format(LOCTEXT("ArrayGetByRef_CopyInjection", "This is a \"{CopyNodeName}\" node. It was injected to preserve functionality since the array \"{ArrayGetNodeName}\" node has changed. It outputs a copy of the input element (so modifications don't get applied to the original)."), Args);
+										NewCopyNode->AddNodeUpgradeNote(UpgradeNoteText);
+
+										// Upgrade note for the log
+										GraphBlueprint->UpgradeNotesLog->Note(TEXT("Auto-Upgrade Note: This is a @@ node. It was injected to preserve functionality since the array @@ node has changed. It outputs a copy of the input element (so modifications don't get applied to the original)."), NewCopyNode, NewArrayGetNode);
+									}
+
+									// Connect the Get-by-ref output pin to the copy node's input
+									TargetPin->MakeLinkTo(NewCopyNode->GetInputReferencePin());
+									NewCopyNode->NotifyPinConnectionListChanged(NewCopyNode->GetInputReferencePin());
+
+									if (OldPin->SubPins.Num())
+									{
+										// Make a separate break node!
+										UK2Node_BreakStruct* NewBreakNode = NewObject<UK2Node_BreakStruct>(Graph);
+										check(NewBreakNode);
+										NewBreakNode->SetFlags(RF_Transactional);
+										Graph->AddNode(NewBreakNode, false, false);
+										NewBreakNode->CreateNewGuid();
+										NewBreakNode->PostPlacedNewNode();
+										NewBreakNode->StructType = CastChecked<UScriptStruct>(OldPin->PinType.PinSubCategoryObject.Get());
+										NewBreakNode->AllocateDefaultPins();
+										NewBreakNode->NodePosX = NewCopyNode->NodePosX + 50.0f;
+										NewBreakNode->NodePosY = NewCopyNode->NodePosY;
+
+										// Add an upgrade note, informing users that the node was auto-injected
+										{
+											FFormatNamedArguments Args;
+											Args.Add(TEXT("ArrayGetNodeName"), NewArrayGetNode->GetNodeTitle(ENodeTitleType::ListView));
+											FText UpgradeNoteText = FText::Format(LOCTEXT("ArrayGetByRef_BreakInjection", "Since the array \"{ArrayGetNodeName}\" node has changed, this node was injected to preserve functionality."), Args);
+											NewBreakNode->AddNodeUpgradeNote(UpgradeNoteText);
+
+											// Upgrade note for the log
+											GraphBlueprint->UpgradeNotesLog->Note(TEXT("Auto-Upgrade Note: Since the array @@ node has changed, @@ was injected to preserve functionality."), NewArrayGetNode, NewBreakNode );
+										}
+
+										NewCopyNode->GetCopyResultPin()->MakeLinkTo(NewBreakNode->Pins[0]);
+										int SubPinIdx = 0;
+										for (UEdGraphPin* BreakNodePin : NewBreakNode->Pins)
+										{
+											if (BreakNodePin->Direction == EGPD_Output)
+											{
+												UEdGraphPin* SplitPin = OldPin->SubPins[SubPinIdx++];
+												if (!Graph->GetSchema()->MovePinLinks(*SplitPin, *BreakNodePin).CanSafeConnect())
+												{
+													GraphBlueprint->UpgradeNotesLog->Warning(TEXT("BackwardCompatibilityNodeConversion Error 'cannot safetly move pin @@ to @@' in blueprint: @@"), SplitPin, BreakNodePin, GraphBlueprint);
+												}
+											}
+										}
+									}
+									else
+									{
+										// Move all pins from the OldPin (the result of the Get node) to the output of the Copy Node
+										if (!Graph->GetSchema()->MovePinLinks(*OldPin, *NewCopyNode->GetCopyResultPin()).CanSafeConnect())
+										{
+											GraphBlueprint->UpgradeNotesLog->Warning(TEXT("BackwardCompatibilityNodeConversion Error 'cannot safetly move pin @@ to @@' in blueprint: @@"), OldPin, NewCopyNode, GraphBlueprint);
+										}
+
+										// We have made the connections already
+										TargetPin = nullptr;
+									}
+								}
+							}
+							else if (OldPin->PinName == PN_Self)
+							{
+								// Do not continue, there's just nothing to do
+							}
+							else if (OldPin->ParentPin)
+							{
+								continue;
+							}
+							else
+							{
+								check(0);
+							}
+
+							if (TargetPin && !Graph->GetSchema()->MovePinLinks(*OldPin, *TargetPin).CanSafeConnect())
+							{
+								GraphBlueprint->UpgradeNotesLog->Warning(TEXT("BackwardCompatibilityNodeConversion Error 'cannot safetly move pin @@ to @@' in blueprint: @@"), OldPin, TargetPin, GraphBlueprint);
+							}
+							else if (TargetPin)
+							{
+								NewArrayGetNode->NotifyPinConnectionListChanged(TargetPin);
+							}
+						}
+
+						if (!bIsArrayObjectType)
+						{
+							// Add an upgrade note to the node informing users that the functionality has changed.
+							NewArrayGetNode->AddNodeUpgradeNote(FText::FromString(TEXT("This node has changed! Its output is now a reference - which means modifying the element that this returns is the same as modifying the element stored in the array.")));
+
+							// Upgrade note for the log
+							GraphBlueprint->UpgradeNotesLog->Note(TEXT("Auto-Upgrade Note: @@ has changed! @@ is now a reference - which means modifying the element that this returns is the same as modifying the element stored in the array."), NewArrayGetNode, NewArrayGetNode->GetResultPin());
+						}
+
+						OldCallFuncNode->DestroyNode();
+					}
+				}
+			}
+
+			// We need to look at Macro Instances and upgrade all "ForEachLoop" and "ForEachLoopWithBreak" nodes because those graphs were prevented from having the Copy node injected
+			// because it is desirable for those Macros to return by-ref. However, to keep Blueprints running as they did prior to the upgrade, the instanced Macros need the Copy node.
+			TArray<UK2Node_MacroInstance*> MacroInstanceNodes;
+			Graph->GetNodesOfClass<UK2Node_MacroInstance>(MacroInstanceNodes);
+			for (UK2Node_MacroInstance* MacroInstanceNode : MacroInstanceNodes)
+			{
+				if (IsGraphForEachLoop(MacroInstanceNode->GetMacroGraph()))
+				{
+					UEdGraphPin* ArrayElementPin = MacroInstanceNode->FindPin(TEXT("Array Element"));
+
+					// Because we have a "ForEachLoop" or "ForEachLoopWithBreak" macro instance, we need to add a COPY node at the end, attached to the "Array Element" pin
+					if (ensure(ArrayElementPin) && 
+						ArrayElementPin->PinType.PinCategory != PC_Object && ArrayElementPin->PinType.PinCategory != PC_Class && 
+						ArrayElementPin->PinType.PinCategory != PC_Asset && ArrayElementPin->PinType.PinCategory != PC_AssetClass && 
+						ArrayElementPin->PinType.PinCategory != PC_Interface)
+					{
+						// Add an upgrade note, informing users that the node was auto-injected
+						{
+							FFormatNamedArguments Args;
+							Args.Add(TEXT("ArrayElementPin"), ArrayElementPin->PinFriendlyName);
+							FText UpgradeNoteText = FText::Format(LOCTEXT("ArrayGetByRef_ForEachLoopChanged", "This node has changed! Its output is now a reference - which means modifying the element that this returns in {ArrayElementPin} is the same as modifying the element stored in the array."), Args);
+							MacroInstanceNode->AddNodeUpgradeNote(UpgradeNoteText);
+
+							// Upgrade note for the log
+							GraphBlueprint->UpgradeNotesLog->Note(TEXT("Auto-Upgrade Note: @@ has changed! @@ is now a reference - which means modifying the element that this returns is the same as modifying the element stored in the array."), MacroInstanceNode, ArrayElementPin);
+						}
+
+						UK2Node_Copy* NewCopyNode = NewObject<UK2Node_Copy>(Graph);
+						check(NewCopyNode);
+						NewCopyNode->SetFlags(RF_Transactional);
+						Graph->AddNode(NewCopyNode, false, false);
+						NewCopyNode->CreateNewGuid();
+						NewCopyNode->PostPlacedNewNode();
+						NewCopyNode->AllocateDefaultPins();
+						NewCopyNode->NodePosX = MacroInstanceNode->NodePosX + 200.0f;
+						NewCopyNode->NodePosY = MacroInstanceNode->NodePosY + 50.0f;
+
+						// Add an upgrade note, informing users that the node was auto-injected
+						{
+							FFormatNamedArguments Args;
+							Args.Add(TEXT("CopyNodeName"), NewCopyNode->GetNodeTitle(ENodeTitleType::ListView));
+							Args.Add(TEXT("ForEachLoopName"), MacroInstanceNode->GetNodeTitle(ENodeTitleType::ListView));
+							FText UpgradeNoteText = FText::Format(LOCTEXT("ArrayGetByRef_ForLoopCopyInjection", "This is a \"{CopyNodeName}\" node. It was injected to preserve functionality since the array \"{ArrayGetNodeName}\" node has changed (inside {ForEachLoopName}). It outputs a copy of the input element (so modifications don't get applied to the original)."), Args);
+							NewCopyNode->AddNodeUpgradeNote(UpgradeNoteText);
+
+							// Upgrade note for the log
+							GraphBlueprint->UpgradeNotesLog->Note(TEXT("Auto-Upgrade Note: This is a @@ node. It was injected to preserve functionality since the array \"Get\" node (inside @@) has changed. It outputs a copy of the input element (so modifications don't get applied to the original)."), NewCopyNode, MacroInstanceNode);
+						}
+
+						if (ArrayElementPin->SubPins.Num())
+						{
+							// Make a separate break node!
+							UK2Node_BreakStruct* NewBreakNode = NewObject<UK2Node_BreakStruct>(Graph);
+							check(NewBreakNode);
+							NewBreakNode->SetFlags(RF_Transactional);
+							Graph->AddNode(NewBreakNode, false, false);
+							NewBreakNode->CreateNewGuid();
+							NewBreakNode->PostPlacedNewNode();
+							NewBreakNode->StructType = CastChecked<UScriptStruct>(ArrayElementPin->PinType.PinSubCategoryObject.Get());
+							NewBreakNode->AllocateDefaultPins();
+							NewBreakNode->NodePosX = NewCopyNode->NodePosX + 50.0f;
+							NewBreakNode->NodePosY = NewCopyNode->NodePosY;
+
+							// Add an upgrade note, informing users that the node was auto-injected
+							{
+								FFormatNamedArguments Args;
+								Args.Add(TEXT("ForEachLoopName"), MacroInstanceNode->GetNodeTitle(ENodeTitleType::ListView));
+								FText UpgradeNoteText = FText::Format(LOCTEXT("ArrayGetByRef_ForLoopBreakInjection", "Since the array \"Get\" node has changed (inside {ForEachLoopName}), this node was injected to preserve functionality."), Args);
+								NewBreakNode->AddNodeUpgradeNote(UpgradeNoteText);
+
+								// Upgrade note for the log
+								GraphBlueprint->UpgradeNotesLog->Note(TEXT("Auto-Upgrade Note: Since the array \"Get\" node has changed (inside @@), @@ was injected to preserve functionality. ."), MacroInstanceNode, NewBreakNode);
+							}
+
+							NewCopyNode->GetCopyResultPin()->MakeLinkTo(NewBreakNode->Pins[0]);
+							int SubPinIdx = 0;
+							for (UEdGraphPin* BreakNodePin : NewBreakNode->Pins)
+							{
+								if (BreakNodePin->Direction == EGPD_Output)
+								{
+									UEdGraphPin* SplitPin = ArrayElementPin->SubPins[SubPinIdx++];
+									if (!Graph->GetSchema()->MovePinLinks(*SplitPin, *BreakNodePin).CanSafeConnect())
+									{
+										GraphBlueprint->UpgradeNotesLog->Warning(TEXT("BackwardCompatibilityNodeConversion Error 'cannot safetly move pin @@ to @@' in blueprint: @@"), SplitPin, BreakNodePin, GraphBlueprint);
+									}
+								}
+							}
+
+							// Recombine the array element pin's children
+							RecombinePin(ArrayElementPin->SubPins[0]);
+						}
+						else
+						{
+							// Move all pins from the OldPin (the result of the Get node) to the output of the Copy Node
+							if (!Graph->GetSchema()->MovePinLinks(*ArrayElementPin, *NewCopyNode->GetCopyResultPin()).CanSafeConnect())
+							{
+								GraphBlueprint->UpgradeNotesLog->Warning(TEXT("BackwardCompatibilityNodeConversion Error 'cannot safetly move pin @@ to @@' in blueprint: @@"), ArrayElementPin, NewCopyNode, GraphBlueprint);
+							}
+						}
+
+						// Connect the Get-by-ref output pin to the copy node's input
+						ArrayElementPin->MakeLinkTo(NewCopyNode->GetInputReferencePin());
+						NewCopyNode->NotifyPinConnectionListChanged(NewCopyNode->GetInputReferencePin());
+					}
+				}
+			}
+		} // End FBlueprintsObjectVersion::ArrayGetByRefUpgrade upgrades
 	}
 }
 
diff --git a/Engine/Source/Editor/BlueprintGraph/Private/K2Node_BreakStruct.cpp b/Engine/Source/Editor/BlueprintGraph/Private/K2Node_BreakStruct.cpp
--- a/Engine/Source/Editor/BlueprintGraph/Private/K2Node_BreakStruct.cpp
+++ b/Engine/Source/Editor/BlueprintGraph/Private/K2Node_BreakStruct.cpp
@@ -91,11 +91,7 @@ public:
 
 			UBlueprintEditorSettings* Settings = GetMutableDefault<UBlueprintEditorSettings>();
 			FBPTerminal* Term = Context.CreateLocalTerminalFromPinAutoChooseScope(Net, Net->PinName);
-
-			if (Settings->bAllowReferencePassThroughsToCacheNonReferences)
-			{
-				Term->bPassedByReference = ContextTerm->bPassedByReference;
-			}
+			Term->bPassedByReference = ContextTerm->bPassedByReference;
 			Term->AssociatedVarProperty = BoundProperty;
 			Context.NetMap.Add(Net, Term);
 			Term->Context = ContextTerm;
diff --git a/Engine/Source/Editor/BlueprintGraph/Private/K2Node_ExecutionSequence.cpp b/Engine/Source/Editor/BlueprintGraph/Private/K2Node_ExecutionSequence.cpp
--- a/Engine/Source/Editor/BlueprintGraph/Private/K2Node_ExecutionSequence.cpp
+++ b/Engine/Source/Editor/BlueprintGraph/Private/K2Node_ExecutionSequence.cpp
@@ -58,6 +58,7 @@ public:
 			if (Context.IsDebuggingOrInstrumentationRequired() && (OutputPins.Num() > 1))
 			{
 				const FString NodeComment = Node->NodeComment.IsEmpty() ? Node->GetName() : Node->NodeComment;
+				const bool bInstrumentedVersion = Context.bInstrumentScriptCode;
 
 				// Assuming sequence X goes to A, B, C, we want to emit:
 				//   X: push X1
@@ -71,6 +72,14 @@ public:
 				// A push statement we need to patch up on the next pass (e.g., push X1 before we know where X1 is)
 				FBlueprintCompiledStatement* LastPushStatement = NULL;
 
+				if (bInstrumentedVersion)
+				{
+					// Emit a sequence instruction
+					FBlueprintCompiledStatement& SequenceStart = Context.AppendStatementForNode(Node);
+					SequenceStart.Type = KCST_InstrumentedStatePush;
+					SequenceStart.Comment = NodeComment;
+				}
+
 				for (int32 i = 0; i < OutputPins.Num(); ++i)
 				{
 					// Emit the debug site and patch up the previous jump if we're on subsequent steps
@@ -79,7 +88,7 @@ public:
 					{
 						// Emit a debug site
 						FBlueprintCompiledStatement& DebugSiteAndJumpTarget = Context.AppendStatementForNode(Node);
-						DebugSiteAndJumpTarget.Type = Context.GetBreakpointType();
+						DebugSiteAndJumpTarget.Type = bInstrumentedVersion ? KCST_InstrumentedStateRestore : Context.GetBreakpointType();
 						DebugSiteAndJumpTarget.Comment = NodeComment;
 						DebugSiteAndJumpTarget.bIsJumpTarget = true;
 
@@ -88,8 +97,8 @@ public:
 						LastPushStatement->TargetLabel = &DebugSiteAndJumpTarget;
 					}
 
-					// Emit a push to get to the next step in the sequence, unless we're the last one
-					const bool bNotLastIndex = (i + 1) < OutputPins.Num();
+					// Emit a push to get to the next step in the sequence, unless we're the last one or this is an instrumented build
+					const bool bNotLastIndex = ((i + 1) < OutputPins.Num()) || bInstrumentedVersion;
 					if (bNotLastIndex)
 					{
 						FBlueprintCompiledStatement& PushExecutionState = Context.AppendStatementForNode(Node);
@@ -102,6 +111,17 @@ public:
 					GotoSequenceLinkedState.Type = KCST_UnconditionalGoto;
 					Context.GotoFixupRequestMap.Add(&GotoSequenceLinkedState, OutputPins[i]);
 				}
+				if (bInstrumentedVersion)
+				{
+					// For instrumented builds jump back to node and pop the execution state then exit the thread.
+					FBlueprintCompiledStatement& SequenceEnd = Context.AppendStatementForNode(Node);
+					SequenceEnd.Type = KCST_InstrumentedStatePop;
+					SequenceEnd.Comment = NodeComment;
+					SequenceEnd.bIsJumpTarget = true;
+					LastPushStatement->TargetLabel = &SequenceEnd;
+					FBlueprintCompiledStatement& NextExecutionState = Context.AppendStatementForNode(Node);
+					NextExecutionState.Type = KCST_EndOfThread;
+				}
 			}
 			else
 			{
diff --git a/Engine/Source/Editor/BlueprintGraph/Private/K2Node_GetArrayItem.cpp b/Engine/Source/Editor/BlueprintGraph/Private/K2Node_GetArrayItem.cpp
new file mode 100644
--- /dev/null
+++ b/Engine/Source/Editor/BlueprintGraph/Private/K2Node_GetArrayItem.cpp
@@ -0,0 +1,221 @@
+// Copyright 1998-2016 Epic Games, Inc. All Rights Reserved.
+
+
+#include "BlueprintGraphPrivatePCH.h"
+
+#include "SlateBasics.h"
+#include "../../../Runtime/Engine/Classes/Kismet/KismetArrayLibrary.h"
+#include "ScopedTransaction.h"
+#include "KismetCompiler.h"
+#include "BlueprintNodeSpawner.h"
+#include "EditorCategoryUtils.h"
+#include "BlueprintActionDatabaseRegistrar.h"
+#include "EdGraph/EdGraphNodeUtils.h" // for FNodeTextCache
+
+#define LOCTEXT_NAMESPACE "GetArrayItem"
+
+/*******************************************************************************
+*  FKCHandler_GetArrayItem
+******************************************************************************/
+class FKCHandler_GetArrayItem : public FNodeHandlingFunctor
+{
+public:
+	FKCHandler_GetArrayItem(FKismetCompilerContext& InCompilerContext)
+		: FNodeHandlingFunctor(InCompilerContext)
+	{
+	}
+
+	virtual void RegisterNets(FKismetFunctionContext& Context, UEdGraphNode* Node) override
+	{
+		UK2Node_GetArrayItem* ArrayNode = CastChecked<UK2Node_GetArrayItem>(Node);
+
+		// return inline term
+		if (Context.NetMap.Find(Node->Pins[2]))
+		{
+			Context.MessageLog.Error(*LOCTEXT("Error_ReturnTermAlreadyRegistered", "ICE: Return term is already registered @@").ToString(), Node);
+			return;
+		}
+
+		{
+			FBPTerminal* Term = new (Context.InlineGeneratedValues) FBPTerminal();
+			Term->CopyFromPin(Node->Pins[2], Context.NetNameMap->MakeValidName(Node->Pins[2]));
+			Context.NetMap.Add(Node->Pins[2], Term);
+		}
+
+		FNodeHandlingFunctor::RegisterNets(Context, Node);
+	}
+
+	virtual void Compile(FKismetFunctionContext& Context, UEdGraphNode* Node) override
+	{
+		UK2Node_GetArrayItem* ArrayNode = CastChecked<UK2Node_GetArrayItem>(Node);
+
+		FBlueprintCompiledStatement* SelectStatementPtr = new FBlueprintCompiledStatement();
+		FBlueprintCompiledStatement& ArrayGetFunction = *SelectStatementPtr;
+		ArrayGetFunction.Type = KCST_ArrayGetByRef;
+
+		UEdGraphPin* ArrayPinNet = FEdGraphUtilities::GetNetFromPin(Node->Pins[0]);
+		UEdGraphPin* ReturnValueNet = FEdGraphUtilities::GetNetFromPin(Node->Pins[2]);
+
+		FBPTerminal** ArrayTerm = Context.NetMap.Find(ArrayPinNet);
+
+		UEdGraphPin* IndexPin = ArrayNode->GetIndexPin();
+		UEdGraphPin* IndexPinNet = IndexPin ? FEdGraphUtilities::GetNetFromPin(IndexPin) : nullptr;
+		FBPTerminal** IndexTermPtr = IndexPinNet ? Context.NetMap.Find(IndexPinNet) : nullptr;
+
+		FBPTerminal** ReturnValue = Context.NetMap.Find(ReturnValueNet);
+		FBPTerminal** ReturnValueOrig = Context.NetMap.Find(Node->Pins[2]);
+		ArrayGetFunction.RHS.Add(*ArrayTerm);
+		ArrayGetFunction.RHS.Add(*IndexTermPtr);
+
+		(*ReturnValue)->InlineGeneratedParameter = &ArrayGetFunction;
+	}
+};
+
+/*******************************************************************************
+*  UK2Node_GetArrayItem
+******************************************************************************/
+UK2Node_GetArrayItem::UK2Node_GetArrayItem(const FObjectInitializer& ObjectInitializer)
+	: Super(ObjectInitializer)
+{
+}
+
+void UK2Node_GetArrayItem::AllocateDefaultPins()
+{
+	// Create the output pin
+	CreatePin(EGPD_Input, UEdGraphSchema_K2::PC_Wildcard, TEXT(""), NULL, true, false, TEXT("Array"));
+	UEdGraphPin* IndexPin = CreatePin(EGPD_Input, UEdGraphSchema_K2::PC_Int, TEXT(""), NULL, false, false, TEXT("Dimension 1"));
+	GetDefault<UEdGraphSchema_K2>()->SetPinDefaultValueBasedOnType(IndexPin);
+
+	// Create the input pins to create the arrays from
+	CreatePin(EGPD_Output, UEdGraphSchema_K2::PC_Wildcard, TEXT(""), NULL, false, true, TEXT("Output"));
+}
+
+void UK2Node_GetArrayItem::PostReconstructNode()
+{
+	if (GetTargetArrayPin()->LinkedTo.Num() > 0)
+	{
+		PropagatePinType(GetTargetArrayPin()->LinkedTo[0]->PinType);
+	}
+	else if (GetResultPin()->LinkedTo.Num() > 0)
+	{
+		PropagatePinType(GetResultPin()->LinkedTo[0]->PinType);
+	}
+}
+
+FText UK2Node_GetArrayItem::GetNodeTitle(ENodeTitleType::Type TitleType) const
+{
+	if (TitleType == ENodeTitleType::FullTitle)
+	{
+		return LOCTEXT("GetArrayItemByRef_FullTitle", "GET");
+	}
+	return LOCTEXT("GetArrayItemByRef", "Get");
+}
+
+class FNodeHandlingFunctor* UK2Node_GetArrayItem::CreateNodeHandler(class FKismetCompilerContext& CompilerContext) const
+{
+	return new FKCHandler_GetArrayItem(CompilerContext);
+}
+
+void UK2Node_GetArrayItem::GetMenuActions(FBlueprintActionDatabaseRegistrar& ActionRegistrar) const
+{
+	// actions get registered under specific object-keys; the idea is that 
+	// actions might have to be updated (or deleted) if their object-key is  
+	// mutated (or removed)... here we use the node's class (so if the node 
+	// type disappears, then the action should go with it)
+	// actions get registered under specific object-keys; the idea is that 
+	// actions might have to be updated (or deleted) if their object-key is  
+	// mutated (or removed)... here we use the node's class (so if the node 
+	// type disappears, then the action should go with it)
+	UClass* ActionKey = GetClass();
+	// to keep from needlessly instantiating a UBlueprintNodeSpawner, first   
+	// check to make sure that the registrar is looking for actions of this type
+	// (could be regenerating actions for a specific asset, and therefore the 
+	// registrar would only accept actions corresponding to that asset)
+	if (ActionRegistrar.IsOpenForRegistration(ActionKey))
+	{
+		UBlueprintNodeSpawner* NodeSpawner = UBlueprintNodeSpawner::Create(GetClass());
+		check(NodeSpawner != nullptr);
+
+		ActionRegistrar.AddBlueprintAction(ActionKey, NodeSpawner);
+	}
+}
+
+void UK2Node_GetArrayItem::NotifyPinConnectionListChanged(UEdGraphPin* Pin)
+{
+	Super::NotifyPinConnectionListChanged(Pin);
+
+	if (Pin != GetIndexPin() && Pin->ParentPin == nullptr)
+	{
+		UEdGraphPin* ArrayPin = Pins[0];
+		UEdGraphPin* ResultPin = Pins[2];
+
+		// If a pin connection is made, then we need to propagate the change to the other pins and validate all connections, otherwise reset them to Wildcard pins
+		if (Pin->LinkedTo.Num() > 0)
+		{
+			PropagatePinType(Pin->LinkedTo[0]->PinType);
+		}
+		else if (ArrayPin->LinkedTo.Num() == 0 && ResultPin->LinkedTo.Num() == 0)
+		{
+			ArrayPin->PinType.PinCategory = UEdGraphSchema_K2::PC_Wildcard;
+			ArrayPin->PinType.PinSubCategory = TEXT("");
+			ArrayPin->PinType.PinSubCategoryObject = NULL;
+
+			ResultPin->PinType.PinCategory = UEdGraphSchema_K2::PC_Wildcard;
+			ResultPin->PinType.PinSubCategory = TEXT("");
+			ResultPin->PinType.PinSubCategoryObject = NULL;
+			ResultPin->PinType.bIsReference = true;
+
+			ArrayPin->BreakAllPinLinks();
+			ResultPin->BreakAllPinLinks();
+		}
+	}
+}
+
+void UK2Node_GetArrayItem::PropagatePinType(FEdGraphPinType& InType)
+{
+	UClass const* CallingContext = NULL;
+	if (UBlueprint const* Blueprint = GetBlueprint())
+	{
+		CallingContext = Blueprint->GeneratedClass;
+		if (CallingContext == NULL)
+		{
+			CallingContext = Blueprint->ParentClass;
+		}
+	}
+
+	const UEdGraphSchema_K2* Schema = GetDefault<UEdGraphSchema_K2>();
+	UEdGraphPin* ArrayPin = Pins[0];
+	UEdGraphPin* ResultPin = Pins[2];
+
+	ArrayPin->PinType = InType;
+	ArrayPin->PinType.bIsArray = true;
+	ArrayPin->PinType.bIsReference = false;
+
+	ResultPin->PinType = InType;
+	ResultPin->PinType.bIsArray = false;
+	ResultPin->PinType.bIsReference = !(ResultPin->PinType.PinCategory == Schema->PC_Object || ResultPin->PinType.PinCategory == Schema->PC_Class || ResultPin->PinType.PinCategory == Schema->PC_Asset || ResultPin->PinType.PinCategory == Schema->PC_AssetClass || ResultPin->PinType.PinCategory == Schema->PC_Interface);
+	ResultPin->PinType.bIsConst = false;
+
+
+	// Verify that all previous connections to this pin are still valid with the new type
+	for (TArray<UEdGraphPin*>::TIterator ConnectionIt(ArrayPin->LinkedTo); ConnectionIt; ++ConnectionIt)
+	{
+		UEdGraphPin* ConnectedPin = *ConnectionIt;
+		if (!Schema->ArePinsCompatible(ArrayPin, ConnectedPin, CallingContext))
+		{
+			ArrayPin->BreakLinkTo(ConnectedPin);
+		}
+	}
+
+	// Verify that all previous connections to this pin are still valid with the new type
+	for (TArray<UEdGraphPin*>::TIterator ConnectionIt(ResultPin->LinkedTo); ConnectionIt; ++ConnectionIt)
+	{
+		UEdGraphPin* ConnectedPin = *ConnectionIt;
+		if (!Schema->ArePinsCompatible(ResultPin, ConnectedPin, CallingContext))
+		{
+			ResultPin->BreakLinkTo(ConnectedPin);
+		}
+	}
+}
+
+#undef LOCTEXT_NAMESPACE
diff --git a/Engine/Source/Editor/BlueprintGraph/Private/K2Node_SetFieldsInStruct.cpp b/Engine/Source/Editor/BlueprintGraph/Private/K2Node_SetFieldsInStruct.cpp
--- a/Engine/Source/Editor/BlueprintGraph/Private/K2Node_SetFieldsInStruct.cpp
+++ b/Engine/Source/Editor/BlueprintGraph/Private/K2Node_SetFieldsInStruct.cpp
@@ -30,10 +30,7 @@ public:
 		: FKCHandler_MakeStruct(InCompilerContext)
 	{
 		UBlueprintEditorSettings* Settings = GetMutableDefault<UBlueprintEditorSettings>();
-		if (Settings->bAllowReferencePassThroughsToCacheNonReferences)
-		{
-			bAutoGenerateGotoForPure = false;
-		}
+		bAutoGenerateGotoForPure = false;
 	}
 
 	virtual UEdGraphPin* FindStructPinChecked(UEdGraphNode* InNode) const override
@@ -48,23 +45,16 @@ public:
 	{
 		UBlueprintEditorSettings* Settings = GetMutableDefault<UBlueprintEditorSettings>();
 
-		if (!Settings->bAllowReferencePassThroughsToCacheNonReferences)
-		{
-			FKCHandler_MakeStruct::RegisterNet(Context, Net);
-		}
-		else
+		if (Net->Direction == EGPD_Output)
 		{
-			if (Net->Direction == EGPD_Output)
+			if (Net->ReferencePassThroughConnection)
 			{
-				if (Net->ReferencePassThroughConnection)
+				UEdGraphPin* InputPinNet = FEdGraphUtilities::GetNetFromPin(Net->ReferencePassThroughConnection);
+				FBPTerminal** InputPinTerm = Context.NetMap.Find(InputPinNet);
+				if (InputPinTerm && !(*InputPinTerm)->bPassedByReference)
 				{
-					UEdGraphPin* InputPinNet = FEdGraphUtilities::GetNetFromPin(Net->ReferencePassThroughConnection);
-					FBPTerminal** InputPinTerm = Context.NetMap.Find(InputPinNet);
-					if (InputPinTerm && !(*InputPinTerm)->bPassedByReference)
-					{
-						// We need a net for the output pin which we have thus far prevented from being registered
-						FKCHandler_MakeStruct::RegisterNet(Context, Net);
-					}
+					// We need a net for the output pin which we have thus far prevented from being registered
+					FKCHandler_MakeStruct::RegisterNet(Context, Net);
 				}
 			}
 		}
@@ -75,27 +65,24 @@ public:
 		FKCHandler_MakeStruct::RegisterNets(Context, Node);
 
 		UBlueprintEditorSettings* Settings = GetMutableDefault<UBlueprintEditorSettings>();
-		if (Settings->bAllowReferencePassThroughsToCacheNonReferences)
-		{
-			UEdGraphPin* ReturnPin = Node->FindPin(SetFieldsInStructHelper::StructOutPinName());
-			UEdGraphPin* ReturnStructNet = FEdGraphUtilities::GetNetFromPin(ReturnPin);
-			FBPTerminal** ReturnTerm = Context.NetMap.Find(ReturnStructNet);
+		UEdGraphPin* ReturnPin = Node->FindPin(SetFieldsInStructHelper::StructOutPinName());
+		UEdGraphPin* ReturnStructNet = FEdGraphUtilities::GetNetFromPin(ReturnPin);
+		FBPTerminal** ReturnTerm = Context.NetMap.Find(ReturnStructNet);
 
-			UEdGraphPin* InputPin = Node->FindPinChecked(SetFieldsInStructHelper::StructRefPinName());
-			UEdGraphPin* InputPinNet = FEdGraphUtilities::GetNetFromPin(InputPin);
-			FBPTerminal** InputTerm = Context.NetMap.Find(InputPinNet);
+		UEdGraphPin* InputPin = Node->FindPinChecked(SetFieldsInStructHelper::StructRefPinName());
+		UEdGraphPin* InputPinNet = FEdGraphUtilities::GetNetFromPin(InputPin);
+		FBPTerminal** InputTerm = Context.NetMap.Find(InputPinNet);
 
-			if (InputTerm == nullptr)
-			{
-				CompilerContext.MessageLog.Error(*LOCTEXT("MakeStruct_NoTerm_Error", "Failed to generate a term for the @@ pin; was it a struct reference that was left unset?").ToString(), InputPin);
-			}
-			else
+		if (InputTerm == nullptr)
+		{
+			CompilerContext.MessageLog.Error(*LOCTEXT("MakeStruct_NoTerm_Error", "Failed to generate a term for the @@ pin; was it a struct reference that was left unset?").ToString(), InputPin);
+		}
+		else
+		{
+			if ((*InputTerm)->bPassedByReference) //InputPinNet->PinType.bIsReference)
 			{
-				if ((*InputTerm)->bPassedByReference) //InputPinNet->PinType.bIsReference)
-				{
-					// Forward the net to the output pin because it's being passed by-ref and this pin is a by-ref pin
-					Context.NetMap.Add(ReturnStructNet, *InputTerm);
-				}
+				// Forward the net to the output pin because it's being passed by-ref and this pin is a by-ref pin
+				Context.NetMap.Add(ReturnStructNet, *InputTerm);
 			}
 		}
 	}
@@ -105,31 +92,28 @@ public:
 		FKCHandler_MakeStruct::Compile(Context, Node);
 
 		UBlueprintEditorSettings* Settings = GetMutableDefault<UBlueprintEditorSettings>();
-		if (Settings->bAllowReferencePassThroughsToCacheNonReferences)
 		{
-			{
-				UEdGraphPin* InputPin = Node->FindPinChecked(SetFieldsInStructHelper::StructRefPinName());
-				UEdGraphPin* InputPinNet = FEdGraphUtilities::GetNetFromPin(InputPin);
-				FBPTerminal** InputTerm = Context.NetMap.Find(InputPinNet);
+			UEdGraphPin* InputPin = Node->FindPinChecked(SetFieldsInStructHelper::StructRefPinName());
+			UEdGraphPin* InputPinNet = FEdGraphUtilities::GetNetFromPin(InputPin);
+			FBPTerminal** InputTerm = Context.NetMap.Find(InputPinNet);
 
-				// If the InputTerm was not a by-ref, then we need to place the modified structure into the local output term with an AssignStatement
-				if (InputTerm && !(*InputTerm)->bPassedByReference)
-				{
-					UEdGraphPin* ReturnPin = Node->FindPin(SetFieldsInStructHelper::StructOutPinName());
-					UEdGraphPin* ReturnStructNet = FEdGraphUtilities::GetNetFromPin(ReturnPin);
-					FBPTerminal** ReturnTerm = Context.NetMap.Find(ReturnStructNet);
-
-					FBlueprintCompiledStatement& AssignStatement = Context.AppendStatementForNode(Node);
-					AssignStatement.Type = KCST_Assignment;
-					// The return term is a reference no matter the way we received it.
-					(*ReturnTerm)->bPassedByReference = true;
-					AssignStatement.LHS = *ReturnTerm;
-					AssignStatement.RHS.Add(*InputTerm);
-				}
+			// If the InputTerm was not a by-ref, then we need to place the modified structure into the local output term with an AssignStatement
+			if (InputTerm && !(*InputTerm)->bPassedByReference)
+			{
+				UEdGraphPin* ReturnPin = Node->FindPin(SetFieldsInStructHelper::StructOutPinName());
+				UEdGraphPin* ReturnStructNet = FEdGraphUtilities::GetNetFromPin(ReturnPin);
+				FBPTerminal** ReturnTerm = Context.NetMap.Find(ReturnStructNet);
+
+				FBlueprintCompiledStatement& AssignStatement = Context.AppendStatementForNode(Node);
+				AssignStatement.Type = KCST_Assignment;
+				// The return term is a reference no matter the way we received it.
+				(*ReturnTerm)->bPassedByReference = true;
+				AssignStatement.LHS = *ReturnTerm;
+				AssignStatement.RHS.Add(*InputTerm);
 			}
-
-			GenerateSimpleThenGoto(Context, *Node);
 		}
+
+		GenerateSimpleThenGoto(Context, *Node);
 	}
 };
 
@@ -300,30 +284,6 @@ void UK2Node_SetFieldsInStruct::FSetFieldsInStructPinManager::GetRecordDefaults(
 	Record.bShowPin = false;
 }
 
-void UK2Node_SetFieldsInStruct::ExpandNode(class FKismetCompilerContext& CompilerContext, UEdGraph* SourceGraph)
-{
-	Super::ExpandNode(CompilerContext, SourceGraph);
-
-	UBlueprintEditorSettings* Settings = GetMutableDefault<UBlueprintEditorSettings>();
-	if (!Settings->bAllowReferencePassThroughsToCacheNonReferences)
-	{
-		UEdGraphPin* OutPin = FindPin(SetFieldsInStructHelper::StructOutPinName());
-		if (OutPin && OutPin->LinkedTo.Num())
-		{
-			UEdGraphPin* InPin = FindPin(SetFieldsInStructHelper::StructRefPinName());
-			UEdGraphPin* SourcePin = (InPin && (1 == InPin->LinkedTo.Num())) ? InPin->LinkedTo[0] : nullptr;
-			auto Schema = CompilerContext.GetSchema();
-			const bool bCopied = SourcePin && Schema->MovePinLinks(*OutPin, *SourcePin, true).CanSafeConnect();
-
-			if (!bCopied)
-			{
-				CompilerContext.MessageLog.Error(*LOCTEXT("ExpansionError", "Cannot copy links from @@").ToString(), OutPin);
-			}
-		}
-		Pins.Remove(OutPin);
-	}
-}
-
 bool UK2Node_SetFieldsInStruct::IsConnectionDisallowed(const UEdGraphPin* MyPin, const UEdGraphPin* OtherPin, FString& OutReason) const
 {
 	const UEdGraphSchema_K2* K2Schema = GetDefault<UEdGraphSchema_K2>();
diff --git a/Engine/Source/Editor/BlueprintGraph/Public/BlueprintEditorSettings.h b/Engine/Source/Editor/BlueprintGraph/Public/BlueprintEditorSettings.h
--- a/Engine/Source/Editor/BlueprintGraph/Public/BlueprintEditorSettings.h
+++ b/Engine/Source/Editor/BlueprintGraph/Public/BlueprintEditorSettings.h
@@ -88,10 +88,6 @@ public:
 	UPROPERTY(EditAnywhere, config, Category=Compiler)
 	bool bJumpToNodeErrors;
 
-	/** If enabled, "Set Members..." nodes will cache any non-reference passed in the input on the output by-ref pin. This allows it to modify literals and other by-value items passed into it. */
-	UPROPERTY(EditAnywhere, config, Category = Experimental, AdvancedDisplay)
-	bool bAllowReferencePassThroughsToCacheNonReferences;
-
 	/** If enabled, nodes can be explicitly disabled via context menu when right-clicking on impure nodes in the Blueprint editor. Disabled nodes will not be compiled, but also will not break existing connections. */
 	UPROPERTY(EditAnywhere, config, Category=Experimental, AdvancedDisplay)
 	bool bAllowExplicitImpureNodeDisabling;
diff --git a/Engine/Source/Editor/BlueprintGraph/Public/BlueprintGraphDefinitions.h b/Engine/Source/Editor/BlueprintGraph/Public/BlueprintGraphDefinitions.h
--- a/Engine/Source/Editor/BlueprintGraph/Public/BlueprintGraphDefinitions.h
+++ b/Engine/Source/Editor/BlueprintGraph/Public/BlueprintGraphDefinitions.h
@@ -35,6 +35,7 @@
 #include "K2Node_FunctionEntry.h"
 #include "K2Node_FunctionResult.h"
 #include "K2Node_GetDataTableRow.h"
+#include "K2Node_GetArrayItem.h"
 #include "K2Node_IfThenElse.h"
 #include "K2Node_InputAction.h"
 #include "K2Node_InputAxisEvent.h"
diff --git a/Engine/Source/Editor/GraphEditor/Private/KismetNodes/SGraphNodeK2ArrayFunction.cpp b/Engine/Source/Editor/GraphEditor/Private/KismetNodes/SGraphNodeK2ArrayFunction.cpp
--- a/Engine/Source/Editor/GraphEditor/Private/KismetNodes/SGraphNodeK2ArrayFunction.cpp
+++ b/Engine/Source/Editor/GraphEditor/Private/KismetNodes/SGraphNodeK2ArrayFunction.cpp
@@ -12,7 +12,7 @@
 //////////////////////////////////////////////////////////////////////////
 // SGraphNodeSwitchStatement
 
-void SGraphNodeK2ArrayFunction::Construct(const FArguments& InArgs, UK2Node_CallArrayFunction* InNode)
+void SGraphNodeK2ArrayFunction::Construct(const FArguments& InArgs, UK2Node* InNode)
 {
 	this->GraphNode = InNode;
 
@@ -28,9 +28,17 @@ FOptionalSize SGraphNodeK2ArrayFunction::GetBackgroundImageSize() const
 
 FSlateColor SGraphNodeK2ArrayFunction::GetTypeIconColor() const
 {
-	const UEdGraphSchema_K2* K2Schema = GetDefault<UEdGraphSchema_K2>();
-	UK2Node_CallArrayFunction* CallNode = CastChecked<UK2Node_CallArrayFunction>(GraphNode);
-	FLinearColor Color = K2Schema->GetPinTypeColor(CallNode->GetTargetArrayPin()->PinType);
+ 	const UEdGraphSchema_K2* K2Schema = GetDefault<UEdGraphSchema_K2>();
+	FLinearColor Color; 
+	if (UK2Node_CallArrayFunction* CallNode = Cast<UK2Node_CallArrayFunction>(GraphNode))
+	{
+		Color = K2Schema->GetPinTypeColor(CallNode->GetTargetArrayPin()->PinType);
+	}
+	else if (UK2Node_GetArrayItem* GetArrayItemByRefNode = Cast<UK2Node_GetArrayItem>(GraphNode))
+	{
+		Color = K2Schema->GetPinTypeColor(GetArrayItemByRefNode->GetTargetArrayPin()->PinType);
+	}
+
 	Color.A = 0.25f;
 	return Color;
 }
diff --git a/Engine/Source/Editor/GraphEditor/Private/KismetNodes/SGraphNodeK2ArrayFunction.h b/Engine/Source/Editor/GraphEditor/Private/KismetNodes/SGraphNodeK2ArrayFunction.h
--- a/Engine/Source/Editor/GraphEditor/Private/KismetNodes/SGraphNodeK2ArrayFunction.h
+++ b/Engine/Source/Editor/GraphEditor/Private/KismetNodes/SGraphNodeK2ArrayFunction.h
@@ -8,7 +8,7 @@ public:
 	SLATE_BEGIN_ARGS(SGraphNodeK2ArrayFunction){}
 	SLATE_END_ARGS()
 
-	void Construct(const FArguments& InArgs, UK2Node_CallArrayFunction* InNode);
+	void Construct(const FArguments& InArgs, UK2Node* InNode);
 
 	virtual const FSlateBrush* GetShadowBrush(bool bSelected) const;
 
diff --git a/Engine/Source/Editor/GraphEditor/Private/KismetNodes/SGraphNodeK2Copy.cpp b/Engine/Source/Editor/GraphEditor/Private/KismetNodes/SGraphNodeK2Copy.cpp
--- a/Engine/Source/Editor/GraphEditor/Private/KismetNodes/SGraphNodeK2Copy.cpp
+++ b/Engine/Source/Editor/GraphEditor/Private/KismetNodes/SGraphNodeK2Copy.cpp
@@ -204,7 +204,21 @@ void SGraphNodeK2Copy::UpdateGraphNode()
 		.HAlign(HAlign_Center)
 		.VAlign(VAlign_Center)
 		[
-			CreateNodeContentArea()
+			SNew(SVerticalBox)
+			+ SVerticalBox::Slot()
+			.AutoHeight()
+			.HAlign(HAlign_Center)
+			.VAlign(VAlign_Top)
+			[
+				CreateNodeContentArea()
+			]
+
+			+ SVerticalBox::Slot()
+			.AutoHeight()
+			.Padding(Settings->GetNonPinNodeBodyPadding())
+			[
+				ErrorReporting->AsWidget()
+			]
 		];
 	CreatePinWidgets();
 }
diff --git a/Engine/Source/Editor/GraphEditor/Private/NodeFactory.cpp b/Engine/Source/Editor/GraphEditor/Private/NodeFactory.cpp
--- a/Engine/Source/Editor/GraphEditor/Private/NodeFactory.cpp
+++ b/Engine/Source/Editor/GraphEditor/Private/NodeFactory.cpp
@@ -221,6 +221,10 @@ TSharedPtr<SGraphNode> FNodeFactory::CreateNodeWidget(UEdGraphNode* InNode)
 		{
 			return SNew(SGraphNodeK2ArrayFunction, CallFunction);
 		}
+		else if (UK2Node_GetArrayItem* GetArrayItemNode = Cast<UK2Node_GetArrayItem>(InNode))
+		{
+			return SNew(SGraphNodeK2ArrayFunction, GetArrayItemNode);
+		}
 		else if (UK2Node_Knot* Knot = Cast<UK2Node_Knot>(InNode))
 		{
 			return SNew(SGraphNodeKnot, Knot);
diff --git a/Engine/Source/Editor/GraphEditor/Private/SGraphNode.cpp b/Engine/Source/Editor/GraphEditor/Private/SGraphNode.cpp
--- a/Engine/Source/Editor/GraphEditor/Private/SGraphNode.cpp
+++ b/Engine/Source/Editor/GraphEditor/Private/SGraphNode.cpp
@@ -634,6 +634,11 @@ void SGraphNode::UpdateErrorInfo()
 			ErrorColor = FEditorStyle::GetColor("InfoReporting.BackgroundColor");
 		}
 	}
+	else if (!GraphNode->NodeUpgradeMessage.IsEmpty())
+	{
+		ErrorMsg = FString(TEXT("UPGRADE NOTE"));
+		ErrorColor = FEditorStyle::GetColor("InfoReporting.BackgroundColor");
+	}
 	else 
 	{
 		ErrorColor = FLinearColor(0,0,0);
@@ -1237,7 +1242,24 @@ void SGraphNode::PositionThisNodeBetweenOtherNodes(const FVector2D& PrevPos, con
 
 FText SGraphNode::GetErrorMsgToolTip( ) const
 {
-	return FText::FromString(GraphNode->ErrorMsg);
+	FText Result;
+	// Append the node's upgrade message, if any.
+	if (!GraphNode->NodeUpgradeMessage.IsEmpty())
+	{
+		if (Result.IsEmpty())
+		{
+			Result = GraphNode->NodeUpgradeMessage;
+		}
+		else
+		{
+			Result = FText::Format(FText::FromString(TEXT("{0}\n\n{1}")), Result, GraphNode->NodeUpgradeMessage);
+		}
+	}
+	else
+	{
+		Result = FText::FromString(GraphNode->ErrorMsg);
+	}
+	return Result;
 }
 
 bool SGraphNode::IsNameReadOnly() const
diff --git a/Engine/Source/Editor/Kismet/Kismet.Build.cs b/Engine/Source/Editor/Kismet/Kismet.Build.cs
--- a/Engine/Source/Editor/Kismet/Kismet.Build.cs
+++ b/Engine/Source/Editor/Kismet/Kismet.Build.cs
@@ -18,7 +18,6 @@ public class Kismet : ModuleRules
                 "DerivedDataCache",
                 "LevelEditor",
 				"GameProjectGeneration",
-                "BlueprintProfiler"
 			}
 			);
 
@@ -48,7 +47,8 @@ public class Kismet : ModuleRules
                 "DerivedDataCache",
 				"DesktopPlatform",
 				"HotReload",
-                "BlueprintNativeCodeGen"
+                "BlueprintNativeCodeGen",
+                "BlueprintProfiler"
 			}
 			);
 
@@ -59,7 +59,6 @@ public class Kismet : ModuleRules
 				"EditorWidgets",
 				"GameProjectGeneration",
                 "BlueprintCompilerCppBackend",
-                "BlueprintProfiler"
 			}
             );
 
diff --git a/Engine/Source/Editor/Kismet/Private/BlueprintEditor.cpp b/Engine/Source/Editor/Kismet/Private/BlueprintEditor.cpp
--- a/Engine/Source/Editor/Kismet/Private/BlueprintEditor.cpp
+++ b/Engine/Source/Editor/Kismet/Private/BlueprintEditor.cpp
@@ -1728,6 +1728,12 @@ void FBlueprintEditor::InitBlueprintEditor(
 				GetToolkitCommands()->ExecuteAction(FFullBlueprintEditorCommands::Get().EditClassDefaults.ToSharedRef());
 			}
 		}
+
+		// There are upgrade notes, open the log and dump the messages to it
+		if (Blueprint->UpgradeNotesLog.IsValid())
+		{
+			DumpMessagesToCompilerLog(Blueprint->UpgradeNotesLog->Messages, true);
+		}
 	}
 }
 
@@ -3163,6 +3169,11 @@ void FBlueprintEditor::Compile()
 			}
 		}
 
+		if (BlueprintObj->UpgradeNotesLog.IsValid())
+		{
+			CompilerResultsListing->AddMessages(BlueprintObj->UpgradeNotesLog->Messages);
+		}
+
 		// send record when player clicks compile and send the result
 		// this will make sure how the users activity is
 		AnalyticsTrackCompileEvent(BlueprintObj, LogResults.NumErrors, LogResults.NumWarnings);
diff --git a/Engine/Source/Editor/Kismet/Private/Profiler/BPProfilerStatisticWidgets.cpp b/Engine/Source/Editor/Kismet/Private/Profiler/BPProfilerStatisticWidgets.cpp
--- a/Engine/Source/Editor/Kismet/Private/Profiler/BPProfilerStatisticWidgets.cpp
+++ b/Engine/Source/Editor/Kismet/Private/Profiler/BPProfilerStatisticWidgets.cpp
@@ -47,6 +47,7 @@ TSharedRef<SWidget> SProfilerStatRow::GenerateWidgetForColumn(const FName& Colum
 			.VAlign(VAlign_Center)
 			[
 				SNew(SExpanderArrow, SharedThis(this))
+				.IndentAmount(15.f)
 			]
 			+SHorizontalBox::Slot()
 			.FillWidth(1.0f)
@@ -202,21 +203,7 @@ void FBPProfilerStatWidget::NavigateTo() const
 {
 	if (ExecNode.IsValid())
 	{
-		if (IBlueprintProfilerInterface* Profiler = FModuleManager::GetModulePtr<IBlueprintProfilerInterface>("BlueprintProfiler"))
-		{
-			if (ExecNode->IsClass())
-			{
-				Profiler->NavigateToBlueprint(ExecNode->GetName());
-			}
-			else if (ExecNode->IsInstance())
-			{
-				Profiler->NavigateToInstance(ExecNode->GetName());
-			}
-			else
-			{
-				Profiler->NavigateToNode(ExecNode->GetName());
-			}
-		}
+		ExecNode->NavigateToObject();
 	}
 }
 
diff --git a/Engine/Source/Editor/Kismet/Private/Profiler/EventExecution.cpp b/Engine/Source/Editor/Kismet/Private/Profiler/EventExecution.cpp
--- a/Engine/Source/Editor/Kismet/Private/Profiler/EventExecution.cpp
+++ b/Engine/Source/Editor/Kismet/Private/Profiler/EventExecution.cpp
@@ -69,14 +69,24 @@ bool FScriptExecutionNode::operator == (const FScriptExecutionNode& NodeIn) cons
 
 void FScriptExecutionNode::GetLinearExecutionPath(TArray<FLinearExecPath>& LinearExecutionNodes, const FTracePath& TracePath)
 {
-	LinearExecutionNodes.Add(FLinearExecPath(AsShared(), TracePath));
-	if (GetNumLinkedNodes() == 1)
+	FExecNodeFilter Filter;
+	GetLinearExecutionPath_Internal(Filter, LinearExecutionNodes, TracePath);
+}
+
+void FScriptExecutionNode::GetLinearExecutionPath_Internal(FExecNodeFilter& Filter, TArray<FLinearExecPath>& LinearExecutionNodes, const FTracePath& TracePath)
+{
+	if (!Filter.Contains(ObservedObject))
 	{
-		FTracePath NewTracePath(TracePath);
-		for (auto NodeIter : LinkedNodes)
+		Filter.Add(ObservedObject);
+		LinearExecutionNodes.Add(FLinearExecPath(AsShared(), TracePath));
+		if (GetNumLinkedNodes() == 1)
 		{
-			NewTracePath.AddExitPin(NodeIter.Key);
-			NodeIter.Value->GetLinearExecutionPath(LinearExecutionNodes, NewTracePath);
+			FTracePath NewTracePath(TracePath);
+			for (auto NodeIter : LinkedNodes)
+			{
+				NewTracePath.AddExitPin(NodeIter.Key);
+				NodeIter.Value->GetLinearExecutionPath_Internal(Filter, LinearExecutionNodes, NewTracePath);
+			}
 		}
 	}
 }
@@ -101,120 +111,112 @@ EScriptStatContainerType::Type FScriptExecutionNode::GetStatisticContainerType()
 
 void FScriptExecutionNode::RefreshStats(const FTracePath& TracePath)
 {
-	bool bRefreshChildStats = false;
-	bool bRefreshLinkStats = true;
-	TArray<TSharedPtr<FScriptPerfData>> BlueprintPooledStats;
-	// Update stats based on type
-	switch(GetStatisticContainerType())
+	FExecNodeFilter StatFilter;
+	RefreshStats_Internal(TracePath, StatFilter);
+}
+
+void FScriptExecutionNode::RefreshStats_Internal(const FTracePath& TracePath, FExecNodeFilter& VisitedStats)
+{
+	if (!VisitedStats.Find(ObservedObject))
 	{
-		case EScriptStatContainerType::Container:
+		// Update visited stats
+		if (!HasFlags(EScriptExecutionNodeFlags::Container))
 		{
-			// Only consider stat data on this path
-			GetInstancePerfDataByTracePath(TracePath, BlueprintPooledStats);
-			bRefreshChildStats = true;
-			break;
+			VisitedStats.Add(ObservedObject);
 		}
-		case EScriptStatContainerType::Standard:
-		case EScriptStatContainerType::SequentialBranch:
+		// Process stat update
+		bool bRefreshChildStats = false;
+		bool bRefreshLinkStats = true;
+		TArray<TSharedPtr<FScriptPerfData>> BlueprintPooledStats;
+		// Update stats based on type
+		switch(GetStatisticContainerType())
 		{
-			// Only consider stat data on this path
-			GetInstancePerfDataByTracePath(TracePath, BlueprintPooledStats);
-			break;
-		}
-		//case EScriptStatContainerType::SequentialBranch:
-		//{
-		//	for (auto InstanceIter : InstanceInputPinToPerfDataMap)
-		//	{
-		//		if (InstanceIter.Key != NAME_None)
-		//		{
-		//			TSharedPtr<FScriptPerfData> InstancePerfData = GetPerfDataByInstanceAndTracePath(InstanceIter.Key, TracePath);
-		//			InstancePerfData->Reset();
-		//			BlueprintPooledStats.Add(InstancePerfData);
-		//		
-		//			for (auto LinkIter : LinkedNodes)
-		//			{
-		//				FTracePath LinkTracePath(TracePath);
-		//				LinkTracePath.AddExitPin(LinkIter.Key);
-		//				TSharedPtr<FScriptPerfData> LinkPerfData = LinkIter.Value->GetPerfDataByInstanceAndTracePath(InstanceIter.Key, LinkTracePath);
-		//				InstancePerfData->AddData(*LinkPerfData.Get());
-		//				LinkIter.Value->RefreshStats(LinkTracePath);
-		//			}
-		//		}
-		//	}
-		//	bRefreshLinkStats = false;
-		//	break;
-		//}
-		case EScriptStatContainerType::NewExecutionPath:
-		{
-			TSet<FName> AllInstances;
-			// Refresh child stats and copy as branch stats - this is a dummy entry
-			for (auto ChildIter : ChildNodes)
+			case EScriptStatContainerType::Container:
+			{
+				// Only consider stat data on this path
+				GetInstancePerfDataByTracePath(TracePath, BlueprintPooledStats);
+				bRefreshChildStats = true;
+				break;
+			}
+			case EScriptStatContainerType::Standard:
+			case EScriptStatContainerType::SequentialBranch:
 			{
-				// Fill out instance data that can be missing because its a dummy node, or the code below can fail.
-				for (auto InstanceIter : ChildIter->InstanceInputPinToPerfDataMap)
+				// Only consider stat data on this path
+				GetInstancePerfDataByTracePath(TracePath, BlueprintPooledStats);
+				break;
+			}
+			case EScriptStatContainerType::NewExecutionPath:
+			{
+				TSet<FName> AllInstances;
+				// Refresh child stats and copy as branch stats - this is a dummy entry
+				for (auto ChildIter : ChildNodes)
 				{
-					if (InstanceIter.Key != NAME_None)
+					// Fill out instance data that can be missing because its a dummy node, or the code below can fail.
+					for (auto InstanceIter : ChildIter->InstanceInputPinToPerfDataMap)
 					{
-						AllInstances.Add(InstanceIter.Key);
+						if (InstanceIter.Key != NAME_None)
+						{
+							AllInstances.Add(InstanceIter.Key);
+						}
 					}
+					ChildIter->RefreshStats_Internal(TracePath, VisitedStats);
 				}
-				ChildIter->RefreshStats(TracePath);
-			}
-			for (auto InstanceIter : AllInstances)
-			{
-				if (InstanceIter != NAME_None)
+				for (auto InstanceIter : AllInstances)
 				{
-					TSharedPtr<FScriptPerfData> InstancePerfData = GetPerfDataByInstanceAndTracePath(InstanceIter, TracePath);
-					InstancePerfData->Reset();
-					BlueprintPooledStats.Add(InstancePerfData);
-				
-					for (auto ChildIter : ChildNodes)
+					if (InstanceIter != NAME_None)
 					{
-						TSharedPtr<FScriptPerfData> ChildPerfData = ChildIter->GetPerfDataByInstanceAndTracePath(InstanceIter, TracePath);
-						InstancePerfData->AddBranchData(*ChildPerfData.Get());
+						TSharedPtr<FScriptPerfData> InstancePerfData = GetPerfDataByInstanceAndTracePath(InstanceIter, TracePath);
+						InstancePerfData->Reset();
+						BlueprintPooledStats.Add(InstancePerfData);
+				
+						for (auto ChildIter : ChildNodes)
+						{
+							TSharedPtr<FScriptPerfData> ChildPerfData = ChildIter->GetPerfDataByInstanceAndTracePath(InstanceIter, TracePath);
+							InstancePerfData->AddBranchData(*ChildPerfData.Get());
+						}
 					}
 				}
-			}
-			// Update Blueprint stats as branches
-			TSharedPtr<FScriptPerfData> BlueprintData = GetBlueprintPerfDataByTracePath(TracePath);
-			BlueprintData->Reset();
+				// Update Blueprint stats as branches
+				TSharedPtr<FScriptPerfData> BlueprintData = GetBlueprintPerfDataByTracePath(TracePath);
+				BlueprintData->Reset();
 
-			for (auto BlueprintChildDataIter : BlueprintPooledStats)
-			{
-				BlueprintData->AddBranchData(*BlueprintChildDataIter.Get());
+				for (auto BlueprintChildDataIter : BlueprintPooledStats)
+				{
+					BlueprintData->AddBranchData(*BlueprintChildDataIter.Get());
+				}
+				BlueprintPooledStats.Reset(0);
+				bRefreshLinkStats = false;
+				break;
 			}
-			BlueprintPooledStats.Reset(0);
-			bRefreshLinkStats = false;
-			break;
 		}
-	}
-	// Refresh Child Links
-	if (bRefreshChildStats)
-	{
-		for (auto ChildIter : ChildNodes)
+		// Refresh Child Links
+		if (bRefreshChildStats)
 		{
-			ChildIter->RefreshStats(TracePath);
+			for (auto ChildIter : ChildNodes)
+			{
+				ChildIter->RefreshStats_Internal(TracePath, VisitedStats);
+			}
 		}
-	}
-	// Refresh All links
-	if (bRefreshLinkStats)
-	{
-		for (auto LinkIter : LinkedNodes)
+		// Refresh All links
+		if (bRefreshLinkStats)
 		{
-			FTracePath LinkTracePath(TracePath);
-			LinkTracePath.AddExitPin(LinkIter.Key);
-			LinkIter.Value->RefreshStats(LinkTracePath);
+			for (auto LinkIter : LinkedNodes)
+			{
+				FTracePath LinkTracePath(TracePath);
+				LinkTracePath.AddExitPin(LinkIter.Key);
+				LinkIter.Value->RefreshStats_Internal(LinkTracePath, VisitedStats);
+			}
 		}
-	}
-	// Update the owning blueprint stats
-	if (BlueprintPooledStats.Num() > 0)
-	{
-		TSharedPtr<FScriptPerfData> BlueprintData = GetBlueprintPerfDataByTracePath(TracePath);
-		BlueprintData->Reset();
-
-		for (auto BlueprintChildDataIter : BlueprintPooledStats)
+		// Update the owning blueprint stats
+		if (BlueprintPooledStats.Num() > 0)
 		{
-			BlueprintData->AddData(*BlueprintChildDataIter.Get());
+			TSharedPtr<FScriptPerfData> BlueprintData = GetBlueprintPerfDataByTracePath(TracePath);
+			BlueprintData->Reset();
+
+			for (auto BlueprintChildDataIter : BlueprintPooledStats)
+			{
+				BlueprintData->AddData(*BlueprintChildDataIter.Get());
+			}
 		}
 	}
 }
@@ -233,6 +235,14 @@ void FScriptExecutionNode::GetAllExecNodes(TMap<FName, TSharedPtr<FScriptExecuti
 	}
 }
 
+void FScriptExecutionNode::NavigateToObject() const
+{
+	if (ObservedObject.IsValid())
+	{
+		FKismetEditorUtilities::BringKismetToFocusAttentionOnObject(ObservedObject.Get());
+	}
+}
+
 //////////////////////////////////////////////////////////////////////////
 // FScriptExecutionBlueprint
 
@@ -302,38 +312,48 @@ void FScriptExecutionBlueprint::GetAllExecNodes(TMap<FName, TSharedPtr<FScriptEx
 	ExecNodesOut.Add(GetName()) = AsShared();
 }
 
-//////////////////////////////////////////////////////////////////////////
-// FScriptExecutionContext
-
-void FScriptExecutionContext::GetAllExecNodes(TMap<FName, TSharedPtr<FScriptExecutionNode>>& ExecNodesOut)
+void FScriptExecutionBlueprint::NavigateToObject() const
 {
-	if (ExecutionNode.IsValid())
+	if (ObservedObject.IsValid())
 	{
-		ExecutionNode->GetAllExecNodes(ExecNodesOut);
-	}
-}
-
-const UEdGraphNode* FScriptExecutionContext::GetNodeFromCodeLocation(const int32 CodeLocation, UFunction* FunctionOverride) const
-{
-	if (BlueprintClass.IsValid())
-	{
-		if (FunctionOverride)
+		if (const UBlueprint* Blueprint = Cast<UBlueprint>(ObservedObject.Get()))
 		{
-			return BlueprintClass.Get()->GetDebugData().FindSourceNodeFromCodeLocation(FunctionOverride, CodeLocation, true);
-		}
-		else if (BlueprintFunction.IsValid())
-		{
-			return BlueprintClass.Get()->GetDebugData().FindSourceNodeFromCodeLocation(BlueprintFunction.Get(), CodeLocation, true);
+			TArray<FAssetData> BlueprintAssets;
+			BlueprintAssets.Add(Blueprint);
+			GEditor->SyncBrowserToObjects(BlueprintAssets);
 		}
 	}
-	return nullptr;
 }
 
-void FScriptExecutionContext::UpdateConnectedStats()
+//////////////////////////////////////////////////////////////////////////
+// FScriptExecutionInstance
+
+void FScriptExecutionInstance::NavigateToObject() const
 {
-	if (ExecutionNode.IsValid())
+	if (ObservedObject.IsValid())
 	{
-		FTracePath InitialTracePath;
-		ExecutionNode->RefreshStats(InitialTracePath);
+		if (const AActor* Actor = Cast<AActor>(ObservedObject.Get()))
+		{
+			if (const FWorldContext* PIEWorldContext = GEditor->GetPIEWorldContext())
+			{
+				if (const UWorld* PIEWorld = PIEWorldContext->World())
+				{
+					for (auto LevelIter : PIEWorld->GetLevels())
+					{
+						if (AActor* PIEActor = Cast<AActor>(FindObject<UObject>(LevelIter, *Actor->GetName())))
+						{
+							GEditor->SelectNone(false, false);
+							GEditor->SelectActor(const_cast<AActor*>(PIEActor), true, true, true);
+							break;
+						}
+					}
+				}
+			}
+			else
+			{
+				GEditor->SelectNone(false, false);
+				GEditor->SelectActor(const_cast<AActor*>(Actor), true, true, true);
+			}
+		}
 	}
 }
diff --git a/Engine/Source/Editor/Kismet/Private/Profiler/SGraphExecutionStatDisplay.cpp b/Engine/Source/Editor/Kismet/Private/Profiler/SGraphExecutionStatDisplay.cpp
--- a/Engine/Source/Editor/Kismet/Private/Profiler/SGraphExecutionStatDisplay.cpp
+++ b/Engine/Source/Editor/Kismet/Private/Profiler/SGraphExecutionStatDisplay.cpp
@@ -6,6 +6,7 @@
 #include "BlueprintEditor.h"
 #include "Public/Profiler/EventExecution.h"
 #include "SDockTab.h"
+#include "Developer/BlueprintProfiler/Public/ScriptInstrumentationPlayback.h"
 
 #define LOCTEXT_NAMESPACE "BlueprintProfilerGraphExecView"
 
@@ -264,7 +265,7 @@ void SGraphExecutionStatDisplay::Tick(const FGeometry& AllottedGeometry, const d
 {
 	if (GetDefault<UEditorExperimentalSettings>()->bBlueprintPerformanceAnalysisTools)
 	{
-		if (IBlueprintProfilerInterface* ProfilerInterface = FModuleManager::GetModulePtr<IBlueprintProfilerInterface>("BlueprintProfiler"))
+		if (IBlueprintProfilerInterface* Profiler = FModuleManager::GetModulePtr<IBlueprintProfilerInterface>("BlueprintProfiler"))
 		{
 			// Find and process execution path timing data.
 			if (BlueprintEditor.IsValid())
@@ -274,36 +275,37 @@ void SGraphExecutionStatDisplay::Tick(const FGeometry& AllottedGeometry, const d
 				if (Blueprints->Num() == 1)
 				{
 					UBlueprint* NewBlueprint = Cast<UBlueprint>((*Blueprints)[0]);
+					UBlueprintGeneratedClass* BPGC = Cast<UBlueprintGeneratedClass>(NewBlueprint->GeneratedClass);
 					TWeakObjectPtr<const UObject> NewInstance = NewBlueprint->GetObjectBeingDebugged();
-					FName NewInstancePath = NewInstance.IsValid() ? FName(*ProfilerInterface->GetValidInstancePath(NewInstance->GetPathName())) : NAME_None;
-					const bool bBlueprintChanged = NewBlueprint != CurrentBlueprint.Get();
-					const bool bInstanceChanged = bScopeToActiveDebugInstance && NewInstancePath != CurrentInstancePath;
-					// Find active graph to filter on
-					bool bGraphFilterChanged = false;
-					TSharedPtr<SDockTab> DockTab = BlueprintEditor.Pin()->DocumentManager->GetActiveTab();
-					if (DockTab.IsValid())
-					{
-						TSharedRef<SGraphEditor> ActiveGraphEditor = StaticCastSharedRef<SGraphEditor>(DockTab->GetContent());
-						bGraphFilterChanged = ActiveGraphEditor->GetCurrentGraph() != CurrentGraph.Get();
-						CurrentGraph = ActiveGraphEditor->GetCurrentGraph();
-					}
+					TSharedPtr<FBlueprintExecutionContext> BlueprintExecContext = Profiler->GetBlueprintContext(BPGC->GetPathName());
 
-					if (bBlueprintChanged||bInstanceChanged||bGraphFilterChanged)
+					if (BlueprintExecContext.IsValid())
 					{
-						const FString BlueprintPath = NewBlueprint->GeneratedClass->GetPathName();
-						TSharedPtr<FScriptExecutionNode> BlueprintEntry = ProfilerInterface->GetProfilerDataForNode(FName(*BlueprintPath));
-						RootTreeItems.Reset(0);
+						FName NewInstancePath = BlueprintExecContext->RemapInstancePath(FName(*NewInstance->GetPathName()));
+						const bool bBlueprintChanged = NewBlueprint != CurrentBlueprint.Get();
+						const bool bInstanceChanged = bScopeToActiveDebugInstance && NewInstancePath != CurrentInstancePath;
+						// Find active graph to filter on
+						bool bGraphFilterChanged = false;
+						TSharedPtr<SDockTab> DockTab = BlueprintEditor.Pin()->DocumentManager->GetActiveTab();
+						if (DockTab.IsValid())
+						{
+							TSharedRef<SGraphEditor> ActiveGraphEditor = StaticCastSharedRef<SGraphEditor>(DockTab->GetContent());
+							bGraphFilterChanged = ActiveGraphEditor->GetCurrentGraph() != CurrentGraph.Get();
+							CurrentGraph = ActiveGraphEditor->GetCurrentGraph();
+						}
 
-						if (BlueprintEntry.IsValid() && BlueprintEntry->IsRootNode())
+						if (bBlueprintChanged||bInstanceChanged||bGraphFilterChanged)
 						{
-							// Cache Active blueprint and Instance
-							CurrentBlueprint = NewBlueprint;
-							CurrentInstancePath = NewInstancePath;
-							// Build Instance widget execution trees
-							TSharedPtr<FScriptExecutionBlueprint> BlueprintExecNode = StaticCastSharedPtr<FScriptExecutionBlueprint>(BlueprintEntry);
+							const FString BlueprintPath = NewBlueprint->GeneratedClass->GetPathName();
+							TSharedPtr<FScriptExecutionBlueprint> BlueprintExecNode = BlueprintExecContext->GetBlueprintExecNode();
+							RootTreeItems.Reset(0);
 
 							if (BlueprintExecNode.IsValid())
 							{
+								// Cache Active blueprint and Instance
+								CurrentBlueprint = NewBlueprint;
+								CurrentInstancePath = NewInstancePath;
+								// Build Instance widget execution trees
 								FName GraphFilter = CurrentGraph.IsValid() && bFilterEventsToGraph ? CurrentGraph.Get()->GetFName() : NAME_None;
 
 								if (bDisplayInstances && bScopeToActiveDebugInstance)
@@ -339,24 +341,24 @@ void SGraphExecutionStatDisplay::Tick(const FGeometry& AllottedGeometry, const d
 									RootTreeItems.Add(BlueprintWidget);
 								}
 							}
-						}
-						// Refresh Tree
-						if (ExecutionStatTree.IsValid())
-						{
-							ExecutionStatTree->RequestTreeRefresh();
-							if (bAutoExpandStats)
+							// Refresh Tree
+							if (ExecutionStatTree.IsValid())
 							{
-								for (auto Iter : RootTreeItems)
+								ExecutionStatTree->RequestTreeRefresh();
+								if (bAutoExpandStats)
 								{
-									Iter->ExpandWidgetState(ExecutionStatTree, bAutoExpandStats);
+									for (auto Iter : RootTreeItems)
+									{
+										Iter->ExpandWidgetState(ExecutionStatTree, bAutoExpandStats);
+									}
 								}
-							}
-							else
-							{
-								for (auto Iter : RootTreeItems)
+								else
 								{
-									Iter->ProbeChildWidgetExpansionStates();
-									Iter->RestoreWidgetExpansionState(ExecutionStatTree);
+									for (auto Iter : RootTreeItems)
+									{
+										Iter->ProbeChildWidgetExpansionStates();
+										Iter->RestoreWidgetExpansionState(ExecutionStatTree);
+									}
 								}
 							}
 						}
diff --git a/Engine/Source/Editor/Kismet/Public/Profiler/EventExecution.h b/Engine/Source/Editor/Kismet/Public/Profiler/EventExecution.h
--- a/Engine/Source/Editor/Kismet/Public/Profiler/EventExecution.h
+++ b/Engine/Source/Editor/Kismet/Public/Profiler/EventExecution.h
@@ -4,9 +4,6 @@
 
 #include "TracePath.h"
 
-struct FSlateBrush;
-class UBlueprintGeneratedClass;
-
 /**  Execution node flags */
 namespace EScriptExecutionNodeFlags
 {
@@ -16,6 +13,7 @@ namespace EScriptExecutionNodeFlags
 		Class				= 0x00000001,	// Class
 		Instance			= 0x00000002,	// Instance
 		Event				= 0x00000004,	// Event
+		Container			= 0x00000007,	// Container type - Class, Instance or Event.
 		FunctionCall		= 0x00000008,	// Function Call
 		MacroCall			= 0x00000010,	// Macro Call
 		CallSite			= 0x00000018,	// Function / Macro call site
@@ -126,6 +124,8 @@ protected:
 //////////////////////////////////////////////////////////////////////////
 // FScriptExecutionNode
 
+typedef TSet<TWeakObjectPtr<const UObject>> FExecNodeFilter;
+
 class KISMET_API FScriptExecutionNode : public FScriptNodeExecLinkage, public FScriptNodePerfData, public TSharedFromThis<FScriptExecutionNode>
 {
 public:
@@ -183,6 +183,15 @@ public:
 	/** Returns if this event potentially multiple exit sites */
 	bool IsBranch() const { return (NodeFlags & EScriptExecutionNodeFlags::BranchNode) != 0U; }
 
+	/** Sets the observed object context */
+	void SetObservedObject(const UObject* ObjectIn) { ObservedObject = ObjectIn; }
+
+	/** Gets the observed object context */
+	const UObject* GetObservedObject() { return ObservedObject.Get(); }
+
+	/** Navigate to object */
+	virtual void NavigateToObject() const;
+
 	/** Returns the display name for widget UI */
 	const FText& GetDisplayName() const { return DisplayName; }
 
@@ -216,9 +225,6 @@ public:
 	/** Sets the current expansion state for widget UI */
 	void SetExpanded(bool bIsExpanded) { bExpansionState = bIsExpanded; }
 
-	/** Returns if this node represents a blueprint */
-	virtual bool IsRootNode() const { return false; }
-
 	/** Return the linear execution path from this node */
 	void GetLinearExecutionPath(TArray<FLinearExecPath>& LinearExecutionNodes, const FTracePath& TracePath);
 
@@ -231,6 +237,14 @@ public:
 	/** Get all exec nodes */
 	virtual void GetAllExecNodes(TMap<FName, TSharedPtr<FScriptExecutionNode>>& ExecNodesOut);
 
+protected:
+
+	/** Get linear execution path - private implementation */
+	void GetLinearExecutionPath_Internal(FExecNodeFilter& Filter, TArray<FLinearExecPath>& LinearExecutionNodes, const FTracePath& TracePath);
+
+	/** Refresh Stats - private implementation */
+	void RefreshStats_Internal(const FTracePath& TracePath, FExecNodeFilter& VisitedStats);
+
 protected:
 
 	/** Node name */
@@ -239,6 +253,8 @@ protected:
 	FName OwningGraphName;
 	/** Node flags to describe the source graph node type */
 	uint32 NodeFlags;
+	/** Oberved object */
+	TWeakObjectPtr<const UObject> ObservedObject;
 	/** Display name for widget UI */
 	FText DisplayName;
 	/** Tooltip for widget UI */
@@ -252,6 +268,18 @@ protected:
 
 };
 
+//////////////////////////////////////////////////////////////////////////
+// FScriptExecutionInstance
+
+class KISMET_API FScriptExecutionInstance : public FScriptExecutionNode
+{
+public:
+
+	// FScriptExecutionNode
+	virtual void NavigateToObject() const override;
+	// ~FScriptExecutionNode
+};
+
 //////////////////////////////////////////////////////////////////////////
 // FScriptExecutionBlueprint
 
@@ -260,49 +288,28 @@ class KISMET_API FScriptExecutionBlueprint : public FScriptExecutionNode
 public:
 
 	/** Adds new blueprint instance node */
-	void AddInstance(TSharedPtr<class FScriptExecutionNode> Instance) { Instances.Add(Instance); }
+	void AddInstance(TSharedPtr<FScriptExecutionNode> Instance) { Instances.Add(Instance); }
 
 	/** Returns the current number of instance nodes */
 	int32 GetInstanceCount() { return Instances.Num(); }
 
 	/** Returns the instance node specified by Index */
-	TSharedPtr<class FScriptExecutionNode> GetInstanceByIndex(const int32 Index) { return Instances[Index]; }
+	TSharedPtr<FScriptExecutionNode> GetInstanceByIndex(const int32 Index) { return Instances[Index]; }
 
 	/** Returns the instance that matches the supplied name if present */
 	TSharedPtr<FScriptExecutionNode> GetInstanceByName(FName InstanceName);
 
-	// ~Begin FScriptExecutionNode
-	virtual bool IsRootNode() const override { return true; }
+	// FScriptExecutionNode
 	virtual EScriptStatContainerType::Type GetStatisticContainerType() const override { return EScriptStatContainerType::Container; }
 	virtual void RefreshStats(const FTracePath& TracePath) override;
 	virtual void GetAllExecNodes(TMap<FName, TSharedPtr<FScriptExecutionNode>>& ExecNodesOut) override;
-	// ~End FScriptExecutionNode
+	virtual void NavigateToObject() const override;
+	// ~FScriptExecutionNode
 
 private:
 
 	/** Exec nodes representing all instances based on this blueprint */
-	TArray<TSharedPtr<class FScriptExecutionNode>> Instances;
+	TArray<TSharedPtr<FScriptExecutionNode>> Instances;
 
 };
 
-//////////////////////////////////////////////////////////////////////////
-// FScriptExecutionContext
-
-class KISMET_API FScriptExecutionContext : public TSharedFromThis<FScriptExecutionContext>
-{
-public:
-
-	bool IsContextValid() const { return ContextObject.IsValid(); }
-	const class UEdGraphNode* GetGraphNode() const { return ContextObject.IsValid() ? Cast<UEdGraphNode>(ContextObject.Get()) : nullptr; }
-	const class UEdGraphNode* GetNodeFromCodeLocation(const int32 CodeLocation, UFunction* FunctionOverride = nullptr) const;
-	void GetAllExecNodes(TMap<FName, TSharedPtr<FScriptExecutionNode>>& ExecNodesOut);
-	void UpdateConnectedStats();
-
-public:
-
-	TSharedPtr<FScriptExecutionNode> ExecutionNode;
-	TWeakObjectPtr<const UObject> ContextObject;
-	TWeakObjectPtr<const UObject> MacroContextObject;
-	TWeakObjectPtr<UBlueprintGeneratedClass> BlueprintClass;
-	TWeakObjectPtr<UFunction> BlueprintFunction;
-};
diff --git a/Engine/Source/Editor/Kismet/Public/Profiler/ScriptPerfData.h b/Engine/Source/Editor/Kismet/Public/Profiler/ScriptPerfData.h
--- a/Engine/Source/Editor/Kismet/Public/Profiler/ScriptPerfData.h
+++ b/Engine/Source/Editor/Kismet/Public/Profiler/ScriptPerfData.h
@@ -2,8 +2,6 @@
 
 #pragma once
 
-#include "Styling/SlateColor.h"
-
 //////////////////////////////////////////////////////////////////////////
 // FScriptPerfData
 
diff --git a/Engine/Source/Editor/KismetCompiler/Private/KismetCompilerMisc.cpp b/Engine/Source/Editor/KismetCompiler/Private/KismetCompilerMisc.cpp
--- a/Engine/Source/Editor/KismetCompiler/Private/KismetCompilerMisc.cpp
+++ b/Engine/Source/Editor/KismetCompiler/Private/KismetCompilerMisc.cpp
@@ -1240,7 +1240,7 @@ void FNodeHandlingFunctor::ResolveAndRegisterScopedTerm(FKismetFunctionContext&
 		FBPTerminal* Term = new (NetArray) FBPTerminal();
 		Term->CopyFromPin(Net, Net->PinName);
 		Term->AssociatedVarProperty = BoundProperty;
-		Term->bPassedByReference = Settings->bAllowReferencePassThroughsToCacheNonReferences;
+		Term->bPassedByReference = true;
 		Context.NetMap.Add(Net, Term);
 
 		// Check if the property is a local variable and mark it so
diff --git a/Engine/Source/Editor/KismetCompiler/Private/KismetCompilerVMBackend.cpp b/Engine/Source/Editor/KismetCompiler/Private/KismetCompilerVMBackend.cpp
--- a/Engine/Source/Editor/KismetCompiler/Private/KismetCompilerVMBackend.cpp
+++ b/Engine/Source/Editor/KismetCompiler/Private/KismetCompilerVMBackend.cpp
@@ -1325,11 +1325,23 @@ public:
 			case KCST_InstrumentedWireExit:			EventType = EScriptInstrumentation::NodeExit; break;
 			case KCST_InstrumentedWireEntry:		EventType = EScriptInstrumentation::NodeEntry; break;
 			case KCST_InstrumentedPureNodeEntry:	EventType = EScriptInstrumentation::PureNodeEntry; break;
+			case KCST_InstrumentedStatePush:		EventType = EScriptInstrumentation::PushState; break;
+			case KCST_InstrumentedStateRestore:		EventType = EScriptInstrumentation::RestoreState; break;
+			case KCST_InstrumentedStatePop:			EventType = EScriptInstrumentation::PopState; break;
 		}
 		Writer << EX_InstrumentationEvent;
 		Writer << EventType;
 	}
 
+	void EmitArrayGetByRef(FBlueprintCompiledStatement& Statement)
+	{
+		Writer << EX_ArrayGetByRef;
+		// The array variable
+		EmitTerm(Statement.RHS[0]);
+		// The index to access in the array
+		EmitTerm(Statement.RHS[1], (UProperty*)(GetDefault<UIntProperty>()));
+	}
+
 	void PushReturnAddress(FBlueprintCompiledStatement& ReturnTarget)
 	{
 		Writer << EX_PushExecutionFlow;
@@ -1486,6 +1498,8 @@ public:
 			break;
 		case KCST_InstrumentedWireExit:
 		case KCST_InstrumentedWireEntry:
+		case KCST_InstrumentedStatePush:
+		case KCST_InstrumentedStatePop:
 			{
 				UEdGraphPin const* TrueSourcePin = Cast<UEdGraphPin const>(FunctionContext.MessageLog.FindSourceObject(Statement.ExecContext));
 				if (TrueSourcePin)
@@ -1514,12 +1528,16 @@ public:
 				}
 			}
 		case KCST_InstrumentedPureNodeEntry:
+		case KCST_InstrumentedStateRestore:
 			{
 				// Emit Statement
 				EmitInstrumentation(Statement);
 				break;
 			}
 
+		case KCST_ArrayGetByRef:
+			EmitArrayGetByRef(Statement);
+			break;
 		default:
 			UE_LOG(LogK2Compiler, Warning, TEXT("VM backend encountered unsupported statement type %d"), (int32)Statement.Type);
 		}
diff --git a/Engine/Source/Editor/KismetCompiler/Public/BlueprintCompiledStatement.h b/Engine/Source/Editor/KismetCompiler/Public/BlueprintCompiledStatement.h
--- a/Engine/Source/Editor/KismetCompiler/Public/BlueprintCompiledStatement.h
+++ b/Engine/Source/Editor/KismetCompiler/Public/BlueprintCompiledStatement.h
@@ -40,7 +40,12 @@ enum EKismetCompiledStatementType
 	// Kismet instrumentation extensions
 	KCST_InstrumentedPureNodeEntry,	// Instrumented pure node entry
 	KCST_InstrumentedWireEntry,		// Instrumented wiretrace entry
-	KCST_InstrumentedWireExit		// Instrumented wiretrace exit
+	KCST_InstrumentedWireExit,		// Instrumented wiretrace exit
+	KCST_InstrumentedStatePush,		// Instrumented state push
+	KCST_InstrumentedStateRestore,	// Instrumented state restore
+	KCST_InstrumentedStatePop,		// Instrumented state pop
+	//
+	KCST_ArrayGetByRef,
 };
 
 //@TODO: Too rigid / icky design
diff --git a/Engine/Source/Editor/UnrealEd/Private/ScriptDisassembler.cpp b/Engine/Source/Editor/UnrealEd/Private/ScriptDisassembler.cpp
--- a/Engine/Source/Editor/UnrealEd/Private/ScriptDisassembler.cpp
+++ b/Engine/Source/Editor/UnrealEd/Private/ScriptDisassembler.cpp
@@ -824,6 +824,15 @@ void FKismetBytecodeDisassembler::ProcessCommon(int32& ScriptIndex, EExprToken O
 				case EScriptInstrumentation::NodeExit:
 					Ar.Logf(TEXT("%s $%X: .. instrumented wire exit site .."), *Indents, (int32)Opcode);
 					break;
+				case EScriptInstrumentation::PushState:
+					Ar.Logf(TEXT("%s $%X: .. push execution state .."), *Indents, (int32)Opcode);
+					break;
+				case EScriptInstrumentation::RestoreState:
+					Ar.Logf(TEXT("%s $%X: .. restore execution state .."), *Indents, (int32)Opcode);
+					break;
+				case EScriptInstrumentation::PopState:
+					Ar.Logf(TEXT("%s $%X: .. pop execution state .."), *Indents, (int32)Opcode);
+					break;
 			}
 			break;
 		}
@@ -858,6 +867,12 @@ void FKismetBytecodeDisassembler::ProcessCommon(int32& ScriptIndex, EExprToken O
 			DropIndent();
 			break;
 		}
+	case EX_ArrayGetByRef:
+		{
+			Ar.Logf(TEXT("%s $%X: Array Get-by-Ref Index"), *Indents, (int32)Opcode);
+			AddIndent();
+			break;
+		}
 	default:
 		{
 			// This should never occur.
diff --git a/Engine/Source/Programs/UnrealHeaderTool/Private/CodeGenerator.cpp b/Engine/Source/Programs/UnrealHeaderTool/Private/CodeGenerator.cpp
--- a/Engine/Source/Programs/UnrealHeaderTool/Private/CodeGenerator.cpp
+++ b/Engine/Source/Programs/UnrealHeaderTool/Private/CodeGenerator.cpp
@@ -1379,7 +1379,7 @@ void FNativeClassHeaderGenerator::ExportNativeGeneratedInitCode(FClass* Class, F
 		{
 			const FString DynamicClassPackageName = FClass::GetTypePackageName(Class);
 			GeneratedClassRegisterFunctionText.Logf(TEXT("\t\tUPackage* OuterPackage = FindOrConstructDynamicTypePackage(TEXT(\"%s\"));\r\n"), *DynamicClassPackageName);
-			GeneratedClassRegisterFunctionText.Logf(TEXT("\t\tUClass* OuterClass = Cast<UClass>(StaticFindObjectFast(UClass::StaticClass(), OuterPackage, TEXT(\"%s\")));\r\n"), *Class->GetName());
+			GeneratedClassRegisterFunctionText.Logf(TEXT("\t\tUClass* OuterClass = Cast<UClass>(StaticFindObjectFast(UClass::StaticClass(), OuterPackage, TEXT(\"%s\")));\r\n"), *FNativeClassHeaderGenerator::GetOverriddenName(Class));
 			GeneratedClassRegisterFunctionText.Logf(TEXT("\t\tif (!OuterClass || !(OuterClass->ClassFlags & CLASS_Constructed))\r\n"));
 		}
 		
@@ -1405,8 +1405,8 @@ void FNativeClassHeaderGenerator::ExportNativeGeneratedInitCode(FClass* Class, F
 
 		FString OuterString = FString(TEXT("OuterClass"));
 
+		TMap<FName, FString>* MetaDataMap = UMetaData::GetMapForObject(Class);
 		{
-			TMap<FName, FString>* MetaDataMap = UMetaData::GetMapForObject(Class);
 			FClassMetaData* ClassMetaData = GScriptHelper.FindClassData(Class);
 			if (MetaDataMap && ClassMetaData && Class)
 			{
@@ -1490,6 +1490,11 @@ void FNativeClassHeaderGenerator::ExportNativeGeneratedInitCode(FClass* Class, F
 
 		if (bIsDynamic)
 		{
+			FString* CustomDynamicClassInitializationMD = MetaDataMap ? MetaDataMap->Find(TEXT("CustomDynamicClassInitialization")) : nullptr;
+			if (CustomDynamicClassInitializationMD)
+			{
+				GeneratedClassRegisterFunctionText.Logf(TEXT("\t\t\t\t%s(CastChecked<UDynamicClass>(OuterClass));\n"), *(*CustomDynamicClassInitializationMD));
+			}
 			GeneratedClassRegisterFunctionText.Logf(TEXT("\t\t\t\tOuterClass->GetDefaultObject();\r\n"));
 		}
 
@@ -2791,6 +2796,7 @@ void FNativeClassHeaderGenerator::ExportGeneratedStructBodyMacros(const TArray<U
 	{
 		UScriptStruct* Struct = NativeStructs[i];
 		const bool bIsDynamic = FClass::IsDynamic(Struct);
+		const FString ActualStructName = FNativeClassHeaderGenerator::GetOverriddenName(Struct);
 
 		// Export struct.
 		if (Struct->StructFlags & STRUCT_Native)
@@ -2847,7 +2853,7 @@ void FNativeClassHeaderGenerator::ExportGeneratedStructBodyMacros(const TArray<U
 			else
 			{
 				GeneratedPackageCPP.Logf(TEXT("\tclass UScriptStruct* Singleton = Cast<UScriptStruct>(StaticFindObjectFast(UScriptStruct::StaticClass(), %s, TEXT(\"%s\")));\r\n"),
-					*OuterName, *Struct->GetName());
+					*OuterName, *ActualStructName);
 			}
 			GeneratedPackageCPP.Logf(TEXT("\tif (!Singleton)\r\n"));
 			GeneratedPackageCPP.Logf(TEXT("\t{\r\n"));
@@ -2855,14 +2861,14 @@ void FNativeClassHeaderGenerator::ExportGeneratedStructBodyMacros(const TArray<U
 			GeneratedPackageCPP.Logf(TEXT("\t\textern %suint32 %s();\r\n"), *FriendApiString, *GetCRCName);
 
 			GeneratedPackageCPP.Logf(TEXT("\t\tSingleton = GetStaticStruct(%s, %s, TEXT(\"%s\"), sizeof(%s), %s());\r\n"),
-				*SingletonName, *OuterName, *Struct->GetName(), StructNameCPP, *GetCRCName);
+				*SingletonName, *OuterName, *ActualStructName, StructNameCPP, *GetCRCName);
 
 			GeneratedPackageCPP.Logf(TEXT("\t}\r\n"));
 			GeneratedPackageCPP.Logf(TEXT("\treturn Singleton;\r\n"));
 			GeneratedPackageCPP.Logf(TEXT("}\r\n"));
 
 			GeneratedPackageCPP.Logf(TEXT("static FCompiledInDeferStruct Z_CompiledInDeferStruct_UScriptStruct_%s(%s::StaticStruct, TEXT(\"%s\"), TEXT(\"%s\"), %s, %s, %s);\r\n"),
-				StructNameCPP, StructNameCPP, *Struct->GetOutermost()->GetName(), *Struct->GetName(),
+				StructNameCPP, StructNameCPP, *Struct->GetOutermost()->GetName(), *ActualStructName,
 				bIsDynamic ? TEXT("true") : TEXT("false"),
 				bIsDynamic ? *AsTEXT(FClass::GetTypePackageName(Struct)) : TEXT("nullptr"),
 				bIsDynamic ? *AsTEXT(FNativeClassHeaderGenerator::GetOverriddenPathName(Struct)) : TEXT("nullptr"));
@@ -2876,7 +2882,7 @@ void FNativeClassHeaderGenerator::ExportGeneratedStructBodyMacros(const TArray<U
 				GeneratedPackageCPP.Logf(TEXT("\tFScriptStruct_%s_StaticRegisterNatives%s()\r\n"), *ShortPackageName, StructNameCPP);
 				GeneratedPackageCPP.Logf(TEXT("\t{\r\n"));
 
-				GeneratedPackageCPP.Logf(TEXT("\t\tUScriptStruct::DeferCppStructOps(FName(TEXT(\"%s\")),new UScriptStruct::TCppStructOps<%s%s>);\r\n"), *Struct->GetName(), Struct->GetPrefixCPP(), *Struct->GetName());
+				GeneratedPackageCPP.Logf(TEXT("\t\tUScriptStruct::DeferCppStructOps(FName(TEXT(\"%s\")),new UScriptStruct::TCppStructOps<%s>);\r\n"), *ActualStructName, StructNameCPP);
 
 				GeneratedPackageCPP.Logf(TEXT("\t}\r\n"));
 				GeneratedPackageCPP.Logf(TEXT("} ScriptStruct_%s_StaticRegisterNatives%s;\r\n"), *ShortPackageName, StructNameCPP);
@@ -2923,11 +2929,11 @@ void FNativeClassHeaderGenerator::ExportGeneratedStructBodyMacros(const TArray<U
 			GeneratedStructRegisterFunctionText.Logf(TEXT("\t\textern uint32 %s();\r\n"), *CRCFuncName);
 			if (!bIsDynamic)
 			{
-				GeneratedStructRegisterFunctionText.Logf(TEXT("\t\tstatic UScriptStruct* ReturnStruct = FindExistingStructIfHotReloadOrDynamic(Outer, TEXT(\"%s\"), sizeof(%s), %s(), false);\r\n"), *ScriptStruct->GetName(), NameLookupCPP.GetNameCPP(Struct), *CRCFuncName);
+				GeneratedStructRegisterFunctionText.Logf(TEXT("\t\tstatic UScriptStruct* ReturnStruct = FindExistingStructIfHotReloadOrDynamic(Outer, TEXT(\"%s\"), sizeof(%s), %s(), false);\r\n"), *ActualStructName, NameLookupCPP.GetNameCPP(Struct), *CRCFuncName);
 			}
 			else
 			{
-				GeneratedStructRegisterFunctionText.Logf(TEXT("\t\tUScriptStruct* ReturnStruct = FindExistingStructIfHotReloadOrDynamic(Outer, TEXT(\"%s\"), sizeof(%s), %s(), true);\r\n"), *ScriptStruct->GetName(), NameLookupCPP.GetNameCPP(Struct), *CRCFuncName);
+				GeneratedStructRegisterFunctionText.Logf(TEXT("\t\tUScriptStruct* ReturnStruct = FindExistingStructIfHotReloadOrDynamic(Outer, TEXT(\"%s\"), sizeof(%s), %s(), true);\r\n"), *ActualStructName, NameLookupCPP.GetNameCPP(Struct), *CRCFuncName);
 			}
 			GeneratedStructRegisterFunctionText.Logf(TEXT("\t\tif (!ReturnStruct)\r\n"));
 			GeneratedStructRegisterFunctionText.Logf(TEXT("\t\t{\r\n"));
@@ -2952,7 +2958,7 @@ void FNativeClassHeaderGenerator::ExportGeneratedStructBodyMacros(const TArray<U
 
 			const TCHAR* UStructObjectFlags = bIsDynamic ? TEXT("RF_Public|RF_Transient") : TEXT("RF_Public|RF_Transient|RF_MarkAsNative");
 			GeneratedStructRegisterFunctionText.Logf(TEXT("\t\t\tReturnStruct = new(EC_InternalUseOnlyConstructor, Outer, TEXT(\"%s\"), %s) UScriptStruct(FObjectInitializer(), %s, %s, EStructFlags(0x%08X)%s);\r\n"),
-				*ScriptStruct->GetName(),
+				*ActualStructName,
 				UStructObjectFlags,
 				*BaseStructString,
 				*CppStructOpsString,
@@ -3053,8 +3059,9 @@ void FNativeClassHeaderGenerator::ExportGeneratedEnumsInitCode(const TArray<UEnu
 		GeneratedPackageCPP.Logf(TEXT("\treturn Singleton;\r\n"));
 		GeneratedPackageCPP.Logf(TEXT("}\r\n"));
 
+		const FString EnumNameCpp = Enum->GetName(); //UserDefinedEnum should already have a valid cpp name.
 		GeneratedPackageCPP.Logf(TEXT("static FCompiledInDeferEnum Z_CompiledInDeferEnum_UEnum_%s(%s_StaticEnum, TEXT(\"%s\"), TEXT(\"%s\"), %s, %s, %s);\r\n"), 
-			*Enum->GetName(), *Enum->GetName(), *Enum->GetOutermost()->GetName(), *Enum->GetName(),
+			*EnumNameCpp, *EnumNameCpp, *Enum->GetOutermost()->GetName(), *FNativeClassHeaderGenerator::GetOverriddenName(Enum),
 			bIsDynamic ? TEXT("true") : TEXT("false"),
 			bIsDynamic ? *AsTEXT(FClass::GetTypePackageName(Enum)) : TEXT("nullptr"),
 			bIsDynamic ? *AsTEXT(FNativeClassHeaderGenerator::GetOverriddenPathName(Enum)) : TEXT("nullptr"));
diff --git a/Engine/Source/Runtime/Core/Public/UObject/BlueprintsObjectVersion.h b/Engine/Source/Runtime/Core/Public/UObject/BlueprintsObjectVersion.h
--- a/Engine/Source/Runtime/Core/Public/UObject/BlueprintsObjectVersion.h
+++ b/Engine/Source/Runtime/Core/Public/UObject/BlueprintsObjectVersion.h
@@ -10,7 +10,7 @@ struct CORE_API FBlueprintsObjectVersion
 		BeforeCustomVersionWasAdded = 0,
 		OverridenEventReferenceFixup,
 		CleanBlueprintFunctionFlags,
-
+		ArrayGetByRefUpgrade,
 
 		// -----<new versions can be added above this line>-------------------------------------------------
 		VersionPlusOne,
diff --git a/Engine/Source/Runtime/CoreUObject/Private/Blueprint/BlueprintSupport.cpp b/Engine/Source/Runtime/CoreUObject/Private/Blueprint/BlueprintSupport.cpp
--- a/Engine/Source/Runtime/CoreUObject/Private/Blueprint/BlueprintSupport.cpp
+++ b/Engine/Source/Runtime/CoreUObject/Private/Blueprint/BlueprintSupport.cpp
@@ -82,6 +82,47 @@ bool FBlueprintSupport::IsDeferredCDOInitializationDisabled()
 #endif
 }
 
+bool FBlueprintSupport::IsDeferredDependencyPlaceholder(UObject* LoadedObj)
+{
+	return LoadedObj && ( LoadedObj->IsA<ULinkerPlaceholderClass>() ||
+		LoadedObj->IsA<ULinkerPlaceholderFunction>() ||
+		LoadedObj->IsA<ULinkerPlaceholderExportObject>() );
+}
+
+bool FBlueprintSupport::IsInBlueprintPackage(UObject* LoadedObj)
+{
+	UPackage* Pkg = LoadedObj->GetOutermost();
+	if (Pkg && !Pkg->HasAnyPackageFlags(PKG_CompiledIn))
+	{
+		TArray<UObject*> PkgObjects;
+		GetObjectsWithOuter(Pkg, PkgObjects, /*bIncludeNestedObjects =*/false);
+		
+		UObject* PkgCDO   = nullptr;
+		UClass*  PkgClass = nullptr;
+
+		for (UObject* PkgObj : PkgObjects)
+		{
+			if (PkgObj->HasAnyFlags(RF_ClassDefaultObject))
+			{
+				PkgCDO = PkgObj;
+			}
+			else if (UClass* AsClass = Cast<UClass>(PkgObj))
+			{
+				PkgClass = AsClass;
+			}
+		}
+		const bool bHasBlueprintClass = PkgClass && PkgClass->HasAnyClassFlags(CLASS_CompiledFromBlueprint);
+
+		return bHasBlueprintClass
+			//&& (PkgCDO && PkgCDO->GetClass() == PkgClass)
+#if WITH_EDITORONLY_DATA
+			//&& (PkgClass->ClassGeneratedBy != nullptr) && (PkgClass->ClassGeneratedBy->GetOuter() == Pkg)
+#endif
+			;
+	}
+	return false;
+}
+
 /*******************************************************************************
  * FScopedClassDependencyGather
  ******************************************************************************/
@@ -329,11 +370,50 @@ bool FLinkerLoad::RegenerateBlueprintClass(UClass* LoadClass, UObject* ExportObj
 	return bSuccessfulRegeneration;
 }
 
-/** Helper utility function to convert UObject pointers to FLinkerPlaceholderBase* */
+/** 
+ * Frivolous helper functions, to provide unique identifying names for our different placeholder types.
+ */
+template<class PlaceholderType>
+static FString GetPlaceholderPrefix()                      { return TEXT("PLACEHOLDER_"); }
+template<>
+FString GetPlaceholderPrefix<ULinkerPlaceholderFunction>() { return TEXT("PLACEHOLDER-FUNCTION_"); }
+template<>
+FString GetPlaceholderPrefix<ULinkerPlaceholderClass>()    { return TEXT("PLACEHOLDER-CLASS_"); }
+
+/** Internal utility function for spawning various type of placeholder objects. */
 template<class PlaceholderType>
-static FLinkerPlaceholderBase* PlaceholderCast(UObject* PlaceholderObj)
+static PlaceholderType* MakeImportPlaceholder(UObject* Outer, const TCHAR* TargetObjName, int32 ImportIndex = INDEX_NONE)
 {
-	return CastChecked<PlaceholderType>(PlaceholderObj, ECastCheckedType::NullAllowed);
+	PlaceholderType* PlaceholderObj = nullptr;
+
+#if USE_CIRCULAR_DEPENDENCY_LOAD_DEFERRING
+	FName PlaceholderName(*FString::Printf(TEXT("%s_%s"), *GetPlaceholderPrefix<PlaceholderType>(), TargetObjName));
+	PlaceholderName = MakeUniqueObjectName(Outer, PlaceholderType::StaticClass(), PlaceholderName);
+
+	PlaceholderObj = NewObject<PlaceholderType>(Outer, PlaceholderType::StaticClass(), PlaceholderName, RF_Public | RF_Transient);
+
+	if (ImportIndex != INDEX_NONE)
+	{
+		PlaceholderObj->PackageIndex = FPackageIndex::FromImport(ImportIndex);
+	}
+	// else, this is probably coming from something like an ImportText() call, 
+	// and isn't referenced by the ImportMap... instead, this should be stored 
+	// in the FLinkerLoad's ImportPlaceholders map
+
+	// make sure the class is fully formed (has its 
+	// ClassAddReferencedObjects/ClassConstructor members set)
+	PlaceholderObj->Bind();
+	PlaceholderObj->StaticLink(/*bRelinkExistingProperties =*/true);
+
+#if USE_DEFERRED_DEPENDENCY_CHECK_VERIFICATION_TESTS
+	if (ULinkerPlaceholderClass* OuterAsPlaceholder = dynamic_cast<ULinkerPlaceholderClass*>(Outer))
+	{
+		OuterAsPlaceholder->AddChildObject(PlaceholderObj);
+	}
+#endif //USE_DEFERRED_DEPENDENCY_CHECK_VERIFICATION_TESTS
+#endif //USE_CIRCULAR_DEPENDENCY_LOAD_DEFERRING
+
+	return PlaceholderObj;
 }
 
 bool FLinkerLoad::DeferPotentialCircularImport(const int32 Index)
@@ -361,78 +441,23 @@ bool FLinkerLoad::DeferPotentialCircularImport(const int32 Index)
 		{
 			if (const UClass* ImportClass = FindObject<UClass>(ClassPackage, *Import.ClassName.ToString()))
 			{
-				UObject* PlaceholderOuter = LinkerRoot;
-				FString  PlaceholderNamePrefix = TEXT("PLACEHOLDER_");
-				UClass*  PlaceholderType = nullptr;
-#if USE_DEFERRED_DEPENDENCY_CHECK_VERIFICATION_TESTS
-				ULinkerPlaceholderClass* Outer = nullptr;
-#endif //USE_DEFERRED_DEPENDENCY_CHECK_VERIFICATION_TESTS
-				FLinkerPlaceholderBase* (*PlaceholderCastFunc)(UObject* RawUObject) = nullptr;
-
-				bool const bIsBlueprintClass = ImportClass->IsChildOf<UClass>();
-				// @TODO: if we could see if the related package is created 
-				//        (without loading it further), AND it already has this 
-				//        import loaded, then we can probably return that rather 
-				//        than allocating a placeholder
-				if (bIsBlueprintClass)
+				if (ImportClass->IsChildOf<UClass>())
 				{
-					PlaceholderNamePrefix = TEXT("PLACEHOLDER-CLASS_");
-					PlaceholderType = ULinkerPlaceholderClass::StaticClass();
-					PlaceholderCastFunc = &PlaceholderCast<ULinkerPlaceholderClass>;
+					Import.XObject = MakeImportPlaceholder<ULinkerPlaceholderClass>(LinkerRoot, *Import.ObjectName.ToString(), Index);
 				}
 				else if (ImportClass->IsChildOf<UFunction>())
-				{		
-					if (Import.OuterIndex.IsImport())
-					{
-						const int32 OuterImportIndex = Import.OuterIndex.ToImport();
-						// @TODO: if the sole reason why we have ULinkerPlaceholderFunction 
-						//        is that it's outer is a placeholder, then we 
-						//        could instead log it (with the placeholder) as 
-						//        a referencer, and then move the function later
-						if (DeferPotentialCircularImport(OuterImportIndex))
-						{
-							PlaceholderNamePrefix = TEXT("PLACEHOLDER-FUNCTION_");
-							PlaceholderOuter = ImportMap[OuterImportIndex].XObject;
-							PlaceholderType  = ULinkerPlaceholderFunction::StaticClass();
-							PlaceholderCastFunc = &PlaceholderCast<ULinkerPlaceholderFunction>;
-#if USE_DEFERRED_DEPENDENCY_CHECK_VERIFICATION_TESTS
-							Outer = dynamic_cast<ULinkerPlaceholderClass*>(PlaceholderOuter);
-							check(Outer);
-#endif //USE_DEFERRED_DEPENDENCY_CHECK_VERIFICATION_TESTS
-						}
-					}
-				}
-
-				if (PlaceholderType != nullptr)
 				{
-					FName PlaceholderName(*FString::Printf(TEXT("%s_%s"), *PlaceholderNamePrefix, *Import.ObjectName.ToString()));
-					PlaceholderName = MakeUniqueObjectName(PlaceholderOuter, PlaceholderType, PlaceholderName);
-
-					UStruct* PlaceholderObj = NewObject<UStruct>(PlaceholderOuter, PlaceholderType, PlaceholderName, RF_Public | RF_Transient);
-					// store the import index in the placeholder, so we can 
-					// easily look it up in the import map, given the 
-					// placeholder (needed, to find the corresponding import for 
-					// ResolvingDeferredPlaceholder)
-					DEFERRED_DEPENDENCY_CHECK(PlaceholderCastFunc != nullptr);
-					if (PlaceholderCastFunc != nullptr)
-					{
-						if (FLinkerPlaceholderBase* AsPlaceholder = PlaceholderCastFunc(PlaceholderObj))
-						{
-							AsPlaceholder->PackageIndex = FPackageIndex::FromImport(Index);
-						}
-					}				
-					// make sure the class is fully formed (has its 
-					// ClassAddReferencedObjects/ClassConstructor members set)
-					PlaceholderObj->Bind();
-					PlaceholderObj->StaticLink(/*bRelinkExistingProperties =*/true);
-#if USE_DEFERRED_DEPENDENCY_CHECK_VERIFICATION_TESTS
-					if (Outer)
+					const int32 OuterImportIndex = Import.OuterIndex.ToImport();
+					// @TODO: if the sole reason why we have ULinkerPlaceholderFunction 
+					//        is that it's outer is a placeholder, then we 
+					//        could instead log it (with the placeholder) as 
+					//        a referencer, and then move the function later
+					if (DeferPotentialCircularImport(OuterImportIndex))
 					{
-						Outer->AddChildObject(PlaceholderObj);
+						UObject* FuncOuter = ImportMap[OuterImportIndex].XObject;
+						Import.XObject = MakeImportPlaceholder<ULinkerPlaceholderFunction>(FuncOuter, *Import.ObjectName.ToString(), Index);
+						DEFERRED_DEPENDENCY_CHECK(dynamic_cast<ULinkerPlaceholderClass*>(FuncOuter) != nullptr);
 					}
-#endif //USE_DEFERRED_DEPENDENCY_CHECK_VERIFICATION_TESTS
-
-					Import.XObject = PlaceholderObj;
 				}
 			}
 		}
@@ -450,28 +475,29 @@ bool FLinkerLoad::DeferPotentialCircularImport(const int32 Index)
 }
 
 #if WITH_EDITOR
-bool FLinkerLoad::IsSuppressableBlueprintImportError(FObjectImport& InImport)
+bool FLinkerLoad::IsSuppressableBlueprintImportError(int32 ImportIndex) const
 {
 	// We want to suppress any import errors that target a BlueprintGeneratedClass
 	// since these issues can occur when an externally referenced Blueprint is saved 
 	// without compiling. This should not be a problem because all Blueprints are
 	// compiled-on-load.
 	static const FName NAME_BlueprintGeneratedClass("BlueprintGeneratedClass");
-	FObjectImport& TestImport = InImport;
-	bool bResult = false;
+
+	bool bImportBelongsToBlueprint = false;
 	// We will look at each outer of the Import to see if any of them are a BPGC
-	while (1)
+	while (ImportMap.IsValidIndex(ImportIndex))
 	{
+		const FObjectImport& TestImport = ImportMap[ImportIndex];
 		if (TestImport.ClassName == NAME_BlueprintGeneratedClass)
 		{
 			// The import is a BPGC, suppress errors
-			bResult = true;
+			bImportBelongsToBlueprint = true;
 			break;
 		}
 
 		if (!TestImport.OuterIndex.IsNull() && TestImport.OuterIndex.IsImport())
 		{
-			TestImport = Imp(TestImport.OuterIndex);
+			ImportIndex = TestImport.OuterIndex.ToImport();
 		}
 		else
 		{
@@ -479,7 +505,8 @@ bool FLinkerLoad::IsSuppressableBlueprintImportError(FObjectImport& InImport)
 			break;
 		}
 	}
-	return bResult;
+
+	return bImportBelongsToBlueprint;
 }
 #endif // WITH_EDITOR
 
@@ -790,69 +817,109 @@ void FLinkerLoad::ResolveDeferredDependencies(UStruct* LoadStruct)
 	// placeholder classes, all member struct's placholders, and its parent's
 	{
 		FUnresolvedStructTracker ScopedResolveTracker(LoadStruct);
+		
+		UClass* const LoadClass = Cast<UClass>(LoadStruct);
 
+		bool bImportMapResolved = false;
 		// this function (for this linker) could be reentrant (see where we 
 		// recursively call ResolveDeferredDependencies() for super-classes below);  
 		// if that's the case, then we want to finish resolving the pending class 
 		// before we continue on
 		if (ResolvingDeferredPlaceholder != nullptr)
 		{
-			int32 ResolvedRefCount = ResolveDependencyPlaceholder(ResolvingDeferredPlaceholder, Cast<UClass>(LoadStruct));
+			FName ReplacementPkgPath = NAME_None;
+			// if this placeholder is not an ImportMap placeholder, then we need
+			// additional info to resolve it (namely the true object's package
+			// path)... if we don't have that, then this resolve will fail
+			if (ResolvingDeferredPlaceholder->PackageIndex.IsNull())
+			{
+				const FName* ImportObjectPath = ImportPlaceholders.FindKey(ResolvingDeferredPlaceholder);
+				DEFERRED_DEPENDENCY_CHECK(ImportObjectPath != nullptr);
+
+				if (ImportObjectPath != nullptr)
+				{
+					ReplacementPkgPath = *ImportObjectPath;
+					// save us from looping through the ImportMap again; if 
+					// we're already resolving an entry from ImportPlaceholders,  
+					// then we've already been through this function and are in  
+					// the ImportPlaceholders loop somewhere up the stack
+					bImportMapResolved = true;
+				}
+			}
+
+			int32 ResolvedRefCount = ResolveDependencyPlaceholder(ResolvingDeferredPlaceholder, LoadClass, ReplacementPkgPath);
 			// @TODO: can we reliably count on this resolving some dependencies?... 
 			//        if so, check verify that!
 			ResolvingDeferredPlaceholder = nullptr;
+			ImportPlaceholders.Remove(ReplacementPkgPath);
 		}
 
-		// because this loop could recurse (and end up finishing all of this for
-		// us), we check HasUnresolvedDependencies() so we can early out  
-		// from this loop in that situation (the loop has been finished elsewhere)
-		for (int32 ImportIndex = 0; ImportIndex < ImportMap.Num() && HasUnresolvedDependencies(); ++ImportIndex)
+		if (!bImportMapResolved)
 		{
-			FObjectImport& Import = ImportMap[ImportIndex];
-
-			const FLinkerLoad* SourceLinker = (Import.SourceLinker != nullptr) ? Import.SourceLinker :
-				(Import.XObject != nullptr) ? Import.XObject->GetLinker() : nullptr;
-
-			const UPackage* SourcePackage = (SourceLinker != nullptr) ? SourceLinker->LinkerRoot : nullptr;
-			// this package may not have introduced any (possible) cyclic 
-			// dependencies, but it still could have been deferred (kept from
-			// fully loading... we need to make sure metadata gets loaded, etc.)
-			if ((SourcePackage != nullptr) && !SourcePackage->HasAnyFlags(RF_WasLoaded))
+			// because this loop could recurse (and end up finishing all of this for
+			// us), we check HasUnresolvedDependencies() so we can early out  
+			// from this loop in that situation (the loop has been finished elsewhere)
+			for (int32 ImportIndex = 0; ImportIndex < ImportMap.Num() && HasUnresolvedDependencies(); ++ImportIndex)
 			{
-				uint32 InternalLoadFlags = LoadFlags & (LOAD_NoVerify | LOAD_NoWarn | LOAD_Quiet);
-				// make sure LoadAllObjects() is called for this package
-				LoadPackageInternal(/*Outer =*/nullptr, *SourceLinker->Filename, InternalLoadFlags, this);
-			}
+				FObjectImport& Import = ImportMap[ImportIndex];
 
-			if (ULinkerPlaceholderClass* PlaceholderClass = Cast<ULinkerPlaceholderClass>(Import.XObject))
-			{
-				DEFERRED_DEPENDENCY_CHECK(PlaceholderClass->PackageIndex.ToImport() == ImportIndex);
+				const FLinkerLoad* SourceLinker = (Import.SourceLinker != nullptr) ? Import.SourceLinker :
+					(Import.XObject != nullptr) ? Import.XObject->GetLinker() : nullptr;
 
-				// NOTE: we don't check that this resolve successfully replaced any
-				//       references (by the return value), because this resolve 
-				//       could have been re-entered and completed by a nested call
-				//       to the same function (for the same placeholder)
-				ResolveDependencyPlaceholder(PlaceholderClass, Cast<UClass>(LoadStruct));
-			}
-			else if (ULinkerPlaceholderFunction* PlaceholderFunction = Cast<ULinkerPlaceholderFunction>(Import.XObject))
-			{
-				if (ULinkerPlaceholderClass* PlaceholderOwner = Cast<ULinkerPlaceholderClass>(PlaceholderFunction->GetOwnerClass()))
+				const UPackage* SourcePackage = (SourceLinker != nullptr) ? SourceLinker->LinkerRoot : nullptr;
+				// this package may not have introduced any (possible) cyclic 
+				// dependencies, but it still could have been deferred (kept from
+				// fully loading... we need to make sure metadata gets loaded, etc.)
+				if ((SourcePackage != nullptr) && !SourcePackage->HasAnyFlags(RF_WasLoaded))
 				{
-					ResolveDependencyPlaceholder(PlaceholderOwner, Cast<UClass>(LoadStruct));
+					uint32 InternalLoadFlags = LoadFlags & (LOAD_NoVerify | LOAD_NoWarn | LOAD_Quiet);
+					// make sure LoadAllObjects() is called for this package
+					LoadPackageInternal(/*Outer =*/nullptr, *SourceLinker->Filename, InternalLoadFlags, this);
 				}
 
-				DEFERRED_DEPENDENCY_CHECK(PlaceholderFunction->PackageIndex.ToImport() == ImportIndex);
-				ResolveDependencyPlaceholder(PlaceholderFunction, Cast<UClass>(LoadStruct)); 
-			}
-			else if (UScriptStruct* StructObj = Cast<UScriptStruct>(Import.XObject))
-			{
-				// in case this is a user defined struct, we have to resolve any 
-				// deferred dependencies in the struct 
-				if (Import.SourceLinker != nullptr)
+				if (ULinkerPlaceholderClass* PlaceholderClass = Cast<ULinkerPlaceholderClass>(Import.XObject))
+				{
+					DEFERRED_DEPENDENCY_CHECK(PlaceholderClass->PackageIndex.ToImport() == ImportIndex);
+
+					// NOTE: we don't check that this resolve successfully replaced any
+					//       references (by the return value), because this resolve 
+					//       could have been re-entered and completed by a nested call
+					//       to the same function (for the same placeholder)
+					ResolveDependencyPlaceholder(PlaceholderClass, LoadClass);
+				}
+				else if (ULinkerPlaceholderFunction* PlaceholderFunction = Cast<ULinkerPlaceholderFunction>(Import.XObject))
+				{
+					if (ULinkerPlaceholderClass* PlaceholderOwner = Cast<ULinkerPlaceholderClass>(PlaceholderFunction->GetOwnerClass()))
+					{
+						ResolveDependencyPlaceholder(PlaceholderOwner, LoadClass);
+					}
+
+					DEFERRED_DEPENDENCY_CHECK(PlaceholderFunction->PackageIndex.ToImport() == ImportIndex);
+					ResolveDependencyPlaceholder(PlaceholderFunction, LoadClass);
+				}
+				else if (UScriptStruct* StructObj = Cast<UScriptStruct>(Import.XObject))
 				{
-					Import.SourceLinker->ResolveDeferredDependencies(StructObj);
+					// in case this is a user defined struct, we have to resolve any 
+					// deferred dependencies in the struct 
+					if (Import.SourceLinker != nullptr)
+					{
+						Import.SourceLinker->ResolveDeferredDependencies(StructObj);
+					}
 				}
 			}
+		} // if (!bImportMapResolved)
+
+		// resolve any placeholders that were imported through methods like 
+		// ImportText() (meaning the ImportMap wouldn't reference them)
+		while (ImportPlaceholders.Num() > 0)
+		{
+			auto PlaceholderIt = ImportPlaceholders.CreateIterator();
+
+			// store off the key before we resolve, in case this has been recursively removed
+			const FName PlaceholderKey = PlaceholderIt.Key();
+			ResolveDependencyPlaceholder(PlaceholderIt.Value(), LoadClass, PlaceholderKey);
+
+			ImportPlaceholders.Remove(PlaceholderKey);
 		}
 
 		if (UStruct* SuperStruct = LoadStruct->GetSuperStruct())
@@ -895,6 +962,7 @@ void FLinkerLoad::ResolveDeferredDependencies(UStruct* LoadStruct)
 			DEFERRED_DEPENDENCY_CHECK(!PlaceholderClass->HasKnownReferences());
 		}
 	}
+	DEFERRED_DEPENDENCY_CHECK(ImportPlaceholders.Num() == 0);
 #endif // USE_DEFERRED_DEPENDENCY_CHECK_VERIFICATION_TESTS
 #endif // USE_CIRCULAR_DEPENDENCY_LOAD_DEFERRING
 }
@@ -918,7 +986,7 @@ bool FLinkerLoad::HasUnresolvedDependencies() const
 #endif // USE_CIRCULAR_DEPENDENCY_LOAD_DEFERRING
 }
 
-int32 FLinkerLoad::ResolveDependencyPlaceholder(FLinkerPlaceholderBase* PlaceholderIn, UClass* ReferencingClass)
+int32 FLinkerLoad::ResolveDependencyPlaceholder(FLinkerPlaceholderBase* PlaceholderIn, UClass* ReferencingClass, const FName ObjectPath)
 {
 #if USE_CIRCULAR_DEPENDENCY_LOAD_DEFERRING
 	TGuardValue<uint32>  LoadFlagsGuard(LoadFlags, (LoadFlags & ~LOAD_DeferDependencyLoads));
@@ -927,27 +995,38 @@ int32 FLinkerLoad::ResolveDependencyPlaceholder(FLinkerPlaceholderBase* Placehol
 	UObject* PlaceholderObj = PlaceholderIn->GetPlaceholderAsUObject();
 	DEFERRED_DEPENDENCY_CHECK(PlaceholderObj != nullptr);
 	DEFERRED_DEPENDENCY_CHECK(PlaceholderObj->GetOutermost() == LinkerRoot);
-	DEFERRED_DEPENDENCY_CHECK(PlaceholderIn->PackageIndex.IsImport());
-	
-	int32 const ImportIndex = PlaceholderIn->PackageIndex.ToImport();
-	FObjectImport& Import = ImportMap[ImportIndex];
 	
 	UObject* RealImportObj = nullptr;
-	if ((Import.XObject != nullptr) && (Import.XObject != PlaceholderObj))
+
+	if (PlaceholderIn->PackageIndex.IsNull())
 	{
-		DEFERRED_DEPENDENCY_CHECK(ResolvingDeferredPlaceholder == PlaceholderIn);
-		RealImportObj = Import.XObject;
+		DEFERRED_DEPENDENCY_CHECK(ObjectPath.IsValid() && !ObjectPath.IsNone());
+		// emulating the StaticLoadObject() call in UObjectPropertyBase::FindImportedObject(),
+		// since this was most likely a placeholder 
+		RealImportObj = StaticLoadObject(UObject::StaticClass(), /*Outer =*/nullptr, *ObjectPath.ToString(), /*Filename =*/nullptr, (LOAD_NoWarn | LOAD_FindIfFail));
 	}
-	else 
+	else
 	{
-		// clear the placeholder from the import, so that a call to CreateImport()
-		// properly fills it in
-		Import.XObject = nullptr;
-		// NOTE: this is a possible point of recursion... CreateImport() could 
-		//       continue to load a package already started up the stack and you 
-		//       could end up in another ResolveDependencyPlaceholder() for some  
-		//       other placeholder before this one has completely finished resolving
-		RealImportObj = CreateImport(ImportIndex);
+		DEFERRED_DEPENDENCY_CHECK(PlaceholderIn->PackageIndex.IsImport());
+		int32 const ImportIndex = PlaceholderIn->PackageIndex.ToImport();
+		FObjectImport& Import = ImportMap[ImportIndex];
+
+		if ((Import.XObject != nullptr) && (Import.XObject != PlaceholderObj))
+		{
+			DEFERRED_DEPENDENCY_CHECK(ResolvingDeferredPlaceholder == PlaceholderIn);
+			RealImportObj = Import.XObject;
+		}
+		else
+		{
+			// clear the placeholder from the import, so that a call to CreateImport()
+			// properly fills it in
+			Import.XObject = nullptr;
+			// NOTE: this is a possible point of recursion... CreateImport() could 
+			//       continue to load a package already started up the stack and you 
+			//       could end up in another ResolveDependencyPlaceholder() for some  
+			//       other placeholder before this one has completely finished resolving
+			RealImportObj = CreateImport(ImportIndex);
+		}
 	}
 
 #if USE_DEFERRED_DEPENDENCY_CHECK_VERIFICATION_TESTS
@@ -1182,6 +1261,7 @@ void FLinkerLoad::FinalizeBlueprint(UClass* LoadClass)
 		// is where script code is regenerated)
 		FStructScriptLoader::ResolveDeferredScriptLoads(this);
 
+		DEFERRED_DEPENDENCY_CHECK(ImportPlaceholders.Num() == 0);
 		DEFERRED_DEPENDENCY_CHECK(LoadClass->GetOutermost() != GetTransientPackage());
 		// just in case we choose to enable the deferred dependency loading for 
 		// cooked builds... we want to keep from regenerating in that scenario
@@ -1222,9 +1302,18 @@ void FLinkerLoad::ResolveDeferredExports(UClass* LoadClass)
 #if USE_DEFERRED_DEPENDENCY_CHECK_VERIFICATION_TESTS
 		auto IsPlaceholderReferenced = [](ULinkerPlaceholderExportObject* ExportPlaceholder)->bool
 		{
-			FReferencerInformationList UnresolvedReferences;
 			UObject* PlaceholderObj = ExportPlaceholder;
-			return IsReferenced(PlaceholderObj, GARBAGE_COLLECTION_KEEPFLAGS, EInternalObjectFlags::GarbageCollectionKeepFlags, /*bCheckSubObjects =*/false, &UnresolvedReferences);
+
+			FReferencerInformationList UnresolvedReferences;
+			bool bIsReferenced = IsReferenced(PlaceholderObj, GARBAGE_COLLECTION_KEEPFLAGS, EInternalObjectFlags::GarbageCollectionKeepFlags, /*bCheckSubObjects =*/false, &UnresolvedReferences);
+
+			if (bIsReferenced && IsAsyncLoading())
+			{
+				// if we're async loading, then we assume a single external 
+				// reference belongs to FAsyncObjectsReferencer, which is allowable 
+				bIsReferenced = (UnresolvedReferences.ExternalReferences.Num() != 1) || (UnresolvedReferences.InternalReferences.Num() > 0);
+			}
+			return bIsReferenced;
 		};
 #endif // USE_DEFERRED_DEPENDENCY_CHECK_VERIFICATION_TESTS
 
@@ -1445,6 +1534,20 @@ bool FLinkerLoad::IsExportBeingResolved(int32 ExportIndex)
 	return bIsExportClassBeingForceRegened;
 }
 
+void FLinkerLoad::ResetDeferredLoadingState()
+{
+#if USE_CIRCULAR_DEPENDENCY_LOAD_DEFERRING
+	DeferredCDOIndex = INDEX_NONE;
+	bForceBlueprintFinalization = false;
+	ResolvingDeferredPlaceholder = nullptr;
+	ImportPlaceholders.Empty();
+	LoadFlags &= ~(LOAD_DeferDependencyLoads);
+
+	FResolvingExportTracker::Get().Reset(this);
+	FUnresolvedStructTracker::Reset(this);
+#endif // USE_CIRCULAR_DEPENDENCY_LOAD_DEFERRING
+}
+
 bool FLinkerLoad::HasPerformedFullExportResolvePass()
 {
 #if USE_DEFERRED_DEPENDENCY_CHECK_VERIFICATION_TESTS
@@ -1455,16 +1558,39 @@ bool FLinkerLoad::HasPerformedFullExportResolvePass()
 	
 }
 
-void FLinkerLoad::ResetDeferredLoadingState()
+UObject* FLinkerLoad::RequestPlaceholderValue(UClass* ObjectType, const TCHAR* ObjectPath)
 {
-#if USE_CIRCULAR_DEPENDENCY_LOAD_DEFERRING
-	DeferredCDOIndex = INDEX_NONE;
-	bForceBlueprintFinalization = false;
-	ResolvingDeferredPlaceholder = nullptr;
-	LoadFlags &= ~(LOAD_DeferDependencyLoads);
+#if !USE_CIRCULAR_DEPENDENCY_LOAD_DEFERRING
+	return nullptr;
+#else
+	FLinkerPlaceholderBase* Placeholder = nullptr;
 
-	FResolvingExportTracker::Get().Reset(this);
-	FUnresolvedStructTracker::Reset(this);
+	if (FBlueprintSupport::UseDeferredDependencyLoading() && (LoadFlags & LOAD_DeferDependencyLoads))
+	{
+		const FName ObjId(ObjectPath);
+		if (FLinkerPlaceholderBase** PlaceholderPtr = ImportPlaceholders.Find(ObjId))
+		{
+			Placeholder = *PlaceholderPtr;
+		}
+		// right now we only support external parties requesting CLASS placeholders;
+		// if there is a scenario where they're, through a different ObjectType, 
+		// loading another Blueprint package when they shouldn't, then we need to 
+		// handle that here as well
+		else if (ObjectType->IsChildOf<UClass>())
+		{
+			const FString ObjectPathStr(ObjectPath);
+			// we don't need placeholders for native object references (the 
+			// calling code should properly handle null return values)
+			if (!FPackageName::IsScriptPackage(ObjectPathStr))
+			{
+				const FString ObjectName = FPackageName::ObjectPathToObjectName(ObjectPathStr);
+				Placeholder = MakeImportPlaceholder<ULinkerPlaceholderClass>(LinkerRoot, *ObjectName);
+				ImportPlaceholders.Add(ObjId, Placeholder);
+			}
+		}
+	}
+
+	return Placeholder ? Placeholder->GetPlaceholderAsUObject() : nullptr;
 #endif // USE_CIRCULAR_DEPENDENCY_LOAD_DEFERRING
 }
 
diff --git a/Engine/Source/Runtime/CoreUObject/Private/UObject/Class.cpp b/Engine/Source/Runtime/CoreUObject/Private/UObject/Class.cpp
--- a/Engine/Source/Runtime/CoreUObject/Private/UObject/Class.cpp
+++ b/Engine/Source/Runtime/CoreUObject/Private/UObject/Class.cpp
@@ -4444,7 +4444,7 @@ void GetPrivateStaticClassBody(
 	}
 	else
 	{
-		ReturnClass = (UClass*)GUObjectAllocator.AllocateUObject(sizeof(UDynamicClass), ALIGNOF(UDynamicClass), true);
+		ReturnClass = (UClass*)GUObjectAllocator.AllocateUObject(sizeof(UDynamicClass), ALIGNOF(UDynamicClass), GIsInitialLoad);
 		ReturnClass = ::new (ReturnClass)
 			UDynamicClass
 			(
@@ -4827,6 +4827,42 @@ UScriptStruct* TBaseStructure<FFallbackStruct>::Get()
 	return ScriptStruct;
 }
 
+UScriptStruct* TBaseStructure<FFloatRangeBound>::Get()
+{
+	static auto ScriptStruct = StaticGetBaseStructureInternal(TEXT("FloatRangeBound"));
+	return ScriptStruct;
+}
+
+UScriptStruct* TBaseStructure<FFloatRange>::Get()
+{
+	static auto ScriptStruct = StaticGetBaseStructureInternal(TEXT("FloatRange"));
+	return ScriptStruct;
+}
+
+UScriptStruct* TBaseStructure<FInt32RangeBound>::Get()
+{
+	static auto ScriptStruct = StaticGetBaseStructureInternal(TEXT("Int32RangeBound"));
+	return ScriptStruct;
+}
+
+UScriptStruct* TBaseStructure<FInt32Range>::Get()
+{
+	static auto ScriptStruct = StaticGetBaseStructureInternal(TEXT("Int32Range"));
+	return ScriptStruct;
+}
+
+UScriptStruct* TBaseStructure<FFloatInterval>::Get()
+{
+	static auto ScriptStruct = StaticGetBaseStructureInternal(TEXT("FloatInterval"));
+	return ScriptStruct;
+}
+
+UScriptStruct* TBaseStructure<FInt32Interval>::Get()
+{
+	static auto ScriptStruct = StaticGetBaseStructureInternal(TEXT("Int32Interval"));
+	return ScriptStruct;
+}
+
 
 IMPLEMENT_CORE_INTRINSIC_CLASS(UFunction, UStruct,
 	{
diff --git a/Engine/Source/Runtime/CoreUObject/Private/UObject/Linker.cpp b/Engine/Source/Runtime/CoreUObject/Private/UObject/Linker.cpp
--- a/Engine/Source/Runtime/CoreUObject/Private/UObject/Linker.cpp
+++ b/Engine/Source/Runtime/CoreUObject/Private/UObject/Linker.cpp
@@ -700,20 +700,6 @@ FLinkerLoad* GetPackageLinker
 
 		Result = FLinkerLoad::CreateLinker( InOuter, *NewFilename, LoadFlags );
 	}
-#if USE_CIRCULAR_DEPENDENCY_LOAD_DEFERRING
-	else if ( !Result->HaveImportsBeenVerified() && !(LoadFlags & (LOAD_NoVerify|LOAD_DeferDependencyLoads)) )
-	{
-		// if we're in here, then this Linker is most likely in FinalizeCreation() 
-		// earlier on in the stack (HasFinishedInitialization() is probably false
-		// at this point in time); since we've recursively requested this Linker
-		// (and we're about to return it), we want to make sure verification 
-		// (what FinalizeCreation() invokes) has completed - this is important 
-		// for cyclic dependencies because we'd rather all these other packages 
-		// are loaded up front, instead of during the serialization process 
-		// (when 'DeferredDependencyLoading' is on, and external loads are problematic)
-		Result->Verify();
-	}
-#endif // USE_CIRCULAR_DEPENDENCY_LOAD_DEFERRING
 
 	// Verify compatibility.
 	if (Result && CompatibleGuid && Result->Summary.Guid != *CompatibleGuid)
diff --git a/Engine/Source/Runtime/CoreUObject/Private/UObject/LinkerLoad.cpp b/Engine/Source/Runtime/CoreUObject/Private/UObject/LinkerLoad.cpp
--- a/Engine/Source/Runtime/CoreUObject/Private/UObject/LinkerLoad.cpp
+++ b/Engine/Source/Runtime/CoreUObject/Private/UObject/LinkerLoad.cpp
@@ -718,6 +718,7 @@ FLinkerLoad::ELinkerStatus FLinkerLoad::Tick( float InTimeLimit, bool bInUseTime
 FLinkerLoad::FLinkerLoad(UPackage* InParent, const TCHAR* InFilename, uint32 InLoadFlags )
 : FLinker(ELinkerType::Load, InParent, InFilename)
 , LoadFlags(InLoadFlags)
+, bHaveImportsBeenVerified(false)
 , bDynamicClassLinker(false)
 , Loader(nullptr)
 , AsyncRoot(nullptr)
@@ -738,7 +739,6 @@ FLinkerLoad::FLinkerLoad(UPackage* InParent, const TCHAR* InFilename, uint32 InL
 , IsTimeLimitExceededCallCount(0)
 , TimeLimit(0.0f)
 , TickStartTime(0.0)
-, VerifiedImportCount(0)
 , bFixupExportMapDone(false)
 #if WITH_EDITOR
 , bExportsDuplicatesFixed(false)
@@ -1885,20 +1885,7 @@ FLinkerLoad::ELinkerStatus FLinkerLoad::FinalizeCreation()
 //			UE_LOG(LogLinker, Log, TEXT("Found a user created pacakge (%s)"), *(FPaths::GetBaseFilename(Filename)));
 		}
 
-#if USE_CIRCULAR_DEPENDENCY_LOAD_DEFERRING
-		// we can't verify our imports if we're currently deferring dependency
-		// loads (or if we've explicitly requested this linker without it);
-		// with the LOAD_DeferDependencyLoads flag, this is most likely a 
-		// UserDefinedStruct package, which we need to load for some other Blueprint
-		// package (further up the stack), but at that Blueprint's request we 
-		// cannot spider out to load anything else - don't worry though, import
-		// verification happens as needed for CreateImport() during export 
-		// serialization, so it's not like this will be skipped completely (just
-		// deferred)
-		if( !(LoadFlags & (LOAD_NoVerify|LOAD_DeferDependencyLoads)) )
-#else 
-		if ( !(LoadFlags & LOAD_NoVerify) )
-#endif // USE_CIRCULAR_DEPENDENCY_LOAD_DEFERRING
+		if( !(LoadFlags & LOAD_NoVerify))
 		{
 			Verify();
 		}
@@ -1988,21 +1975,16 @@ void FLinkerLoad::Verify()
 {
 	if(!FApp::IsGame() || GIsEditor || IsRunningCommandlet())
 	{
-		if (!HaveImportsBeenVerified())
+		if (!bHaveImportsBeenVerified)
 		{
 #if WITH_EDITOR
 			FScopedSlowTask SlowTask(Summary.ImportCount, NSLOCTEXT("Core", "LinkerLoad_Imports", "Loading Imports"), ShouldReportProgress());
 #endif
 			// Validate all imports and map them to their remote linkers.
-			while (VerifiedImportCount < Summary.ImportCount)
+			for (int32 ImportIndex = 0; ImportIndex < Summary.ImportCount; ImportIndex++)
 			{
-				// use an incrementing class-scoped counter in order to make 
-				// this function reentrant (so we can continue verifying from
-				// where we last left off, up the callstack)
-				const int32 ImportIndex = VerifiedImportCount;
-				++VerifiedImportCount;
-
 				FObjectImport& Import = ImportMap[ImportIndex];
+
 #if WITH_EDITOR
 				SlowTask.EnterProgressFrame(1, FText::Format(NSLOCTEXT("Core", "LinkerLoad_LoadingImportName", "Loading Import '{0}'"), FText::FromString(Import.ObjectName.ToString())));
 #endif
@@ -2010,6 +1992,7 @@ void FLinkerLoad::Verify()
 			}
 		}
 	}
+	bHaveImportsBeenVerified = true;
 }
 
 FName FLinkerLoad::GetExportClassPackage( int32 i )
@@ -2327,7 +2310,7 @@ FLinkerLoad::EVerifyResult FLinkerLoad::VerifyImport(int32 ImportIndex)
 #if WITH_EDITOR
 				if( GIsEditor && !IsRunningCommandlet())
 				{
-					bSupressLinkerError = IsSuppressableBlueprintImportError(Import);
+					bSupressLinkerError = IsSuppressableBlueprintImportError(ImportIndex);
 					if (!bSupressLinkerError)
 					{
 						FDeferredMessageLog LoadErrors(NAME_LoadErrors);
@@ -3556,7 +3539,7 @@ UObject* FLinkerLoad::CreateExport( int32 Index )
 				// SuperStruct needs to be fully linked so that UStruct::Link will have access to UObject::SuperStruct->PropertySize. 
 				// There are other attempts to force our super struct to load, and I have not verified that they can all be removed
 				// in favor of this one:
-				if (!SuperStruct->HasAnyFlags(RF_LoadCompleted)
+				if (!SuperStruct->HasAnyFlags(RF_LoadCompleted | RF_Dynamic)
 					&& !SuperStruct->IsNative()
 					&& SuperStruct->GetLinker()
 					&& Export.SuperIndex.IsImport())
diff --git a/Engine/Source/Runtime/CoreUObject/Private/UObject/LinkerPlaceholderBase.cpp b/Engine/Source/Runtime/CoreUObject/Private/UObject/LinkerPlaceholderBase.cpp
--- a/Engine/Source/Runtime/CoreUObject/Private/UObject/LinkerPlaceholderBase.cpp
+++ b/Engine/Source/Runtime/CoreUObject/Private/UObject/LinkerPlaceholderBase.cpp
@@ -4,48 +4,29 @@
 #include "LinkerPlaceholderBase.h"
 #include "Blueprint/BlueprintSupport.h"
 
+#if USE_DEFERRED_DEPENDENCY_CHECK_VERIFICATION_TESTS
+	#define DEFERRED_DEPENDENCY_ENSURE(EnsueExpr) ensure(EnsueExpr)
+#else  // USE_DEFERRED_DEPENDENCY_CHECK_VERIFICATION_TESTS
+	#define DEFERRED_DEPENDENCY_ENSURE(EnsueExpr)  (EnsueExpr)
+#endif // USE_DEFERRED_DEPENDENCY_CHECK_VERIFICATION_TESTS
+
 /*******************************************************************************
- *FPlaceholderContainerTracker
+ * LinkerPlaceholderObjectImpl
  ******************************************************************************/
 
 /**  */
 struct FPlaceholderContainerTracker : TThreadSingleton<FPlaceholderContainerTracker>
-{
+{	
 	TArray<UObject*> PerspectiveReferencerStack;
+	// as far as I can tell, structs are going to be the only bridging point 
+	// between property ownership
+	TArray<const UStructProperty*> IntermediatePropertyStack; 
 };
 
-//------------------------------------------------------------------------------
-FScopedPlaceholderContainerTracker::FScopedPlaceholderContainerTracker(UObject* InPlaceholderContainerCandidate)
-	: PlaceholderReferencerCandidate(InPlaceholderContainerCandidate)
-{
-	FPlaceholderContainerTracker::Get().PerspectiveReferencerStack.Add(InPlaceholderContainerCandidate);
-}
-
-//------------------------------------------------------------------------------
-FScopedPlaceholderContainerTracker::~FScopedPlaceholderContainerTracker()
-{
-	UObject* StackTop = FPlaceholderContainerTracker::Get().PerspectiveReferencerStack.Pop();
-#if USE_DEFERRED_DEPENDENCY_CHECK_VERIFICATION_TESTS
-	check(StackTop == PlaceholderReferencerCandidate);
-#endif // USE_DEFERRED_DEPENDENCY_CHECK_VERIFICATION_TESTS
-}
-
-/*******************************************************************************
- * FLinkerPlaceholderBase
- ******************************************************************************/
-
-namespace LinkerPlaceholderObjectImpl
+/**  */
+class FLinkerPlaceholderObjectImpl
 {
-	/**
-	 * Traces the specified property up its outer chain, all the way up to the 
-	 * OwnerClass. Constructs a UProperty path along the way, so that we can use
-	 * traverse a property tree down, from the class owner down.
-	 * 
-	 * @param  LeafProperty	The property you want a path for.
-	 * @param  ChainOut		Output array that tracks the specified property back up its outer chain (first entry should be what was passed in LeafProperty). 
-	 */
-	static void BuildPropertyChain(const UProperty* LeafProperty, TArray<const UProperty*>& ChainOut);
-
+public:
 	/**
 	 * A recursive method that replaces all leaf references to this object with 
 	 * the supplied ReplacementValue.
@@ -65,29 +46,20 @@ namespace LinkerPlaceholderObjectImpl
 	 * @return The number of references that were replaced.
 	 */
 	static int32 ResolvePlaceholderValues(const TArray<const UProperty*>& PropertyChain, int32 ChainIndex, uint8* ValueAddress, UObject* OldValue, UObject* ReplacementValue);
-}
-
-//------------------------------------------------------------------------------
-static void LinkerPlaceholderObjectImpl::BuildPropertyChain(const UProperty* LeafProperty, TArray<const UProperty*>& ChainOut)
-{
-	ChainOut.Empty();
-	ChainOut.Add(LeafProperty);
-
-	UClass* ClassOwner = LeafProperty->GetOwnerClass();
-
-	UObject* PropertyOuter = LeafProperty->GetOuter();
-	while ((PropertyOuter != nullptr) && (PropertyOuter != ClassOwner))
-	{
-		if (const UProperty* PropertyOwner = Cast<const UProperty>(PropertyOuter))
-		{
-			ChainOut.Add(PropertyOwner);
-		}
-		PropertyOuter = PropertyOuter->GetOuter();
-	}
-}
+	
+	/**
+	 * Uses the current FPlaceholderContainerTracker::PerspectiveReferencerStack
+	 * to search for a viable placeholder container (expected that it will be 
+	 * the top of the stack).
+	 * 
+	 * @param  PropertyChainRef    Defines the nested property path through a UObject, where the end leaf property is one left referencing a placeholder.
+	 * @return The UObject instance that is assumably referencing a placeholder (null if one couldn't be found)
+	 */
+	static UObject* FindPlaceholderContainer(const FLinkerPlaceholderBase::FPlaceholderValuePropertyPath& PropertyChainRef);
+};
 
 //------------------------------------------------------------------------------
-static int32 LinkerPlaceholderObjectImpl::ResolvePlaceholderValues(const TArray<const UProperty*>& PropertyChain, int32 ChainIndex, uint8* ValueAddress, UObject* OldValue, UObject* ReplacementValue)
+int32 FLinkerPlaceholderObjectImpl::ResolvePlaceholderValues(const TArray<const UProperty*>& PropertyChain, int32 ChainIndex, uint8* ValueAddress, UObject* OldValue, UObject* ReplacementValue)
 {
 	int32 ReplacementCount = 0;
 
@@ -101,7 +73,9 @@ static int32 LinkerPlaceholderObjectImpl::ResolvePlaceholderValues(const TArray<
 #endif // USE_DEFERRED_DEPENDENCY_CHECK_VERIFICATION_TESTS
 
 			const UObjectProperty* ReferencingProperty = (const UObjectProperty*)Property;
-			if (ReferencingProperty->GetObjectPropertyValue(ValueAddress) == OldValue)
+
+			UObject* CurrentValue = ReferencingProperty->GetObjectPropertyValue(ValueAddress);
+			if (CurrentValue == OldValue)
 			{
 				// @TODO: use an FArchiver with ReferencingProperty->SerializeItem() 
 				//        so that we can utilize CheckValidObject()
@@ -135,10 +109,7 @@ static int32 LinkerPlaceholderObjectImpl::ResolvePlaceholderValues(const TArray<
 		else
 		{
 			const UProperty* NextProperty = PropertyChain[PropertyIndex - 1];
-#if USE_DEFERRED_DEPENDENCY_CHECK_VERIFICATION_TESTS
-			check(NextProperty->GetOuter() == Property);
-#endif // USE_DEFERRED_DEPENDENCY_CHECK_VERIFICATION_TESTS
-			ValueAddress = Property->ContainerPtrToValuePtr<uint8>(ValueAddress, /*ArrayIndex =*/0);
+			ValueAddress = NextProperty->ContainerPtrToValuePtr<uint8>(ValueAddress, /*ArrayIndex =*/0);
 		}
 	}
 
@@ -146,59 +117,211 @@ static int32 LinkerPlaceholderObjectImpl::ResolvePlaceholderValues(const TArray<
 }
 
 //------------------------------------------------------------------------------
-FLinkerPlaceholderBase::FLinkerPlaceholderBase()
-	: bResolveWasInvoked(false)
+UObject* FLinkerPlaceholderObjectImpl::FindPlaceholderContainer(const FLinkerPlaceholderBase::FPlaceholderValuePropertyPath& PropertyChainRef)
 {
+	UObject* ContainerObj = nullptr;
+	TArray<UObject*>& PossibleReferencers = FPlaceholderContainerTracker::Get().PerspectiveReferencerStack;
+
+	UClass* OwnerClass = PropertyChainRef.GetOwnerClass();
+	if ((OwnerClass != nullptr) && (PossibleReferencers.Num() > 0))
+	{
+		UObject* ReferencerCandidate = PossibleReferencers.Top();
+		// we expect that the current object we're looking for (the one we're 
+		// serializing) is at the top of the stack
+		if (DEFERRED_DEPENDENCY_ENSURE(ReferencerCandidate->GetClass()->IsChildOf(OwnerClass)))
+		{
+			ContainerObj = ReferencerCandidate;
+		}
+		else
+		{
+			// if it's not the top element, then iterate backwards because this 
+			// is meant to act as a stack, where the last entry is most likely 
+			// the one we're looking for
+			for (int32 CandidateIndex = PossibleReferencers.Num() - 2; CandidateIndex >= 0; --CandidateIndex)
+			{
+				ReferencerCandidate = PossibleReferencers[CandidateIndex];
+
+				UClass* CandidateClass = ReferencerCandidate->GetClass();
+				if (CandidateClass->IsChildOf(OwnerClass))
+				{
+					ContainerObj = ReferencerCandidate;
+					break;
+				}
+			}
+		}
+	}
+	return ContainerObj;
 }
 
+/*******************************************************************************
+ * FPlaceholderContainerTracker / FScopedPlaceholderPropertyTracker
+ ******************************************************************************/
+
 //------------------------------------------------------------------------------
-FLinkerPlaceholderBase::~FLinkerPlaceholderBase()
+FScopedPlaceholderContainerTracker::FScopedPlaceholderContainerTracker(UObject* InPlaceholderContainerCandidate)
+	: PlaceholderReferencerCandidate(InPlaceholderContainerCandidate)
 {
+	FPlaceholderContainerTracker::Get().PerspectiveReferencerStack.Add(InPlaceholderContainerCandidate);
+}
+
+//------------------------------------------------------------------------------
+FScopedPlaceholderContainerTracker::~FScopedPlaceholderContainerTracker()
+{
+	UObject* StackTop = FPlaceholderContainerTracker::Get().PerspectiveReferencerStack.Pop();
 #if USE_DEFERRED_DEPENDENCY_CHECK_VERIFICATION_TESTS
-	check(!HasKnownReferences());
+	check(StackTop == PlaceholderReferencerCandidate);
 #endif // USE_DEFERRED_DEPENDENCY_CHECK_VERIFICATION_TESTS
 }
 
 //------------------------------------------------------------------------------
-bool FLinkerPlaceholderBase::AddReferencingPropertyValue(const UObjectProperty* ReferencingProperty, void* DataPtr)
+FScopedPlaceholderPropertyTracker::FScopedPlaceholderPropertyTracker(const UStructProperty* InIntermediateProperty)
+	: IntermediateProperty(nullptr) // leave null, as a sentinel value (implying that PerspectiveReferencerStack was empty)
 {
-	TArray<UObject*>& PossibleReferencers = FPlaceholderContainerTracker::Get().PerspectiveReferencerStack;
+	FPlaceholderContainerTracker& ContainerRepository = FPlaceholderContainerTracker::Get();
+	if (ContainerRepository.PerspectiveReferencerStack.Num() > 0)
+	{
+		IntermediateProperty = InIntermediateProperty;
+		ContainerRepository.IntermediatePropertyStack.Add(InIntermediateProperty);
+	}
+	// else, if there's nothing in the PerspectiveReferencerStack, then caching 
+	// a property here would be pointless (the whole point of this is to be able 
+	// to use this to look up the referencing object)
+}
 
-	UObject* FoundReferencer = nullptr;
-	// iterate backwards because this is meant to act as a stack, where the last
-	// entry is most likely the one we're looking for
-	for (int32 CandidateIndex = PossibleReferencers.Num() - 1; CandidateIndex >= 0; --CandidateIndex)
+//------------------------------------------------------------------------------
+FScopedPlaceholderPropertyTracker::~FScopedPlaceholderPropertyTracker()
+{
+	FPlaceholderContainerTracker& ContainerRepository = FPlaceholderContainerTracker::Get();
+	if (IntermediateProperty != nullptr)
+	{
+		const UStructProperty* StackTop = ContainerRepository.IntermediatePropertyStack.Pop();
+#if USE_DEFERRED_DEPENDENCY_CHECK_VERIFICATION_TESTS
+		check(StackTop == IntermediateProperty);
+	}
+	else
 	{
-		UObject* ReferencerCandidate = PossibleReferencers[CandidateIndex];
+		check(ContainerRepository.IntermediatePropertyStack.Num()  == 0);
+		check(ContainerRepository.PerspectiveReferencerStack.Num() == 0);
+#endif // USE_DEFERRED_DEPENDENCY_CHECK_VERIFICATION_TESTS
+	}
+}
 
-		UClass* CandidateClass = ReferencerCandidate->GetClass();
-		UClass* PropOwnerClass = ReferencingProperty->GetOwnerClass();
+/*******************************************************************************
+ * FLinkerPlaceholderBase::FPlaceholderValuePropertyPath
+ ******************************************************************************/
+
+//------------------------------------------------------------------------------
+FLinkerPlaceholderBase::FPlaceholderValuePropertyPath::FPlaceholderValuePropertyPath(const UProperty* ReferencingProperty)
+{
+	PropertyChain.Add(ReferencingProperty);
+	const UObject* PropertyOuter = ReferencingProperty->GetOuter();
+	
+	const TArray<const UStructProperty*>& StructPropertyStack = FPlaceholderContainerTracker::Get().IntermediatePropertyStack;
+	int32 StructStackIndex = StructPropertyStack.Num() - 1; // "top" of the array is the last element
 
-		if (CandidateClass->IsChildOf(PropOwnerClass))
+	while (PropertyOuter && !PropertyOuter->GetClass()->IsChildOf<UClass>())
+	{
+		// handle nested properties (like array members)
+		if (const UProperty* PropertyOwner = Cast<const UProperty>(PropertyOuter))
 		{
-			FoundReferencer = ReferencerCandidate;
-			break;
+			PropertyChain.Add(PropertyOwner);
 		}
+		// handle nested struct properties (use the FPlaceholderContainerTracker::IntermediatePropertyStack
+		// to help trace the property path)
+		else if (const UScriptStruct* StructOwner = Cast<const UScriptStruct>(PropertyOuter))
+		{
+			if (DEFERRED_DEPENDENCY_ENSURE(StructStackIndex != INDEX_NONE))
+			{
+				// we expect the top struct property to be the one we're currently serializing
+				const UStructProperty* SerializingStructProp = StructPropertyStack[StructStackIndex];
+				if (DEFERRED_DEPENDENCY_ENSURE(SerializingStructProp->Struct->IsChildOf(StructOwner)))
+				{
+					PropertyOuter = SerializingStructProp;
+					PropertyChain.Add(SerializingStructProp);
+				}
+				else
+				{
+#if USE_DEFERRED_DEPENDENCY_CHECK_VERIFICATION_TESTS
+					PropertyChain.Empty(); // invalidate this FPlaceholderValuePropertyPath
+					checkf(false, TEXT("Looks like we couldn't reliably determine the object that a placeholder value should belong to - are we missing a FScopedPlaceholderPropertyTracker?"));
+#endif // USE_DEFERRED_DEPENDENCY_CHECK_VERIFICATION_TEST
+					break;
+				}
+				--StructStackIndex;
+			}
+		}
+		PropertyOuter = PropertyOuter->GetOuter();
 	}
 
 #if USE_DEFERRED_DEPENDENCY_CHECK_VERIFICATION_TESTS
-	check(ReferencingProperty->GetObjectPropertyValue(DataPtr) == GetPlaceholderAsUObject());
-	check(FoundReferencer != nullptr);
+	if (!DEFERRED_DEPENDENCY_ENSURE(PropertyOuter != nullptr))
+	{
+		PropertyChain.Empty(); // invalidate this FPlaceholderValuePropertyPath
+	}
+#endif // USE_DEFERRED_DEPENDENCY_CHECK_VERIFICATION_TEST
+}
+
+//------------------------------------------------------------------------------
+bool FLinkerPlaceholderBase::FPlaceholderValuePropertyPath::IsValid() const
+{
+	return (PropertyChain.Num() > 0) && PropertyChain[0]->IsA<UObjectProperty>() && PropertyChain.Last()->GetOuter()->IsA<UClass>();
+}
+
+//------------------------------------------------------------------------------
+UClass* FLinkerPlaceholderBase::FPlaceholderValuePropertyPath::GetOwnerClass() const
+{
+	return (PropertyChain.Num() > 0) ? Cast<UClass>(PropertyChain.Last()->GetOuter()) : nullptr;
+}
+
+//------------------------------------------------------------------------------
+int32 FLinkerPlaceholderBase::FPlaceholderValuePropertyPath::Resolve(FLinkerPlaceholderBase* Placeholder, UObject* Replacement, UObject* Container) const
+{
+	const UProperty* OutermostProperty = PropertyChain.Last();
+
+#if USE_DEFERRED_DEPENDENCY_CHECK_VERIFICATION_TESTS
+	UClass* OwnerClass = OutermostProperty->GetOwnerClass();
+	check(OwnerClass && Container->IsA(OwnerClass));
+#endif 
+
+	uint8* OutermostAddress = OutermostProperty->ContainerPtrToValuePtr<uint8>((uint8*)Container, /*ArrayIndex =*/0);
+	return FLinkerPlaceholderObjectImpl::ResolvePlaceholderValues(PropertyChain, PropertyChain.Num() - 1, OutermostAddress, Placeholder->GetPlaceholderAsUObject(), Replacement);
+}
+ 
+/*******************************************************************************
+ * FLinkerPlaceholderBase
+ ******************************************************************************/
+
+//------------------------------------------------------------------------------
+FLinkerPlaceholderBase::FLinkerPlaceholderBase()
+	: bResolveWasInvoked(false)
+{
+}
+
+//------------------------------------------------------------------------------
+FLinkerPlaceholderBase::~FLinkerPlaceholderBase()
+{
+#if USE_DEFERRED_DEPENDENCY_CHECK_VERIFICATION_TESTS
+	check(!HasKnownReferences());
 #endif // USE_DEFERRED_DEPENDENCY_CHECK_VERIFICATION_TESTS
+}
+
+//------------------------------------------------------------------------------
+bool FLinkerPlaceholderBase::AddReferencingPropertyValue(const UObjectProperty* ReferencingProperty, void* DataPtr)
+{
+	FPlaceholderValuePropertyPath PropertyChain(ReferencingProperty);
+	UObject* ReferencingContainer = FLinkerPlaceholderObjectImpl::FindPlaceholderContainer(PropertyChain);
 
-	if (FoundReferencer != nullptr)
-	{
 #if USE_DEFERRED_DEPENDENCY_CHECK_VERIFICATION_TESTS
-		// @TODO: verify that DataPtr belongs to FoundReferencer
-// 		uint8* ContainerStart = (uint8*)FoundReferencer;
-// 		uint8* ContainerEnd   = ContainerStart + FoundReferencer->GetClass()->GetStructureSize();
-// 		// check that we picked the right container object, and that DataPtr exists within it 
-// 		check(DataPtr >= ContainerStart && DataPtr <= ContainerEnd);
+	check(ReferencingProperty->GetObjectPropertyValue(DataPtr) == GetPlaceholderAsUObject());
+	check(PropertyChain.IsValid());
+	checkf(ReferencingContainer != nullptr, TEXT("We couldn't reliably determine the object that a placeholder value should belong to - are we missing a FScopedPlaceholderContainerTracker somewhere?"));
 #endif // USE_DEFERRED_DEPENDENCY_CHECK_VERIFICATION_TESTS
 
-		ReferencingContainers.FindOrAdd(FoundReferencer).Add(ReferencingProperty);
+	if (ReferencingContainer != nullptr)
+	{
+		ReferencingContainers.FindOrAdd(ReferencingContainer).Add(PropertyChain);
 	}
-	return (FoundReferencer != nullptr);
+	return (ReferencingContainer != nullptr);
 }
 
 //------------------------------------------------------------------------------
@@ -240,7 +363,6 @@ int32 FLinkerPlaceholderBase::ResolvePlaceholderPropertyValues(UObject* NewObjec
 {
 	int32 ResolvedTotal = 0;
 
-	UObject* ThisAsUObject = GetPlaceholderAsUObject();
 	for (auto& ReferencingPair : ReferencingContainers)
 	{
 		TWeakObjectPtr<UObject> ContainerPtr = ReferencingPair.Key;
@@ -250,18 +372,13 @@ int32 FLinkerPlaceholderBase::ResolvePlaceholderPropertyValues(UObject* NewObjec
 		}
 		UObject* Container = ContainerPtr.Get();
 
-		for (const UObjectProperty* Property : ReferencingPair.Value)
+		for (const FPlaceholderValuePropertyPath& PropertyRef : ReferencingPair.Value)
 		{
 #if USE_DEFERRED_DEPENDENCY_CHECK_VERIFICATION_TESTS
-			check(Container->GetClass()->IsChildOf(Property->GetOwnerClass()));
+			check(Container->GetClass()->IsChildOf(PropertyRef.GetOwnerClass()));
 #endif // USE_DEFERRED_DEPENDENCY_CHECK_VERIFICATION_TESTS
 
-			TArray<const UProperty*> PropertyChain;
-			LinkerPlaceholderObjectImpl::BuildPropertyChain(Property, PropertyChain);
-			const UProperty* OutermostProperty = PropertyChain.Last();
-
-			uint8* OutermostAddress = OutermostProperty->ContainerPtrToValuePtr<uint8>((uint8*)Container, /*ArrayIndex =*/0);
-			int32 ResolvedCount = LinkerPlaceholderObjectImpl::ResolvePlaceholderValues(PropertyChain, PropertyChain.Num() - 1, OutermostAddress, ThisAsUObject, NewObjectValue);
+			int32 ResolvedCount = PropertyRef.Resolve(this, NewObjectValue, Container);
 			ResolvedTotal += ResolvedCount;
 
 #if USE_DEFERRED_DEPENDENCY_CHECK_VERIFICATION_TESTS
@@ -378,3 +495,5 @@ int32 TLinkerImportPlaceholder<UFunction>::ResolvePropertyReferences(UFunction*
 	ReferencingProperties.Empty();
 	return ReplacementCount;
 }
+
+#undef DEFERRED_DEPENDENCY_ENSURE
diff --git a/Engine/Source/Runtime/CoreUObject/Private/UObject/LinkerPlaceholderBase.h b/Engine/Source/Runtime/CoreUObject/Private/UObject/LinkerPlaceholderBase.h
--- a/Engine/Source/Runtime/CoreUObject/Private/UObject/LinkerPlaceholderBase.h
+++ b/Engine/Source/Runtime/CoreUObject/Private/UObject/LinkerPlaceholderBase.h
@@ -7,8 +7,9 @@ class UProperty;
 class UObjectProperty;
 
 /*******************************************************************************
- *FPlaceholderContainerTracker
+ * FPlaceholderContainerTracker / FScopedPlaceholderPropertyTracker
  ******************************************************************************/
+
 /** 
  * To track placeholder property values, we need to know the root container 
  * instance that is set with the placeholder value (so we can reset it later). 
@@ -26,6 +27,23 @@ private:
 	UObject* PlaceholderReferencerCandidate;
 };
 
+/** 
+ * Sometimes using FScopedPlaceholderContainerTracker above is not enough; we
+ * could be working with a series of nested structs, where the owning object is 
+ * somewhere up the chain, but the lower properties have no idea who that is. 
+ * This provides us context, so we can navigate backwards and truly determine 
+ * the object a property is writing to.
+ */
+struct FScopedPlaceholderPropertyTracker
+{
+public:
+	 FScopedPlaceholderPropertyTracker(const UStructProperty* IntermediateProperty);
+	~FScopedPlaceholderPropertyTracker();
+
+private:
+	const UStructProperty* IntermediateProperty;
+};
+
 /*******************************************************************************
  * TLinkerPlaceholderBase<>
  ******************************************************************************/
@@ -117,8 +135,77 @@ private:
 	/** Used to catch references that are added after we've already resolved all references */
 	bool bResolveWasInvoked;
 
+	/** 
+	 * Handily tracks a series of nested properties through an object's class, 
+	 * specifically for scenarios where the leaf property's value is referencing
+	 * a linker-placeholder object. Used so we can later come back and easily  
+	 * resolve (replace) the placeholder value with a legitamate object.
+	 */
+	struct FPlaceholderValuePropertyPath
+	{
+		FPlaceholderValuePropertyPath(const UProperty* ReferencingProperty);
+
+		/** 
+		 * Validates that the internal property path points to a UObjectProperty, 
+		 * and that the whole thing has a class owner. 
+		 */
+		bool IsValid() const;
+
+		/**  
+		 * Returns the outer class that seemingly owns the property path 
+		 * represented by this struct. 
+		 */
+		UClass* GetOwnerClass() const;
+
+		/**
+		 * Replaces the (placeholder) object value at the end of this property
+		 * chain with the specified Replacement object.
+		 * 
+		 * @param  Placeholder	The object value that you're seeking to replace.
+		 * @param  Replacement	The new object value for the property represented by this struct.
+		 * @param  Container    The object instance that you want the value changed for (within).
+		 * @return The number of values successfully replaced (could be multiple for array properties).
+		 */
+		int32 Resolve(FLinkerPlaceholderBase* Placeholder, UObject* Replacement, UObject* Container) const;
+
+	private:
+		/** Denotes the property hierarchy used to reach this leaf property that is referencing a placeholder*/
+		TArray<const UProperty*> PropertyChain;
+
+	public:
+		/** Support comparison functions that make this usable as a KeyValue for a TSet<> */
+		friend uint32 GetTypeHash(const FPlaceholderValuePropertyPath& PlaceholderPropertyRef)
+		{
+			uint32 Hash = 0;
+			for (const UProperty* Propererty : PlaceholderPropertyRef.PropertyChain)
+			{
+				Hash = HashCombine(Hash, GetTypeHash(Propererty));
+			}
+			return Hash;
+		}
+
+		/**  */
+		friend bool operator==(const FPlaceholderValuePropertyPath& Lhs, const FPlaceholderValuePropertyPath& Rhs)
+		{
+			if (Lhs.PropertyChain.Num() == Rhs.PropertyChain.Num())
+			{
+				for (int32 PropIndex = 0; PropIndex < Lhs.PropertyChain.Num(); ++PropIndex)
+				{
+					if (Lhs.PropertyChain[PropIndex] != Rhs.PropertyChain[PropIndex])
+					{
+						return false;
+					}
+				}
+				return true;
+			}
+			return false;
+		};
+	};
+	friend class FLinkerPlaceholderObjectImpl; // for access to FPlaceholderValuePropertyPath
+	typedef TSet<FPlaceholderValuePropertyPath> FReferencingPropertySet;
+
 	/** Tracks container objects that have property values set to reference this placeholder (references that need to be replaced later) */
-	TMap< TWeakObjectPtr<UObject>, TSet<const UObjectProperty*> > ReferencingContainers;
+	TMap< TWeakObjectPtr<UObject>, FReferencingPropertySet > ReferencingContainers;
 };
 
 /*******************************************************************************
diff --git a/Engine/Source/Runtime/CoreUObject/Private/UObject/LinkerPlaceholderClass.h b/Engine/Source/Runtime/CoreUObject/Private/UObject/LinkerPlaceholderClass.h
--- a/Engine/Source/Runtime/CoreUObject/Private/UObject/LinkerPlaceholderClass.h
+++ b/Engine/Source/Runtime/CoreUObject/Private/UObject/LinkerPlaceholderClass.h
@@ -32,7 +32,7 @@ public:
 	// End of UField interface.
 
 	// FLinkerPlaceholderBase interface 
-	virtual UObject* GetPlaceholderAsUObject() override { return this; }
+	virtual UObject* GetPlaceholderAsUObject() override { return (UObject*)(this); }
 	// End of FLinkerPlaceholderBase interface
 }; 
 
diff --git a/Engine/Source/Runtime/CoreUObject/Private/UObject/LinkerPlaceholderExportObject.h b/Engine/Source/Runtime/CoreUObject/Private/UObject/LinkerPlaceholderExportObject.h
--- a/Engine/Source/Runtime/CoreUObject/Private/UObject/LinkerPlaceholderExportObject.h
+++ b/Engine/Source/Runtime/CoreUObject/Private/UObject/LinkerPlaceholderExportObject.h
@@ -25,6 +25,6 @@ public:
 	// End of UObject interface
 
 	// FLinkerPlaceholderBase interface
-	virtual UObject* GetPlaceholderAsUObject() override { return this; }
+	virtual UObject* GetPlaceholderAsUObject() override { return (UObject*)(this); }
 	// End of FLinkerPlaceholderBase interface
 }; 
diff --git a/Engine/Source/Runtime/CoreUObject/Private/UObject/LinkerPlaceholderFunction.h b/Engine/Source/Runtime/CoreUObject/Private/UObject/LinkerPlaceholderFunction.h
--- a/Engine/Source/Runtime/CoreUObject/Private/UObject/LinkerPlaceholderFunction.h
+++ b/Engine/Source/Runtime/CoreUObject/Private/UObject/LinkerPlaceholderFunction.h
@@ -16,6 +16,6 @@ public:
 	ULinkerPlaceholderFunction(const FObjectInitializer& ObjectInitializer);
 
 	// FLinkerPlaceholderBase interface 
-	virtual UObject* GetPlaceholderAsUObject() override { return this; }
+	virtual UObject* GetPlaceholderAsUObject() override { return (UObject*)(this); }
 	// End of FLinkerPlaceholderBase interface
 };
diff --git a/Engine/Source/Runtime/CoreUObject/Private/UObject/Property.cpp b/Engine/Source/Runtime/CoreUObject/Private/UObject/Property.cpp
--- a/Engine/Source/Runtime/CoreUObject/Private/UObject/Property.cpp
+++ b/Engine/Source/Runtime/CoreUObject/Private/UObject/Property.cpp
@@ -438,6 +438,7 @@ void UProperty::ExportCppDeclaration(FOutputDevice& Out, EExportedDeclaration::T
 		ExtendedTypeText = *ActualExtendedType;
 	}
 
+	const bool bCanHaveRef = 0 == (AdditionalExportCPPFlags & CPPF_NoRef);
 	const bool bCanHaveConst = 0 == (AdditionalExportCPPFlags & CPPF_NoConst);
 	if (!dynamic_cast<const UBoolProperty*>(this) && bCanHaveConst) // can't have const bitfields because then we cannot determine their offset and mask from the compiler
 	{
@@ -446,18 +447,20 @@ void UProperty::ExportCppDeclaration(FOutputDevice& Out, EExportedDeclaration::T
 		// export 'const' for parameters
 		const bool bIsConstParam   = bIsParameter && (HasAnyPropertyFlags(CPF_ConstParm) || (bIsInterfaceProp && !HasAllPropertyFlags(CPF_OutParm)));
 		const bool bIsOnConstClass = ObjectProp && ObjectProp->PropertyClass && ObjectProp->PropertyClass->HasAnyClassFlags(CLASS_Const);
-		if (bIsConstParam || bIsOnConstClass)
+		const bool bShouldHaveRef = bCanHaveRef && HasAnyPropertyFlags(CPF_OutParm | CPF_ReferenceParm);
+
+		const bool bConstAtTheBeginning = bIsOnConstClass || (bIsConstParam && !bShouldHaveRef);
+		if (bConstAtTheBeginning)
 		{
 			TypeText = FString::Printf(TEXT("const %s"), *TypeText);
 		}
 
-		if (DeclarationType == EExportedDeclaration::Member)
+		const UClass* const MyPotentialConstClass = (DeclarationType == EExportedDeclaration::Member) ? dynamic_cast<UClass*>(GetOuter()) : nullptr;
+		const bool bFromConstClass = MyPotentialConstClass && MyPotentialConstClass->HasAnyClassFlags(CLASS_Const);
+		const bool bConstAtTheEnd = bFromConstClass || (bIsConstParam && bShouldHaveRef);
+		if (bConstAtTheEnd)
 		{
-			UClass* MyClass = dynamic_cast<UClass*>(GetOuter());
-			if (MyClass && MyClass->HasAnyClassFlags(CLASS_Const))
-			{
-				ExtendedTypeText += TEXT(" const");
-			}
+			ExtendedTypeText += TEXT(" const");
 		}
 	}
 
@@ -486,7 +489,6 @@ void UProperty::ExportCppDeclaration(FOutputDevice& Out, EExportedDeclaration::T
 		}
 	}
 
-	const bool bCanHaveRef = 0 == (AdditionalExportCPPFlags & CPPF_NoRef);
 	if(auto BoolProperty = dynamic_cast<const UBoolProperty*>(this) )
 	{
 		// if this is a member variable, export it as a bitfield
diff --git a/Engine/Source/Runtime/CoreUObject/Private/UObject/PropertyBaseObject.cpp b/Engine/Source/Runtime/CoreUObject/Private/UObject/PropertyBaseObject.cpp
--- a/Engine/Source/Runtime/CoreUObject/Private/UObject/PropertyBaseObject.cpp
+++ b/Engine/Source/Runtime/CoreUObject/Private/UObject/PropertyBaseObject.cpp
@@ -4,6 +4,7 @@
 #include "PropertyHelper.h"
 #include "LinkerPlaceholderClass.h"
 #include "LinkerPlaceholderExportObject.h"
+#include "BlueprintSupport.h" // for IsInBlueprintPackage()
 
 /*-----------------------------------------------------------------------------
 	UObjectPropertyBase.
@@ -406,10 +407,27 @@ UObject* UObjectPropertyBase::FindImportedObject( const UProperty* Property, UOb
 		// If we still can't find it, try to load it. (Only try to load fully qualified names)
 		if(!Result && Dot)
 		{
-			uint32 LoadFlags = LOAD_NoWarn | LOAD_FindIfFail;
+#if USE_CIRCULAR_DEPENDENCY_LOAD_DEFERRING
+			FLinkerLoad* Linker = (OwnerObject != nullptr) ? OwnerObject->GetClass()->GetLinker() : nullptr;
+			const bool bDeferAssetImports = (Linker != nullptr) && (Linker->LoadFlags & LOAD_DeferDependencyLoads);
+
+			if (bDeferAssetImports)
+			{
+				Result = Linker->RequestPlaceholderValue(ObjectClass, Text);
+			}
+			
+			if (Result == nullptr)
+#endif // USE_CIRCULAR_DEPENDENCY_LOAD_DEFERRING
+			{
+				uint32 LoadFlags = LOAD_NoWarn | LOAD_FindIfFail;
+
+				UE_LOG(LogProperty, Verbose, TEXT("FindImportedObject is attempting to import [%s] (class = %s) with StaticLoadObject"), Text, *GetFullNameSafe(ObjectClass));
+				Result = StaticLoadObject(ObjectClass, NULL, Text, NULL, LoadFlags, NULL);
 
-			UE_LOG(LogProperty, Verbose, TEXT("FindImportedObject is attempting to import [%s] (class = %s) with StaticLoadObject"), Text, *GetFullNameSafe(ObjectClass));
-			Result = StaticLoadObject(ObjectClass, NULL, Text, NULL, LoadFlags, NULL);
+#if USE_DEFERRED_DEPENDENCY_CHECK_VERIFICATION_TESTS
+				check(!bDeferAssetImports || !Result || !FBlueprintSupport::IsInBlueprintPackage(Result));
+#endif // USE_DEFERRED_DEPENDENCY_CHECK_VERIFICATION_TESTS
+			}
 		}
 	}
 
diff --git a/Engine/Source/Runtime/CoreUObject/Private/UObject/PropertyClass.cpp b/Engine/Source/Runtime/CoreUObject/Private/UObject/PropertyClass.cpp
--- a/Engine/Source/Runtime/CoreUObject/Private/UObject/PropertyClass.cpp
+++ b/Engine/Source/Runtime/CoreUObject/Private/UObject/PropertyClass.cpp
@@ -76,16 +76,15 @@ const TCHAR* UClassProperty::ImportText_Internal( const TCHAR* Buffer, void* Dat
 	const TCHAR* Result = UObjectProperty::ImportText_Internal( Buffer, Data, PortFlags, Parent, ErrorText );
 	if( Result )
 	{
-		CheckValidObject(Data);
 		if (UClass* AssignedPropertyClass = dynamic_cast<UClass*>(GetObjectPropertyValue(Data)))
 		{
 #if USE_CIRCULAR_DEPENDENCY_LOAD_DEFERRING
-			FLinkerLoad* PropertyLinker = GetLinker();
-			bool const bIsDeferringValueLoad = ((PropertyLinker == nullptr) || (PropertyLinker->LoadFlags & LOAD_DeferDependencyLoads)) &&
-				Cast<ULinkerPlaceholderClass>(MetaClass);
+			FLinkerLoad* ObjectLinker = (Parent != nullptr) ? Parent->GetClass()->GetLinker() : GetLinker();
+			bool const bIsDeferringValueLoad = ((ObjectLinker == nullptr) || (ObjectLinker->LoadFlags & LOAD_DeferDependencyLoads)) &&
+				(Cast<ULinkerPlaceholderClass>(MetaClass) || Cast<ULinkerPlaceholderClass>(AssignedPropertyClass));
 
 #if USE_DEFERRED_DEPENDENCY_CHECK_VERIFICATION_TESTS
-			check(bIsDeferringValueLoad || !Cast<ULinkerPlaceholderClass>(MetaClass));
+			check( bIsDeferringValueLoad || !(Cast<ULinkerPlaceholderClass>(MetaClass) || Cast<ULinkerPlaceholderClass>(AssignedPropertyClass)) );
 #endif // USE_DEFERRED_DEPENDENCY_CHECK_VERIFICATION_TESTS
 
 #else  // USE_CIRCULAR_DEPENDENCY_LOAD_DEFERRING 
diff --git a/Engine/Source/Runtime/CoreUObject/Private/UObject/PropertyObject.cpp b/Engine/Source/Runtime/CoreUObject/Private/UObject/PropertyObject.cpp
--- a/Engine/Source/Runtime/CoreUObject/Private/UObject/PropertyObject.cpp
+++ b/Engine/Source/Runtime/CoreUObject/Private/UObject/PropertyObject.cpp
@@ -3,6 +3,7 @@
 #include "CoreUObjectPrivate.h"
 #include "LinkerPlaceholderExportObject.h"
 #include "LinkerPlaceholderClass.h"
+#include "BlueprintSupport.h" // for IsDeferredDependencyPlaceholder()
 
 /*-----------------------------------------------------------------------------
 	UObjectProperty.
@@ -61,6 +62,36 @@ void UObjectProperty::SerializeItem( FArchive& Ar, void* Value, void const* Defa
 	}
 }
 
+const TCHAR* UObjectProperty::ImportText_Internal(const TCHAR* Buffer, void* Data, int32 PortFlags, UObject* OwnerObject, FOutputDevice* ErrorText) const
+{
+	const TCHAR* Result = TUObjectPropertyBase<UObject*>::ImportText_Internal(Buffer, Data, PortFlags, OwnerObject, ErrorText);
+	if (Result)
+	{
+		CheckValidObject(Data);
+
+#if USE_CIRCULAR_DEPENDENCY_LOAD_DEFERRING
+		UObject* ObjectValue = GetObjectPropertyValue(Data);
+
+		if (ULinkerPlaceholderClass* PlaceholderClass = Cast<ULinkerPlaceholderClass>(ObjectValue))
+		{
+			// we use this tracker mechanism to help record the instance that is
+			// referencing the placeholder (so we can replace it later on fixup)
+			FScopedPlaceholderContainerTracker ImportingObjTracker(OwnerObject);
+
+			PlaceholderClass->AddReferencingPropertyValue(this, Data);
+		}
+#if USE_DEFERRED_DEPENDENCY_CHECK_VERIFICATION_TESTS
+		else
+		{
+			// as far as we know, ULinkerPlaceholderClass is the only type we have to handle through ImportText()
+			check(!FBlueprintSupport::IsDeferredDependencyPlaceholder(ObjectValue));
+		}
+#endif // USE_DEFERRED_DEPENDENCY_CHECK_VERIFICATION_TESTS
+#endif // USE_CIRCULAR_DEPENDENCY_LOAD_DEFERRING
+	}
+	return Result;
+}
+
 void UObjectProperty::SetObjectPropertyValue(void* PropertyValueAddress, UObject* Value) const
 {
 	SetPropertyValue(PropertyValueAddress, Value);
diff --git a/Engine/Source/Runtime/CoreUObject/Private/UObject/PropertyStruct.cpp b/Engine/Source/Runtime/CoreUObject/Private/UObject/PropertyStruct.cpp
--- a/Engine/Source/Runtime/CoreUObject/Private/UObject/PropertyStruct.cpp
+++ b/Engine/Source/Runtime/CoreUObject/Private/UObject/PropertyStruct.cpp
@@ -3,6 +3,7 @@
 #include "CoreUObjectPrivate.h"
 #include "Archive.h"
 #include "PropertyHelper.h"
+#include "LinkerPlaceholderBase.h" // for FScopedPlaceholderPropertyTracker
 
 static inline void PreloadInnerStructMembers(UStructProperty* StructProperty)
 {
@@ -105,6 +106,10 @@ void UStructProperty::SerializeItem( FArchive& Ar, void* Value, void const* Defa
 {
 	check(Struct);
 
+#if USE_CIRCULAR_DEPENDENCY_LOAD_DEFERRING
+	FScopedPlaceholderPropertyTracker ImportPropertyTracker(this);
+#endif
+
 	Struct->SerializeItem(Ar, Value, Defaults);
 }
 
@@ -279,6 +284,9 @@ void UStructProperty::ExportTextItem( FString& ValueStr, const void* PropertyVal
 
 const TCHAR* UStructProperty::ImportText_Internal(const TCHAR* InBuffer, void* Data, int32 PortFlags, UObject* Parent, FOutputDevice* ErrorText) const
 {
+#if USE_CIRCULAR_DEPENDENCY_LOAD_DEFERRING
+	FScopedPlaceholderPropertyTracker ImportPropertyTracker(this);
+#endif 
 	return ImportText_Static(Struct, GetName(), InBuffer, Data, PortFlags, Parent, ErrorText);
 }
 
diff --git a/Engine/Source/Runtime/CoreUObject/Private/UObject/SavePackage.cpp b/Engine/Source/Runtime/CoreUObject/Private/UObject/SavePackage.cpp
--- a/Engine/Source/Runtime/CoreUObject/Private/UObject/SavePackage.cpp
+++ b/Engine/Source/Runtime/CoreUObject/Private/UObject/SavePackage.cpp
@@ -237,7 +237,7 @@ public:
 	{
 		if (const IBlueprintNativeCodeGenCore* Coordinator = IBlueprintNativeCodeGenCore::Get())
 		{
-			if (const UClass* ReplObjClass = Coordinator->FindReplacedClass(Obj->GetClass()))
+			if (const UClass* ReplObjClass = Coordinator->FindReplacedClassForObject(Obj))
 			{
 				MarkNameAsReferenced(ReplObjClass->GetFName());
 			}
@@ -785,6 +785,12 @@ FArchive& FArchiveSaveTagImports::operator<<( UObject*& Obj )
 				if (!IsCooking() || !Obj->HasAllMarks(ObjectMarks))
 				{
 					UObject* Parent = Obj->GetOuter();
+#if WITH_EDITOR
+					const IBlueprintNativeCodeGenCore* Coordinator = IBlueprintNativeCodeGenCore::Get();
+					FName UnusedName;
+					UObject* ReplacedOuter = Coordinator ? Coordinator->FindReplacedNameAndOuter(Obj, /*out*/UnusedName) : nullptr;
+					Parent = ReplacedOuter ? ReplacedOuter : Obj->GetOuter();
+#endif //WITH_EDITOR
 					if( Parent )
 					{
 						*this << Parent;
@@ -3110,6 +3116,10 @@ ESavePackageResult UPackage::Save(UPackage* InOuter, UObject* Base, EObjectFlags
 {
 	UE_START_LOG_COOK_TIME( Filename );
 
+#if WITH_EDITOR
+	TMap<UObject*, UObject*> ReplacedImportOuters;
+#endif //WITH_EDITOR
+
 	if (FPlatformProperties::HasEditorOnlyData())
 	{
 		if (GIsSavingPackage)
@@ -3502,14 +3512,13 @@ ESavePackageResult UPackage::Save(UPackage* InOuter, UObject* Base, EObjectFlags
 						EReplacementResult ReplacmentResult = IBlueprintNativeCodeGenCore::Get()->IsTargetedForReplacement(InOuter);
 						if (ReplacmentResult == EReplacementResult::ReplaceCompletely)
 						{
+							UE_LOG(LogSavePackage, Display, TEXT("Package %s contains assets, that were converted into native code. Package will not be saved."), *InOuter->GetName());
 							return ESavePackageResult::ReplaceCompletely;
 						}
 						else if (ReplacmentResult == EReplacementResult::GenerateStub)
 						{
 							bRequestStub = true;
 						}
-						UE_LOG(LogSavePackage, Display, TEXT("Package %s contains assets, that were converted into native code. Package will not be saved."), *InOuter->GetName());
-						// Or replace with the NativeScriptPackage ?
 					}
 #endif
 				}
@@ -3957,31 +3966,36 @@ ESavePackageResult UPackage::Save(UPackage* InOuter, UObject* Base, EObjectFlags
 							}
 						}
 					}
-
 					for(int32 Index = 0; Index < TagImpObjects.Num(); Index++)
 					{
 						UObject* Obj = TagImpObjects[Index];
 						check(Obj->HasAnyMarks(OBJECTMARK_TagImp) || Linker->IsCooking());
-						UClass* ObjClass = nullptr;
+						UClass* ObjClass = Obj->GetClass();
 #if WITH_EDITOR
-						bool bReplaced = false;
+						FName ReplacedName = NAME_None;
 						if (const IBlueprintNativeCodeGenCore* Coordinator = IBlueprintNativeCodeGenCore::Get())
 						{
-							if (UClass* ReplacedClass = Coordinator->FindReplacedClass(Obj->GetClass()))
+							if (UClass* ReplacedClass = Coordinator->FindReplacedClassForObject(Obj))
 							{
 								ObjClass = ReplacedClass;
-								bReplaced = true;
+							}
+							if (UObject* ReplacedOuter = Coordinator->FindReplacedNameAndOuter(Obj, /*out*/ReplacedName))
+							{
+								ReplacedImportOuters.Add(Obj, ReplacedOuter);
 							}
 						}
-						if (!bReplaced)
 #endif //WITH_EDITOR
+						FObjectImport* LocObjectImport = nullptr;
+						if (Obj->HasAnyMarks(OBJECTMARK_TagImp))
 						{
-							ObjClass = Obj->GetClass();
+							LocObjectImport = new(Linker->ImportMap)FObjectImport(Obj, ObjClass);
 						}
-						if (Obj->HasAnyMarks(OBJECTMARK_TagImp))
+#if WITH_EDITOR
+						if (LocObjectImport && (ReplacedName != NAME_None))
 						{
-							new(Linker->ImportMap)FObjectImport(Obj, ObjClass);
+							LocObjectImport->ObjectName = ReplacedName;
 						}
+#endif //WITH_EDITOR
 					}
 				}
 
@@ -4578,7 +4592,18 @@ ESavePackageResult UPackage::Save(UPackage* InOuter, UObject* Base, EObjectFlags
 								}
 							}
 							check(!Import.XObject->GetOuter()->IsIn(InOuter) || Import.XObject->HasAllFlags(RF_Transient) || Import.XObject->IsNative());
-							Import.OuterIndex = Linker->MapObject(Import.XObject->GetOuter());
+#if WITH_EDITOR
+							UObject** ReplacedOuter = ReplacedImportOuters.Find(Import.XObject);
+							if (ReplacedOuter && *ReplacedOuter)
+							{
+								Import.OuterIndex = Linker->MapObject(*ReplacedOuter);
+								ensure(Import.OuterIndex != FPackageIndex());
+							}
+							else
+#endif
+							{
+								Import.OuterIndex = Linker->MapObject(Import.XObject->GetOuter());
+							}
 						}
 					}
 					else
diff --git a/Engine/Source/Runtime/CoreUObject/Private/UObject/ScriptCore.cpp b/Engine/Source/Runtime/CoreUObject/Private/UObject/ScriptCore.cpp
--- a/Engine/Source/Runtime/CoreUObject/Private/UObject/ScriptCore.cpp
+++ b/Engine/Source/Runtime/CoreUObject/Private/UObject/ScriptCore.cpp
@@ -1765,6 +1765,46 @@ void UObject::execSwitchValue(FFrame& Stack, RESULT_DECL)
 }
 IMPLEMENT_VM_FUNCTION(EX_SwitchValue, execSwitchValue);
 
+void UObject::execArrayGetByRef(FFrame& Stack, RESULT_DECL)
+{
+	// Get variable address.
+	Stack.MostRecentPropertyAddress = NULL;
+	Stack.Step( Stack.Object, NULL ); // Evaluate variable.
+
+	if (Stack.MostRecentPropertyAddress == NULL)
+	{
+		static FBlueprintExceptionInfo ExceptionInfo(EBlueprintExceptionType::AccessViolation, LOCTEXT("ArrayGetRefException", "Attempt to assign variable through None"));
+		FBlueprintCoreDelegates::ThrowScriptException(this, Stack, ExceptionInfo);
+	}
+
+	void* ArrayAddr = Stack.MostRecentPropertyAddress;
+	UArrayProperty* ArrayProperty = ExactCast<UArrayProperty>(Stack.MostRecentProperty);
+
+ 	int32 ArrayIndex;
+ 	Stack.Step( Stack.Object, &ArrayIndex);
+
+	FScriptArrayHelper ArrayHelper(ArrayProperty, ArrayAddr);
+	Stack.MostRecentProperty = ArrayProperty->Inner;
+	// Add a little safety for Blueprints to not hard crash
+	if (ArrayHelper.IsValidIndex(ArrayIndex))
+	{
+		Stack.MostRecentPropertyAddress = ArrayHelper.GetRawPtr(ArrayIndex);
+
+		if (RESULT_PARAM)
+		{
+			ArrayProperty->Inner->CopyCompleteValueToScriptVM(RESULT_PARAM, ArrayHelper.GetRawPtr(ArrayIndex));
+		}
+	}
+	else
+	{
+		// The index is not valid. We will report this as an error (it is fatal in C++ conversions) but otherwise let the Blueprint handle the issue.
+		const int32 PropertySize = ArrayProperty->Inner->ElementSize * ArrayProperty->Inner->ArrayDim;
+		ArrayProperty->Inner->InitializeValue(Stack.MostRecentPropertyAddress);
+		FFrame::KismetExecutionMessage(*FString::Printf(TEXT("Attempted to get an item from array %s out of bounds [%d/%d]!"), *ArrayProperty->GetName(), ArrayIndex, ArrayHelper.Num()? ArrayHelper.Num() - 1 : 0), ELogVerbosity::Warning);
+	}
+}
+IMPLEMENT_VM_FUNCTION(EX_ArrayGetByRef, execArrayGetByRef);
+
 void UObject::execLet(FFrame& Stack, RESULT_DECL)
 {
 	Stack.MostRecentProperty = nullptr;
diff --git a/Engine/Source/Runtime/CoreUObject/Public/Blueprint/BlueprintSupport.h b/Engine/Source/Runtime/CoreUObject/Public/Blueprint/BlueprintSupport.h
--- a/Engine/Source/Runtime/CoreUObject/Public/Blueprint/BlueprintSupport.h
+++ b/Engine/Source/Runtime/CoreUObject/Public/Blueprint/BlueprintSupport.h
@@ -27,6 +27,12 @@ struct FBlueprintSupport
 	static bool UseDeferredDependencyLoading();
 	static bool IsDeferredExportCreationDisabled();
 	static bool IsDeferredCDOInitializationDisabled();
+
+	/** Tells if the specified object is one of the many flavors of FLinkerPlaceholderBase that we have. */
+	static bool IsDeferredDependencyPlaceholder(UObject* LoadedObj);
+
+	/** Not a particularly fast function. Mostly intended for validation in debug builds. */
+	static bool IsInBlueprintPackage(UObject* LoadedObj);
 };
 
 #if WITH_EDITOR
@@ -118,10 +124,19 @@ struct IBlueprintNativeCodeGenCore
 	 * UDynamicClass. Cooking (and conversion in general) must be order independent so
 	 * The scope of this kind of type swap is limited.
 	 * 
-	 * @param Class The class to check for a replacement
+	 * @param Object whose class will be replaced
 	 * @return A replacement class ptr, null if none
 	 */
-	virtual UClass* FindReplacedClass(const UClass* Class) const = 0;
+	virtual UClass* FindReplacedClassForObject(const UObject* Object) const = 0;
+	
+	/** 
+	 * Function used to change the path of subobject from a nativized class.
+	 * 
+	 * @param Object Imported Object.
+	 * @param OutName Referenced to name, that will be saved in import table.
+	 * @return An Outer object that should be saved in import table.
+	 */
+	virtual UObject* FindReplacedNameAndOuter(UObject* Object, FName& OutName) const = 0;
 };
 
 #endif // WITH_EDITOR
diff --git a/Engine/Source/Runtime/CoreUObject/Public/UObject/Class.h b/Engine/Source/Runtime/CoreUObject/Public/UObject/Class.h
--- a/Engine/Source/Runtime/CoreUObject/Public/UObject/Class.h
+++ b/Engine/Source/Runtime/CoreUObject/Public/UObject/Class.h
@@ -2943,3 +2943,33 @@ template<> struct TBaseStructure<FFallbackStruct>
 {
 	COREUOBJECT_API static UScriptStruct* Get();
 };
+
+template<> struct TBaseStructure<FFloatRangeBound>
+{
+	COREUOBJECT_API static UScriptStruct* Get();
+};
+
+template<> struct TBaseStructure<FFloatRange>
+{
+	COREUOBJECT_API static UScriptStruct* Get();
+};
+
+template<> struct TBaseStructure<FInt32RangeBound>
+{
+	COREUOBJECT_API static UScriptStruct* Get();
+};
+
+template<> struct TBaseStructure<FInt32Range>
+{
+	COREUOBJECT_API static UScriptStruct* Get();
+};
+
+template<> struct TBaseStructure<FFloatInterval>
+{
+	COREUOBJECT_API static UScriptStruct* Get();
+};
+
+template<> struct TBaseStructure<FInt32Interval>
+{
+	COREUOBJECT_API static UScriptStruct* Get();
+};
\ No newline at end of file
diff --git a/Engine/Source/Runtime/CoreUObject/Public/UObject/Linker.h b/Engine/Source/Runtime/CoreUObject/Public/UObject/Linker.h
--- a/Engine/Source/Runtime/CoreUObject/Public/UObject/Linker.h
+++ b/Engine/Source/Runtime/CoreUObject/Public/UObject/Linker.h
@@ -1302,6 +1302,8 @@ public:
 
 	/** Flags determining loading behavior.																					*/
 	uint32					LoadFlags;
+	/** Indicates whether the imports for this loader have been verified													*/
+	bool					bHaveImportsBeenVerified;
 	/** Indicates that this linker was created for a dynamic class package and will not use Loader */
 	bool					bDynamicClassLinker;
 	/** Hash table for exports.																								*/
@@ -1405,8 +1407,6 @@ private:
 	float					TimeLimit;
 	/** Time at begin of Tick function. Used for time limit determination.													*/
 	double					TickStartTime;
-	/**  Tracks the last verified import (making Verify() reentrant) - used to know when all the imports have been verified */
-	int32					VerifiedImportCount;
 
 	/** Used for ActiveClassRedirects functionality */
 	bool					bFixupExportMapDone;
@@ -1520,12 +1520,6 @@ public:
 		return OwnerThread;
 	}
 
-	/** Indicates when Verify() has gone through the entire  */
-	FORCEINLINE bool HaveImportsBeenVerified() const
-	{
-		return VerifiedImportCount >= Summary.ImportCount;
-	}
-
 	/**
 	 * If this archive is a FLinkerLoad or FLinkerSave, returns a pointer to the FLinker portion.
 	 */
@@ -1984,6 +1978,22 @@ public:
 	* @return True if FinalizeBlueprint() is currently being ran (or about to be ran) for an export (Blueprint) class.
 	*/
 	bool IsBlueprintFinalizationPending() const;
+
+	/**
+	 * Gives external code the ability to create FLinkerPlaceholderBase objects
+	 * in place of loads that may violate the LOAD_DeferDependencyLoads state.
+	 * This will only produce a placeholder if LOAD_DeferDependencyLoads is set
+	 * for this linker.
+	 *
+	 * NOTE: For now, this will only produce UClass placeholders, as that is the 
+	 *       only type we've identified needing.
+	 * 
+	 * @param  ObjectType    The expected type of the object you want to defer loading of.
+	 * @param  ObjectPath    The full object/package path for the expected object.
+	 * @return A FLinkerPlaceholderBase UObject that can be used in place of the import dependency.
+	 */
+	UObject* RequestPlaceholderValue(UClass* ObjectType, const TCHAR* ObjectPath);
+
 private:
 	/**
 	 * Regenerates/Refreshes a blueprint class
@@ -2007,10 +2017,10 @@ private:
 	/**
 	* Determines if the Object Import error should be suppressed
 	*
-	* @param  InImport    Import to check for suppression.
+	* @param  ImportIndex    Internal index into this linker's ImportMap, references the import to check for suppression.
 	* @return True if the import error should be suppressed
 	*/
-	bool IsSuppressableBlueprintImportError(FObjectImport& InImport);
+	bool IsSuppressableBlueprintImportError(int32 ImportIndex) const;
 #endif // WITH_EDITOR
 
 	/**
@@ -2049,9 +2059,11 @@ private:
 	 * 
 	 * @param  Placeholder		A ULinkerPlaceholderClass that was substituted in place of a deferred dependency.
 	 * @param  ReferencingClass	The (Blueprint) class that was loading, while we deferred dependencies (now referencing the placeholder).
+	 * @param  ObjectPath		Optional param that denotes the full object/package path for the object the placeholder is supposed to represent 
+	 *                          (used when the passed placeholder is not tied to an import in the linker's ImportMap).
 	 * @return The number of placeholder references replaced (could be none, if this was recursively resolved).
 	 */
-	int32 ResolveDependencyPlaceholder(class FLinkerPlaceholderBase* Placeholder, UClass* ReferencingClass = nullptr);
+	int32 ResolveDependencyPlaceholder(class FLinkerPlaceholderBase* Placeholder, UClass* ReferencingClass = nullptr, const FName ObjectPath = NAME_None);
 
 	/**
 	 * Query method to help catch recursive behavior. When this returns true, a 
@@ -2120,7 +2132,6 @@ private:
 	 */
 	bool IsExportBeingResolved(int32 ExportIndex);
 
-
 	void ResetDeferredLoadingState();
 
 	bool HasPerformedFullExportResolvePass();
@@ -2155,6 +2166,12 @@ private:
 	 * to the next.
 	 */
 	class FLinkerPlaceholderBase* ResolvingDeferredPlaceholder;
+
+	/** 
+	 * Internal list to track imports that were deferred, but don't belong to 
+	 * the ImportMap (thinks ones loaded through config files via UProperty::ImportText).
+	 */
+	TMap<FName, FLinkerPlaceholderBase*> ImportPlaceholders;
 #endif // USE_CIRCULAR_DEPENDENCY_LOAD_DEFERRING
 
 
diff --git a/Engine/Source/Runtime/CoreUObject/Public/UObject/ObjectBase.h b/Engine/Source/Runtime/CoreUObject/Public/UObject/ObjectBase.h
--- a/Engine/Source/Runtime/CoreUObject/Public/UObject/ObjectBase.h
+++ b/Engine/Source/Runtime/CoreUObject/Public/UObject/ObjectBase.h
@@ -474,7 +474,7 @@ enum EObjectFlags
 };
 
 	// Special all and none masks
-#define RF_AllFlags				(EObjectFlags)0x03ffffff	///< All flags, used mainly for error checking
+#define RF_AllFlags				(EObjectFlags)0x07ffffff	///< All flags, used mainly for error checking
 
 	// Predefined groups of the above
 #define RF_Load						((EObjectFlags)(RF_Public | RF_Standalone | RF_Transactional | RF_ClassDefaultObject | RF_ArchetypeObject | RF_DefaultSubObject | RF_TextExportTransient | RF_InheritableComponentTemplate)) // Flags to load from Unrealfiles.
diff --git a/Engine/Source/Runtime/CoreUObject/Public/UObject/Script.h b/Engine/Source/Runtime/CoreUObject/Public/UObject/Script.h
--- a/Engine/Source/Runtime/CoreUObject/Public/UObject/Script.h
+++ b/Engine/Source/Runtime/CoreUObject/Public/UObject/Script.h
@@ -255,6 +255,7 @@ enum EExprToken
 	EX_CallMath				= 0x68, // static pure function from on local call space
 	EX_SwitchValue			= 0x69,
 	EX_InstrumentationEvent	= 0x6A, // Instrumentation event
+	EX_ArrayGetByRef		= 0x6B,
 	EX_Max					= 0x100,
 };
 
@@ -296,6 +297,9 @@ namespace EScriptInstrumentation
 		PureNodeEntry,
 		NodeEntry,
 		NodeExit,
+		PushState,
+		RestoreState,
+		PopState,
 		NodePin,
 		Stop
 	};
diff --git a/Engine/Source/Runtime/CoreUObject/Public/UObject/ScriptSerialization.h b/Engine/Source/Runtime/CoreUObject/Public/UObject/ScriptSerialization.h
--- a/Engine/Source/Runtime/CoreUObject/Public/UObject/ScriptSerialization.h
+++ b/Engine/Source/Runtime/CoreUObject/Public/UObject/ScriptSerialization.h
@@ -471,6 +471,12 @@
 			SerializeExpr(iCode, Ar);	//default term
 			break;
 		}
+		case EX_ArrayGetByRef:
+			{
+				SerializeExpr( iCode, Ar );
+				SerializeExpr( iCode, Ar );
+				break;
+			}
 		default:
 		{
 			// This should never occur.
diff --git a/Engine/Source/Runtime/CoreUObject/Public/UObject/UObject.h b/Engine/Source/Runtime/CoreUObject/Public/UObject/UObject.h
--- a/Engine/Source/Runtime/CoreUObject/Public/UObject/UObject.h
+++ b/Engine/Source/Runtime/CoreUObject/Public/UObject/UObject.h
@@ -1152,6 +1152,8 @@ public:
 
 	DECLARE_FUNCTION(execSwitchValue);
 
+	DECLARE_FUNCTION(execArrayGetByRef);
+
 	// -- K2 support functions
 	struct Object_eventExecuteUbergraph_Parms
 	{
diff --git a/Engine/Source/Runtime/CoreUObject/Public/UObject/UnrealType.h b/Engine/Source/Runtime/CoreUObject/Public/UObject/UnrealType.h
--- a/Engine/Source/Runtime/CoreUObject/Public/UObject/UnrealType.h
+++ b/Engine/Source/Runtime/CoreUObject/Public/UObject/UnrealType.h
@@ -1906,6 +1906,7 @@ class COREUOBJECT_API UObjectProperty : public TUObjectPropertyBase<UObject*>
 	// UProperty interface
 	virtual void SerializeItem( FArchive& Ar, void* Value, void const* Defaults ) const override;
 	virtual void EmitReferenceInfo(UClass& OwnerClass, int32 BaseOffset) override;
+	virtual const TCHAR* ImportText_Internal(const TCHAR* Buffer, void* Data, int32 PortFlags, UObject* OwnerObject, FOutputDevice* ErrorText) const override;
 
 private:
 	virtual uint32 GetValueTypeHashInternal(const void* Src) const override
diff --git a/Engine/Source/Runtime/Engine/Classes/Components/SkinnedMeshComponent.h b/Engine/Source/Runtime/Engine/Classes/Components/SkinnedMeshComponent.h
--- a/Engine/Source/Runtime/Engine/Classes/Components/SkinnedMeshComponent.h
+++ b/Engine/Source/Runtime/Engine/Classes/Components/SkinnedMeshComponent.h
@@ -839,10 +839,23 @@ public:
 	 * @param TestLocation the location to test against
 	 * @param BoneLocation (optional, out) if specified, set to the world space location of the bone that was found, or (0,0,0) if no bone was found
 	 * @param IgnoreScale (optional) if specified, only bones with scaling larger than the specified factor are considered
+	 * @param bRequirePhysicsAsset (optional) if true, only bones with physics will be considered
 	 *
 	 * @return the name of the bone that was found, or 'None' if no bone was found
 	 */
-	FName FindClosestBone(FVector TestLocation, FVector* BoneLocation = NULL, float IgnoreScale = 0.f) const;
+	FName FindClosestBone(FVector TestLocation, FVector* BoneLocation = NULL, float IgnoreScale = 0.f, bool bRequirePhysicsAsset = false) const;
+
+	/** finds the closest bone to the given location
+	*
+	* @param TestLocation the location to test against
+	* @param BoneLocation (optional, out) if specified, set to the world space location of the bone that was found, or (0,0,0) if no bone was found
+	* @param IgnoreScale (optional) if specified, only bones with scaling larger than the specified factor are considered
+	* @param bRequirePhysicsAsset (optional) if true, only bones with physics will be considered
+	*
+	* @return the name of the bone that was found, or 'None' if no bone was found
+	*/
+	UFUNCTION(BlueprintCallable, Category = "Components|SkinnedMesh", meta=(DisplayName="FindClosestBone", AdvancedDisplay="bRequirePhysicsAsset"))
+	FName FindClosestBone_K2(FVector TestLocation, FVector& BoneLocation, float IgnoreScale = 0.f, bool bRequirePhysicsAsset = false) const;
 
 	/**
 	 * Find a named MorphTarget from the current SkeletalMesh
diff --git a/Engine/Source/Runtime/Engine/Classes/EdGraph/EdGraphNode.h b/Engine/Source/Runtime/Engine/Classes/EdGraph/EdGraphNode.h
--- a/Engine/Source/Runtime/Engine/Classes/EdGraph/EdGraphNode.h
+++ b/Engine/Source/Runtime/Engine/Classes/EdGraph/EdGraphNode.h
@@ -134,6 +134,10 @@ class ENGINE_API UEdGraphNode : public UObject
 	/** If true, this node can be renamed in the editor */
 	UPROPERTY()
 	uint32 bCanRenameNode:1;
+
+	/** Note for a node that lingers until saved */
+	UPROPERTY(Transient)
+	FText NodeUpgradeMessage;
 #endif // WITH_EDITORONLY_DATA
 
 	/** Comment string that is drawn on the node */
@@ -210,6 +214,7 @@ public:
 	// UObject interface
 	static void AddReferencedObjects(UObject* InThis, FReferenceCollector& Collector);
 	virtual void Serialize(FArchive& Ar) override;
+	virtual void PreSave() override;
 	virtual void PostLoad() override;
 	virtual void PostEditChangeProperty(FPropertyChangedEvent& PropertyChangedEvent) override;
 	// End of UObject interface
@@ -493,6 +498,8 @@ public:
 	/** Create a visual widget to represent this node in a graph editor or graph panel.  If not implemented, the default node factory will be used. */
 	virtual TSharedPtr<SGraphNode> CreateVisualWidget() { return TSharedPtr<SGraphNode>(); }
 
+	/** Adds an upgrade note to this node */
+	void AddNodeUpgradeNote(FText InUpgradeNote);
 #endif // WITH_EDITOR
 
 };
diff --git a/Engine/Source/Runtime/Engine/Classes/Engine/Blueprint.h b/Engine/Source/Runtime/Engine/Classes/Engine/Blueprint.h
--- a/Engine/Source/Runtime/Engine/Classes/Engine/Blueprint.h
+++ b/Engine/Source/Runtime/Engine/Classes/Engine/Blueprint.h
@@ -727,6 +727,9 @@ public:
 	/** If this blueprint is currently being compiled, the CurrentMessageLog will be the log currently being used to send logs to. */
 	class FCompilerResultsLog* CurrentMessageLog;
 
+	/** Message log for storing upgrade notes that were generated within the Blueprint, will be displayed to the compiler results each compiler and will remain until saving */
+	TSharedPtr<class FCompilerResultsLog> UpgradeNotesLog;
+
 	/** 
 	 * Sends a message to the CurrentMessageLog, if there is one available.  Otherwise, defaults to logging to the normal channels.
 	 * Should use this for node and blueprint actions that happen during compilation!
diff --git a/Engine/Source/Runtime/Engine/Classes/GameFramework/Actor.h b/Engine/Source/Runtime/Engine/Classes/GameFramework/Actor.h
--- a/Engine/Source/Runtime/Engine/Classes/GameFramework/Actor.h
+++ b/Engine/Source/Runtime/Engine/Classes/GameFramework/Actor.h
@@ -261,6 +261,7 @@ public:
 	void CopyRemoteRoleFrom(const AActor* CopyFromActor);
 
 	/** Returns how much control the remote machine has over this actor. */
+	UFUNCTION(BlueprintCallable, Category = "Replication")
 	ENetRole GetRemoteRole() const;
 
 	/** Used for replication of our RootComponent's position and velocity */
diff --git a/Engine/Source/Runtime/Engine/Classes/Kismet/KismetArrayLibrary.h b/Engine/Source/Runtime/Engine/Classes/Kismet/KismetArrayLibrary.h
--- a/Engine/Source/Runtime/Engine/Classes/Kismet/KismetArrayLibrary.h
+++ b/Engine/Source/Runtime/Engine/Classes/Kismet/KismetArrayLibrary.h
@@ -110,14 +110,15 @@ class ENGINE_API UKismetArrayLibrary : public UBlueprintFunctionLibrary
 	UFUNCTION(BlueprintPure, CustomThunk, meta=(DisplayName = "Last Index", CompactNodeTitle = "LAST INDEX", ArrayParm = "TargetArray"), Category="Utilities|Array")
 	static int32 Array_LastIndex(const TArray<int32>& TargetArray);
 
-	/* 
+	/*
+	 * Officially deprecated in favor of UK2Node_ArrayGetItem, which returns the array item by-ref. This function is no longer accessible inside Blueprints.
 	 *Given an array and an index, returns the item found at that index
 	 *
 	 *@param	TargetArray		The array to get an item from
 	 *@param	Index			The index in the array to get an item from
 	 *@return	The item stored at the index
 	*/
-	UFUNCTION(BlueprintPure, CustomThunk, meta=(DisplayName = "Get", CompactNodeTitle = "GET", ArrayParm = "TargetArray", ArrayTypeDependentParams = "Item"), Category="Utilities|Array")
+	UFUNCTION(BlueprintPure, CustomThunk, meta=(DisplayName = "Get", CompactNodeTitle = "GET", ArrayParm = "TargetArray", ArrayTypeDependentParams = "Item", BlueprintInternalUseOnly = "True"))
 	static void Array_Get(const TArray<int32>& TargetArray, int32 Index, int32& Item);
 
 	/* 
diff --git a/Engine/Source/Runtime/Engine/Private/ActorConstruction.cpp b/Engine/Source/Runtime/Engine/Private/ActorConstruction.cpp
--- a/Engine/Source/Runtime/Engine/Private/ActorConstruction.cpp
+++ b/Engine/Source/Runtime/Engine/Private/ActorConstruction.cpp
@@ -442,6 +442,19 @@ void AActor::ExecuteConstruction(const FTransform& Transform, const FComponentIn
 	TArray<const UBlueprintGeneratedClass*> ParentBPClassStack;
 	const bool bErrorFree = UBlueprintGeneratedClass::GetGeneratedClassesHierarchy(GetClass(), ParentBPClassStack);
 
+	TArray<const UDynamicClass*> ParentDynamicClassStack;
+	for (UClass* ClassIt = GetClass(); ClassIt; ClassIt = ClassIt->GetSuperClass())
+	{
+		if (UDynamicClass* DynamicClass = Cast<UDynamicClass>(ClassIt))
+		{
+			ParentDynamicClassStack.Add(DynamicClass);
+		}
+	}
+	for (int32 i = ParentDynamicClassStack.Num() - 1; i >= 0; i--)
+	{
+		UBlueprintGeneratedClass::CreateComponentsForActor(ParentDynamicClassStack[i], this);
+	}
+
 	// If this actor has a blueprint lineage, go ahead and run the construction scripts from least derived to most
 	if( (ParentBPClassStack.Num() > 0)  )
 	{
@@ -521,7 +534,6 @@ void AActor::ExecuteConstruction(const FTransform& Transform, const FComponentIn
 	}
 	else
 	{
-		UBlueprintGeneratedClass::CreateComponentsForActor(GetClass(), this);
 #if WITH_EDITOR
 		bool bDoUserConstructionScript;
 		GConfig->GetBool(TEXT("Kismet"), TEXT("bTurnOffEditorConstructionScript"), bDoUserConstructionScript, GEngineIni);
diff --git a/Engine/Source/Runtime/Engine/Private/Blueprint.cpp b/Engine/Source/Runtime/Engine/Private/Blueprint.cpp
--- a/Engine/Source/Runtime/Engine/Private/Blueprint.cpp
+++ b/Engine/Source/Runtime/Engine/Private/Blueprint.cpp
@@ -7,6 +7,7 @@
 #include "LatentActions.h"
 
 #if WITH_EDITOR
+#include "UObject/DevObjectVersion.h"
 #include "Editor/UnrealEd/Public/Kismet2/BlueprintEditorUtils.h"
 #include "Editor/UnrealEd/Public/Kismet2/KismetEditorUtilities.h"
 #include "Editor/UnrealEd/Public/Kismet2/CompilerResultsLog.h"
@@ -215,6 +216,10 @@ void UBlueprintCore::Serialize(FArchive& Ar)
 {
 	Super::Serialize(Ar);
 
+#if WITH_EDITOR
+	Ar.UsingCustomVersion(FBlueprintsObjectVersion::GUID);
+#endif
+
 	Ar << bLegacyGeneratedClassIsAuthoritative;	
 
 	if ((Ar.UE4Ver() < VER_UE4_BLUEPRINT_SKEL_CLASS_TRANSIENT_AGAIN)
@@ -309,6 +314,9 @@ void UBlueprint::PreSave()
 {
 	Super::PreSave();
 
+	// Clear all upgrade notes, the user has saved and should not see them anymore
+	UpgradeNotesLog.Reset();
+
 	// Cache the BP for use
 	FFindInBlueprintSearchManager::Get().AddOrUpdateBlueprintSearchMetadata(this);
 }
@@ -607,6 +615,18 @@ void UBlueprint::PostLoad()
 		}
 	}
 
+#if WITH_EDITOR
+	if (UpgradeNotesLog.IsValid() && GetLinker())
+	{
+		const FCustomVersion* CustomVersion = GetLinker()->GetCustomVersions().GetVersion(FBlueprintsObjectVersion::GUID);
+		if (CustomVersion == nullptr || CustomVersion->Version < FBlueprintsObjectVersion::ArrayGetByRefUpgrade)
+		{
+			UpgradeNotesLog->Note(TEXT("Auto-Upgrade Note: Updated all array \"Get\" nodes in Blueprint Class \'@@\' to return a reference. Nodes have been injected to ensure they continue to work as they always have."), this);
+			UpgradeNotesLog->Note(TEXT("Auto-Upgrade Note: To remove these upgrade notes, resave @@. Be sure to verify nothing has broken in the upgrade process."), this);
+		}
+	}
+#endif
+
 	FStructureEditorUtils::RemoveInvalidStructureMemberVariableFromBlueprint(this);
 
 
diff --git a/Engine/Source/Runtime/Engine/Private/Components/SkinnedMeshComponent.cpp b/Engine/Source/Runtime/Engine/Private/Components/SkinnedMeshComponent.cpp
--- a/Engine/Source/Runtime/Engine/Private/Components/SkinnedMeshComponent.cpp
+++ b/Engine/Source/Runtime/Engine/Private/Components/SkinnedMeshComponent.cpp
@@ -1479,7 +1479,7 @@ void USkinnedMeshComponent::TransformFromBoneSpace(FName BoneName, FVector InPos
 
 
 
-FName USkinnedMeshComponent::FindClosestBone(FVector TestLocation, FVector* BoneLocation, float IgnoreScale) const
+FName USkinnedMeshComponent::FindClosestBone(FVector TestLocation, FVector* BoneLocation, float IgnoreScale, bool bRequirePhysicsAsset) const
 {
 	if (SkeletalMesh == NULL)
 	{
@@ -1491,6 +1491,17 @@ FName USkinnedMeshComponent::FindClosestBone(FVector TestLocation, FVector* Bone
 	}
 	else
 	{
+		// cache the physics asset
+		const UPhysicsAsset* PhysAsset = GetPhysicsAsset();
+		if (bRequirePhysicsAsset && !PhysAsset)
+		{
+			if (BoneLocation != NULL)
+			{
+				*BoneLocation = FVector::ZeroVector;
+			}
+			return NAME_None;
+		}
+
 		// transform the TestLocation into mesh local space so we don't have to transform the (mesh local) bone locations
 		TestLocation = ComponentToWorld.InverseTransformPosition(TestLocation);
 		
@@ -1499,7 +1510,15 @@ FName USkinnedMeshComponent::FindClosestBone(FVector TestLocation, FVector* Bone
 		int32 BestIndex = -1;
 		for (int32 i = 0; i < GetNumSpaceBases(); i++)
 		{
-			if (IgnoreScale < 0.f || GetSpaceBases()[i].GetScaledAxis( EAxis::X ).SizeSquared() > IgnoreScaleSquared)
+			// If we require a physics asset, then look it up in the map
+			bool bPassPACheck = !bRequirePhysicsAsset;
+			if (bRequirePhysicsAsset)
+			{
+				FName BoneName = SkeletalMesh->RefSkeleton.GetBoneName(i);
+				bPassPACheck = (PhysAsset->BodySetupIndexMap.FindRef(BoneName) != INDEX_NONE);
+			}
+
+			if (bPassPACheck && (IgnoreScale < 0.f || GetSpaceBases()[i].GetScaledAxis(EAxis::X).SizeSquared() > IgnoreScaleSquared))
 			{
 				float DistSquared = (TestLocation - GetSpaceBases()[i].GetLocation()).SizeSquared();
 				if (DistSquared < BestDistSquared)
@@ -1530,6 +1549,12 @@ FName USkinnedMeshComponent::FindClosestBone(FVector TestLocation, FVector* Bone
 	}
 }
 
+FName USkinnedMeshComponent::FindClosestBone_K2(FVector TestLocation, FVector& BoneLocation, float IgnoreScale, bool bRequirePhysicsAsset) const
+{
+	BoneLocation = FVector::ZeroVector;
+	return FindClosestBone(TestLocation, &BoneLocation, IgnoreScale, bRequirePhysicsAsset);
+}
+
 void USkinnedMeshComponent::ShowMaterialSection(int32 MaterialID, bool bShow, int32 LODIndex)
 {
 	if (!SkeletalMesh)
diff --git a/Engine/Source/Runtime/Engine/Private/EdGraph/EdGraphNode.cpp b/Engine/Source/Runtime/Engine/Private/EdGraph/EdGraphNode.cpp
--- a/Engine/Source/Runtime/Engine/Private/EdGraph/EdGraphNode.cpp
+++ b/Engine/Source/Runtime/Engine/Private/EdGraph/EdGraphNode.cpp
@@ -272,6 +272,19 @@ void UEdGraphNode::Serialize(FArchive& Ar)
 	}
 }
 
+void UEdGraphNode::PreSave()
+{
+	Super::PreSave();
+
+#if WITH_EDITORONLY_DATA
+	if (!NodeUpgradeMessage.IsEmpty())
+	{
+		// When saving, we clear any upgrade messages
+		NodeUpgradeMessage = FText::GetEmpty();
+	}
+#endif // WITH_EDITORONLY_DATA
+}
+
 void UEdGraphNode::PostLoad()
 {
 	Super::PostLoad();
@@ -417,6 +430,19 @@ FText UEdGraphNode::GetKeywords() const
 	return GetClass()->GetMetaDataText(TEXT("Keywords"), TEXT("UObjectKeywords"), GetClass()->GetFullGroupName(false));
 }
 
+void UEdGraphNode::AddNodeUpgradeNote(FText InUpgradeNote)
+{
+#if WITH_EDITORONLY_DATA
+	if (NodeUpgradeMessage.IsEmpty())
+	{
+		NodeUpgradeMessage = InUpgradeNote;
+	}
+	else
+	{
+		NodeUpgradeMessage = FText::Format(FText::FromString(TEXT("{0}\n{1}")), NodeUpgradeMessage, InUpgradeNote);
+	}
+#endif
+}
 #endif	//#if WITH_EDITOR
 
 bool UEdGraphNode::IsInDevelopmentMode() const
