Removing debug code that came with Dev-Blueprints integration

diff --git a/Engine/Source/Runtime/CoreUObject/Private/UObject/GarbageCollection.cpp b/Engine/Source/Runtime/CoreUObject/Private/UObject/GarbageCollection.cpp
--- a/Engine/Source/Runtime/CoreUObject/Private/UObject/GarbageCollection.cpp
+++ b/Engine/Source/Runtime/CoreUObject/Private/UObject/GarbageCollection.cpp
@@ -1182,12 +1182,6 @@ void CollectGarbageInternal(EObjectFlags KeepFlags, bool bPerformFullPurge)
 			bool bShouldAssert = false;
 			FUObjectItem* ObjectItem = *It;
 			UObject* Object = (UObject*)ObjectItem->Object;
-
-			if (Object->HasAnyFlags(RF_ClassDefaultObject) && Object->GetName().Contains(TEXT("RecastFilter_UseDefaultArea")))
-			{
-				printf("");
-			}
-
 			// Don't require UGCObjectReferencer's references to adhere to the assumptions.
 			// Although we want the referencer itself to sit in the disregard for gc set, most of the objects
 			// it's referencing will not be in the root set.
@@ -1199,25 +1193,19 @@ void CollectGarbageInternal(EObjectFlags KeepFlags, bool bPerformFullPurge)
 				ObjectReferenceCollector.FindReferences( Object );
 
 				// Iterate over referenced objects, finding bad ones.
-				for (int32 ReferenceIndex = 0; ReferenceIndex < CollectedReferences.Num(); ReferenceIndex++)
+				for( int32 ReferenceIndex=0; ReferenceIndex<CollectedReferences.Num(); ReferenceIndex++ )
 				{
 					UObject* ReferencedObject = CollectedReferences[ReferenceIndex];
-
-
-					if (ReferencedObject)
+					if (ReferencedObject && 
+						!(ReferencedObject->IsRooted() || 
+						  UObjectArray.IsDisregardForGC(ReferencedObject) || 
+							UObjectArray.ObjectToObjectItem(ReferencedObject)->GetOwnerIndex() ||
+							UObjectArray.ObjectToObjectItem(ReferencedObject)->HasAnyFlags(EInternalObjectFlags::ClusterRoot)))
 					{
-						bool bIsRooted = ReferencedObject->IsRooted();
-						bool bIsDisregard = UObjectArray.IsDisregardForGC(ReferencedObject);
-						int32 OwnerIndex = UObjectArray.ObjectToObjectItem(ReferencedObject)->GetOwnerIndex();
-						bool bHasClusterRoot = UObjectArray.ObjectToObjectItem(ReferencedObject)->HasAnyFlags(EInternalObjectFlags::ClusterRoot);
-
-						if (!(bIsRooted || bIsDisregard || OwnerIndex || bHasClusterRoot))
-						{
-							UE_LOG(LogGarbage, Warning, TEXT("Disregard for GC object %s referencing %s which is not part of root set"),
-								*Object->GetFullName(),
-								*ReferencedObject->GetFullName());
-							bShouldAssert = true;
-						}
+						UE_LOG(LogGarbage, Warning, TEXT("Disregard for GC object %s referencing %s which is not part of root set"),
+							*Object->GetFullName(),
+							*ReferencedObject->GetFullName());						
+						bShouldAssert = true;
 					}
 				}
 			}
