Temporarily disabling Blueprint spawning InitProperties() optimizations, as it was causing fallout (from CL 2916140).

diff --git a/Engine/Source/Runtime/CoreUObject/Private/UObject/UObjectGlobals.cpp b/Engine/Source/Runtime/CoreUObject/Private/UObject/UObjectGlobals.cpp
--- a/Engine/Source/Runtime/CoreUObject/Private/UObject/UObjectGlobals.cpp
+++ b/Engine/Source/Runtime/CoreUObject/Private/UObject/UObjectGlobals.cpp
@@ -2761,6 +2761,12 @@ void FObjectInitializer::InitProperties(UObject* Obj, UClass* DefaultsClass, UOb
 	}
 	else
 	{
+		// @TODO: temporarily disabling the PostConstructLink optimization here
+		//        since it is unknowingly causing issues in gameplay (see 
+		//        UE-28522) - we're still unclear on what specifically this is
+		//        doing wrong, but will investigate (UE-29449)
+		bCanUsePostConstructLink = false;
+
 		QUICK_SCOPE_CYCLE_COUNTER(STAT_InitProperties_Blueprint);
 
 		UObject* ClassDefaults = bCopyTransientsFromClassDefaults ? DefaultsClass->GetDefaultObject() : NULL;		
