Copying //UE4/Orion-Staging to //UE4/Main (Source: //Orion/Dev-General @ 2997507)

==========================
MAJOR FEATURES + CHANGES
==========================

Change 2997066 on 2016/06/01 by [REDACTED_PERSON]

	Engine: Marked engine performance target cvars ECVF_Scalability so they can be adjusted via scalability buckets at runtime for games that target different framerates on different levels of hardware
	#tests Ran Paragon and changed video settings and tested t.TargetFrameTimeThreshold

Change 2996816 on 2016/06/01 by [REDACTED_PERSON]

	Add FixedWorldDirection option for Root Motion Radial Forces; code/BP hook-up for allowing Price's reworked RMB to send all targets in the same (correct) direction
	- FixedWorldDirection added to both root motion system and corresponding ability tasks
	- Exposed AddHitResult for EffectContexts for modifying EffectContext hit results from BP
	- Hooked up to Price RMB - we (somewhat dirtily) route Price's location and facing through the HitResult of the EffectContext for his displacement GE

	#tests MultiPIE

Change 2996526 on 2016/06/01 by [REDACTED_PERSON]

	Fixed tube light typo

	JB made this robomerge up.  Shader recompiling in our future.

	#Tests:none

Change 2996428 on 2016/06/01 by [REDACTED_PERSON]

	O - Made r.D3D.RemoveUnusedInterpolators a system setting which now also alters ddc key; fix r.Shaders.FastMath actually affecting compilation when =0
	#tests Run with and without r.D3D.RemoveUnusedInterpolators=1 on DefaultEngine.ini

Change 2996090 on 2016/06/01 by [REDACTED_PERSON]

	#ORION_MAIN - Merge 26.2 @ CL 2995754

	#Tests:none

Change 2995785 on 2016/06/01 by [REDACTED_PERSON]

	Don't delete non-backup log files when cleaning up the logs folder.

	#tests Tested in the editor with multiple old log files

Change 2995556 on 2016/05/31 by [REDACTED_PERSON]

	More info about timers on crash (OR-21986).

	- Somewhat desperate attempt to get more clue about timer crash. Intended to be removed later.

	#tests Compiled the Linux server, ran it, crashed a few times.

Change 2995397 on 2016/05/31 by [REDACTED_PERSON]

	Rendering: Made the optimization to combine upscaling/downscaling and tonemapping optional based on the amount of upscaling that will occur
	- r.Tonemapper.ScreenPercentage has been renamed to r.Tonemapper.MergeWithUpscale.Mode
	- r.Tonemapper.MergeWithUpscale.Threshold is a new setting used when r.Tonemapper.MergeWithUpscale.Mode is set to 2, which indcates to only try to merge the passes if the ratio of the area before upscale/downscale to the area afterwards is greater than the threshold
	This prevents running the tonemapper on all of the target res pixels when the source res is far smaller, as that can cause it to be a loss to merge the passes

	Upgrade Notes: r.Tonemapper.ScreenPercentage has been renamed to r.Tonemapper.MergeWithUpscale.Mode

	#tests Ran Paragon at various resolutions on Intel and NV cards
	#rn

Change 2995118 on 2016/05/31 by [REDACTED_PERSON]

	- Fix for build failure

	#tests golden path game

Change 2994929 on 2016/05/31 by [REDACTED_PERSON]

	#Orion_Analytics
	-Added PacketRecievedHistogram event that fires every minute in game the rate is configurable in DefaultGame.ini
	-Moved FHistogram from PerfCountersModule to ProfilingHelpers
	-Re-enabled Location event

	#tests golden path game

Change 2994920 on 2016/05/31 by [REDACTED_PERSON]

	Added some more cooking stats to save package.
	#test Cook Orion.

Change 2994622 on 2016/05/31 by [REDACTED_PERSON]

	#orion - Pickup and Coin filter collision optimizations.

	- Added coin collision profile preset and made it ignore everything but Pawn (importantly, now ignores triggers).
	- Pickups filter out collision with AI earlier. Profile already did this but this avoids more branches and cache misses in PreFilter.
	- Coins now additionally filter out more efficiently Heroes that can't pick them (no overlap events generated at all).

	#tests PlayGo MultiPIE

Change 2994305 on 2016/05/31 by [REDACTED_PERSON]

	Restoring prompt/exit on signed archive issue to help identify causes
	#tests compiled

Change 2994226 on 2016/05/31 by [REDACTED_PERSON]

	#ORION_MAIN - Merge 27 @ CL 2993946

	#Tests:none

Change 2994204 on 2016/05/31 by bruce.nesbit

	More banner changes
	-Fixed an issue with InterpToComponent with very short times
	-revisions to test banner
	-added fade to banner/charms
	-tries to set team color on decativate FX

	#tests Game+PIE

Change 2993973 on 2016/05/30 by [REDACTED_PERSON]

	Updated protection handling

	#tests Compiled and applied protection

Change 2993588 on 2016/05/27 by [REDACTED_PERSON]

	Engine: Removed a bogus autocomplete for ShowMaterialDrawEvents, which was previously renamed r.ShowMaterialDrawEvents
	#tests Typed in ShowMat in the console and verified that no autocomplete appeared
	#rn

Change 2993510 on 2016/05/27 by [REDACTED_PERSON]

	Fix issue with root motion sources and replays, fixes TwinBlast RMB ability animation issue, and other artifacts

	#tests Replays

Change 2993484 on 2016/05/27 by [REDACTED_PERSON]

	New logic for computing the skel mesh streaming scales
	#tests loaded editor, played with streaming scale

Change 2993211 on 2016/05/27 by [REDACTED_PERSON]

	Workaround for lightmap streaming flags not being correctly set.
	#tests building lighting and investigating streaming

Change 2993068 on 2016/05/27 by [REDACTED_PERSON]

	Duplicate 2989729
	Fix for lightmaps and shadowmaps having low resolutions after building lightings
	#test PC at various scalability

Change 2993066 on 2016/05/27 by [REDACTED_PERSON]

	fixed behavior tree getting stuck on ResumeLogic call
	copied from //UE4/Dev-Framework, CL# 2993058
	#tests none

Change 2992706 on 2016/05/27 by [REDACTED_PERSON]

	Duplicate 2991726
	Fix for grey skin in simple lighting model (shadows off)

	#test lowest settings on PC

Change 2992705 on 2016/05/27 by [REDACTED_PERSON]

	Duplicate 2991727
	Fix emissive decals in simple forward renderer
	#test PC lowest settings

Change 2992658 on 2016/05/27 by [REDACTED_PERSON]

	Remove all occurrences of Ability.PersistPastDeath from granted tags. Fix code to *only* check asset tags for this tag, instead of both.

	#test pie

Change 2992646 on 2016/05/27 by [REDACTED_PERSON]

	BuildGraph: Add a BuildGraph task to run a UE4 commandlet. Syntax is <Commandlet Name="..." Project="..." Arguments="...">.

	#tests none

Change 2992252 on 2016/05/26 by [REDACTED_PERSON]

	#ORION_DG - Unclog ROBO Merge in DG

	#Tests:none

Change 2992180 on 2016/05/26 by [REDACTED_PERSON]

	Fix issue where external data wasn't saving out properly

	#tests Replays

Change 2992159 on 2016/05/26 by [REDACTED_PERSON]

	CVar to disable/freeze GPU particle simulation (r.GPUParticle.Simulate)
	[Replicated from Dev-Rendering checkin CL# 2989752 by Olaf]
	#tests Tested the command in Agora and verified that GPU particles were not being drawn

Change 2992158 on 2016/05/26 by [REDACTED_PERSON]

	Rendering: Added a cvar that controls unbinding of all texture resources between materials changes in the DX11 renderer (r.UnbindResourcesBetweenDrawsInDX11) to improve the readability of GPA captures
	Note: This will probably be moved to be on when markers are on rather than an independent cvar, but it is currently separate for testing
	#tests Ran with the var off and on and verified in GPA captures

Change 2991645 on 2016/05/26 by [REDACTED_PERSON]

	Fix for filesize returning 0 if file not found
	#tests bugit now works

Change 2991290 on 2016/05/26 by [REDACTED_PERSON]

	Added a static flag to NavigationSystem that can be used to short-circuit dynamic navigation related functions #UE4

	Will save some perf on PS4 in Orion, since clients do use navigation system there.

	#test golden path

Change 2991288 on 2016/05/26 by [REDACTED_PERSON]

	CL#2990243 manually redone in for Orion #UE4

	Original description:
	> Fixed behavior tree observers not being applied correctly

	#test golden path

Change 2991271 on 2016/05/26 by [REDACTED_PERSON]

	#ORION_MAIN - Merge 26.2 @ CL 2990688

	#Tests:none

Change 2991185 on 2016/05/26 by [REDACTED_PERSON]

	Fixed UAITask_MoveTo not releasing paths properly, or clearing path observing when task is being reused #UE4

	Also, fixed FNavigationPath::DoneUpdating not converting ENavPathUpdateType properly

	#test golden path

Change 2990788 on 2016/05/25 by [REDACTED_PERSON]

	Fix deprecation warning
	#test none

Change 2990582 on 2016/05/25 by [REDACTED_PERSON]

	Now that render commands are enqueued again on servers, we shouldn't outright crash if an allocation gets to nullrhi
	#test none

Change 2990450 on 2016/05/25 by [REDACTED_PERSON]

	OR-22233 GPU Sprites invisible unless solo'd
	#test:PC

Change 2990199 on 2016/05/25 by [REDACTED_PERSON]

	Remove experimental HDR support in tonemapper. Brings tonemapper cost back down into line

Change 2989908 on 2016/05/25 by [REDACTED_PERSON]

	Changed warning about DDC cache full to Display
	#tests compiled

Change 2989903 on 2016/05/25 by [REDACTED_PERSON]

	Made BT component ignore subtree injection request if relevant BT nodes already use indicated asset #UE4

	#test golden path

Change 2989795 on 2016/05/25 by [REDACTED_PERSON]

	Fix for storing the correct URL on the pending net game for replay playback.
	Re-implemented this fix from Dev-Networking CL 2981198, fixes deathcam after latest main integration.

	#tests played a reply, enabled deathcam

Change 2989483 on 2016/05/25 by [REDACTED_PERSON]

	ToggleJuggernaut cheat
	#tests pie

Change 2989384 on 2016/05/25 by [REDACTED_PERSON]

	Extra chunk decryption tests and logging to help diagnose the random failure we're seeing in the wild
	 - retry decrypt three times
	 - after the first attempt, re-decrypt original source, just incase the decrypt cache has been corrupted

	#tests cooked pc client + dedicated server

Change 2989225 on 2016/05/24 by [REDACTED_PERSON]

	Fix rare crash in Linux threading code (OR-22193).

	- Sometimes, for some reason, freeing memory for an alternate thread from a thread in PostRun() can crash because the jemalloc apparently does not have an arena for this thread anymore.
	- This change works around the problem by allocating the said memory statically in LinuxThread class.

	#tests Compiled Linux server, started it.

Change 2988768 on 2016/05/24 by [REDACTED_PERSON]

	Added support for SkinnedMesh in the texture streaming MeshCoordSize accuracy viewmode.
	#tests loaded game and editor

Change 2988462 on 2016/05/24 by [REDACTED_PERSON]

	Added a piece of logging to both scenarios or movement aborting in UPathFollowingComponent::UpdatePathSegment to be able to tell them appart while reading the log #Orion

	#test golden path

Change 2988036 on 2016/05/24 by [REDACTED_PERSON]

	#ORION_MAIN - Merge 26.2 @ CL 2987910

	#Tests:none

Change 2987457 on 2016/05/23 by [REDACTED_PERSON]

	Redone changes from CL#2981193 #UE4

	Original description: fixed missing observers in behavior tree when dynamic subtree is changed while waiting for full restart (out of nodes)

	#test golden path

Change 2987388 on 2016/05/23 by [REDACTED_PERSON]

	Replicating CL 2985226; don't push mesh emitter transform to pixel shader unless used in the material graph

	#tests editor game PC

Change 2986255 on 2016/05/22 by [REDACTED_PERSON]

	Manually resolving conflict that stoped robomerge from Main to DG #Orion

	#test compile

Change 2986209 on 2016/05/21 by [REDACTED_PERSON]

	Removed hitchunter logging from http thread
	#tests compiled

Change 2986202 on 2016/05/21 by [REDACTED_PERSON]

	Merging //UE4/Main @ 2981382 from //UE4/Orion-Staging
	#tests engine & game QA passed, built locally

Change 2985899 on 2016/05/20 by [REDACTED_PERSON]

	Move PS4 HTTP processing to HTTP thread
	#tests golden path

Change 2985884 on 2016/05/20 by [REDACTED_PERSON]

	Fix issue where oodle wasn't enabled in shipping correctly.

	#tests ran orion server in shipping and connecting with a client

Change 2985778 on 2016/05/20 by [REDACTED_PERSON]

	#ORION_MAIN - Merge 26.2 @ CL 2985753

	#Tests:none

Change 2985760 on 2016/05/20 by [REDACTED_PERSON]

	Second pass on HTTP threading
	Move threaded objects into separate class, FHttpThread.  FCurlHttpThread derives from FHttpThread and the curl multi work is performed in FCurlHttpThread
	Removed code that limited number of curl easy requests that were added to the multi simultaneously / per frame as now that curl work is performed on a separate thread the performance no longer directly impacts the game thread
	Remove lock from CurlHttp and instead of use FThreadSafeCounter
	#tests golden path (PC & PS4)

Change 2985658 on 2016/05/20 by [REDACTED_PERSON]

	Fixed issue with cached http replay results making time go backwards

	#tests replays

Change 2985640 on 2016/05/20 by [REDACTED_PERSON]

	Ability System: call OnRemove event for gameplay cues that are mispredicted. Previously if a looping GC was predictively added, it would only get the OnRemove event if the replicated GC was removed. In the case of a mis prediction there is no replicated version, so the OnRemove was never called and cleanup was never happening.

	#tests multi pie

Change 2985631 on 2016/05/20 by [REDACTED_PERSON]

	Fast TArray serialization fixes:
	1. Fix case where Array ReplicationKey has changed no items between base and current state have changed. Previously server would early out and not send an update: but this needs be sent so that the client can potentially perform an implicit delete.

	This fixes the case where client TArray would have stale items hanging around until a new update was sent (which could potentially be never).

	2. Fix case where an array item would be deleted by both explicit delete and implicit delete: causing other items in the array to be deleted (!).

	#tests golden path

Change 2985542 on 2016/05/20 by [REDACTED_PERSON]

	Added per package stats.
	Optimized cooker, moved FTextureSource::Compress from UTexture::Presave to UTexture::Serialize so we can avoid it in the cooker.
	#test cook paragon, save packages paragon editor

Change 2985152 on 2016/05/20 by [REDACTED_PERSON]

	#ORION_MAIN - Merge 26.2 @ CL 2985092

	#Tests:none

Change 2985001 on 2016/05/20 by [REDACTED_PERSON]

	Move fullcrashdump location for Paragon from QA deptartment storage to Paragon project storage. Changes CrashReportClient config only.

Change 2984839 on 2016/05/20 by [REDACTED_PERSON]

	Renaming some confusing function names and updating messages related to exception handling.

	#tests Cooked Win64 Client and Server, Tested crash reporting in cooked game

Change 2984517 on 2016/05/19 by [REDACTED_PERSON]

	Adjusted 'PlatformHeadroom' audio volume settings to DB-3 on both Windows and PS4

Change 2983932 on 2016/05/19 by jason.bestimt

	#ORION_MAIN - Merge 26.2 @ CL 2983814

	#Tests:none

Change 2983864 on 2016/05/19 by [REDACTED_PERSON]

	Added missing assignment copy/move ops to FAnalyticsEventAttribute. Doh, should have looked at more usages of PLATFORM_COMPILER_HAS_DEFAULTED_FUNCTIONS...
	#tests compiled Orion Editor Win64

Change 2983780 on 2016/05/19 by [REDACTED_PERSON]

	Modernize FAnalyticsEventAttribute usage. #jira UE-30551.

	Replaced FAnalyticsEventAttribute various ctors with a perfect forwarding one that can convert them to strings.
	* The Name must be convertible to a string
	* The value must be convertible to a string via an AnalyticsConversion::ToString() overload.
	* Added/expanded the supported conversions to strings to analytics attribute values. See AnalyticsConversion.h which contains all the previously supported conversions and more.

	Added MakeAnalyticsEventAttributeArray(), which uses variadic templates to create an array of event attributes inline, which can be passed to RecordEvent[Json] and efficiently taken ownership of:
	    RecordEvent("EventName", MakeAnalyticsEventAttributeArray(
	        "Attr1", false,
	        "Attr2", 42.0,
	        "Attr3",  SomeMap,
	        "Attr4", SomeArray);

	SamS - you've been asking for better attribute conversion facilities for years. Finally got it, haha.
	SteveR - I only added copy/move ctors to FAnalyticsEventAttribute. Do I also need to explicitly add the copy/move assignment ops?
	DavidD - This will allow you to create attributes a lot more easily and efficiently now.
	#tests orion editor runs

Change 2983702 on 2016/05/19 by [REDACTED_PERSON]

	Renumbered cooking stats to be more correct
	#test cook paragon.

Change 2983392 on 2016/05/19 by [REDACTED_PERSON]

	#ORION_MAIN - Merge 26.2 @ CL 2983342

	#Tests:none

Change 2982910 on 2016/05/18 by [REDACTED_PERSON]

	Don't fail entire deployments because obsolete manifest can't find the files to delete
	#test none

Change 2982902 on 2016/05/18 by [REDACTED_PERSON]

	Disable HTTP networkfilesystem because it constantly crashes cookonthefly servers.  platform team is aware
	#test cookonthefly

Change 2982837 on 2016/05/18 by [REDACTED_PERSON]

	Spot merge safety check in ~FAgggregator. From BobT CL 2966255.

	#tests compile

Change 2982723 on 2016/05/18 by [REDACTED_PERSON]

	Analytics no longer adds IsEditor attribute to all events. Wasn't actually used by anyone anymore. #jira UE-30559
	#tests none

Change 2982716 on 2016/05/18 by [REDACTED_PERSON]

	Remove Analytics code to divert legacy code to source data collector. #jira UE-27794
	#tests run orion editor

Change 2982707 on 2016/05/18 by [REDACTED_PERSON]

	AnalyticsET support for arbitrary Json events. #jira UE-30375
	* AnalyticsET supports a new API, RecordEventJson.
	* API supports rvalue refs to avoid unnecessary copies of the attribute array.
		/**
		 * Sends an event where each attribute value is expected to be a string-ified Json value.
		 * Meaning, each attribute value can be an integer, float, bool, string,
		 * arbitrarily complex Json array, or arbitrarily complex Json object.
		 *
		 * The main thing to remember is that if you pass a Json string as an attribute value, it is up to you to
		 * quote the string, as the string you pass is expected to be able to be pasted directly into a Json value. ie:
		 *
		 * {
		 *     "EventName": "MyStringEvent",
		 *     "IntAttr": 42                 <--- You simply pass this in as "42"
		 *     "StringAttr": "SomeString"    <--- You must pass SomeString as "\"SomeString\""
		 * }
		 *
		 * [REDACTED_HANDLE] EventName			The name of the event.
		 * [REDACTED_HANDLE] AttributesJson	array of key/value attribute pairs where each value is a Json value (pure Json strings mustbe quoted by the caller).
		 */
		virtual void RecordEventJson(const FString& EventName, TArray<FAnalyticsEventAttribute>&& AttributesJson) = 0;
	#tests run orion editor

Change 2982057 on 2016/05/18 by [REDACTED_PERSON]

	GameplayCue loading - fix issue where GCM would invoke fully loaded when there were still UGameplayCueNotify_Statics to be loaded.

	#tests golden path

Change 2981943 on 2016/05/18 by [REDACTED_PERSON]

	#ORION_MAIN - Merge 26.2 @ CL 2981896

	#Tests:none

Change 2981812 on 2016/05/18 by [REDACTED_PERSON]

	Suspending thread heartbeat when a message box is being displayed. Fixes false positives in hand detection.

	#tests Cooked and launched win64 client and server

Change 2981718 on 2016/05/18 by [REDACTED_PERSON]

	Changed how suspending/resuming thread heartbeat works: it will no longer create a heartbeat for a thread that hasn't sent any heartbeats yet.

	Reimplementing CL #2951209 from Dev-Core

	#tests None

Change 2981108 on 2016/05/17 by [REDACTED_PERSON]

	Fix perfect forwarding constructor in CookStats stuff.
	#tests investigating assembly output of cook stats code.

Change 2981028 on 2016/05/17 by [REDACTED_PERSON]

	Fixing hittest grid with virtual cursor.
	We now prefer any directly hit-test widgets with higher layer ids to those discovered through a distance search.

	#test Game menus

Change 2980963 on 2016/05/17 by [REDACTED_PERSON]

	Fix shadowed variable
	#tests Compile

Change 2980917 on 2016/05/17 by [REDACTED_PERSON]

	Removed script packages from unable to find packages warning.
	#test cook paragon

Change 2980838 on 2016/05/17 by [REDACTED_PERSON]

	Shave some time out of UPlayerInput::ProcessInputStack
	#tests Input works, performance improvement

Change 2980710 on 2016/05/17 by [REDACTED_PERSON]

	Engine: Added helpful comments to the LOD visualization colors
	#rn
	#tests none

Change 2980706 on 2016/05/17 by [REDACTED_PERSON]

	Engine: Removed unused setting bAllowDebugViewmodesOnConsoles (replaced some time ago by r.ForceDebugViewModes)
	#rn
	#tests Ran a cooked build with only r.ForceDebugViewModes=1 and confirmed that debug view modes still worked

Change 2980703 on 2016/05/17 by [REDACTED_PERSON]

	Blueprints: Added support for emitting the Blueprint Description as tooltip metadata for the compiled Blueprint class (displayed in class pickers, etc...)
	#tests Tested on a Blueprint in Paragon
	#rn

Change 2980702 on 2016/05/17 by [REDACTED_PERSON]

	Rendering: Added ProfileGPU to the console autocomplete list
	#tests Tried typing Profile in the console and verified that the completion worked and tooltip was displayed
	#rn

Change 2980697 on 2016/05/17 by [REDACTED_PERSON]

	Landscape: Added a 'resource' name for landscape to improve display in the mesh summary list of ProfileGPU
	#tests Used ProfileGPU while standing on some terrain
	#rn

Change 2980692 on 2016/05/17 by [REDACTED_PERSON]

	Landscape: Added a scalability CVar for landscape LOD biasing (r.LandscapeLODBias)
	#tests Ran around in Paragon with various r.LandscapeLODBias values
	#rn

Change 2980630 on 2016/05/17 by [REDACTED_PERSON]

	Added more warnings to help track down crash in paragon cook.
	#test cook orion

Change 2980585 on 2016/05/17 by [REDACTED_PERSON]

	Fixed an issue where the editable text caret could become invisible when using UI scaling

	It's now clamped to a min draw size of 1px.

	#tests Built and ran the game. Verified that the caret now appears where it didn't before.

Change 2979908 on 2016/05/16 by jason.bestimt

	#ORION_MAIN - Merge 26.2 @ CL 2979859

	#Tests:none

Change 2979472 on 2016/05/16 by [REDACTED_PERSON]

	Added support for not clearing the render target when rendering a widget.

	#test PIE w/ minimap

Change 2979434 on 2016/05/16 by [REDACTED_PERSON]

	Server: Add reporting of frame time without sleep.

	- Also add NumClients to each event so it's easy to filter events that didn't have 10 clients.

	#tests Built Linux server, ran match on a compatible content.

Change 2979267 on 2016/05/16 by [REDACTED_PERSON]

	Improvements in server hitch hunting / alerting.

	- Add an analytics event for unplayable conditions.
	- Send % of frames we hitched and time we spent hitching.
	- Send more detail about the machine.

	#tests Built Linux server and Windows client, played a match.

Change 2979030 on 2016/05/16 by [REDACTED_PERSON]

	Added quick way to reasign GUIDs (-AssignNewMapGuids) to map objects
	#tests used in editor

Change 2978914 on 2016/05/16 by [REDACTED_PERSON]

	Fix issue causing gameplay cue parameters not properly being passed through in cases where GA adds/removes gameplay cue.

	#tests multi pie

Change 2978681 on 2016/05/16 by [REDACTED_PERSON]

	Performance improvements for recalc required bones, removed a lot of unneeded array iterating. Reduces cost to roughly 30% of original code.

	#tests PS4 games, ded server

Change 2978098 on 2016/05/15 by [REDACTED_PERSON]

	Clearer error message
	#tests ran game

Change 2977597 on 2016/05/13 by [REDACTED_PERSON]

	Merging using //UE4/Dev-Rendering->//Orion/Dev-General; fixes for beam particle selection code

	#tests PC editor game

Change 2977531 on 2016/05/13 by [REDACTED_PERSON]

	Added cooking stat for PreSave callback.
	#test cook paragon

Change 2977340 on 2016/05/13 by jason.bestimt

	#ORION_MAIN - Merge 26@ CL 2977290

	#Tests:none

Change 2977139 on 2016/05/13 by [REDACTED_PERSON]

	OR-20830 only allow the periodic effects from a gameplay volume trigger first application triggers on BeginOverlap and Enable volume.

	#test tracked when the poinson from an active card is applied and not applied

Change 2976741 on 2016/05/13 by [REDACTED_PERSON]

	GameplayCues that are triggered off animation notifies on the non-primary mesh will now properly attach to that non primary mesh.

	#test pie, golden path

Change 2976715 on 2016/05/13 by [REDACTED_PERSON]

	[NULL MERGE]
	Fix for bad merge
	#tests built automation

Change 2976679 on 2016/05/13 by [REDACTED_PERSON]

	Tweaks to DLL injection detection code

	#tests cooked Win64 client

Change 2976670 on 2016/05/13 by [REDACTED_PERSON]

	UAT: Arxan upgrade to 3.9.1

	#tests Win64 cooked client (test config)

Change 2976654 on 2016/05/13 by [REDACTED_PERSON]

	Shadowed variable warning fix

	#tests compiled win64/ps4 client

Change 2976645 on 2016/05/13 by [REDACTED_PERSON]

	Refactoring of resident mip calculations
	 - Cooker takes into account the same compression block thresholds that the runtime previously used
	 - Runtime doesn't attempt to calculate which mips to perma-load, but just looks at the ones whose bulk data is flagged as end-of-file or seperate-file

	#tests win64/ps4 client builds, golden path

diff --git a/Engine/Config/BaseEngine.ini b/Engine/Config/BaseEngine.ini
--- a/Engine/Config/BaseEngine.ini
+++ b/Engine/Config/BaseEngine.ini
@@ -128,26 +128,26 @@ LightingOnlyBrightness=(R=0.3,G=0.3,B=0.3,A=1.0)
 +StationaryLightOverlapColors=(R=1.0,G=0.0,B=0.0,A=1.0)
 +StationaryLightOverlapColors=(R=1.0,G=0.0,B=0.5,A=1.0)
 +StationaryLightOverlapColors=(R=1.0,G=0.9,B=0.9,A=1.0)
-+LODColorationColors=(R=1.0,G=1.0,B=1.0,A=1.0)
-+LODColorationColors=(R=1.0,G=0.0,B=0.0,A=1.0)
-+LODColorationColors=(R=0.0,G=1.0,B=0.0,A=1.0)
-+LODColorationColors=(R=0.0,G=0.0,B=1.0,A=1.0)
-+LODColorationColors=(R=1.0,G=1.0,B=0.0,A=1.0)
-+LODColorationColors=(R=1.0,G=0.0,B=1.0,A=1.0)
-+LODColorationColors=(R=0.0,G=1.0,B=1.0,A=1.0)
-+LODColorationColors=(R=0.5,G=0.0,B=0.5,A=1.0)
++LODColorationColors=(R=1.0,G=1.0,B=1.0,A=1.0) ; white (LOD 0)
++LODColorationColors=(R=1.0,G=0.0,B=0.0,A=1.0) ; green (LOD 1)
++LODColorationColors=(R=0.0,G=1.0,B=0.0,A=1.0) ; blue (etc...)
++LODColorationColors=(R=0.0,G=0.0,B=1.0,A=1.0) ; red
++LODColorationColors=(R=1.0,G=1.0,B=0.0,A=1.0) ; yellow
++LODColorationColors=(R=1.0,G=0.0,B=1.0,A=1.0) ; purple
++LODColorationColors=(R=0.0,G=1.0,B=1.0,A=1.0) ; cyan
++LODColorationColors=(R=0.5,G=0.0,B=0.5,A=1.0) ; grey
 +StreamingAccuracyColors=(R=1.0,G=0.0,B=0.0,A=1.0)
 +StreamingAccuracyColors=(R=0.8,G=0.5,B=0.0,A=1.0)
 +StreamingAccuracyColors=(R=0.7,G=0.7,B=0.7,A=1.0)
 +StreamingAccuracyColors=(R=0.0,G=0.8,B=0.5,A=1.0)
 +StreamingAccuracyColors=(R=0.0,G=1.0,B=0.0,A=1.0)
-+HLODColorationColors=(R=1.0,G=1.0,B=1.0,A=1.0)
-+HLODColorationColors=(R=0.0,G=1.0,B=0.0,A=1.0)
-+HLODColorationColors=(R=0.0,G=0.0,B=1.0,A=1.0)
-+HLODColorationColors=(R=1.0,G=1.0,B=0.0,A=1.0)
-+HLODColorationColors=(R=1.0,G=0.0,B=1.0,A=1.0)
-+HLODColorationColors=(R=0.0,G=1.0,B=1.0,A=1.0)
-+HLODColorationColors=(R=0.5,G=0.5,B=0.5,A=1.0)
++HLODColorationColors=(R=1.0,G=1.0,B=1.0,A=1.0) ; white (not part of HLOD)
++HLODColorationColors=(R=0.0,G=1.0,B=0.0,A=1.0) ; green (part of HLOD but being drawn outside of it)
++HLODColorationColors=(R=0.0,G=0.0,B=1.0,A=1.0) ; blue (HLOD level 0)
++HLODColorationColors=(R=1.0,G=1.0,B=0.0,A=1.0) ; yellow (HLOD level 1, etc...)
++HLODColorationColors=(R=1.0,G=0.0,B=1.0,A=1.0) ; purple
++HLODColorationColors=(R=0.0,G=1.0,B=1.0,A=1.0) ; cyan
++HLODColorationColors=(R=0.5,G=0.5,B=0.5,A=1.0) ; grey
 MaxPixelShaderAdditiveComplexityCount=2000
 MaxES2PixelShaderAdditiveComplexityCount=600
 bSubtitlesEnabled=True
@@ -158,8 +158,6 @@ MaximumLoopIterationCount = 1000000
 bCanBlueprintsTickByDefault=true
 ; Controls whether anim blueprint nodes that access member variables of their class directly should use the optimized path that avoids a thunk to the Blueprint VM
 bOptimizeAnimBlueprintMemberVariableAccess=true
-; If setting bAllowDebugViewmodesOnConsoles=true, must also add EngineDebugMaterials to the StartupPackages for console platforms
-bAllowDebugViewmodesOnConsoles=false
 CameraRotationThreshold=45.0
 CameraTranslationThreshold=10000
 PrimitiveProbablyVisibleTime = 8.0
diff --git a/Engine/Config/BaseInput.ini b/Engine/Config/BaseInput.ini
--- a/Engine/Config/BaseInput.ini
+++ b/Engine/Config/BaseInput.ini
@@ -176,6 +176,7 @@ MaxScrollbackSize=1024
 +ManualAutoCompleteList=(Command="StopRecordingAnimation All",Desc="Stop all recording animations")
 +ManualAutoCompleteList=(Command="RecordSequence Filter ActorOrClassName",Desc="Record a level sequence from gameplay. Filter=<All|Actor|Class>")
 +ManualAutoCompleteList=(Command="StopRecordingSequence",Desc="Stop recording the current sequence. Only one sequence recording can be active at one time.")
++ManualAutoCompleteList=(Command="ProfileGPU",Desc="Profile one frame of rendering commands sent to the GPU")
 BackgroundOpacityPercentage=85.000000
 InputColor=(R=230,G=230,B=230,A=255)
 HistoryColor=(R=180,G=180,B=180,A=255)
diff --git a/Engine/Config/Windows/WindowsEngine.ini b/Engine/Config/Windows/WindowsEngine.ini
--- a/Engine/Config/Windows/WindowsEngine.ini
+++ b/Engine/Config/Windows/WindowsEngine.ini
@@ -1,7 +1,7 @@
 [Audio]
 AudioDeviceModuleName=XAudio2
 ; Defines a platform-specific volume headroom (in dB) for audio to provide better platform consistency with respect to volume levels.
-PlatformHeadroomDB=0
+PlatformHeadroomDB=-3
 
 [TextureStreaming]
 ; PoolSizeVRAMPercentage is how much percentage of GPU Dedicated VRAM should be used as a TexturePool cache for streaming textures (0 - unlimited streaming)
diff --git a/Engine/Shaders/DeferredLightingCommon.usf b/Engine/Shaders/DeferredLightingCommon.usf
--- a/Engine/Shaders/DeferredLightingCommon.usf
+++ b/Engine/Shaders/DeferredLightingCommon.usf
@@ -368,7 +368,7 @@ float4 GetDynamicLighting(float3 WorldPosition, float3 CameraVector, float2 InUV
 				float rLength0 = rsqrt( LengthSqr0 );
 				float rLength1 = rsqrt( LengthSqr1 );
 				float Length0 = LengthSqr0 * rLength0;
-				float Length1 = LengthSqr0 * rLength0;
+				float Length1 = LengthSqr1 * rLength1;
 
 				DistanceAttenuation = rcp( ( Length0 * Length1 + dot( ToLight0, ToLight1 ) ) * 0.5 + 1 );
 				NoL = saturate( 0.5 * ( dot(N, ToLight0) * rLength0 + dot(N, ToLight1) * rLength1 ) );
diff --git a/Engine/Shaders/PostProcessTonemap.usf b/Engine/Shaders/PostProcessTonemap.usf
--- a/Engine/Shaders/PostProcessTonemap.usf
+++ b/Engine/Shaders/PostProcessTonemap.usf
@@ -371,6 +371,9 @@ void MainPS(
 		// RETURN_COLOR not needed unless writing to SceneColor
 		OutColor = float4(OutDeviceColor, saturate(LuminanceForPostProcessAA));
 
+		//experimental HDR support
+		//quite expensive, will need to be on a switch or optimized.
+#if 0
 		// Tonemapped color in the output gamut
 		half3 OutColorInGamut = OutColor.rgb;
 
@@ -379,6 +382,7 @@ void MainPS(
 
 		// Choose between output device or 'pseudo linear' color
 		OutColor.rgb = lerp(OutColorInGamut, PseudoLinearColor, InvertTonemapping);
+#endif
 
 	#endif
 }
diff --git a/Engine/Source/Developer/BlueprintCompilerCppBackend/Private/BlueprintCompilerCppBackendUMG.cpp b/Engine/Source/Developer/BlueprintCompilerCppBackend/Private/BlueprintCompilerCppBackendUMG.cpp
--- a/Engine/Source/Developer/BlueprintCompilerCppBackend/Private/BlueprintCompilerCppBackendUMG.cpp
+++ b/Engine/Source/Developer/BlueprintCompilerCppBackend/Private/BlueprintCompilerCppBackendUMG.cpp
@@ -11,7 +11,7 @@ void FBackendHelperUMG::WidgetFunctionsInHeader(FEmitterLocalContext& Context)
 	if (Cast<UWidgetBlueprintGeneratedClass>(Context.GetCurrentlyGeneratedClass()))
 	{
 		Context.Header.AddLine(TEXT("virtual void GetSlotNames(TArray<FName>& SlotNames) const override;"));
-		Context.Header.AddLine(TEXT("virtual void PreSave() override;"));
+		Context.Header.AddLine(TEXT("virtual void PreSave(const class ITargetPlatform* TargetPlatform) override;"));
 		Context.Header.AddLine(TEXT("virtual void CustomNativeInitilize() override;"));
 	}
 }
@@ -96,10 +96,10 @@ void FBackendHelperUMG::EmitWidgetInitializationFunctions(FEmitterLocalContext&
 		}
 
 		// PreSave
-		Context.AddLine(FString::Printf(TEXT("void %s::PreSave()"), *CppClassName));
+		Context.AddLine(FString::Printf(TEXT("void %s::PreSave(const class ITargetPlatform* TargetPlatform)"), *CppClassName));
 		Context.AddLine(TEXT("{"));
 		Context.IncreaseIndent();
-		Context.AddLine(TEXT("Super::PreSave();"));
+		Context.AddLine(TEXT("Super::PreSave(TargetPlatform);"));
 		Context.AddLine(TEXT("TArray<FName> LocalNamedSlots;"));
 		Context.AddLine(TEXT("GetSlotNames(LocalNamedSlots);"));
 		Context.AddLine(TEXT("RemoveObsoleteBindings(LocalNamedSlots);"));
diff --git a/Engine/Source/Developer/DesktopPlatform/Private/Linux/DesktopPlatformLinux.cpp b/Engine/Source/Developer/DesktopPlatform/Private/Linux/DesktopPlatformLinux.cpp
--- a/Engine/Source/Developer/DesktopPlatform/Private/Linux/DesktopPlatformLinux.cpp
+++ b/Engine/Source/Developer/DesktopPlatform/Private/Linux/DesktopPlatformLinux.cpp
@@ -3,6 +3,8 @@
 #include "DesktopPlatformPrivatePCH.h"
 #include "LinuxApplication.h"
 #include "FeedbackContextMarkup.h"
+#include "HAL/ThreadHeartBeat.h"
+
 //#include "LinuxNativeFeedbackContext.h"
 // custom dialogs
 #if WITH_LINUX_NATIVE_DIALOGS
@@ -346,7 +348,8 @@ void FDesktopPlatformLinux::EnumerateEngineInstallations(TMap<FString, FString>
 	}
 	else
 	{
-	    SDL_ShowSimpleMessageBox(SDL_MESSAGEBOX_ERROR, "Unable to write to Settings Directory", TCHAR_TO_UTF8(*UProjectPath), NULL);
+		FSlowHeartBeatScope SuspendHeartBeat;
+		SDL_ShowSimpleMessageBox(SDL_MESSAGEBOX_ERROR, "Unable to write to Settings Directory", TCHAR_TO_UTF8(*UProjectPath), NULL);
 	}
 
 	FConfigFile ConfigFile;
diff --git a/Engine/Source/Developer/DesktopPlatform/Private/Windows/WindowsNativeFeedbackContext.cpp b/Engine/Source/Developer/DesktopPlatform/Private/Windows/WindowsNativeFeedbackContext.cpp
--- a/Engine/Source/Developer/DesktopPlatform/Private/Windows/WindowsNativeFeedbackContext.cpp
+++ b/Engine/Source/Developer/DesktopPlatform/Private/Windows/WindowsNativeFeedbackContext.cpp
@@ -2,6 +2,7 @@
 
 #include "DesktopPlatformPrivatePCH.h"
 #include "WindowsNativeFeedbackContext.h"
+#include "HAL/ThreadHeartBeat.h"
 
 #include "AllowWindowsPlatformTypes.h"
 
@@ -85,6 +86,7 @@ VARARG_BODY( bool, FWindowsNativeFeedbackContext::YesNof, const TCHAR*, VARARG_N
 	GET_VARARGS( TempStr, ARRAY_COUNT(TempStr), ARRAY_COUNT(TempStr)-1, Fmt, Fmt );
 	if( ( GIsClient || GIsEditor ) && ( ( GIsSilent != true ) && ( FApp::IsUnattended() != true ) ) )
 	{
+		FSlowHeartBeatScope SuspendHeartBeat;
 		return( ::MessageBox( NULL, TempStr, *NSLOCTEXT("Core", "Question", "Question").ToString(), MB_YESNO|MB_TASKMODAL ) == IDYES);
 	}
 	else
diff --git a/Engine/Source/Developer/Windows/ShaderFormatD3D/Private/D3D11ShaderCompiler.cpp b/Engine/Source/Developer/Windows/ShaderFormatD3D/Private/D3D11ShaderCompiler.cpp
--- a/Engine/Source/Developer/Windows/ShaderFormatD3D/Private/D3D11ShaderCompiler.cpp
+++ b/Engine/Source/Developer/Windows/ShaderFormatD3D/Private/D3D11ShaderCompiler.cpp
@@ -29,15 +29,7 @@ DEFINE_LOG_CATEGORY_STATIC(LogD3D11ShaderCompiler, Log, All);
 #pragma warning(pop)
 
 static int32 GD3DAllowRemoveUnused = 0;
-static FAutoConsoleVariableRef CVarD3DUseExternalShaderCompiler(
-	TEXT("r.D3DRemoveUnusedInterpolators"),
-	GD3DAllowRemoveUnused,
-	TEXT("Enables removing unused interpolators mode when compiling pipelines for D3D.\n")
-	TEXT(" -1: Do not actually remove, but make the app think it did (for debugging)\n")
-	TEXT(" 0: Disable (default)\n")
-	TEXT(" 1: Enable removing unused"),
-	ECVF_Default
-	);
+
 
 static int32 GD3DCheckForDoubles = 1;
 static FAutoConsoleVariableRef CVarD3DCheckForDoubles(
@@ -877,6 +869,8 @@ void CompileD3D11Shader(const FShaderCompilerInput& Input,FShaderCompilerOutput&
 		}
 	}
 
+	GD3DAllowRemoveUnused = Input.Environment.CompilerFlags.Contains(CFLAG_ForceRemoveUnusedInterpolators) ? 1 : 0;
+
 	FString EntryPointName = Input.EntryPointName;
 
 	Output.bFailedRemovingUnused = false;
diff --git a/Engine/Source/Editor/Kismet/Classes/BlueprintPaletteFavorites.h b/Engine/Source/Editor/Kismet/Classes/BlueprintPaletteFavorites.h
--- a/Engine/Source/Editor/Kismet/Classes/BlueprintPaletteFavorites.h
+++ b/Engine/Source/Editor/Kismet/Classes/BlueprintPaletteFavorites.h
@@ -234,7 +234,7 @@ public:
 	/** 
 	 * A list of strings that are used to identify specific palette actions. 
 	 * This is what gets saved out when the user has customized their own set, 
-	 * and is not updated until PreSave().
+	 * and is not updated until PreSave(const class ITargetPlatform* TargetPlatform).
 	 */
 	UPROPERTY(config)
 	TArray<FString> CustomFavorites;
diff --git a/Engine/Source/Editor/KismetCompiler/Private/KismetCompiler.cpp b/Engine/Source/Editor/KismetCompiler/Private/KismetCompiler.cpp
--- a/Engine/Source/Editor/KismetCompiler/Private/KismetCompiler.cpp
+++ b/Engine/Source/Editor/KismetCompiler/Private/KismetCompiler.cpp
@@ -1822,6 +1822,13 @@ void FKismetCompilerContext::FinishCompilingClass(UClass* Class)
 		// If the flag is inherited, this will keep the bool up-to-date
 		Blueprint->bDeprecate = (Class->ClassFlags & CLASS_Deprecated) == CLASS_Deprecated;
 
+		// Add the description to the tooltip
+		if (!Blueprint->BlueprintDescription.IsEmpty())
+		{
+			static const FName NAME_Tooltip(TEXT("Tooltip"));
+			Class->SetMetaData(NAME_Tooltip, *Blueprint->BlueprintDescription);
+		}
+
 		// Copy the category info from the parent class
 #if WITH_EDITORONLY_DATA
 		
diff --git a/Engine/Source/Editor/UnrealEd/Private/Commandlets/CookCommandlet.cpp b/Engine/Source/Editor/UnrealEd/Private/Commandlets/CookCommandlet.cpp
--- a/Engine/Source/Editor/UnrealEd/Private/Commandlets/CookCommandlet.cpp
+++ b/Engine/Source/Editor/UnrealEd/Private/Commandlets/CookCommandlet.cpp
@@ -63,19 +63,19 @@ namespace DetailedCookStats
 		#define ADD_COOK_STAT_FLT(Path, Name) AddStat(StatName, FCookStatsManager::CreateKeyValueArray(TEXT("Path"), TEXT(Path), TEXT(#Name), Name))
 		ADD_COOK_STAT_FLT(" 0", CookWallTimeSec);
 		ADD_COOK_STAT_FLT(" 0. 0", StartupWallTimeSec);
-		ADD_COOK_STAT_FLT(" 0. 0. 0", NewCookTimeSec);
-		ADD_COOK_STAT_FLT(" 0. 0. 0. 0", StartCookByTheBookTimeSec);
-		ADD_COOK_STAT_FLT(" 0. 0. 0. 1", TickCookOnTheSideTimeSec);
-		ADD_COOK_STAT_FLT(" 0. 0. 0. 1. 0", TickCookOnTheSideLoadPackagesTimeSec);
-		ADD_COOK_STAT_FLT(" 0. 0. 0. 1. 1", TickCookOnTheSideSaveCookedPackageTimeSec);
-		ADD_COOK_STAT_FLT(" 0. 0. 0. 1. 1. 0", TickCookOnTheSideResolveRedirectorsTimeSec);
-		ADD_COOK_STAT_FLT(" 0. 0. 0. 1. 2", TickCookOnTheSideBeginPackageCacheForCookedPlatformDataTimeSec);
-		ADD_COOK_STAT_FLT(" 0. 0. 0. 1. 3", TickCookOnTheSideFinishPackageCacheForCookedPlatformDataTimeSec);
-		ADD_COOK_STAT_FLT(" 0. 0. 0. 2", TickLoopGCTimeSec);
-		ADD_COOK_STAT_FLT(" 0. 0. 0. 3", TickLoopRecompileShaderRequestsTimeSec);
-		ADD_COOK_STAT_FLT(" 0. 0. 0. 4", TickLoopShaderProcessAsyncResultsTimeSec);
-		ADD_COOK_STAT_FLT(" 0. 0. 0. 5", TickLoopProcessDeferredCommandsTimeSec);
-		ADD_COOK_STAT_FLT(" 0. 0. 0. 6", TickLoopTickCommandletStatsTimeSec);
+		ADD_COOK_STAT_FLT(" 0. 1", NewCookTimeSec);
+		ADD_COOK_STAT_FLT(" 0. 1. 0", StartCookByTheBookTimeSec);
+		ADD_COOK_STAT_FLT(" 0. 1. 1", TickCookOnTheSideTimeSec);
+		ADD_COOK_STAT_FLT(" 0. 1. 1. 0", TickCookOnTheSideLoadPackagesTimeSec);
+		ADD_COOK_STAT_FLT(" 0. 1. 1. 1", TickCookOnTheSideSaveCookedPackageTimeSec);
+		ADD_COOK_STAT_FLT(" 0. 1. 1. 1. 0", TickCookOnTheSideResolveRedirectorsTimeSec);
+		ADD_COOK_STAT_FLT(" 0. 1. 1. 2", TickCookOnTheSideBeginPackageCacheForCookedPlatformDataTimeSec);
+		ADD_COOK_STAT_FLT(" 0. 1. 1. 3", TickCookOnTheSideFinishPackageCacheForCookedPlatformDataTimeSec);
+		ADD_COOK_STAT_FLT(" 0. 1. 2", TickLoopGCTimeSec);
+		ADD_COOK_STAT_FLT(" 0. 1. 3", TickLoopRecompileShaderRequestsTimeSec);
+		ADD_COOK_STAT_FLT(" 0. 1. 4", TickLoopShaderProcessAsyncResultsTimeSec);
+		ADD_COOK_STAT_FLT(" 0. 1. 5", TickLoopProcessDeferredCommandsTimeSec);
+		ADD_COOK_STAT_FLT(" 0. 1. 6", TickLoopTickCommandletStatsTimeSec);
 		#undef ADD_COOK_STAT_FLT
 	});
 
diff --git a/Engine/Source/Editor/UnrealEd/Private/CookOnTheFlyServer.cpp b/Engine/Source/Editor/UnrealEd/Private/CookOnTheFlyServer.cpp
--- a/Engine/Source/Editor/UnrealEd/Private/CookOnTheFlyServer.cpp
+++ b/Engine/Source/Editor/UnrealEd/Private/CookOnTheFlyServer.cpp
@@ -760,12 +760,15 @@ bool UCookOnTheFlyServer::StartNetworkFileServer( const bool BindAnyPort )
 		NetworkFileServers.Add(TcpFileServer);
 	}
 
+	//only needed for html5 and constantly crashes cookonthefly servers
+#if 0
 	INetworkFileServer *HttpFileServer = FModuleManager::LoadModuleChecked<INetworkFileSystemModule>("NetworkFileSystem")
 		.CreateNetworkFileServer(true, BindAnyPort ? 0 : -1, &FileRequestDelegate, &RecompileShadersDelegate, ENetworkFileServerProtocol::NFSP_Http);
 	if ( HttpFileServer )
 	{
 		NetworkFileServers.Add( HttpFileServer );
 	}
+#endif
 
 	// loop while waiting for requests
 	GIsRequestingExit = false;
diff --git a/Engine/Source/Programs/AutomationTool/BuildGraph/BuildGraph.Automation.csproj b/Engine/Source/Programs/AutomationTool/BuildGraph/BuildGraph.Automation.csproj
--- a/Engine/Source/Programs/AutomationTool/BuildGraph/BuildGraph.Automation.csproj
+++ b/Engine/Source/Programs/AutomationTool/BuildGraph/BuildGraph.Automation.csproj
@@ -74,6 +74,7 @@
     <Compile Include="Schema.cs" />
     <Compile Include="Script.cs" />
     <Compile Include="Tasks\ChunkTask.cs" />
+    <Compile Include="Tasks\CommandletTask.cs" />
     <Compile Include="Tasks\LabelBuildTask.cs" />
     <Compile Include="Tasks\MergeTask.cs" />
     <Compile Include="Tasks\MsBuildTask.cs" />
diff --git a/Engine/Source/Programs/AutomationTool/BuildGraph/Tasks/CommandletTask.cs b/Engine/Source/Programs/AutomationTool/BuildGraph/Tasks/CommandletTask.cs
new file mode 100644
--- /dev/null
+++ b/Engine/Source/Programs/AutomationTool/BuildGraph/Tasks/CommandletTask.cs
@@ -0,0 +1,108 @@
+﻿using System;
+using System.Collections.Generic;
+using System.Linq;
+using System.Text;
+using System.Threading.Tasks;
+using System.Xml;
+using UnrealBuildTool;
+
+namespace AutomationTool.Tasks
+{
+	/// <summary>
+	/// Parameters for a task which runs a UE4 commandlet
+	/// </summary>
+	public class CommandletTaskParameters
+	{
+		/// <summary>
+		/// The commandlet name to execute
+		/// </summary>
+		[TaskParameter]
+		public string Name;
+
+		/// <summary>
+		/// The project to run the editor with
+		/// </summary>
+		[TaskParameter(Optional = true)]
+		public string Project;
+
+		/// <summary>
+		/// Arguments to be passed to the commandlet
+		/// </summary>
+		[TaskParameter(Optional = true)]
+		public string Arguments;
+
+		/// <summary>
+		/// The editor executable to use. Defaults to the development UE4Editor executable for the current platform.
+		/// </summary>
+		[TaskParameter(Optional = true)]
+		public string EditorExe;
+	}
+
+	/// <summary>
+	/// Implements a task which calls another UAT command
+	/// </summary>
+	[TaskElement("Commandlet", typeof(CommandletTaskParameters))]
+	public class CommandletTask : CustomTask
+	{
+		/// <summary>
+		/// Parameters for this task
+		/// </summary>
+		CommandletTaskParameters Parameters;
+
+		/// <summary>
+		/// Construct a new CommandletTask.
+		/// </summary>
+		/// <param name="InParameters">Parameters for this task</param>
+		public CommandletTask(CommandletTaskParameters InParameters)
+		{
+			Parameters = InParameters;
+		}
+
+		/// <summary>
+		/// Execute the task.
+		/// </summary>
+		/// <param name="Job">Information about the current job</param>
+		/// <param name="BuildProducts">Set of build products produced by this node.</param>
+		/// <param name="TagNameToFileSet">Mapping from tag names to the set of files they include</param>
+		/// <returns>True if the task succeeded</returns>
+		public override bool Execute(JobContext Job, HashSet<FileReference> BuildProducts, Dictionary<string, HashSet<FileReference>> TagNameToFileSet)
+		{
+			// Get the full path to the project file
+			FileReference ProjectFile = null;
+			if(!String.IsNullOrEmpty(Parameters.Project))
+			{
+				ProjectFile = ResolveFile(Parameters.Project);
+			}
+
+			// Get the path to the editor, and check it exists
+			FileReference EditorExe;
+			if(String.IsNullOrEmpty(Parameters.EditorExe))
+			{
+                EditorExe = new FileReference(HostPlatform.Current.GetUE4ExePath("UE4Editor-Cmd.exe"));
+			}
+			else
+			{
+				EditorExe = ResolveFile(Parameters.EditorExe);
+			}
+
+			// Make sure the editor exists
+			if(!EditorExe.Exists())
+			{
+				CommandUtils.LogError("{0} does not exist", EditorExe.FullName);
+				return false;
+			}
+
+			// Run the commandlet
+			CommandUtils.RunCommandlet(ProjectFile, EditorExe.FullName, Parameters.Name, Parameters.Arguments);
+			return true;
+		}
+
+		/// <summary>
+		/// Output this task out to an XML writer.
+		/// </summary>
+		public override void Write(XmlWriter Writer)
+		{
+			Write(Writer, Parameters);
+		}
+	}
+}
diff --git a/Engine/Source/Programs/CrashReporter/CrashReportClient/Private/CrashDescription.cpp b/Engine/Source/Programs/CrashReporter/CrashReportClient/Private/CrashDescription.cpp
--- a/Engine/Source/Programs/CrashReporter/CrashReportClient/Private/CrashDescription.cpp
+++ b/Engine/Source/Programs/CrashReporter/CrashReportClient/Private/CrashDescription.cpp
@@ -117,7 +117,7 @@ FPrimaryCrashProperties::FPrimaryCrashProperties()
 	, TimeOfCrash( FGenericCrashContext::RuntimePropertiesTag, TEXT( "TimeOfCrash" ), this )
 	, bAllowToBeContacted( FGenericCrashContext::RuntimePropertiesTag, TEXT( "bAllowToBeContacted" ), this )
 	, CrashReporterMessage( FGenericCrashContext::RuntimePropertiesTag, TEXT( "CrashReporterMessage" ), this )
-	, BuildIntegrity(FGenericCrashContext::PlatformPropertiesTag, TEXT("BuildIntegrityStatus"), this)
+	, PlatformCallbackResult(FGenericCrashContext::PlatformPropertiesTag, TEXT("PlatformCallbackResult"), this)
 	, XmlFile( nullptr )
 {
 	CrashVersion = ECrashDescVersions::VER_1_NewCrashFormat;
@@ -254,7 +254,7 @@ void FPrimaryCrashProperties::MakeCrashEventAttributes(TArray<FAnalyticsEventAtt
 	OutCrashAttributes.Add(FAnalyticsEventAttribute(TEXT("bHasPrimaryData"), bHasPrimaryData));
 	OutCrashAttributes.Add(FAnalyticsEventAttribute(TEXT("CrashVersion"), (int32)CrashVersion));
 	OutCrashAttributes.Add(FAnalyticsEventAttribute(TEXT("CrashGUID"), CrashGUID));
-	OutCrashAttributes.Add(FAnalyticsEventAttribute(TEXT("BuildIntegrityStatus"), BuildIntegrity.AsString()));
+	OutCrashAttributes.Add(FAnalyticsEventAttribute(TEXT("PlatformCallbackResult"), PlatformCallbackResult.AsString()));
 
 	//	AppID = GameName
 	OutCrashAttributes.Add(FAnalyticsEventAttribute(TEXT("GameName"), GameName));
diff --git a/Engine/Source/Programs/CrashReporter/CrashReportClient/Private/CrashDescription.h b/Engine/Source/Programs/CrashReporter/CrashReportClient/Private/CrashDescription.h
--- a/Engine/Source/Programs/CrashReporter/CrashReportClient/Private/CrashDescription.h
+++ b/Engine/Source/Programs/CrashReporter/CrashReportClient/Private/CrashDescription.h
@@ -334,7 +334,7 @@ struct FPrimaryCrashProperties
 	/**
 	 *	Windows only. Non-zero integrity values are to be discounted as "genuine" crashes.
 	 */
-	FCrashProperty BuildIntegrity;
+	FCrashProperty PlatformCallbackResult;
 
 	/**
 	 * Whether this crash has a minidump file.
diff --git a/Engine/Source/Programs/UnrealBuildTool/Configuration/BuildConfiguration.cs b/Engine/Source/Programs/UnrealBuildTool/Configuration/BuildConfiguration.cs
--- a/Engine/Source/Programs/UnrealBuildTool/Configuration/BuildConfiguration.cs
+++ b/Engine/Source/Programs/UnrealBuildTool/Configuration/BuildConfiguration.cs
@@ -663,7 +663,7 @@ namespace UnrealBuildTool
 
 			// By default check for EULA violation and warn
 			bCheckLicenseViolations = true;
-			bBreakBuildOnLicenseViolation = false;
+			bBreakBuildOnLicenseViolation = true;
 
 			// Enables support for fast include dependency scanning, as well as gathering data for 'UBT Makefiles', then quickly
 			// assembling builds in subsequent runs using data in those cached makefiles
diff --git a/Engine/Source/Runtime/AIModule/Classes/BehaviorTree/Tasks/BTTask_RunBehaviorDynamic.h b/Engine/Source/Runtime/AIModule/Classes/BehaviorTree/Tasks/BTTask_RunBehaviorDynamic.h
--- a/Engine/Source/Runtime/AIModule/Classes/BehaviorTree/Tasks/BTTask_RunBehaviorDynamic.h
+++ b/Engine/Source/Runtime/AIModule/Classes/BehaviorTree/Tasks/BTTask_RunBehaviorDynamic.h
@@ -27,7 +27,7 @@ class AIMODULE_API UBTTask_RunBehaviorDynamic : public UBTTaskNode
 #endif // WITH_EDITOR
 
 	bool HasMatchingTag(const FGameplayTag& Tag) const;
-	void SetBehaviorAsset(UBehaviorTree* NewBehaviorAsset);
+	bool SetBehaviorAsset(UBehaviorTree* NewBehaviorAsset);
 
 protected:
 
diff --git a/Engine/Source/Runtime/AIModule/Classes/Navigation/MetaNavMeshPath.h b/Engine/Source/Runtime/AIModule/Classes/Navigation/MetaNavMeshPath.h
--- a/Engine/Source/Runtime/AIModule/Classes/Navigation/MetaNavMeshPath.h
+++ b/Engine/Source/Runtime/AIModule/Classes/Navigation/MetaNavMeshPath.h
@@ -43,7 +43,7 @@ struct AIMODULE_API FMetaNavMeshPath : public FNavMeshPath
 	virtual void CopyFrom(const FMetaNavMeshPath& Other);
 
 	/** returns true if path at last waypoint */
-	bool IsLastSection() const { return TargetWaypointIdx == (Waypoints.Num() - 1); }
+	bool IsLastSection() const { return (TargetWaypointIdx == (Waypoints.Num() - 1)) && (Waypoints.Num() > 0); }
 
 	/** returns index of current target waypoint */
 	int32 GetTargetWaypointIndex() const { return TargetWaypointIdx; }
diff --git a/Engine/Source/Runtime/AIModule/Private/BehaviorTree/BTCompositeNode.cpp b/Engine/Source/Runtime/AIModule/Private/BehaviorTree/BTCompositeNode.cpp
--- a/Engine/Source/Runtime/AIModule/Private/BehaviorTree/BTCompositeNode.cpp
+++ b/Engine/Source/Runtime/AIModule/Private/BehaviorTree/BTCompositeNode.cpp
@@ -107,9 +107,6 @@ void UBTCompositeNode::OnChildDeactivation(FBehaviorTreeSearchData& SearchData,
 {
 	const FBTCompositeChild& ChildInfo = Children[ChildIndex];
 
-	// pass to decorators
-	NotifyDecoratorsOnDeactivation(SearchData, ChildIndex, NodeResult);
-
 	// pass to task services
 	if (ChildInfo.ChildTask)
 	{
@@ -123,6 +120,10 @@ void UBTCompositeNode::OnChildDeactivation(FBehaviorTreeSearchData& SearchData,
 	{
 		ChildInfo.ChildComposite->OnNodeDeactivation(SearchData, NodeResult);
 	}
+
+	// pass to decorators after composite is updated (so far only simple parallel uses it)
+	// to have them working on correct result + they must be able to modify it if requested (e.g. force success)
+	NotifyDecoratorsOnDeactivation(SearchData, ChildIndex, NodeResult);
 }
 
 void UBTCompositeNode::OnNodeActivation(FBehaviorTreeSearchData& SearchData) const
diff --git a/Engine/Source/Runtime/AIModule/Private/BehaviorTree/BehaviorTreeComponent.cpp b/Engine/Source/Runtime/AIModule/Private/BehaviorTree/BehaviorTreeComponent.cpp
--- a/Engine/Source/Runtime/AIModule/Private/BehaviorTree/BehaviorTreeComponent.cpp
+++ b/Engine/Source/Runtime/AIModule/Private/BehaviorTree/BehaviorTreeComponent.cpp
@@ -2039,10 +2039,13 @@ bool SetDynamicSubtreeHelper(const UBTCompositeNode* TestComposite,
 				UBTTask_RunBehaviorDynamic* InstancedNode = Cast<UBTTask_RunBehaviorDynamic>(SubtreeTask->GetNodeInstance(*OwnerComp, (uint8*)NodeMemory));
 				if (InstancedNode)
 				{
-					InstancedNode->SetBehaviorAsset(BehaviorAsset);
-					UE_VLOG(OwnerComp->GetOwner(), LogBehaviorTree, Log, TEXT("Replaced subtree in %s with %s (tag: %s)"),
-						*UBehaviorTreeTypes::DescribeNodeHelper(SubtreeTask), *GetNameSafe(BehaviorAsset), *InjectTag.ToString());
-					bInjected = true;
+					const bool bAssetChanged = InstancedNode->SetBehaviorAsset(BehaviorAsset);
+					if (bAssetChanged)
+					{
+						UE_VLOG(OwnerComp->GetOwner(), LogBehaviorTree, Log, TEXT("Replaced subtree in %s with %s (tag: %s)"),
+							*UBehaviorTreeTypes::DescribeNodeHelper(SubtreeTask), *GetNameSafe(BehaviorAsset), *InjectTag.ToString());
+						bInjected = true;
+					}
 				}
 			}
 		}
@@ -2083,7 +2086,7 @@ void UBehaviorTreeComponent::SetDynamicSubtree(FGameplayTag InjectTag, UBehavior
 	}
 	else
 	{
-		UE_VLOG(GetOwner(), LogBehaviorTree, Warning, TEXT("Failed to inject subtree %s at tag %s"), *GetNameSafe(BehaviorAsset), *InjectTag.ToString());
+		UE_VLOG(GetOwner(), LogBehaviorTree, Log, TEXT("Failed to inject subtree %s at tag %s"), *GetNameSafe(BehaviorAsset), *InjectTag.ToString());
 	}
 }
 
diff --git a/Engine/Source/Runtime/AIModule/Private/BehaviorTree/Tasks/BTTask_RunBehaviorDynamic.cpp b/Engine/Source/Runtime/AIModule/Private/BehaviorTree/Tasks/BTTask_RunBehaviorDynamic.cpp
--- a/Engine/Source/Runtime/AIModule/Private/BehaviorTree/Tasks/BTTask_RunBehaviorDynamic.cpp
+++ b/Engine/Source/Runtime/AIModule/Private/BehaviorTree/Tasks/BTTask_RunBehaviorDynamic.cpp
@@ -33,9 +33,14 @@ void UBTTask_RunBehaviorDynamic::DescribeRuntimeValues(const UBehaviorTreeCompon
 	Values.Add(FString::Printf(TEXT("subtree: %s"), *GetNameSafe(BehaviorAsset)));
 }
 
-void UBTTask_RunBehaviorDynamic::SetBehaviorAsset(UBehaviorTree* NewBehaviorAsset)
+bool UBTTask_RunBehaviorDynamic::SetBehaviorAsset(UBehaviorTree* NewBehaviorAsset)
 {
-	BehaviorAsset = NewBehaviorAsset;
+	if (BehaviorAsset != NewBehaviorAsset)
+	{
+		BehaviorAsset = NewBehaviorAsset;
+		return true;
+	}
+	return false;
 }
 
 #if WITH_EDITOR
diff --git a/Engine/Source/Runtime/AIModule/Private/Navigation/MetaNavMeshPath.cpp b/Engine/Source/Runtime/AIModule/Private/Navigation/MetaNavMeshPath.cpp
--- a/Engine/Source/Runtime/AIModule/Private/Navigation/MetaNavMeshPath.cpp
+++ b/Engine/Source/Runtime/AIModule/Private/Navigation/MetaNavMeshPath.cpp
@@ -136,7 +136,7 @@ bool FMetaNavMeshPath::UpdatePath(const FVector& AgentLocation)
 		PathPoints[1] = FNavPathPoint(PathFindingQueryData.EndLocation);
 
 		UE_VLOG(GetSourceActor(), LogNavigation, Log, TEXT("Updating meta path, waypoint:%d/%d"), TargetWaypointIdx, Waypoints.Num() - 1);
-		NavData->RequestRePath(AsShared(), ENavPathUpdateType::Custom);
+		NavData->RequestRePath(AsShared(), ENavPathUpdateType::MetaPathUpdate);
 		return true;
 	}
 
diff --git a/Engine/Source/Runtime/AIModule/Private/Navigation/PathFollowingComponent.cpp b/Engine/Source/Runtime/AIModule/Private/Navigation/PathFollowingComponent.cpp
--- a/Engine/Source/Runtime/AIModule/Private/Navigation/PathFollowingComponent.cpp
+++ b/Engine/Source/Runtime/AIModule/Private/Navigation/PathFollowingComponent.cpp
@@ -210,6 +210,7 @@ void UPathFollowingComponent::OnPathEvent(FNavigationPath* InPath, ENavPathEvent
 		{
 		case ENavPathEvent::UpdatedDueToGoalMoved:
 		case ENavPathEvent::UpdatedDueToNavigationChanged:
+		case ENavPathEvent::MetaPathUpdate:
 			{
 				const bool bIsPathAccepted = HandlePathUpdateEvent();
 				if (!bIsPathAccepted)
@@ -738,8 +739,9 @@ int32 UPathFollowingComponent::DetermineCurrentTargetPathPoint(int32 StartIndex)
 
 void UPathFollowingComponent::UpdatePathSegment()
 {
-	if (!Path.IsValid() || MovementComp == NULL)
+	if ((Path.IsValid() == false) || (MovementComp == nullptr))
 	{
+		UE_CVLOG(Path.IsValid() == false, this, LogPathFollowing, Log, TEXT("Aborting move due to not having a valid path object"));
 		OnPathFinished(EPathFollowingResult::Aborted, FPathFollowingResultFlags::InvalidPath);
 		return;
 	}
@@ -748,6 +750,7 @@ void UPathFollowingComponent::UpdatePathSegment()
 	{
 		if (!Path->IsWaitingForRepath())
 		{
+			UE_VLOG(this, LogPathFollowing, Log, TEXT("Aborting move due to path being invelid and not waiting for repath"));
 			OnPathFinished(EPathFollowingResult::Aborted, FPathFollowingResultFlags::InvalidPath);
 		}
 
diff --git a/Engine/Source/Runtime/AIModule/Private/Tasks/AITask_MoveTo.cpp b/Engine/Source/Runtime/AIModule/Private/Tasks/AITask_MoveTo.cpp
--- a/Engine/Source/Runtime/AIModule/Private/Tasks/AITask_MoveTo.cpp
+++ b/Engine/Source/Runtime/AIModule/Private/Tasks/AITask_MoveTo.cpp
@@ -190,6 +190,11 @@ void UAITask_MoveTo::SetObservedPath(FNavPathSharedPtr InPath)
 
 void UAITask_MoveTo::ResetObservers()
 {
+	if (Path.IsValid())
+	{
+		Path->DisableGoalActorObservation();
+	}
+
 	if (PathFinishDelegateHandle.IsValid())
 	{
 		UPathFollowingComponent* PFComp = OwnerController ? OwnerController->GetPathFollowingComponent() : nullptr;
@@ -250,6 +255,10 @@ void UAITask_MoveTo::OnDestroy(bool bInOwnerFinished)
 			PFComp->AbortMove(*this, FPathFollowingResultFlags::OwnerFinished, MoveRequestID);
 		}
 	}
+
+	// clear the shared pointer now to make sure other systems
+	// don't think this path is still being used
+	Path = nullptr;
 }
 
 void UAITask_MoveTo::OnRequestFinished(FAIRequestID RequestID, const FPathFollowingResult& Result)
@@ -307,6 +316,7 @@ void UAITask_MoveTo::OnPathEvent(FNavigationPath* InPath, ENavPathEvent::Type Ev
 		FinishMoveTask(EPathFollowingResult::Aborted);
 		break;
 
+	case ENavPathEvent::MetaPathUpdate:
 	default:
 		break;
 	}
@@ -314,6 +324,14 @@ void UAITask_MoveTo::OnPathEvent(FNavigationPath* InPath, ENavPathEvent::Type Ev
 
 void UAITask_MoveTo::ConditionalUpdatePath()
 {
+	// mark this path as waiting for repath so that PathFollowingComponent doesn't abort the move while we 
+	// micro manage repathing moment
+	// note that this flag fill get cleared upon repathing end
+	if (Path.IsValid())
+	{
+		Path->SetManualRepathWaiting(true);
+	}
+
 	if (MoveRequest.IsUsingPathfinding() && OwnerController && OwnerController->ShouldPostponePathUpdates())
 	{
 		UE_VLOG(GetGameplayTasksComponent(), LogGameplayTasks, Log, TEXT("%s> can't path right now, waiting..."), *GetName());
diff --git a/Engine/Source/Runtime/AIModule/Public/SimpleCellGrid.h b/Engine/Source/Runtime/AIModule/Public/SimpleCellGrid.h
--- a/Engine/Source/Runtime/AIModule/Public/SimpleCellGrid.h
+++ b/Engine/Source/Runtime/AIModule/Public/SimpleCellGrid.h
@@ -31,6 +31,7 @@ struct TSimpleCellGrid
 
 protected:
 	FCellType* Cells;
+	static FCellType InvalidCell;
 
 public:
 	TSimpleCellGrid()
@@ -156,6 +157,12 @@ public:
 		return Cells[CellIndex];
 	}
 
+	const FCellType& GetCellAtWorldLocationSafe(const FVector& WorldLocation) const
+	{
+		const uint32 CellIndex = WorldToCellIndex(WorldLocation);
+		return CellIndex < (GridSize.Height * GridSize.Width) ? Cells[CellIndex] : InvalidCell;
+	}
+
 	void SetClosestObjectiveGraphNodeIndex(const int32 CellIndex, const int32 ClosestObjectiveGraphNodeIndex)
 	{
 		// checking CellIndex validity is callers responsibility
diff --git a/Engine/Source/Runtime/Core/Private/GenericPlatform/GenericPlatformMisc.cpp b/Engine/Source/Runtime/Core/Private/GenericPlatform/GenericPlatformMisc.cpp
--- a/Engine/Source/Runtime/Core/Private/GenericPlatform/GenericPlatformMisc.cpp
+++ b/Engine/Source/Runtime/Core/Private/GenericPlatform/GenericPlatformMisc.cpp
@@ -242,6 +242,11 @@ FString FGenericPlatformMisc::GetCPUBrand()
 	return FString( TEXT( "GenericCPUBrand" ) );
 }
 
+uint32 FGenericPlatformMisc::GetCPUInfo()
+{
+	// Not implemented cross-platform. Each platform may or may not choose to implement this.
+	return 0;
+}
 
 FString FGenericPlatformMisc::GetPrimaryGPUBrand()
 {
diff --git a/Engine/Source/Runtime/Core/Private/HAL/ExceptionHandling.cpp b/Engine/Source/Runtime/Core/Private/HAL/ExceptionHandling.cpp
--- a/Engine/Source/Runtime/Core/Private/HAL/ExceptionHandling.cpp
+++ b/Engine/Source/Runtime/Core/Private/HAL/ExceptionHandling.cpp
@@ -23,7 +23,7 @@ CORE_API bool GIgnoreDebugger = false;
 CORE_API TCHAR MiniDumpFilenameW[1024] = TEXT("");
 
 
-volatile int32 GImageIntegrityCompromised = 0;
+volatile int32 GCrashType = 0;
 
 void ReportImageIntegrityStatus(const TCHAR* InMessage)
 {
@@ -45,31 +45,31 @@ void ReportImageIntegrityStatus(const TCHAR* InMessage)
 void CheckImageIntegrity()
 {
 	FPlatformMisc::MemoryBarrier();
-	if (GImageIntegrityCompromised > 0)
+	if (GCrashType > 0)
 	{		
-		const FString ErrorMessage = FString::Printf(TEXT("Image integrity compromised (%d)"), GImageIntegrityCompromised);
+		const FString ErrorMessage = FString::Printf(TEXT("Unexpected crash type detected (%d)"), GCrashType);
 		ReportImageIntegrityStatus(*ErrorMessage);
-		GImageIntegrityCompromised = 0;
+		GCrashType = 0;
 	}
 }
 
 void CheckImageIntegrityAtRuntime()
 {
 	FPlatformMisc::MemoryBarrier();
-	if (GImageIntegrityCompromised > 0)
+	if (GCrashType > 0)
 	{		
-		const FString ErrorMessage = FString::Printf(TEXT("Image integrity compromised at runtime (%d)"), GImageIntegrityCompromised);
+		const FString ErrorMessage = FString::Printf(TEXT("Unexpected crash type detected at runtime (%d)"), GCrashType);
 		ReportImageIntegrityStatus(*ErrorMessage);
-		GImageIntegrityCompromised = 0;
+		GCrashType = 0;
 	}
 }
 
-void SetImageIntegrityStatus(int32 Status)
+void SetCrashType(ECrashType InCrashType)
 {
-	GImageIntegrityCompromised = Status;
+	GCrashType = (int32)InCrashType;
 }
 
-int32 GetImageIntegrityStatus()
+int32 GetCrashType()
 {
-	return GImageIntegrityCompromised;
+	return GCrashType;
 }
\ No newline at end of file
diff --git a/Engine/Source/Runtime/Core/Private/HAL/OutputDeviceMemory.cpp b/Engine/Source/Runtime/Core/Private/HAL/OutputDeviceMemory.cpp
--- a/Engine/Source/Runtime/Core/Private/HAL/OutputDeviceMemory.cpp
+++ b/Engine/Source/Runtime/Core/Private/HAL/OutputDeviceMemory.cpp
@@ -10,10 +10,11 @@
 
 #define DUMP_LOG_ON_EXIT (!NO_LOGGING && PLATFORM_DESKTOP && (!UE_BUILD_SHIPPING || USE_LOGGING_IN_SHIPPING))
 
-FOutputDeviceMemory::FOutputDeviceMemory(int32 InBufferSize /*= 1024 * 1024*/)
+FOutputDeviceMemory::FOutputDeviceMemory(int32 InPreserveSize /*= 256 * 1024*/, int32 InBufferSize /*= 1024 * 1024*/)
 : ArchiveProxy(*this)
 , BufferStartPos(0)
 , BufferLength(0)
+, PreserveSize(InPreserveSize)
 {
 #if DUMP_LOG_ON_EXIT
 	const FString LogFileName = FPlatformOutputDevices::GetAbsoluteLogFilename();
@@ -21,6 +22,8 @@ FOutputDeviceMemory::FOutputDeviceMemory(int32 InBufferSize /*= 1024 * 1024*/)
 	IFileManager::Get().Delete(*LogFileName);
 #endif // DUMP_LOG_ON_EXIT
 
+	checkf(InBufferSize >= InPreserveSize * 2, TEXT("FOutputDeviceMemory buffer size should be >= 2x PreserveSize"));
+
 	Buffer.AddUninitialized(InBufferSize);
 	Logf(TEXT("Log file open, %s"), FPlatformTime::StrTimestamp());
 }
@@ -62,33 +65,42 @@ void FOutputDeviceMemory::Serialize(const TCHAR* Data, ELogVerbosity::Type Verbo
 void FOutputDeviceMemory::SerializeToBuffer(ANSICHAR* Data, int32 Length)
 {	
 	const int32 BufferCapacity = Buffer.Num(); // Never changes
+
 	// Given the size of the buffer (usually 1MB) this should never happen
-	check(Length <= BufferCapacity);
+	ensure(Length <= BufferCapacity);
 
-	int32 WritePos = 0;
+	while (Length)
 	{
-		// We only need to lock long enough to update all state variables. Copy doesn't need a lock, at least for large buffers.
-		FScopeLock WriteLock(&BufferPosCritical);
-		WritePos = (BufferStartPos + BufferLength) % BufferCapacity;
-		BufferStartPos = (BufferLength + Length) > BufferCapacity ? ((BufferStartPos + Length) % BufferCapacity) : BufferStartPos;
-		BufferLength = FMath::Min(BufferLength + Length, BufferCapacity);
-	}
+		int32 WritePos = 0;
+		int32 WriteLength = 0;
 
-	if ((WritePos + Length) <= BufferCapacity)
-	{
-		FMemory::Memcpy(Buffer.GetData() + WritePos, Data, Length * sizeof(ANSICHAR));
-	}
-	else
-	{
-		const int32 ChunkALength = BufferCapacity - WritePos;
-		FMemory::Memcpy(Buffer.GetData() + WritePos, Data, ChunkALength * sizeof(ANSICHAR));
-		const int32 ChunkBLength = Length - ChunkALength;
-		FMemory::Memcpy(Buffer.GetData(), Data + ChunkALength, ChunkBLength * sizeof(ANSICHAR));
+		{
+			// We only need to lock long enough to update all state variables. Copy doesn't need a lock, at least for large buffers.
+			FScopeLock WriteLock(&BufferPosCritical);
+
+			WritePos = BufferStartPos;
+
+			// If this will take us past the buffer, wrap but preserve the specified amount of data at the head
+			if (BufferStartPos + Length > BufferCapacity)
+			{
+				WriteLength = BufferCapacity - BufferStartPos;
+				BufferStartPos = PreserveSize;
+			}
+			else
+			{
+				WriteLength = Length;
+				BufferStartPos += WriteLength;
+			}
+
+			BufferLength = FMath::Min(BufferLength + WriteLength, BufferCapacity);
+		}
+
+		FMemory::Memcpy(Buffer.GetData() + WritePos, Data, WriteLength * sizeof(ANSICHAR));
+		Length -= WriteLength;
 	}
 }
 
 void FOutputDeviceMemory::Dump(FArchive& Ar)
 {
-	Ar.Serialize(Buffer.GetData() + BufferStartPos, (BufferLength - BufferStartPos) * sizeof(ANSICHAR));
-	Ar.Serialize(Buffer.GetData(), BufferStartPos * sizeof(ANSICHAR));
+	Ar.Serialize(Buffer.GetData(), BufferLength * sizeof(ANSICHAR));
 }
\ No newline at end of file
diff --git a/Engine/Source/Runtime/Core/Private/HAL/OutputDevices.cpp b/Engine/Source/Runtime/Core/Private/HAL/OutputDevices.cpp
--- a/Engine/Source/Runtime/Core/Private/HAL/OutputDevices.cpp
+++ b/Engine/Source/Runtime/Core/Private/HAL/OutputDevices.cpp
@@ -1228,6 +1228,11 @@ void FOutputDeviceFile::CreateBackupCopy(const TCHAR* Filename)
 	}
 }
 
+bool FOutputDeviceFile::IsBackupCopy(const TCHAR* Filename)
+{
+	return Filename != nullptr && FCString::Stristr(const_cast<TCHAR*>(Filename), BACKUP_LOG_FILENAME_POSTFIX) != nullptr;
+}
+
 void FOutputDeviceFile::WriteByteOrderMarkToArchive(EByteOrderMark ByteOrderMark)
 {
 	switch (ByteOrderMark)
diff --git a/Engine/Source/Runtime/Core/Private/Linux/LinuxPlatformMisc.cpp b/Engine/Source/Runtime/Core/Private/Linux/LinuxPlatformMisc.cpp
--- a/Engine/Source/Runtime/Core/Private/Linux/LinuxPlatformMisc.cpp
+++ b/Engine/Source/Runtime/Core/Private/Linux/LinuxPlatformMisc.cpp
@@ -22,6 +22,7 @@
 #include <net/if_arp.h>
 
 #include "ModuleManager.h"
+#include "HAL/ThreadHeartBeat.h"
 
 // define for glibc 2.12.2 and lower (which is shipped with CentOS 6.x and which we target by default)
 #define __secure_getenv getenv
@@ -581,6 +582,8 @@ EAppReturnType::Type FLinuxPlatformMisc::MessageBoxExt(EAppMsgType::Type MsgType
 	};
 
 	int ButtonPressed = -1;
+
+	FSlowHeartBeatScope SuspendHeartBeat;
 	if (SDL_ShowMessageBox(&MessageBoxData, &ButtonPressed) == -1) 
 	{
 		UE_LOG(LogInit, Fatal, TEXT("Error Presenting MessageBox: %s\n"), ANSI_TO_TCHAR(SDL_GetError()));
diff --git a/Engine/Source/Runtime/Core/Private/Mac/MacPlatformMisc.cpp b/Engine/Source/Runtime/Core/Private/Mac/MacPlatformMisc.cpp
--- a/Engine/Source/Runtime/Core/Private/Mac/MacPlatformMisc.cpp
+++ b/Engine/Source/Runtime/Core/Private/Mac/MacPlatformMisc.cpp
@@ -19,6 +19,7 @@
 #include "MacPlatformCrashContext.h"
 #include "PLCrashReporter.h"
 #include "GenericPlatformDriver.h"
+#include "HAL/ThreadHeartBeat.h"
 
 #include <dlfcn.h>
 #include <IOKit/IOKitLib.h>
@@ -821,6 +822,8 @@ void FMacPlatformMisc::CreateGuid(FGuid& Result)
 
 EAppReturnType::Type FMacPlatformMisc::MessageBoxExt(EAppMsgType::Type MsgType, const TCHAR* Text, const TCHAR* Caption)
 {
+	FSlowHeartBeatScope SuspendHeartBeat;
+
 	SCOPED_AUTORELEASE_POOL;
 
 	EAppReturnType::Type ReturnValue = MainThreadReturn(^{
diff --git a/Engine/Source/Runtime/Core/Private/Misc/CoreMisc.cpp b/Engine/Source/Runtime/Core/Private/Misc/CoreMisc.cpp
--- a/Engine/Source/Runtime/Core/Private/Misc/CoreMisc.cpp
+++ b/Engine/Source/Runtime/Core/Private/Misc/CoreMisc.cpp
@@ -849,33 +849,40 @@ void FMaintenance::DeleteOldLogs()
 			Filename = FPaths::GameLogDir() / Filename;
 		}
 
-		// If required, trim the number of files on disk
-		if (MaxLogFilesOnDisk >= 0 && Files.Num() > MaxLogFilesOnDisk)
+		struct FSortByDateNewestFirst
 		{
-			struct FSortByDateNewestFirst
-			{
-				bool operator()(const FString& A, const FString& B) const
-				{
-					const FDateTime TimestampA = IFileManager::Get().GetTimeStamp(*A);
-					const FDateTime TimestampB = IFileManager::Get().GetTimeStamp(*B);
-					return TimestampB < TimestampA;
-				}
-			};
-			Files.Sort(FSortByDateNewestFirst());
-			while (Files.Num() && Files.Num() > MaxLogFilesOnDisk)
+			bool operator()(const FString& A, const FString& B) const
 			{
-				IFileManager::Get().Delete(*Files.Pop());
+				const FDateTime TimestampA = IFileManager::Get().GetTimeStamp(*A);
+				const FDateTime TimestampB = IFileManager::Get().GetTimeStamp(*B);
+				return TimestampB < TimestampA;
 			}
-		}
+		};
+		Files.Sort(FSortByDateNewestFirst());
 
 		// delete all those with the backup text in their name and that are older than the specified number of days
 		double MaxFileAgeSeconds = 60.0 * 60.0 * 24.0 * double(PurgeLogsDays);
-		for (const FString& Filename : Files)
+		for (int32 FileIndex = Files.Num() - 1; FileIndex >= 0; --FileIndex)
 		{
-			if (Filename.Contains(BACKUP_LOG_FILENAME_POSTFIX) && IFileManager::Get().GetFileAgeSeconds(*Filename) > MaxFileAgeSeconds)
+			const FString& Filename = Files[FileIndex];
+			if (FOutputDeviceFile::IsBackupCopy(*Filename) && IFileManager::Get().GetFileAgeSeconds(*Filename) > MaxFileAgeSeconds)
 			{
 				UE_LOG(LogStreaming, Log, TEXT("Deleting old log file %s"), *Filename);
 				IFileManager::Get().Delete(*Filename);
+				Files.RemoveAt(FileIndex);
+			}
+		}
+
+		// If required, trim the number of files on disk
+		if (MaxLogFilesOnDisk >= 0 && Files.Num() > MaxLogFilesOnDisk)
+		{
+			for (int32 FileIndex = Files.Num() - 1; FileIndex >= 0 && Files.Num() > MaxLogFilesOnDisk; --FileIndex)
+			{
+				if (FOutputDeviceFile::IsBackupCopy(*Files[FileIndex]))
+				{
+					IFileManager::Get().Delete(*Files[FileIndex]);
+					Files.RemoveAt(FileIndex);
+				}
 			}
 		}
 
diff --git a/Engine/Source/Runtime/Core/Private/Misc/EngineVersion.cpp b/Engine/Source/Runtime/Core/Private/Misc/EngineVersion.cpp
--- a/Engine/Source/Runtime/Core/Private/Misc/EngineVersion.cpp
+++ b/Engine/Source/Runtime/Core/Private/Misc/EngineVersion.cpp
@@ -8,11 +8,7 @@
 #include "UObject/ReleaseObjectVersion.h"
 
 /** Version numbers for networking */
-#if LOCKED_NETWORK_VERSION
-int32 GEngineNetVersion = LOCKED_NETWORK_VERSION;
-#else
-int32 GEngineNetVersion = BUILT_FROM_CHANGELIST;
-#endif
+int32 GEngineNetVersion = ENGINE_NET_VERSION;
 
 const int32 GEngineMinNetVersion		= 7038;
 const int32 GEngineNegotiationVersion	= 3077;
diff --git a/Engine/Source/Runtime/Core/Private/ProfilingDebugging/ProfilingHelpers.cpp b/Engine/Source/Runtime/Core/Private/ProfilingDebugging/ProfilingHelpers.cpp
--- a/Engine/Source/Runtime/Core/Private/ProfilingDebugging/ProfilingHelpers.cpp
+++ b/Engine/Source/Runtime/Core/Private/ProfilingDebugging/ProfilingHelpers.cpp
@@ -13,11 +13,166 @@
 
 // find where these are really defined
 const static int32 MaxFilenameLen = 100;
+DEFINE_LOG_CATEGORY(LogHistograms);
 
 #if WITH_ENGINE
 FGetMapNameDelegate GGetMapNameDelegate;
 #endif
 
+void FHistogram::InitLinear(double MinTime, double MaxTime, double BinSize)
+{
+	double CurrentBinMin = MinTime;
+	while (CurrentBinMin < MaxTime)
+	{
+		Bins.Add(FBin(CurrentBinMin, CurrentBinMin + BinSize));
+		CurrentBinMin += BinSize;
+	}
+	Bins.Add(FBin(CurrentBinMin));	// catch-all-above bin
+}
+
+void FHistogram::InitHitchTracking()
+{
+	Bins.Empty();
+
+	Bins.Add(FBin(0.0, 9.0));		// >= 120 fps
+	Bins.Add(FBin(9.0, 17.0));		// 60 - 120 fps
+	Bins.Add(FBin(17.0, 34.0));		// 30 - 60 fps
+	Bins.Add(FBin(34.0, 50.0));		// 20 - 30 fps
+	Bins.Add(FBin(50.0, 67.0));		// 15 - 20 fps
+	Bins.Add(FBin(67.0, 100.0));	// 10 - 15 fps
+	Bins.Add(FBin(100.0, 200.0));	// 5 - 10 fps
+	Bins.Add(FBin(200.0, 300.0));	// < 5 fps
+	Bins.Add(FBin(300.0, 500.0));
+	Bins.Add(FBin(500.0, 750.0));
+	Bins.Add(FBin(750.0, 1000.0));
+	Bins.Add(FBin(1000.0, 1500.0));
+	Bins.Add(FBin(1500.0, 2000.0));
+	Bins.Add(FBin(2000.0, 2500.0));
+	Bins.Add(FBin(2500.0, 5000.0));
+	Bins.Add(FBin(5000.0));
+}
+
+void FHistogram::Reset()
+{
+	for (FBin& Bin : Bins)
+	{
+		Bin.Count = 0;
+		Bin.Sum = 0.0;
+	}
+}
+
+void FHistogram::AddMeasurement(double Value)
+{
+	if (LIKELY(Bins.Num()))
+	{
+		FBin& FirstBin = Bins[0];
+		if (UNLIKELY(Value < FirstBin.MinValue))
+		{
+			return;
+		}
+
+		for (int BinIdx = 0, LastBinIdx = Bins.Num() - 1; BinIdx < LastBinIdx; ++BinIdx)
+		{
+			FBin& Bin = Bins[BinIdx];
+			if (Bin.UpperBound > Value)
+			{
+				++Bin.Count;
+				Bin.Sum += Value;
+				return;
+			}
+		}
+
+		// if we got here, Value did not fit in any of the previous bins
+		FBin& LastBin = Bins.Last();
+		++LastBin.Count;
+		LastBin.Sum += Value;
+	}
+}
+
+void FHistogram::DumpToAnalytics(const FString& ParamNamePrefix, TArray<TPair<FString, double>>& OutParamArray)
+{
+	double TotalSum = 0;
+	double TotalObservations = 0;
+	double AverageObservation = 0;
+
+	if (LIKELY(Bins.Num()))
+	{
+		for (int BinIdx = 0, LastBinIdx = Bins.Num() - 1; BinIdx < LastBinIdx; ++BinIdx)
+		{
+			FBin& Bin = Bins[BinIdx];
+			FString ParamName = FString::Printf(TEXT("_%.0f_%.0f"), Bin.MinValue, Bin.UpperBound);
+			OutParamArray.Add(TPairInitializer<FString, double>(ParamNamePrefix + ParamName + TEXT("_Count"), Bin.Count));
+			OutParamArray.Add(TPairInitializer<FString, double>(ParamNamePrefix + ParamName + TEXT("_Sum"), Bin.Sum));
+
+			TotalObservations += Bin.Count;
+			TotalSum += Bin.Sum;
+		}
+
+		FBin& LastBin = Bins.Last();
+		FString ParamName = FString::Printf(TEXT("_%.0f_AndAbove"), LastBin.MinValue);
+		OutParamArray.Add(TPairInitializer<FString, double>(ParamNamePrefix + ParamName + TEXT("_Count"), LastBin.Count));
+		OutParamArray.Add(TPairInitializer<FString, double>(ParamNamePrefix + ParamName + TEXT("_Sum"), LastBin.Sum));
+
+		TotalObservations += LastBin.Count;
+		TotalSum += LastBin.Sum;
+	}
+
+	if (TotalObservations > 0)
+	{
+		AverageObservation = TotalSum / TotalObservations;
+	}
+
+	// add an average value for ease of monitoring/analyzing
+	OutParamArray.Add(TPairInitializer<FString, double>(ParamNamePrefix + TEXT("_Average"), AverageObservation));
+}
+
+void FHistogram::SetBinCountByIndex(int32 IdxBin, int32 Count)
+{
+	if (Bins.IsValidIndex(IdxBin))
+	{
+		Bins[IdxBin].Count = Count;
+	}
+}
+
+void FHistogram::SetBinSumByIndex(int IdxBin, double InSum)
+{
+	if (Bins.IsValidIndex(IdxBin))
+	{
+		Bins[IdxBin].Sum = InSum;
+	}
+}
+
+void FHistogram::DumpToLog(const FString& HistogramName)
+{
+	UE_LOG(LogHistograms, Log, TEXT("Histogram '%s': %d bins"), *HistogramName, Bins.Num());
+
+	if (LIKELY(Bins.Num()))
+	{
+		double TotalSum = 0;
+		double TotalObservations = 0;
+
+		for (int BinIdx = 0, LastBinIdx = Bins.Num() - 1; BinIdx < LastBinIdx; ++BinIdx)
+		{
+			FBin& Bin = Bins[BinIdx];
+			UE_LOG(LogHistograms, Log, TEXT("Bin %4.0f - %4.0f: %5d observation(s) which sum up to %f"), Bin.MinValue, Bin.UpperBound, Bin.Count, Bin.Sum);
+
+			TotalObservations += Bin.Count;
+			TotalSum += Bin.Sum;
+		}
+
+		FBin& LastBin = Bins.Last();
+		UE_LOG(LogHistograms, Log, TEXT("Bin %4.0f +     : %5d observation(s) which sum up to %f"), LastBin.MinValue, LastBin.Count, LastBin.Sum);
+		TotalObservations += LastBin.Count;
+		TotalSum += LastBin.Sum;
+
+		if (TotalObservations > 0)
+		{
+			UE_LOG(LogHistograms, Log, TEXT("Average value for observation: %f"), TotalSum / TotalObservations);
+		}
+	}
+}
+
+
 /**
  * This will get the changelist that should be used with the Automated Performance Testing
  * If one is passed in we use that otherwise we use the FEngineVersion::Current().GetChangelist().  This allows
diff --git a/Engine/Source/Runtime/Core/Private/Windows/WindowsApplication.cpp b/Engine/Source/Runtime/Core/Private/Windows/WindowsApplication.cpp
--- a/Engine/Source/Runtime/Core/Private/Windows/WindowsApplication.cpp
+++ b/Engine/Source/Runtime/Core/Private/Windows/WindowsApplication.cpp
@@ -9,6 +9,7 @@
 #include "IInputDeviceModule.h"
 #include "IInputDevice.h"
 #include "IHapticDevice.h"
+#include "HAL/ThreadHeartBeat.h"
 
 #if WITH_EDITOR
 #include "ModuleManager.h"
@@ -224,6 +225,7 @@ bool FWindowsApplication::RegisterClass( const HINSTANCE HInstance, const HICON
 		//ShowLastError();
 
 		// @todo Slate: Error message should be localized!
+		FSlowHeartBeatScope SuspendHeartBeat;
 		MessageBox(NULL, TEXT("Window Registration Failed!"), TEXT("Error!"), MB_ICONEXCLAMATION | MB_OK);
 
 		return false;
diff --git a/Engine/Source/Runtime/Core/Private/Windows/WindowsPlatformCrashContext.cpp b/Engine/Source/Runtime/Core/Private/Windows/WindowsPlatformCrashContext.cpp
--- a/Engine/Source/Runtime/Core/Private/Windows/WindowsPlatformCrashContext.cpp
+++ b/Engine/Source/Runtime/Core/Private/Windows/WindowsPlatformCrashContext.cpp
@@ -24,7 +24,8 @@
 void FWindowsPlatformCrashContext::AddPlatformSpecificProperties()
 {
 	AddCrashProperty(TEXT("PlatformIsRunningWindows"), 1);
-	AddCrashProperty(TEXT("BuildIntegrityStatus"), GetImageIntegrityStatus());
+	// On windows track the crash type
+	AddCrashProperty(TEXT("PlatformCallbackResult"), GetCrashType());
 }
 
 namespace
diff --git a/Engine/Source/Runtime/Core/Private/Windows/WindowsPlatformFeedbackContextPrivate.h b/Engine/Source/Runtime/Core/Private/Windows/WindowsPlatformFeedbackContextPrivate.h
--- a/Engine/Source/Runtime/Core/Private/Windows/WindowsPlatformFeedbackContextPrivate.h
+++ b/Engine/Source/Runtime/Core/Private/Windows/WindowsPlatformFeedbackContextPrivate.h
@@ -8,7 +8,7 @@
 /**
  * Feedback context implementation for windows.
  */
-class FFeedbackContextWindows : public FFeedbackContext
+class CORE_API FFeedbackContextWindows : public FFeedbackContext
 {
 	/** Context information for warning and error messages */
 	FContextSupplier*	Context;
@@ -67,17 +67,7 @@ public:
 		}
 	}
 
-	virtual bool YesNof(const FText& Question) override
-	{
-		if( ( GIsClient || GIsEditor ) && ( ( GIsSilent != true ) && ( FApp::IsUnattended() != true ) ) )
-		{
-			return( ::MessageBox( NULL, Question.ToString().GetCharArray().GetData(), *NSLOCTEXT("Core", "Question", "Question").ToString(), MB_YESNO|MB_TASKMODAL ) == IDYES);
-		}
-		else
-		{
-			return false;
-		}
-	}
+	virtual bool YesNof(const FText& Question) override;
 
 	FContextSupplier* GetContext() const
 	{
diff --git a/Engine/Source/Runtime/Core/Private/Windows/WindowsPlatformMisc.cpp b/Engine/Source/Runtime/Core/Private/Windows/WindowsPlatformMisc.cpp
--- a/Engine/Source/Runtime/Core/Private/Windows/WindowsPlatformMisc.cpp
+++ b/Engine/Source/Runtime/Core/Private/Windows/WindowsPlatformMisc.cpp
@@ -9,6 +9,7 @@
 
 #include "GenericPlatformChunkInstall.h"
 #include "GenericPlatformDriver.h"			// FGPUDriverInfo
+#include "HAL/ThreadHeartBeat.h"
 
 // Resource includes.
 #include "Runtime/Launch/Resources/Windows/Resource.h"
@@ -1415,6 +1416,8 @@ int MessageBoxExtInternal( EAppMsgType::Type MsgType, HWND HandleWnd, const TCHA
 
 EAppReturnType::Type FWindowsPlatformMisc::MessageBoxExt( EAppMsgType::Type MsgType, const TCHAR* Text, const TCHAR* Caption )
 {
+	FSlowHeartBeatScope SuspendHeartBeat;
+
 	HWND ParentWindow = (HWND)NULL;
 	switch( MsgType )
 	{
@@ -2146,6 +2149,7 @@ void FWindowsPlatformMisc::PromptForRemoteDebugging(bool bIsEnsure)
 			TEXT("hit YES to allow him to debug the crash.\n")
 			TEXT("[Changelist = %d]"),
 			FEngineVersion::Current().GetChangelist());
+		FSlowHeartBeatScope SuspendHeartBeat;
 		if (MessageBox(0, GErrorRemoteDebugPromptMessage, TEXT("CRASHED"), MB_YESNO|MB_SYSTEMMODAL) == IDYES)
 		{
 			::DebugBreak();
diff --git a/Engine/Source/Runtime/Core/Private/Windows/WindowsPlatformOutputDevices.cpp b/Engine/Source/Runtime/Core/Private/Windows/WindowsPlatformOutputDevices.cpp
--- a/Engine/Source/Runtime/Core/Private/Windows/WindowsPlatformOutputDevices.cpp
+++ b/Engine/Source/Runtime/Core/Private/Windows/WindowsPlatformOutputDevices.cpp
@@ -5,6 +5,7 @@
 #include "FeedbackContextAnsi.h"
 #include "../Private/Windows/WindowsPlatformOutputDevicesPrivate.h"
 #include "../Private/Windows/WindowsPlatformFeedbackContextPrivate.h"
+#include "HAL/ThreadHeartBeat.h"
 
 #include "AllowWindowsPlatformTypes.h"
 
@@ -403,3 +404,15 @@ bool FOutputDeviceConsoleWindows::IsAttached()
 	return false;
 }
 
+bool FFeedbackContextWindows::YesNof(const FText& Question)
+{
+	if ((GIsClient || GIsEditor) && ((GIsSilent != true) && (FApp::IsUnattended() != true)))
+	{
+		FSlowHeartBeatScope SuspendHeartBeat;
+		return(::MessageBox(NULL, Question.ToString().GetCharArray().GetData(), *NSLOCTEXT("Core", "Question", "Question").ToString(), MB_YESNO | MB_TASKMODAL) == IDYES);
+	}
+	else
+	{
+		return false;
+	}
+}
\ No newline at end of file
diff --git a/Engine/Source/Runtime/Core/Private/Windows/WindowsWindow.cpp b/Engine/Source/Runtime/Core/Private/Windows/WindowsWindow.cpp
--- a/Engine/Source/Runtime/Core/Private/Windows/WindowsWindow.cpp
+++ b/Engine/Source/Runtime/Core/Private/Windows/WindowsWindow.cpp
@@ -3,6 +3,7 @@
 #include "CorePrivatePCH.h"
 #include "WindowsWindow.h"
 #include "WindowsApplication.h"
+#include "HAL/ThreadHeartBeat.h"
 
 #include "AllowWindowsPlatformTypes.h"
 #if WINVER > 0x502	// Windows Vista or better required for DWM
@@ -176,6 +177,8 @@ void FWindowsWindow::Initialize( FWindowsApplication* const Application, const T
 
 	if( HWnd == NULL )
 	{
+		FSlowHeartBeatScope SuspendHeartBeat;
+
 		// @todo Error message should be localized!
 		MessageBox(NULL, TEXT("Window Creation Failed!"), TEXT("Error!"), MB_ICONEXCLAMATION | MB_OK);
 		checkf(0, TEXT("Window Creation Failed (%d)"), ::GetLastError() );
diff --git a/Engine/Source/Runtime/Core/Public/GenericPlatform/GenericPlatformMisc.h b/Engine/Source/Runtime/Core/Public/GenericPlatform/GenericPlatformMisc.h
--- a/Engine/Source/Runtime/Core/Public/GenericPlatform/GenericPlatformMisc.h
+++ b/Engine/Source/Runtime/Core/Public/GenericPlatform/GenericPlatformMisc.h
@@ -308,6 +308,22 @@ struct CORE_API FGenericPlatformMisc
 	*/
 	static FString GetCPUVendor();
 
+	/**
+	 * On x86(-64) platforms, uses cpuid instruction to get the CPU signature
+	 *
+	 * @return	CPU info bitfield
+	 *
+	 *			Bits 0-3	Stepping ID
+	 *			Bits 4-7	Model
+	 *			Bits 8-11	Family
+	 *			Bits 12-13	Processor type (Intel) / Reserved (AMD)
+	 *			Bits 14-15	Reserved
+	 *			Bits 16-19	Extended model
+	 *			Bits 20-27	Extended family
+	 *			Bits 28-31	Reserved
+	 */
+	static uint32 GetCPUInfo();
+
 	/** 
 	 * Uses cpuid instruction to get the CPU brand string
 	 *
diff --git a/Engine/Source/Runtime/Core/Public/HAL/ExceptionHandling.h b/Engine/Source/Runtime/Core/Public/HAL/ExceptionHandling.h
--- a/Engine/Source/Runtime/Core/Public/HAL/ExceptionHandling.h
+++ b/Engine/Source/Runtime/Core/Public/HAL/ExceptionHandling.h
@@ -36,6 +36,12 @@ extern CORE_API void NewReportEnsure( const TCHAR* ErrorMessage );
 extern CORE_API void CheckImageIntegrity();
 extern CORE_API void CheckImageIntegrityAtRuntime();
 
+enum class ECrashType
+{
+	Debug = -1,
+	Normal = 0
+};
+
 /** In debug builds we can crash on demand so we want to let the crash reporting system know about it */
-extern CORE_API void SetImageIntegrityStatus(int32 Status);
-extern CORE_API int32 GetImageIntegrityStatus();
\ No newline at end of file
+extern CORE_API void SetCrashType(ECrashType InCrashType);
+extern CORE_API int32 GetCrashType();
\ No newline at end of file
diff --git a/Engine/Source/Runtime/Core/Public/HAL/IConsoleManager.h b/Engine/Source/Runtime/Core/Public/HAL/IConsoleManager.h
--- a/Engine/Source/Runtime/Core/Public/HAL/IConsoleManager.h
+++ b/Engine/Source/Runtime/Core/Public/HAL/IConsoleManager.h
@@ -1003,35 +1003,5 @@ public:
 	
 };
 
-/** Helper class for accessing an int CVar. */
-struct FCVarAccessorInt
-{
-	FCVarAccessorInt(IConsoleVariable* CVar) : Data(CVar->AsVariableInt())	{ check(Data); }
-	FORCEINLINE int32 GetGameThreadValue(){ return Data->GetValueOnAnyThread(true); }
-	FORCEINLINE int32 GetRenderThreadValue(){ return Data->GetValueOnRenderThread(); }
-private:
-	const TConsoleVariableData<int32>* Data;
-};
-
-/** Helper class for accessing a float CVar. */
-struct FCVarAccessorFloat
-{
-	FCVarAccessorFloat(IConsoleVariable* CVar) : Data(CVar->AsVariableFloat())	{ check(Data); }
-	FORCEINLINE float GetGameThreadValue(){ return Data->GetValueOnAnyThread(true); }
-	FORCEINLINE float GetRenderThreadValue(){ return Data->GetValueOnRenderThread(); }
-private:
-	const TConsoleVariableData<float>* Data;
-};
-
-/** Helper class for accessing a string CVar. */
-struct FCVarAccessorString
-{
-	FCVarAccessorString(IConsoleVariable* CVar) : Data(CVar->AsVariableString())	{ check(Data); }
-	FORCEINLINE FString GetGameThreadValue(){ return Data->GetValueOnAnyThread(true); }
-	FORCEINLINE FString GetRenderThreadValue(){ return Data->GetValueOnRenderThread(); }
-private:
-	const TConsoleVariableData<FString>* Data;
-};
-
 CORE_API DECLARE_LOG_CATEGORY_EXTERN(LogConsoleResponse, Log, All);
 
diff --git a/Engine/Source/Runtime/Core/Public/HAL/OutputDeviceMemory.h b/Engine/Source/Runtime/Core/Public/HAL/OutputDeviceMemory.h
--- a/Engine/Source/Runtime/Core/Public/HAL/OutputDeviceMemory.h
+++ b/Engine/Source/Runtime/Core/Public/HAL/OutputDeviceMemory.h
@@ -24,9 +24,10 @@ public:
 	/** 
 	 * Constructor, initializing member variables.
 	 *
-	 * @param InBufferSize Maximum size of the memory ring buffer
+	 * @param InPreserveSize	Bytes of the rung buffer not to overwrite (startup info etc)
+	 * @param InBufferSize		Maximum size of the memory ring buffer
 	 */
-	FOutputDeviceMemory(int32 InBufferSize = 1024 * 1024);
+	FOutputDeviceMemory(int32 InPreserveSize = 256 * 1024, int32 InBufferSize = 1024 * 1024);
 
 	/** Dumps the contents of the buffer to an archive */
 	virtual void Dump(FArchive& Ar) override;
@@ -68,6 +69,8 @@ private:
 	int32 BufferStartPos;
 	/** Used data size */
 	int32 BufferLength;
+	/** Amount of data not to overwrite */
+	int32 PreserveSize;
 	/** Sync object for the buffer pos */
 	FCriticalSection BufferPosCritical;
 };
diff --git a/Engine/Source/Runtime/Core/Public/HAL/OutputDevices.h b/Engine/Source/Runtime/Core/Public/HAL/OutputDevices.h
--- a/Engine/Source/Runtime/Core/Public/HAL/OutputDevices.h
+++ b/Engine/Source/Runtime/Core/Public/HAL/OutputDevices.h
@@ -433,6 +433,9 @@ public:
 	/** Creates a backup copy of a log file if it already exists */
 	static void CreateBackupCopy(const TCHAR* Filename);
 
+	/** Checks if the filename represents a backup copy of a log file */
+	static bool IsBackupCopy(const TCHAR* Filename);
+
 private:
 
 	/** Writes to a file on a separate thread */
diff --git a/Engine/Source/Runtime/Core/Public/Linux/LinuxPlatformRunnableThread.h b/Engine/Source/Runtime/Core/Public/Linux/LinuxPlatformRunnableThread.h
--- a/Engine/Source/Runtime/Core/Public/Linux/LinuxPlatformRunnableThread.h
+++ b/Engine/Source/Runtime/Core/Public/Linux/LinuxPlatformRunnableThread.h
@@ -23,16 +23,15 @@ class FRunnableThreadLinux : public FRunnableThreadPThread
 	};
 
 	/** Each thread needs a separate stack for the signal handler, so possible stack overflows in the thread are handled */
-	void* ThreadCrashHandlingStack;
+	char ThreadCrashHandlingStack[EConstants::CrashHandlerStackSize];
 
 public:
 
-	/** Separate stack for the signal handler (so possible stack overflows don't go unnoticed), shared between threads. */
+	/** Separate stack for the signal handler (so possible stack overflows don't go unnoticed), for the main thread specifically. */
 	static char MainThreadSignalHandlerStack[EConstants::CrashHandlerStackSize];
 
 	FRunnableThreadLinux()
 		:	FRunnableThreadPThread()
-		,	ThreadCrashHandlingStack(nullptr)
 	{
 	}
 
@@ -76,35 +75,18 @@ private:
 		}
 
 		// set the alternate stack for handling crashes due to stack overflow
-		ThreadCrashHandlingStack = FMemory::Malloc(EConstants::CrashHandlerStackSize);
-		if (ThreadCrashHandlingStack)
-		{
-			stack_t SignalHandlerStack;
-			FMemory::Memzero(SignalHandlerStack);
-			SignalHandlerStack.ss_sp = ThreadCrashHandlingStack;
-			SignalHandlerStack.ss_size = EConstants::CrashHandlerStackSize;
+		stack_t SignalHandlerStack;
+		FMemory::Memzero(SignalHandlerStack);
+		SignalHandlerStack.ss_sp = ThreadCrashHandlingStack;
+		SignalHandlerStack.ss_size = EConstants::CrashHandlerStackSize;
 
-			if (sigaltstack(&SignalHandlerStack, nullptr) < 0)
-			{
-				int ErrNo = errno;
-				UE_LOG(LogLinux, Warning, TEXT("Unable to set alternate stack for crash handler, sigaltstack() failed with errno=%d (%s)"),
-					ErrNo,
-					UTF8_TO_TCHAR(strerror(ErrNo))
-					);
-
-				// free the memory immediately
-				FMemory::Free(ThreadCrashHandlingStack);
-				ThreadCrashHandlingStack = nullptr;
-			}
-		}
-	}
-
-	virtual void PostRun()
-	{
-		if (ThreadCrashHandlingStack)
+		if (sigaltstack(&SignalHandlerStack, nullptr) < 0)
 		{
-			FMemory::Free(ThreadCrashHandlingStack);
-			ThreadCrashHandlingStack = nullptr;
+			int ErrNo = errno;
+			UE_LOG(LogLinux, Warning, TEXT("Unable to set alternate stack for crash handler, sigaltstack() failed with errno=%d (%s)"),
+				ErrNo,
+				UTF8_TO_TCHAR(strerror(ErrNo))
+				);
 		}
 	}
 
diff --git a/Engine/Source/Runtime/Core/Public/Misc/Build.h b/Engine/Source/Runtime/Core/Public/Misc/Build.h
--- a/Engine/Source/Runtime/Core/Public/Misc/Build.h
+++ b/Engine/Source/Runtime/Core/Public/Misc/Build.h
@@ -219,12 +219,6 @@
 	#error Exactly one of [UE_BUILD_DEBUG UE_BUILD_DEVELOPMENT UE_BUILD_TEST UE_BUILD_SHIPPING] should be defined to be 1
 #endif
 
-// ORION TEMP - turning on stats to judge perf
-#if UE_SERVER && UE_BUILD_TEST
-	#undef STATS
-	#define STATS 1
-#endif
-
 
 /**
  * This is a global setting which will turn on logging / checks for things which are
diff --git a/Engine/Source/Runtime/Core/Public/ProfilingDebugging/CookStats.h b/Engine/Source/Runtime/Core/Public/ProfilingDebugging/CookStats.h
--- a/Engine/Source/Runtime/Core/Public/ProfilingDebugging/CookStats.h
+++ b/Engine/Source/Runtime/Core/Public/ProfilingDebugging/CookStats.h
@@ -109,7 +109,7 @@ public:
 	{
 		TArray<FCookStatsManager::StringKeyValue> Attrs;
 		Attrs.Empty(sizeof...(Args) / 2);
-		ImplCreateKeyValueArray(Attrs, Args...);
+		ImplCreateKeyValueArray(Attrs, Forward<ArgTypes>(Args)...);
 		return Attrs;
 	}
 
diff --git a/Engine/Source/Runtime/Core/Public/ProfilingDebugging/ProfilingHelpers.h b/Engine/Source/Runtime/Core/Public/ProfilingDebugging/ProfilingHelpers.h
--- a/Engine/Source/Runtime/Core/Public/ProfilingDebugging/ProfilingHelpers.h
+++ b/Engine/Source/Runtime/Core/Public/ProfilingDebugging/ProfilingHelpers.h
@@ -8,6 +8,128 @@
 
 #pragma once
 
+DECLARE_LOG_CATEGORY_EXTERN(LogHistograms, Log, All);
+
+ /** Fairly generic histogram for values that have natural lower bound and possibly no upper bound, e.g. frame time */
+struct CORE_API FHistogram
+{
+	/** Inits histogram with linear, equally sized bins */
+	void InitLinear(double MinTime, double MaxTime, double BinSize);
+
+	/** Inits histogram to mimic our existing hitch buckets */
+	void InitHitchTracking();
+
+	/** Resets measurements, without resetting the configured bins. */
+	void Reset();
+
+	/** Adds an observed measurement. */
+	void AddMeasurement(double Value);
+
+	/** Prints histogram contents to the log. */
+	void DumpToLog(const FString& HistogramName);
+
+	/** Populates array commonly used in analytics events, adding two pairs per bin (count and sum). */
+	void DumpToAnalytics(const FString& ParamNamePrefix, TArray<TPair<FString, double>>& OutParamArray);
+
+	/** Gets number of bins. */
+	inline int32 GetNumBins() const {
+		return Bins.Num();
+	}
+
+	/** Gets lower bound of the bin, i.e. minimum value that goes into it. */
+	inline double GetBinLowerBound(int IdxBin) const {
+		return Bins[IdxBin].MinValue;
+	}
+
+	/** Gets upper bound of the bin, i.e. first value that does not go into it. */
+	inline double GetBinUpperBound(int IdxBin) const {
+		return Bins[IdxBin].UpperBound;
+	}
+
+	/** Gets number of observations in the bin. */
+	inline int32 GetBinObservationsCount(int IdxBin) const {
+		return Bins[IdxBin].Count;
+	}
+
+	/** Gets sum of observations in the bin. */
+	inline double GetBinObservationsSum(int IdxBin) const {
+		return Bins[IdxBin].Sum;
+	}
+
+	/** Sets the count of the indexed bin */
+	void SetBinCountByIndex(int32 IdxBin, int32 Count);
+
+	/** Sets the sum of the indexed bin */
+	void SetBinSumByIndex(int IdxBin, double InSum);
+
+	FORCEINLINE FHistogram operator-(const FHistogram& Other) const
+	{
+		FHistogram Tmp;
+		for (int32 BinIndex = 0, NumBins = Bins.Num(); BinIndex < NumBins; BinIndex++)
+		{
+			Tmp.Bins.Add(Bins[BinIndex] - Other.Bins[BinIndex]);
+		}
+		return Tmp;
+	}
+
+protected:
+
+	/** Bin */
+	struct FBin
+	{
+		/** MinValue to be stored in the bin, inclusive. */
+		double				MinValue;
+
+		/** First value NOT to be stored in the bin. */
+		double				UpperBound;
+
+		/** Sum of all values that were put into this bin. */
+		double				Sum;
+
+		/** How many elements are in this bin. */
+		int32				Count;
+
+		FBin()
+		{
+		}
+
+		/** Constructor for a pre-seeded bin */
+		FBin(double MinInclusive, double MaxExclusive, int32 InSum, int32 InCount)
+			: MinValue(MinInclusive)
+			, UpperBound(MaxExclusive)
+			, Sum(InSum)
+			, Count(InCount)
+		{
+		}
+
+		/** Constructor for any bin */
+		FBin(double MinInclusive, double MaxExclusive)
+			: MinValue(MinInclusive)
+			, UpperBound(MaxExclusive)
+			, Sum(0)
+			, Count(0)
+		{
+		}
+
+		/** Constructor for the last bin. */
+		FBin(double MinInclusive)
+			: MinValue(MinInclusive)
+			, UpperBound(FLT_MAX)
+			, Sum(0)
+			, Count(0)
+		{
+		}
+
+		FORCEINLINE FBin operator-(const FBin& Other) const
+		{
+			return FBin(MinValue, UpperBound, Sum - Other.Sum, Count - Other.Count);
+		}
+	};
+
+	/** Bins themselves, should be continous in terms of [MinValue; UpperBound) and sorted ascending by MinValue. Last bin's UpperBound doesn't matter */
+	TArray<FBin>			Bins;
+};
+
 enum EStreamingStatus
 {
 	LEVEL_Unloaded,
diff --git a/Engine/Source/Runtime/CoreUObject/Private/UObject/LazyObjectPtr.cpp b/Engine/Source/Runtime/CoreUObject/Private/UObject/LazyObjectPtr.cpp
--- a/Engine/Source/Runtime/CoreUObject/Private/UObject/LazyObjectPtr.cpp
+++ b/Engine/Source/Runtime/CoreUObject/Private/UObject/LazyObjectPtr.cpp
@@ -80,7 +80,6 @@ FUniqueObjectGuid FUniqueObjectGuid::GetOrCreateIDForObject(const class UObject
 
 FThreadSafeCounter FUniqueObjectGuid::CurrentAnnotationTag(1);
 
-
 /*-----------------------------------------------------------------------------------------------------------
 	FLazyObjectPtr
 -------------------------------------------------------------------------------------------------------------*/
@@ -129,10 +128,20 @@ void FLazyObjectPtr::PossiblySerializeObjectGuid(UObject *Object, FArchive& Ar)
 				UObject* OtherObject = Guid.ResolveObject();
 				if (OtherObject != Object) // on undo/redo, the object (potentially) already exists
 				{
-					if (OtherObject != NULL)
+					bool Duplicate = OtherObject != NULL;
+					bool Reassigning = FParse::Param(FCommandLine::Get(), TEXT("AssignNewMapGuids"));
+
+					if (Duplicate || Reassigning)
 					{
 						// IsGame returns true for GIsPlayInEditorWorld
-						UE_CLOG(!(FApp::IsGame() && Package && Package->ContainsMap()), LogUObjectGlobals, Warning, TEXT("Guid referenced by %s is already used by %s, which should never happen in the editor but could happen at runtime with duplicate level loading or PIE"), !!Object ? *Object->GetFullName() : TEXT("NULL"), *OtherObject->GetFullName());
+						if (!Reassigning)
+						{
+							UE_CLOG(!(FApp::IsGame() && Package && Package->ContainsMap()), LogUObjectGlobals, Warning, TEXT("Guid referenced by %s is already used by %s, which should never happen in the editor but could happen at runtime with duplicate level loading or PIE"), !!Object ? *Object->GetFullName() : TEXT("NULL"), *OtherObject->GetFullName());
+						}
+						else
+						{
+							UE_LOG(LogUObjectGlobals, Warning, TEXT("Assigning new Guid to %s"), *Object->GetFullName());
+						}
 						// This guid is in use, which should never happen in the editor but could happen at runtime with duplicate level loading or PIE. If so give it an invalid GUID and don't add to the annotation map.
 						Guid = FGuid();
 					}
diff --git a/Engine/Source/Runtime/CoreUObject/Private/UObject/Obj.cpp b/Engine/Source/Runtime/CoreUObject/Private/UObject/Obj.cpp
--- a/Engine/Source/Runtime/CoreUObject/Private/UObject/Obj.cpp
+++ b/Engine/Source/Runtime/CoreUObject/Private/UObject/Obj.cpp
@@ -842,7 +842,7 @@ void UObject::ConditionalPostLoadSubobjects( FObjectInstancingGraph* OuterInstan
 	CheckDefaultSubobjects();
 }
 
-void UObject::PreSave()
+void UObject::PreSave(const class ITargetPlatform* TargetPlatform)
 {
 #if WITH_EDITOR
 	FCoreUObjectDelegates::OnObjectSaved.Broadcast(this);
diff --git a/Engine/Source/Runtime/CoreUObject/Private/UObject/ObjectRedirector.cpp b/Engine/Source/Runtime/CoreUObject/Private/UObject/ObjectRedirector.cpp
--- a/Engine/Source/Runtime/CoreUObject/Private/UObject/ObjectRedirector.cpp
+++ b/Engine/Source/Runtime/CoreUObject/Private/UObject/ObjectRedirector.cpp
@@ -13,7 +13,7 @@
 
 // If this object redirector is pointing to an object that won't be serialized anyway, set the RF_Transient flag
 // so that this redirector is also removed from the package.
-void UObjectRedirector::PreSave()
+void UObjectRedirector::PreSave(const class ITargetPlatform* TargetPlatform)
 {
 	if (DestinationObject == NULL
 	||	DestinationObject->HasAnyFlags(RF_Transient)
diff --git a/Engine/Source/Runtime/CoreUObject/Private/UObject/SavePackage.cpp b/Engine/Source/Runtime/CoreUObject/Private/UObject/SavePackage.cpp
--- a/Engine/Source/Runtime/CoreUObject/Private/UObject/SavePackage.cpp
+++ b/Engine/Source/Runtime/CoreUObject/Private/UObject/SavePackage.cpp
@@ -23,30 +23,41 @@ static const FName WorldClassName = FName("World");
 // Uncomment this code to measure UObject::Serialize time taken per UClass type during save.
 // This code tracks each type in a hash, so is too heavyweight to leave on in general,
 // but can be really useful for telling what classes are causing the most save time.
-//#define ENABLE_PACKAGE_CLASS_SERIALIZATION_TIMES 1
+#define ENABLE_PACKAGE_CLASS_SERIALIZATION_TIMES 0
+// uncomment this code to measure UObject::PreSave time taken per uclass type during save
+#define ENABLE_TAGEXPORTS_CLASS_PRESAVE_TIMES 0
 #include "ScopedTimers.h"
 
 namespace SavePackageStats
 {
 	static int32 NumPackagesSaved = 0;
 	static double SavePackageTimeSec = 0.0;
+	static double TagPackageExportsPresaveTimeSec = 0.0;
 	static double TagPackageExportsTimeSec = 0.0;
+	static double ResetLoadersForSaveTimeSec = 0.0;
+	static double TagPackageExportsGetObjectsWithOuter = 0.0;
+	static double TagPackageExportsGetObjectsWithMarks = 0.0;
 	static double SerializeImportsTimeSec = 0.0;
 	static double SortExportsSeekfreeInnerTimeSec = 0.0;
 	static double SerializeExportsTimeSec = 0.0;
 	static double SerializeBulkDataTimeSec = 0.0;
 	static double AsyncWriteTimeSec = 0.0;
 	static TMap<FName, FCookStatsManager::TKeyValuePair<double, uint32>> PackageClassSerializeTimes;
-
+	static TMap<FName, FCookStatsManager::TKeyValuePair<double, uint32>> TagExportSerializeTimes;
+	static TMap<FName, FCookStatsManager::TKeyValuePair<double, uint32>> ClassPreSaveTimes;
 	static FCookStatsManager::FAutoRegisterCallback RegisterCookStats([](FCookStatsManager::AddStatFuncRef AddStat)
 	{
 		// Don't use FCookStatsManager::CreateKeyValueArray because there's just too many arguments. Don't need to overburden the compiler here.
 		TArray<FCookStatsManager::StringKeyValue> StatsList;
-		StatsList.Empty(10);
+		StatsList.Empty(15);
 		#define ADD_COOK_STAT(Name) StatsList.Emplace(TEXT(#Name), FCookStatsManager::ToString(Name))
 		ADD_COOK_STAT(NumPackagesSaved);
 		ADD_COOK_STAT(SavePackageTimeSec);
+		ADD_COOK_STAT(TagPackageExportsPresaveTimeSec);
 		ADD_COOK_STAT(TagPackageExportsTimeSec);
+		ADD_COOK_STAT(ResetLoadersForSaveTimeSec);
+		ADD_COOK_STAT(TagPackageExportsGetObjectsWithOuter);
+		ADD_COOK_STAT(TagPackageExportsGetObjectsWithMarks);
 		ADD_COOK_STAT(SerializeImportsTimeSec);
 		ADD_COOK_STAT(SortExportsSeekfreeInnerTimeSec);
 		ADD_COOK_STAT(SerializeExportsTimeSec);
@@ -55,6 +66,8 @@ namespace SavePackageStats
 		#undef ADD_COOK_STAT
 
 		AddStat(TEXT("Package.Save"), StatsList);
+		
+		const FString TotalString = TEXT("Total");
 
 		if (PackageClassSerializeTimes.Num() > 0)
 		{
@@ -88,6 +101,94 @@ namespace SavePackageStats
 				ClassesLogged++;
 			}
 		}
+
+		if (TagExportSerializeTimes.Num() > 0)
+		{
+			// Sort the class serialize times in reverse order.
+			typedef FCookStatsManager::TKeyValuePair<FName, FCookStatsManager::TKeyValuePair<double, uint32>> ClassSerializeTimeData;
+			TArray<ClassSerializeTimeData> SerializeTimesArray;
+			SerializeTimesArray.Empty(TagExportSerializeTimes.Num());
+			for (const auto& KV : TagExportSerializeTimes)
+			{
+				SerializeTimesArray.Emplace(FCookStatsManager::MakePair(KV.Key, FCookStatsManager::MakePair(KV.Value.Key, KV.Value.Value)));
+			}
+			SerializeTimesArray.Sort([](const ClassSerializeTimeData& LHS, const ClassSerializeTimeData& RHS)
+			{
+				return LHS.Value.Key > RHS.Value.Key;
+			});
+
+			double TotalSerializeTime = 0.0;
+			int TotalSerializeCalls = 0;
+
+			// always print at least the top n, but not anything < 0.1% of total save time.
+			int ClassesLogged = 0;
+			for (const auto& KV : SerializeTimesArray)
+			{
+				TotalSerializeTime += KV.Value.Key;
+				TotalSerializeCalls += KV.Value.Value;
+
+				// since we're sorted on size already. Just find the first one below the threshold and stop there
+				if (ClassesLogged >= 10 && KV.Value.Key < 0.001 * SavePackageTimeSec)
+				{
+					break;
+				}
+				const FString ClassName = KV.Key.ToString();
+				AddStat(TEXT("Package.TagExportSerialize"), FCookStatsManager::CreateKeyValueArray(
+					TEXT("Class"), ClassName,
+					TEXT("TimeSec"), KV.Value.Key,
+					TEXT("Calls"), KV.Value.Value));
+				ClassesLogged++;
+			}
+
+			AddStat(TEXT("Package.TagExportSerialize"), FCookStatsManager::CreateKeyValueArray(
+				TEXT("Class"), TotalString,
+				TEXT("TimeSec"), TotalSerializeTime,
+				TEXT("Calls"), TotalSerializeCalls));
+		}
+
+		if (ClassPreSaveTimes.Num() > 0)
+		{
+			// Sort the class serialize times in reverse order.
+			typedef FCookStatsManager::TKeyValuePair<FName, FCookStatsManager::TKeyValuePair<double, uint32>> ClassSerializeTimeData;
+			TArray<ClassSerializeTimeData> SerializeTimesArray;
+			SerializeTimesArray.Empty(ClassPreSaveTimes.Num());
+			for (const auto& KV : ClassPreSaveTimes)
+			{
+				SerializeTimesArray.Emplace(FCookStatsManager::MakePair(KV.Key, FCookStatsManager::MakePair(KV.Value.Key, KV.Value.Value)));
+			}
+			SerializeTimesArray.Sort([](const ClassSerializeTimeData& LHS, const ClassSerializeTimeData& RHS)
+			{
+				return LHS.Value.Key > RHS.Value.Key;
+			});
+
+			// always print at least the top n, but not anything < 0.1% of total save time.
+			// even if we don't print them then add up the time spent inside presave so we can at least account for it
+			double TotalPreSaveTime = 0.0;
+			int TotalPreSaveCalls = 0;
+			int ClassesLogged = 0;
+			for (const auto& KV : SerializeTimesArray)
+			{
+				TotalPreSaveTime += KV.Value.Key;
+				TotalPreSaveCalls += KV.Value.Value;
+
+				// since we're sorted on size already. Just find the first one below the threshold and stop there
+				if (ClassesLogged >= 10 && KV.Value.Key < 0.001 * SavePackageTimeSec)
+				{
+					break;
+				}
+				const FString ClassName = KV.Key.ToString();
+				AddStat(TEXT("Package.PreSave"), FCookStatsManager::CreateKeyValueArray(
+					TEXT("Class"), ClassName,
+					TEXT("TimeSec"), KV.Value.Key,
+					TEXT("Calls"), KV.Value.Value));
+				ClassesLogged++;
+			}
+			
+			AddStat(TEXT("Package.PreSave"), FCookStatsManager::CreateKeyValueArray(
+				TEXT("Class"), TotalString,
+				TEXT("TimeSec"), TotalPreSaveTime,
+				TEXT("Calls"), TotalPreSaveCalls));
+		}
 	});
 }
 #endif
@@ -698,7 +799,7 @@ void FArchiveSaveTagExports::ProcessTaggedObjects()
 			// In the CDO case the above would serialize most of the references, including transient properties
 			// but we still want to serialize the object using the normal path to collect all custom versions it might be using.
 #if ENABLE_PACKAGE_CLASS_SERIALIZATION_TIMES
-			auto& TimingInfo = SavePackageStats::PackageClassSerializeTimes.FindOrAdd(Obj->GetClass()->GetFName());
+			auto& TimingInfo = SavePackageStats::TagExportSerializeTimes.FindOrAdd(Obj->GetClass()->GetFName());
 			TimingInfo.Value++;
 			FScopedDurationTimer SerializeTimer(TimingInfo.Key);
 #endif
@@ -2805,11 +2906,13 @@ struct FPackageExportTagger
 	UObject*		Base;
 	EObjectFlags	TopLevelFlags;
 	UObject*		Outer;
+	const class ITargetPlatform* TargetPlatform;
 
-	FPackageExportTagger( UObject* CurrentBase, EObjectFlags CurrentFlags, UObject* InOuter )
+	FPackageExportTagger(UObject* CurrentBase, EObjectFlags CurrentFlags, UObject* InOuter, const class ITargetPlatform* InTargetPlatform)
 	:	Base(CurrentBase)
 	,	TopLevelFlags(CurrentFlags)
 	,	Outer(InOuter)
+	,	TargetPlatform(InTargetPlatform)
 	{}
 
 	void TagPackageExports( FArchiveSaveTagExports& ExportTagger, bool bRoutePresave )
@@ -2819,13 +2922,21 @@ struct FPackageExportTagger
 		{
 			if ( bRoutePresave )
 			{
-				Base->PreSave();
+#if ENABLE_TAGEXPORTS_CLASS_PRESAVE_TIMES
+				auto& TimingInfo = SavePackageStats::ClassPreSaveTimes.FindOrAdd(Base->GetClass()->GetFName());
+				TimingInfo.Value++;
+				FScopedDurationTimer SerializeTimer(TimingInfo.Key);
+#endif
+				Base->PreSave(TargetPlatform);
 			}
 
 			ExportTagger.ProcessBaseObject(Base);
 		}
 		TArray<UObject *> ObjectsInOuter;
-		GetObjectsWithOuter(Outer, ObjectsInOuter);
+		{
+			COOK_STAT(FScopedDurationTimer SerializeTimer(SavePackageStats::TagPackageExportsGetObjectsWithOuter));
+			GetObjectsWithOuter(Outer, ObjectsInOuter);
+		}
 		// Serialize objects to tag them as OBJECTMARK_TagExp.
 		for( int32 Index = 0; Index < ObjectsInOuter.Num(); Index++ )
 		{
@@ -2840,13 +2951,21 @@ struct FPackageExportTagger
 			// Route PreSave.
 			{
 				TArray<UObject*> TagExpObjects;
-				GetObjectsWithAnyMarks(TagExpObjects, OBJECTMARK_TagExp);
+				{
+					COOK_STAT(FScopedDurationTimer SerializeTimer(SavePackageStats::TagPackageExportsGetObjectsWithMarks));
+					GetObjectsWithAnyMarks(TagExpObjects, OBJECTMARK_TagExp);
+				}
 				for(int32 Index = 0; Index < TagExpObjects.Num(); Index++)
 				{
 					UObject* Obj = TagExpObjects[Index];
+#if ENABLE_TAGEXPORTS_CLASS_PRESAVE_TIMES
+					auto& TimingInfo = SavePackageStats::ClassPreSaveTimes.FindOrAdd(Obj->GetClass()->GetFName());
+					TimingInfo.Value++;
+					FScopedDurationTimer SerializeTimer(TimingInfo.Key);
+#endif
 					check(Obj->HasAnyMarks(OBJECTMARK_TagExp));
 					//@warning: Objects created from within PreSave will NOT have PreSave called on them!!!
-					Obj->PreSave();
+					Obj->PreSave(TargetPlatform);
 				}
 			}
 		}
@@ -3210,7 +3329,7 @@ ESavePackageResult UPackage::Save(UPackage* InOuter, UObject* Base, EObjectFlags
 
 			// Don't save packages marked as editor-only.
 			if (CanSkipEditorReferencedPackagesWhenCooking && InOuter->IsLoadedByEditorPropertiesOnly())
-			{				
+			{
 				UE_CLOG(!(SaveFlags & SAVE_NoError), LogSavePackage, Display, TEXT("Package loaded by editor-only properties: %s. Package will not be saved."), *InOuter->GetName());
 				return ESavePackageResult::ReferencedOnlyByEditorOnlyData;
 			}
@@ -3223,9 +3342,9 @@ ESavePackageResult UPackage::Save(UPackage* InOuter, UObject* Base, EObjectFlags
 #endif
 		// if we are cooking we should be doing it in the editor
 		// otherwise some other assumptions are bad
-		check( !bIsCooking || WITH_EDITOR );
-	#if WITH_EDITOR
-		if ( !bIsCooking )
+		check(!bIsCooking || WITH_EDITOR);
+#if WITH_EDITOR
+		if (!bIsCooking)
 		{
 			// Attempt to create a backup of this package before it is saved, if applicable
 			if (FCoreUObjectDelegates::AutoPackageBackupDelegate.IsBound())
@@ -3233,8 +3352,8 @@ ESavePackageResult UPackage::Save(UPackage* InOuter, UObject* Base, EObjectFlags
 				FCoreUObjectDelegates::AutoPackageBackupDelegate.Execute(*InOuter);
 			}
 		}
-		
-	#endif	// #if WITH_EDITOR
+
+#endif	// #if WITH_EDITOR
 
 		// do any path replacements on the source DestFile
 		const FString NewPath = FString(Filename);
@@ -3251,26 +3370,26 @@ ESavePackageResult UPackage::Save(UPackage* InOuter, UObject* Base, EObjectFlags
 		(*GFlushStreamingFunc)();
 		FIOSystem::Get().BlockTillAllRequestsFinishedAndFlushHandles();
 
-		uint32 Time=0; CLOCK_CYCLES(Time);
+		uint32 Time = 0; CLOCK_CYCLES(Time);
 
 		// Make sure package is fully loaded before saving. 
-		if( !Base && !InOuter->IsFullyLoaded() )
+		if (!Base && !InOuter->IsFullyLoaded())
 		{
 			if (!(SaveFlags & SAVE_NoError))
 			{
 				// We cannot save packages that aren't fully loaded as it would clobber existing not loaded content.
 				FText ErrorText;
-				if( InOuter->ContainsMap() )
+				if (InOuter->ContainsMap())
 				{
 					FFormatNamedArguments Arguments;
-					Arguments.Add(TEXT("Name"), FText::FromString( NewPath ));
-					ErrorText = FText::Format( NSLOCTEXT( "SavePackage", "CannotSaveMapPartiallyLoaded", "Map '{Name}' cannot be saved as it has only been partially loaded" ), Arguments );
+					Arguments.Add(TEXT("Name"), FText::FromString(NewPath));
+					ErrorText = FText::Format(NSLOCTEXT("SavePackage", "CannotSaveMapPartiallyLoaded", "Map '{Name}' cannot be saved as it has only been partially loaded"), Arguments);
 				}
 				else
 				{
 					FFormatNamedArguments Arguments;
-					Arguments.Add(TEXT("Name"), FText::FromString( NewPath ));
-					ErrorText = FText::Format( NSLOCTEXT( "SavePackage", "CannotSaveAssetPartiallyLoaded", "Asset '{Name}' cannot be saved as it has only been partially loaded" ), Arguments );
+					Arguments.Add(TEXT("Name"), FText::FromString(NewPath));
+					ErrorText = FText::Format(NSLOCTEXT("SavePackage", "CannotSaveAssetPartiallyLoaded", "Asset '{Name}' cannot be saved as it has only been partially loaded"), Arguments);
 				}
 				Error->Logf(ELogVerbosity::Warning, *ErrorText.ToString());
 			}
@@ -3278,7 +3397,7 @@ ESavePackageResult UPackage::Save(UPackage* InOuter, UObject* Base, EObjectFlags
 		}
 
 		// Make sure package is allowed to be saved.
-		if( !TargetPlatform && FCoreUObjectDelegates::IsPackageOKToSaveDelegate.IsBound() )
+		if (!TargetPlatform && FCoreUObjectDelegates::IsPackageOKToSaveDelegate.IsBound())
 		{
 			bool bIsOKToSave = FCoreUObjectDelegates::IsPackageOKToSaveDelegate.Execute(InOuter, Filename, Error);
 			if (!bIsOKToSave)
@@ -3286,17 +3405,17 @@ ESavePackageResult UPackage::Save(UPackage* InOuter, UObject* Base, EObjectFlags
 				if (!(SaveFlags & SAVE_NoError))
 				{
 					FText ErrorText;
-					if( InOuter->ContainsMap() )
+					if (InOuter->ContainsMap())
 					{
 						FFormatNamedArguments Arguments;
-						Arguments.Add(TEXT("Name"), FText::FromString( NewPath ));
-						ErrorText = FText::Format( NSLOCTEXT( "SavePackage", "MapSaveNotAllowed", "Map '{Name}' is not allowed to save (see log for reason)" ), Arguments );
+						Arguments.Add(TEXT("Name"), FText::FromString(NewPath));
+						ErrorText = FText::Format(NSLOCTEXT("SavePackage", "MapSaveNotAllowed", "Map '{Name}' is not allowed to save (see log for reason)"), Arguments);
 					}
 					else
 					{
 						FFormatNamedArguments Arguments;
-						Arguments.Add(TEXT("Name"), FText::FromString( NewPath ));
-						ErrorText = FText::Format( NSLOCTEXT( "SavePackage", "AssetSaveNotAllowed", "Asset '{Name}' is not allowed to save (see log for reason)" ), Arguments );
+						Arguments.Add(TEXT("Name"), FText::FromString(NewPath));
+						ErrorText = FText::Format(NSLOCTEXT("SavePackage", "AssetSaveNotAllowed", "Asset '{Name}' is not allowed to save (see log for reason)"), Arguments);
 					}
 					Error->Logf(ELogVerbosity::Warning, *ErrorText.ToString());
 				}
@@ -3310,17 +3429,17 @@ ESavePackageResult UPackage::Save(UPackage* InOuter, UObject* Base, EObjectFlags
 			if (!(SaveFlags & SAVE_NoError))
 			{
 				FText ErrorText;
-				if( InOuter->ContainsMap() )
+				if (InOuter->ContainsMap())
 				{
 					FFormatNamedArguments Arguments;
-					Arguments.Add(TEXT("Name"), FText::FromString( NewPath ));
-					ErrorText = FText::Format( NSLOCTEXT( "SavePackage", "CannotSaveMapConformIncompatibility", "Conformed Map '{Name}' cannot be saved as it is incompatible with the original" ), Arguments );
+					Arguments.Add(TEXT("Name"), FText::FromString(NewPath));
+					ErrorText = FText::Format(NSLOCTEXT("SavePackage", "CannotSaveMapConformIncompatibility", "Conformed Map '{Name}' cannot be saved as it is incompatible with the original"), Arguments);
 				}
 				else
 				{
 					FFormatNamedArguments Arguments;
-					Arguments.Add(TEXT("Name"), FText::FromString( NewPath ));
-					ErrorText = FText::Format( NSLOCTEXT( "SavePackage", "CannotSaveAssetConformIncompatibility", "Conformed Asset '{Name}' cannot be saved as it is incompatible with the original" ), Arguments );
+					Arguments.Add(TEXT("Name"), FText::FromString(NewPath));
+					ErrorText = FText::Format(NSLOCTEXT("SavePackage", "CannotSaveAssetConformIncompatibility", "Conformed Asset '{Name}' cannot be saved as it is incompatible with the original"), Arguments);
 				}
 				Error->Logf(ELogVerbosity::Error, *ErrorText.ToString());
 			}
@@ -3331,27 +3450,27 @@ ESavePackageResult UPackage::Save(UPackage* InOuter, UObject* Base, EObjectFlags
 		// store a list of additional packages to cook when this is cooked (ie streaming levels)
 		TArray<FString> AdditionalPackagesToCook;
 
-	
+
 		// Route PreSaveRoot to allow e.g. the world to attach components for the persistent level.
 		bool bCleanupIsRequired = false;
-		if( Base )
+		if (Base)
 		{
 			bCleanupIsRequired = Base->PreSaveRoot(Filename, AdditionalPackagesToCook);
 		}
 
-		const FString BaseFilename = FPaths::GetBaseFilename( Filename );
+		const FString BaseFilename = FPaths::GetBaseFilename(Filename);
 		// Make temp file. CreateTempFilename guarantees unique, non-existing filename.
 		// The temp file will be saved in the game save folder to not have to deal with potentially too long paths.
-		FString TempFilename;	
+		FString TempFilename;
 		TempFilename = FPaths::CreateTempFilename(*FPaths::GameSavedDir(), *BaseFilename);
 
 		// Init.
-		FString CleanFilename = FPaths::GetCleanFilename( Filename );
+		FString CleanFilename = FPaths::GetCleanFilename(Filename);
 
 		FFormatNamedArguments Args;
-		Args.Add( TEXT("CleanFilename"), FText::FromString( CleanFilename ) );
+		Args.Add(TEXT("CleanFilename"), FText::FromString(CleanFilename));
 
-		FText StatusMessage = FText::Format( NSLOCTEXT("Core", "SavingFile", "Saving file: {CleanFilename}..."), Args );
+		FText StatusMessage = FText::Format(NSLOCTEXT("Core", "SavingFile", "Saving file: {CleanFilename}..."), Args);
 
 		const int32 TotalSaveSteps = 33;
 		FScopedSlowTask SlowTask(TotalSaveSteps, StatusMessage, bSlowTask);
@@ -3361,9 +3480,10 @@ ESavePackageResult UPackage::Save(UPackage* InOuter, UObject* Base, EObjectFlags
 
 		bool Success = true;
 		bool bRequestStub = false;
-
-		ResetLoadersForSave(InOuter,Filename);
-
+		{
+			COOK_STAT(FScopedDurationTimer SaveTimer(SavePackageStats::ResetLoadersForSaveTimeSec));
+			ResetLoadersForSave(InOuter, Filename);
+		}
 		SlowTask.EnterProgressFrame();
 
 	
@@ -3387,9 +3507,9 @@ ESavePackageResult UPackage::Save(UPackage* InOuter, UObject* Base, EObjectFlags
 			check( ExportTaggerArchive.IsCooking() == bIsCooking );
 
 			// Tag exports and route presave.
-			FPackageExportTagger PackageExportTagger(Base, TopLevelFlags, InOuter);
+			FPackageExportTagger PackageExportTagger(Base, TopLevelFlags, InOuter, TargetPlatform);
 			{
-				COOK_STAT(FScopedDurationTimer SaveTimer(SavePackageStats::TagPackageExportsTimeSec));
+				COOK_STAT(FScopedDurationTimer SaveTimer(SavePackageStats::TagPackageExportsPresaveTimeSec));
 				PackageExportTagger.TagPackageExports(ExportTaggerArchive, true);
 				ExportTaggerArchive.SetFilterEditorOnly(FilterEditorOnly);
 			}
diff --git a/Engine/Source/Runtime/CoreUObject/Private/UObject/ScriptCore.cpp b/Engine/Source/Runtime/CoreUObject/Private/UObject/ScriptCore.cpp
--- a/Engine/Source/Runtime/CoreUObject/Private/UObject/ScriptCore.cpp
+++ b/Engine/Source/Runtime/CoreUObject/Private/UObject/ScriptCore.cpp
@@ -363,6 +363,9 @@ void FFrame::KismetExecutionMessage(const TCHAR* Message, ELogVerbosity::Type Ve
 
 	FString ScriptStack;
 
+	// Tracking down some places that display warnings but no message..
+	ensure(Verbosity > ELogVerbosity::Warning || FCString::Strlen(Message) > 0);
+
 	// Show the stackfor fatal/error, and on warning if that option is enabled
 	if (Verbosity <= ELogVerbosity::Error || (ShowKismetScriptStackOnWarnings() && Verbosity == ELogVerbosity::Warning))
 	{
@@ -372,13 +375,14 @@ void FFrame::KismetExecutionMessage(const TCHAR* Message, ELogVerbosity::Type Ve
 
 	if (Verbosity == ELogVerbosity::Fatal)
 	{
-		UE_LOG(LogScriptCore, Fatal, TEXT("%s\n%s"), Message, *ScriptStack);
+		UE_LOG(LogScriptCore, Fatal, TEXT("Script Msg: %s\n%s"), Message, *ScriptStack);
 	}
 #if !NO_LOGGING
 	else if (!LogScriptCore.IsSuppressed(Verbosity))
 	{
 		// Call directly so we can pass verbosity through
-		FMsg::Logf_Internal(__FILE__, __LINE__, LogScriptCore.GetCategoryName(), Verbosity, TEXT("%s\n%s"), Message, *ScriptStack);
+		FMsg::Logf_Internal(__FILE__, __LINE__, LogScriptCore.GetCategoryName(), Verbosity, TEXT("Script Msg: %s"), Message);
+		FMsg::Logf_Internal(__FILE__, __LINE__, LogScriptCore.GetCategoryName(), Verbosity, TEXT("%s"), *ScriptStack);
 	}	
 #endif
 }
diff --git a/Engine/Source/Runtime/CoreUObject/Public/UObject/ObjectRedirector.h b/Engine/Source/Runtime/CoreUObject/Public/UObject/ObjectRedirector.h
--- a/Engine/Source/Runtime/CoreUObject/Public/UObject/ObjectRedirector.h
+++ b/Engine/Source/Runtime/CoreUObject/Public/UObject/ObjectRedirector.h
@@ -19,7 +19,7 @@ class UObjectRedirector : public UObject
 	// Variables.
 	UObject*		DestinationObject;
 	// UObject interface.
-	virtual void PreSave() override;
+	virtual void PreSave(const class ITargetPlatform* TargetPlatform) override;
 	void Serialize( FArchive& Ar ) override;
 	virtual bool NeedsLoadForClient() const override;
 	virtual bool NeedsLoadForServer() const override;
diff --git a/Engine/Source/Runtime/CoreUObject/Public/UObject/UObject.h b/Engine/Source/Runtime/CoreUObject/Public/UObject/UObject.h
--- a/Engine/Source/Runtime/CoreUObject/Public/UObject/UObject.h
+++ b/Engine/Source/Runtime/CoreUObject/Public/UObject/UObject.h
@@ -187,7 +187,7 @@ public:
 	 *
 	 * @warning: Objects created from within PreSave will NOT have PreSave called on them!!!
 	 */
-	virtual void PreSave();
+	virtual void PreSave(const class ITargetPlatform* TargetPlatform);
 
 	/**
 	 * Note that the object will be modified.  If we are currently recording into the 
diff --git a/Engine/Source/Runtime/Engine/Classes/AI/Navigation/NavigationData.h b/Engine/Source/Runtime/Engine/Classes/AI/Navigation/NavigationData.h
--- a/Engine/Source/Runtime/Engine/Classes/AI/Navigation/NavigationData.h
+++ b/Engine/Source/Runtime/Engine/Classes/AI/Navigation/NavigationData.h
@@ -92,6 +92,10 @@ struct ENGINE_API FNavigationPath : public TSharedFromThis<FNavigationPath, ESPM
 	{
 		return bWaitingForRepath;
 	}
+	FORCEINLINE void SetManualRepathWaiting(const bool bInWaitingForRepath)
+	{
+		bWaitingForRepath = bInWaitingForRepath;
+	}
 	FORCEINLINE bool ShouldUpdateStartPointOnRepath() const
 	{
 		return bUpdateStartPointOnRepath;
@@ -181,6 +185,13 @@ struct ENGINE_API FNavigationPath : public TSharedFromThis<FNavigationPath, ESPM
 
 	FORCEINLINE void DoneUpdating(ENavPathUpdateType::Type UpdateType)
 	{
+		static const ENavPathEvent::Type PathUpdateTypeToPathEvent[] = {
+			ENavPathEvent::UpdatedDueToGoalMoved // GoalMoved,
+			, ENavPathEvent::UpdatedDueToNavigationChanged // NavigationChanged,
+			, ENavPathEvent::MetaPathUpdate // MetaPathUpdate,
+			, ENavPathEvent::Custom // Custom,
+		};
+
 		bUpToDate = true;
 		bWaitingForRepath = false;
 
@@ -189,8 +200,8 @@ struct ENGINE_API FNavigationPath : public TSharedFromThis<FNavigationPath, ESPM
 			// notify path before observers
 			OnPathUpdated(UpdateType);
 		}
-
-		ObserverDelegate.Broadcast(this, UpdateType == ENavPathUpdateType::GoalMoved ? ENavPathEvent::UpdatedDueToGoalMoved : ENavPathEvent::UpdatedDueToNavigationChanged);
+		
+		ObserverDelegate.Broadcast(this, PathUpdateTypeToPathEvent[uint8(UpdateType)]);
 	}
 
 	FORCEINLINE float GetTimeStamp() const { return LastUpdateTimeStamp; }
diff --git a/Engine/Source/Runtime/Engine/Classes/AI/Navigation/NavigationSystem.h b/Engine/Source/Runtime/Engine/Classes/AI/Navigation/NavigationSystem.h
--- a/Engine/Source/Runtime/Engine/Classes/AI/Navigation/NavigationSystem.h
+++ b/Engine/Source/Runtime/Engine/Classes/AI/Navigation/NavigationSystem.h
@@ -788,13 +788,30 @@ public:
 
 	static FORCEINLINE bool ShouldUpdateNavOctreeOnComponentChange()
 	{
-		return bUpdateNavOctreeOnComponentChange
+		return (bUpdateNavOctreeOnComponentChange && !bStaticRuntimeNavigation)
 #if WITH_EDITOR
 			|| (GIsEditor && !GIsPlayInEditorWorld)
 #endif
 			;
 	}
 
+	static FORCEINLINE bool IsNavigationSystemStatic()
+	{
+		return bStaticRuntimeNavigation
+#if WITH_EDITOR
+			&& !(GIsEditor && !GIsPlayInEditorWorld)
+#endif
+			;
+	}
+
+	/** Use this function to signal the NavigationSystem it doesn't need to store
+	 *	any navigation-generation-related data at game runtime, because 
+	 *	nothing is going to use it anyway. This will short-circuit all code related 
+	 *	to navmesh rebuilding, so use it only if you have fully static navigation in 
+	 *	your game.
+	 *	Note: this is not a runtime switch. Call it before any actual game starts. */
+	static void ConfigureAsStatic();
+
 	/** 
 	 * Exec command handlers
 	 */
@@ -864,8 +881,8 @@ protected:
 	
 	/** whether seamless navigation building is enabled */
 	static bool bNavigationAutoUpdateEnabled;
-
 	static bool bUpdateNavOctreeOnComponentChange;
+	static bool bStaticRuntimeNavigation;
 
 	static TMap<INavLinkCustomInterface*, FWeakObjectPtr> PendingCustomLinkRegistration;
 	TSet<const UClass*> NavAreaClasses;
diff --git a/Engine/Source/Runtime/Engine/Classes/AI/Navigation/NavigationTypes.h b/Engine/Source/Runtime/Engine/Classes/AI/Navigation/NavigationTypes.h
--- a/Engine/Source/Runtime/Engine/Classes/AI/Navigation/NavigationTypes.h
+++ b/Engine/Source/Runtime/Engine/Classes/AI/Navigation/NavigationTypes.h
@@ -363,6 +363,7 @@ namespace ENavPathEvent
 		UpdatedDueToNavigationChanged,
 		Invalidated,
 		RePathFailed,
+		MetaPathUpdate,
 		Custom,
 	};
 }
@@ -373,6 +374,7 @@ namespace ENavPathUpdateType
 	{
 		GoalMoved,
 		NavigationChanged,
+		MetaPathUpdate,
 		Custom,
 	};
 }
diff --git a/Engine/Source/Runtime/Engine/Classes/Animation/AnimSequence.h b/Engine/Source/Runtime/Engine/Classes/Animation/AnimSequence.h
--- a/Engine/Source/Runtime/Engine/Classes/Animation/AnimSequence.h
+++ b/Engine/Source/Runtime/Engine/Classes/Animation/AnimSequence.h
@@ -531,7 +531,7 @@ public:
 	virtual void Serialize(FArchive& Ar) override;
 	virtual void PostInitProperties() override;
 	virtual void PostLoad() override;
-	virtual void PreSave() override;
+	virtual void PreSave(const class ITargetPlatform* TargetPlatform) override;
 #if WITH_EDITOR
 	virtual void PostEditChangeProperty(FPropertyChangedEvent& PropertyChangedEvent) override;
 	virtual void PostDuplicate(bool bDuplicateForPIE) override;
diff --git a/Engine/Source/Runtime/Engine/Classes/Components/InputComponent.h b/Engine/Source/Runtime/Engine/Classes/Components/InputComponent.h
--- a/Engine/Source/Runtime/Engine/Classes/Components/InputComponent.h
+++ b/Engine/Source/Runtime/Engine/Classes/Components/InputComponent.h
@@ -634,7 +634,7 @@ public:
 	 * @param BindingIndex The index of the binding to get.
 	 * @see AddActionBinding, ClearActionBindings, GetNumActionBindings, RemoveActionBinding
 	 */
-	FInputActionBinding& GetActionBinding( const int32 BindingIndex );
+	FInputActionBinding& GetActionBinding(const int32 BindingIndex) { return ActionBindings[BindingIndex]; }
 
 	/**
 	 * Gets the number of action bindings.
@@ -642,7 +642,7 @@ public:
 	 * @return Number of bindings.
 	 * @see AddActionBinding, ClearActionBindings, GetActionBinding, RemoveActionBinding
 	 */
-	int32 GetNumActionBindings( ) const;
+	int32 GetNumActionBindings() const { return ActionBindings.Num(); }
 
 	/**
 	 * Removes the action binding at the specified index.
diff --git a/Engine/Source/Runtime/Engine/Classes/Engine/CollisionProfile.h b/Engine/Source/Runtime/Engine/Classes/Engine/CollisionProfile.h
--- a/Engine/Source/Runtime/Engine/Classes/Engine/CollisionProfile.h
+++ b/Engine/Source/Runtime/Engine/Classes/Engine/CollisionProfile.h
@@ -176,6 +176,10 @@ public:
 	/** Accessor and initializer **/
 	ENGINE_API static UCollisionProfile* Get();
 
+	/** Begin UObject interface */
+	virtual void PostReloadConfig(class UProperty* PropertyThatWasLoaded) override;
+	/** End UObject interface */
+
 	/** Fill up the array with the profile names **/
 	ENGINE_API static void GetProfileNames(TArray<TSharedPtr<FName>>& OutNameList);
 
