Restoring AddReferencedObjects functions of UByteProperty and UArrrayProperty so that FReferenceFinder can find their references until FReferenceFinder can be reworked to use TFastReferenceCollector.

(fixes potential crashes when recompiling BPs)

diff --git a/Engine/Source/Runtime/CoreUObject/Private/UObject/PropertyArray.cpp b/Engine/Source/Runtime/CoreUObject/Private/UObject/PropertyArray.cpp
--- a/Engine/Source/Runtime/CoreUObject/Private/UObject/PropertyArray.cpp
+++ b/Engine/Source/Runtime/CoreUObject/Private/UObject/PropertyArray.cpp
@@ -218,6 +218,12 @@ void UArrayProperty::Serialize( FArchive& Ar )
 	Ar << Inner;
 	checkSlow(Inner || HasAnyFlags(RF_ClassDefaultObject) || IsPendingKill());
 }
+void UArrayProperty::AddReferencedObjects(UObject* InThis, FReferenceCollector& Collector)
+{
+	UArrayProperty* This = CastChecked<UArrayProperty>(InThis);
+	Collector.AddReferencedObject( This->Inner, This );
+	Super::AddReferencedObjects( This, Collector );
+}
 
 FString UArrayProperty::GetCPPTypeCustom(FString* ExtendedTypeText, uint32 CPPExportFlags, const FString& InnerTypeText, const FString& InInnerExtendedTypeText) const
 {
diff --git a/Engine/Source/Runtime/CoreUObject/Private/UObject/PropertyByte.cpp b/Engine/Source/Runtime/CoreUObject/Private/UObject/PropertyByte.cpp
--- a/Engine/Source/Runtime/CoreUObject/Private/UObject/PropertyByte.cpp
+++ b/Engine/Source/Runtime/CoreUObject/Private/UObject/PropertyByte.cpp
@@ -76,7 +76,12 @@ void UByteProperty::Serialize( FArchive& Ar )
 		Ar.Preload(Enum);
 	}
 }
-
+void UByteProperty::AddReferencedObjects(UObject* InThis, FReferenceCollector& Collector)
+{
+	UByteProperty* This = CastChecked<UByteProperty>(InThis);
+	Collector.AddReferencedObject( This->Enum, This );
+	Super::AddReferencedObjects( This, Collector );
+}
 FString UByteProperty::GetCPPType( FString* ExtendedTypeText/*=NULL*/, uint32 CPPExportFlags/*=0*/ ) const
 {
 	if (Enum)
diff --git a/Engine/Source/Runtime/CoreUObject/Public/UObject/UnrealType.h b/Engine/Source/Runtime/CoreUObject/Public/UObject/UnrealType.h
--- a/Engine/Source/Runtime/CoreUObject/Public/UObject/UnrealType.h
+++ b/Engine/Source/Runtime/CoreUObject/Public/UObject/UnrealType.h
@@ -1302,6 +1302,7 @@ class COREUOBJECT_API UByteProperty : public TProperty_Numeric<uint8>
 
 	// UObject interface.
 	virtual void Serialize( FArchive& Ar ) override;
+	static void AddReferencedObjects(UObject* InThis, FReferenceCollector& Collector);
 	// End of UObject interface
 
 	// UHT interface
@@ -2431,6 +2432,7 @@ public:
 
 	// UObject interface
 	virtual void Serialize( FArchive& Ar ) override;
+	static void AddReferencedObjects(UObject* InThis, FReferenceCollector& Collector);
 	// End of UObject interface
 
 	// UField interface
