Copying //UE4/Dev-Framework to Dev-Main (//UE4/Dev-Main)

diff --git a/Engine/Source/Runtime/CoreUObject/Private/UObject/Class.cpp b/Engine/Source/Runtime/CoreUObject/Private/UObject/Class.cpp
--- a/Engine/Source/Runtime/CoreUObject/Private/UObject/Class.cpp
+++ b/Engine/Source/Runtime/CoreUObject/Private/UObject/Class.cpp
@@ -853,58 +853,6 @@ void UStruct::InitTaggedPropertyRedirectsMap()
 	}
 }
 
-template <typename OldIntType, typename NewIntType>
-struct TConvertIntToIntProperty
-{
-	static void Convert(FArchive& Ar, UProperty* Property, void* Obj, int32 ArrayIndex, const TCHAR* PropName, const TCHAR* ObjName, const TCHAR* ArName)
-	{
-		OldIntType OldValue;
-		Ar << OldValue;
-		NewIntType NewValue = OldValue;
-		static_cast<typename UE4Types_Private::TIntegerPropertyMapping<NewIntType>::Type*>(Property)->SetPropertyValue_InContainer(Obj, NewValue, ArrayIndex);
-
-		UE_CLOG(
-			(TIsSigned<OldIntType>::Value && !TIsSigned<NewIntType>::Value && OldValue < 0) || (sizeof(NewIntType) < sizeof(OldIntType) && (OldIntType)NewValue != OldValue),
-			LogClass,
-			Warning,
-			TEXT("Potential data loss during conversion of integer property %s of %s - was (%s) now (%s) - for package: %s"),
-			PropName,
-			ObjName,
-			*LexicalConversion::ToString(OldValue),
-			*LexicalConversion::ToString(NewValue),
-			ArName
-		);
-	}
-};
-
-template <typename OldIntType>
-struct TConvertIntToEnumProperty
-{
-	static void Convert(FArchive& Ar, UByteProperty* Property, UEnum* Enum, void* Obj, int32 ArrayIndex, const TCHAR* PropName, const TCHAR* ObjName, const TCHAR* ArName)
-	{
-		OldIntType OldValue;
-		Ar << OldValue;
-
-		uint8 NewValue = OldValue;
-		if (OldValue > (OldIntType)TNumericLimits<uint8>::Max() || !Enum->IsValidEnumValue(NewValue))
-		{
-			UE_LOG(
-				LogClass,
-				Warning,
-				TEXT("Failed to find valid enum value '%d' for enum type '%s' when converting property '%s' during property loading - setting to '%s'"),
-				OldValue,
-				*Enum->GetName(),
-				PropName,
-				*Enum->GetNameByValue(Enum->GetMaxEnumValue()).ToString()
-			);
-
-			NewValue = Enum->GetMaxEnumValue();
-		}
-
-		Property->SetPropertyValue_InContainer(Obj, NewValue, ArrayIndex);
-	}
-};
-
 void UStruct::SerializeTaggedProperties(FArchive& Ar, uint8* Data, UStruct* DefaultsStruct, uint8* Defaults, const UObject* BreakRecursionIfFullyLoad) const
 {
 	//SCOPED_LOADTIMER(SerializeTaggedPropertiesTime);
@@ -920,58 +868,9 @@ void UStruct::SerializeTaggedProperties(FArchive& Ar, uint8* Data, UStruct* Defa
 		// and makes it an O(n) when properties are saved in the same order as they are loaded (default case). In the 
 		// case that a property was reordered the code falls back to a slower search.
 		UProperty*	Property			= PropertyLink;
-		bool		AdvanceProperty		= false;
+		bool		bAdvanceProperty	= false;
 		int32		RemainingArrayDim	= Property ? Property->ArrayDim : 0;
 
-		auto CanSerializeFromStructWithDifferentName = [](const FArchive& InAr, const FPropertyTag& PropertyTag, const UStructProperty* StructProperty)
-		{
-			if (InAr.UE4Ver() < VER_UE4_STRUCT_GUID_IN_PROPERTY_TAG)
-			{
-				// Old Implementation
-				return StructProperty && !StructProperty->UseBinaryOrNativeSerialization(InAr);
-			}
-			return PropertyTag.StructGuid.IsValid() && StructProperty && StructProperty->Struct && (PropertyTag.StructGuid == StructProperty->Struct->GetCustomGuid());
-		};
-
-		auto ReadEnumAsUint8 = [DefaultsStruct, &Ar](FName TaggedEnumName, const TCHAR* PropertyName) -> uint8
-		{
-			//@warning: mirrors loading code in UByteProperty::SerializeItem()
-			FName EnumName;
-			Ar << EnumName;
-
-			UEnum* Enum = FindField<UEnum>(dynamic_cast<UClass*>(DefaultsStruct) ? static_cast<UClass*>(DefaultsStruct) : DefaultsStruct->GetTypedOuter<UClass>(), TaggedEnumName);
-			if (!Enum)
-			{
-				Enum = FindObject<UEnum>(ANY_PACKAGE, *TaggedEnumName.ToString(), true);
-			}
-
-			if (!Enum)
-			{
-				UE_LOG(LogClass, Warning, TEXT("Failed to find enum '%s' when converting property '%s' during property loading - setting to 0"), *TaggedEnumName.ToString(), PropertyName);
-				return 0;
-			}
-
-			Ar.Preload(Enum);
-
-			uint8 Result = Enum->GetValueByName(EnumName);
-			if (!Enum->IsValidEnumValue(Result))
-			{
-				UE_LOG(
-					LogClass,
-					Warning,
-					TEXT("Failed to find valid enum value '%s' for enum type '%s' when converting property '%s' during property loading - setting to '%s'"),
-					*EnumName.ToString(),
-					*Enum->GetName(),
-					PropertyName,
-					*Enum->GetNameByValue(Enum->GetMaxEnumValue()).ToString()
-				);
-
-				return Enum->GetMaxEnumValue();
-			}
-
-			return Result;
-		};
-
 		// Load all stored properties, potentially skipping unknown ones.
 		while (1)
 		{
@@ -989,7 +888,7 @@ void UStruct::SerializeTaggedProperties(FArchive& Ar, uint8* Data, UStruct* Defa
 			}
 
 			// Move to the next property to be serialized
-			if( AdvanceProperty && --RemainingArrayDim <= 0 )
+			if( bAdvanceProperty && --RemainingArrayDim <= 0 )
 			{
 				Property = Property->PropertyLinkNext;
 				// Skip over properties that don't need to be serialized.
@@ -997,7 +896,7 @@ void UStruct::SerializeTaggedProperties(FArchive& Ar, uint8* Data, UStruct* Defa
 				{
 					Property = Property->PropertyLinkNext;
 				}
-				AdvanceProperty		= 0;
+				bAdvanceProperty		= 0;
 				RemainingArrayDim	= Property ? Property->ArrayDim : 0;
 			}
 			
@@ -1012,7 +911,7 @@ void UStruct::SerializeTaggedProperties(FArchive& Ar, uint8* Data, UStruct* Defa
 				}
 			}
 			// If this property is not the one we expect (e.g. skipped as it matches the default value), do the brute force search.
-			if( Property == NULL || Property->GetFName() != Tag.Name )
+			if( Property == nullptr || Property->GetFName() != Tag.Name )
 			{
 				// No need to check redirects on platforms where everything is cooked. Always check for save games
 				if (!FPlatformProperties::RequiresCookedData() || Ar.IsSaveGame())
@@ -1044,7 +943,7 @@ void UStruct::SerializeTaggedProperties(FArchive& Ar, uint8* Data, UStruct* Defa
 								}
 							}
 							// If theres another class name to check get it, otherwise flag the end.
-							if( SuperClass != NULL )
+							if( SuperClass != nullptr )
 							{
 								EachName = SuperClass->GetFName();
 								SuperClass = SuperClass->GetSuperStruct();
@@ -1067,7 +966,7 @@ void UStruct::SerializeTaggedProperties(FArchive& Ar, uint8* Data, UStruct* Defa
 					}
 				}
 				// ... and then search from the beginning till we reach the current property if it's not found.
-				if( Property == NULL )
+				if( Property == nullptr )
 				{
 					for( Property = PropertyLink; Property && Property != CurrentProperty; Property = Property->PropertyLinkNext )
 					{
@@ -1080,7 +979,7 @@ void UStruct::SerializeTaggedProperties(FArchive& Ar, uint8* Data, UStruct* Defa
 					if( Property == CurrentProperty )
 					{
 						// Property wasn't found.
-						Property = NULL;
+						Property = nullptr;
 					}
 				}
 
@@ -1097,15 +996,19 @@ void UStruct::SerializeTaggedProperties(FArchive& Ar, uint8* Data, UStruct* Defa
 			FName ArrayInnerID = NAME_None;
 
 			// Check if this is a struct property and we have a redirector
+			// No need to check redirects on platforms where everything is cooked. Always check for save games
+			if (!FPlatformProperties::RequiresCookedData() || Ar.IsSaveGame())
+			{
 			if (Tag.Type == NAME_StructProperty && PropID == NAME_StructProperty)
 			{
 				FName* NewName = FLinkerLoad::StructNameRedirects.Find(Tag.StructName);
 				FName StructName = CastChecked<UStructProperty>(Property)->Struct->GetFName();
-				if (NewName != NULL && *NewName == StructName)
+					if (NewName != nullptr && *NewName == StructName)
 				{
 					Tag.StructName = *NewName;
 				}
 			}
+			}
 
 			const int64 StartOfProperty = Ar.Tell();
 			if( !Property )
@@ -1137,738 +1040,19 @@ void UStruct::SerializeTaggedProperties(FArchive& Ar, uint8* Data, UStruct* Defa
 			// Convert properties from old type to new type automatically if types are compatible
 			// If you add an entry to this, you will also need to add an entry to the array case below
 			// For converting to a struct, you can just implement SerializeFromMismatchedTag on the struct
-
-			else if( dynamic_cast<UStructProperty*>(Property) && static_cast<UStructProperty*>(Property)->Struct && (Tag.Type != PropID || (Tag.Type == NAME_StructProperty && Tag.StructName != static_cast<UStructProperty*>(Property)->Struct->GetFName())) && (static_cast<UStructProperty*>(Property)->Struct->StructFlags & STRUCT_SerializeFromMismatchedTag))
+			else if (Property->ConvertFromType(Tag, Ar, Data, DefaultsStruct, bAdvanceProperty))
 			{
-				UScriptStruct::ICppStructOps* CppStructOps = static_cast<UStructProperty*>(Property)->Struct->GetCppStructOps();
-				check(CppStructOps && CppStructOps->HasSerializeFromMismatchedTag()); // else should not have STRUCT_SerializeFromMismatchedTag
-				void* DestAddress = Property->ContainerPtrToValuePtr<void>(Data, Tag.ArrayIndex);  
-				if (CppStructOps->SerializeFromMismatchedTag(Tag, Ar, DestAddress))
+				if (bAdvanceProperty)
 				{
-					AdvanceProperty = true;
 					continue;
 				}
-				else
-				{
-					UE_LOG(LogClass, Warning, TEXT("SerializeFromMismatchedTag failed: Type mismatch in %s of %s - Previous (%s) Current(%s) for package:  %s"), *Tag.Name.ToString(), *GetName(), *Tag.Type.ToString(), *PropID.ToString(), *Ar.GetArchiveName() );
-				}
 			}
+
 			else if (Tag.Type != PropID)
 			{
-				if( Tag.Type==NAME_StrProperty && dynamic_cast<UNameProperty*>(Property) )
-				{
-					FString str;
-					Ar << str;
-					CastChecked<UNameProperty>(Property)->SetPropertyValue_InContainer(Data, FName(*str), Tag.ArrayIndex);
-					AdvanceProperty = true;
-					continue; 
-				}
-				else if( Tag.Type==NAME_StrProperty && dynamic_cast<UTextProperty*>(Property) ) // Convert serialized string to text.
-				{ 
-					FString str;
-					Ar << str;
-					FText Text = FText::FromString(str);
-					Text.TextData->PersistText();
-					Text.Flags |= ETextFlag::ConvertedProperty;
-					CastChecked<UTextProperty>(Property)->SetPropertyValue_InContainer(Data, Text, Tag.ArrayIndex);
-					AdvanceProperty = true;
-					continue; 
-				}
-				else if( Tag.Type==NAME_TextProperty && dynamic_cast<UStrProperty*>(Property) ) // Convert serialized text to string.
-				{ 
-					FText Text;  
-					Ar << Text;
-					FString String = FTextInspector::GetSourceString(Text) ? *FTextInspector::GetSourceString(Text) : TEXT("");
-					CastChecked<UStrProperty>(Property)->SetPropertyValue_InContainer(Data, String, Tag.ArrayIndex);
-					AdvanceProperty = true;
-					continue; 
-				}
-				else if( Tag.Type==NAME_NameProperty && dynamic_cast<UTextProperty*>(Property) ) // Convert serialized name to text.
-				{ 
-					FName Name;  
-					Ar << Name;
-					FText Text = FText::FromName(Name);
-					Text.Flags |= ETextFlag::ConvertedProperty;
-					CastChecked<UTextProperty>(Property)->SetPropertyValue_InContainer(Data, Text, Tag.ArrayIndex);
-					AdvanceProperty = true;
-					continue; 
-				}
-				else if( Tag.Type==NAME_TextProperty && dynamic_cast<UNameProperty*>(Property) ) // Convert serialized text to name.
-				{ 
-					FText Text;  
-					Ar << Text;
-					FName Name = FName(*Text.ToString());
-					CastChecked<UNameProperty>(Property)->SetPropertyValue_InContainer(Data, Name, Tag.ArrayIndex);
-					AdvanceProperty = true;
-					continue; 
-				}
-				else if ((Tag.Type == NAME_AssetObjectProperty || Tag.Type == NAME_AssetSubclassOfProperty) && (PropID == NAME_ObjectProperty || PropID == NAME_ClassProperty))
-				{
-					// This property used to be a TAssetPtr<Foo> but is now a raw UObjectProperty Foo*, we can convert without loss of data
-					FAssetPtr PreviousValue;
-					Ar << PreviousValue;
-
-					// now copy the value into the object's address space
-					UObject* PreviousValueObj = PreviousValue.Get();
-					CastChecked<UObjectProperty>(Property)->SetPropertyValue_InContainer(Data, PreviousValueObj, Tag.ArrayIndex);
-
-					AdvanceProperty = true;
-					continue;
-				}
-				else if ((Tag.Type == NAME_ObjectProperty || Tag.Type == NAME_ClassProperty) && (PropID == NAME_AssetObjectProperty || PropID == NAME_AssetSubclassOfProperty))
-				{
-					// This property used to be a raw UObjectProperty Foo* but is now a TAssetPtr<Foo>
-					UObject* PreviousValue = NULL;
-					Ar << PreviousValue;
-
-					// now copy the value into the object's address space
-					FAssetPtr PreviousValueAssetPtr(PreviousValue);
-					CastChecked<UAssetObjectProperty>(Property)->SetPropertyValue_InContainer(Data, PreviousValueAssetPtr, Tag.ArrayIndex);
-
-					AdvanceProperty = true;
-					continue;
-				}
-				else if (Tag.Type == NAME_IntProperty && PropID == NAME_BoolProperty)
-				{
-					// Property was saved as an int32, but has been changed to a bool (bitfield)
-					int32 IntValue;
-					Ar << IntValue;
-
-					if( IntValue != 0 )
-					{
-#if !(UE_BUILD_SHIPPING || UE_BUILD_TEST)
-						if(IntValue != 1)
-						{
-							UE_LOG(LogClass, Log, TEXT("Loading int32 property (%s) that is now a uint32 - value '%d', expecting 0 or 1. Value set to true."), *Property->GetPathName(), IntValue);
-						}
-#endif
-						CastChecked<UBoolProperty>(Property)->SetPropertyValue_InContainer(Data, true, Tag.ArrayIndex);
-					}
-					else
-					{
-						CastChecked<UBoolProperty>(Property)->SetPropertyValue_InContainer(Data, false, Tag.ArrayIndex);
-					}
-
-					AdvanceProperty = true;
-					continue; 
-				}
-				else if (Tag.Type == NAME_StructProperty && PropID == NAME_AssetObjectProperty)
-				{
-					// This property used to be a FStringAssetReference but is now a TAssetPtr<Foo>
-					FStringAssetReference PreviousValue;
-					// explicitly call Serialize to ensure that the various delegates needed for cooking are fired
-					PreviousValue.Serialize(Ar);
-
-					// now copy the value into the object's address space
-					FAssetPtr PreviousValueAssetPtr;
-					PreviousValueAssetPtr = PreviousValue;
-					CastChecked<UAssetObjectProperty>(Property)->SetPropertyValue_InContainer(Data, PreviousValueAssetPtr, Tag.ArrayIndex);
-
-					AdvanceProperty = true;
-					continue;
-				}
-				/* BEGIN - Up-converting ints to larger ints */
-				else if( Tag.Type == NAME_Int8Property && PropID == NAME_Int16Property )
-				{
-					TConvertIntToIntProperty<int8, int16>::Convert(Ar, Property, Data, Tag.ArrayIndex, *Tag.Name.ToString(), *GetName(), *Ar.GetArchiveName());
-					AdvanceProperty = true;
-					continue;
-				}
-				else if( Tag.Type == NAME_Int8Property && PropID == NAME_IntProperty )
-				{
-					TConvertIntToIntProperty<int8, int32>::Convert(Ar, Property, Data, Tag.ArrayIndex, *Tag.Name.ToString(), *GetName(), *Ar.GetArchiveName());
-					AdvanceProperty = true;
-					continue;
-				}
-				else if( Tag.Type == NAME_Int8Property && PropID == NAME_Int64Property )
-				{
-					TConvertIntToIntProperty<int8, int64>::Convert(Ar, Property, Data, Tag.ArrayIndex, *Tag.Name.ToString(), *GetName(), *Ar.GetArchiveName());
-					AdvanceProperty = true;
-					continue;
-				}
-				else if( Tag.Type == NAME_Int16Property && PropID == NAME_IntProperty )
-				{
-					TConvertIntToIntProperty<int16, int32>::Convert(Ar, Property, Data, Tag.ArrayIndex, *Tag.Name.ToString(), *GetName(), *Ar.GetArchiveName());
-					AdvanceProperty = true;
-					continue;
-				}
-				else if( Tag.Type == NAME_Int16Property && PropID == NAME_Int64Property )
-				{
-					TConvertIntToIntProperty<int16, int64>::Convert(Ar, Property, Data, Tag.ArrayIndex, *Tag.Name.ToString(), *GetName(), *Ar.GetArchiveName());
-					AdvanceProperty = true;
-					continue;
-				}
-				else if( Tag.Type == NAME_IntProperty && PropID == NAME_Int64Property )
-				{
-					TConvertIntToIntProperty<int32, int64>::Convert(Ar, Property, Data, Tag.ArrayIndex, *Tag.Name.ToString(), *GetName(), *Ar.GetArchiveName());
-					AdvanceProperty = true;
-					continue;
-				}
-				/* BEGIN - Up-converting uints to larger uints */
-				else if( Tag.Type == NAME_ByteProperty && PropID == NAME_UInt16Property )
-				{
-					if (Tag.EnumName != NAME_None)
-					{
-						uint8 PreviousValue = ReadEnumAsUint8(Tag.EnumName, *Tag.Name.ToString());
-						CastChecked<UIntProperty>(Property)->SetPropertyValue_InContainer(Data, PreviousValue, Tag.ArrayIndex);
-					}
-					else
-					{
-						TConvertIntToIntProperty<uint8, uint16>::Convert(Ar, Property, Data, Tag.ArrayIndex, *Tag.Name.ToString(), *GetName(), *Ar.GetArchiveName());
-					}
-
-					AdvanceProperty = true;
-					continue;
-				}
-				else if( Tag.Type == NAME_ByteProperty && PropID == NAME_UInt32Property )
-				{
-					if (Tag.EnumName != NAME_None)
-					{
-						uint8 PreviousValue = ReadEnumAsUint8(Tag.EnumName, *Tag.Name.ToString());
-						CastChecked<UIntProperty>(Property)->SetPropertyValue_InContainer(Data, PreviousValue, Tag.ArrayIndex);
-					}
-					else
-					{
-						TConvertIntToIntProperty<uint8, uint32>::Convert(Ar, Property, Data, Tag.ArrayIndex, *Tag.Name.ToString(), *GetName(), *Ar.GetArchiveName());
-					}
-
-					AdvanceProperty = true;
-					continue;
-				}
-				else if( Tag.Type == NAME_ByteProperty && PropID == NAME_UInt64Property )
-				{
-					if (Tag.EnumName != NAME_None)
-					{
-						uint8 PreviousValue = ReadEnumAsUint8(Tag.EnumName, *Tag.Name.ToString());
-						CastChecked<UIntProperty>(Property)->SetPropertyValue_InContainer(Data, PreviousValue, Tag.ArrayIndex);
-					}
-					else
-					{
-						TConvertIntToIntProperty<uint8, uint64>::Convert(Ar, Property, Data, Tag.ArrayIndex, *Tag.Name.ToString(), *GetName(), *Ar.GetArchiveName());
-					}
-
-					AdvanceProperty = true;
-					continue;
-				}
-				else if( Tag.Type == NAME_UInt16Property && PropID == NAME_UInt32Property )
-				{
-					TConvertIntToIntProperty<uint16, uint32>::Convert(Ar, Property, Data, Tag.ArrayIndex, *Tag.Name.ToString(), *GetName(), *Ar.GetArchiveName());
-					AdvanceProperty = true;
-					continue;
-				}
-				else if( Tag.Type == NAME_UInt16Property && PropID == NAME_UInt64Property )
-				{
-					TConvertIntToIntProperty<uint16, uint64>::Convert(Ar, Property, Data, Tag.ArrayIndex, *Tag.Name.ToString(), *GetName(), *Ar.GetArchiveName());
-					AdvanceProperty = true;
-					continue;
-				}
-				else if( Tag.Type == NAME_UInt32Property && PropID == NAME_UInt64Property )
-				{
-					TConvertIntToIntProperty<uint32, uint64>::Convert(Ar, Property, Data, Tag.ArrayIndex, *Tag.Name.ToString(), *GetName(), *Ar.GetArchiveName());
-					AdvanceProperty = true;
-					continue;
-				}
-				/* BEGIN - Up-converting ints to larger uints */
-				else if( Tag.Type == NAME_Int8Property && PropID == NAME_UInt16Property )
-				{
-					TConvertIntToIntProperty<int8, uint16>::Convert(Ar, Property, Data, Tag.ArrayIndex, *Tag.Name.ToString(), *GetName(), *Ar.GetArchiveName());
-					AdvanceProperty = true;
-					continue;
-				}
-				else if( Tag.Type == NAME_Int8Property && PropID == NAME_UInt32Property )
-				{
-					TConvertIntToIntProperty<int8, uint32>::Convert(Ar, Property, Data, Tag.ArrayIndex, *Tag.Name.ToString(), *GetName(), *Ar.GetArchiveName());
-					AdvanceProperty = true;
-					continue;
-				}
-				else if( Tag.Type == NAME_Int8Property && PropID == NAME_UInt64Property )
-				{
-					TConvertIntToIntProperty<int8, uint64>::Convert(Ar, Property, Data, Tag.ArrayIndex, *Tag.Name.ToString(), *GetName(), *Ar.GetArchiveName());
-					AdvanceProperty = true;
-					continue;
-				}
-				else if( Tag.Type == NAME_Int16Property && PropID == NAME_UInt32Property )
-				{
-					TConvertIntToIntProperty<int16, uint32>::Convert(Ar, Property, Data, Tag.ArrayIndex, *Tag.Name.ToString(), *GetName(), *Ar.GetArchiveName());
-					AdvanceProperty = true;
-					continue;
-				}
-				else if( Tag.Type == NAME_Int16Property && PropID == NAME_UInt64Property )
-				{
-					TConvertIntToIntProperty<int16, uint64>::Convert(Ar, Property, Data, Tag.ArrayIndex, *Tag.Name.ToString(), *GetName(), *Ar.GetArchiveName());
-					AdvanceProperty = true;
-					continue;
-				}
-				else if( Tag.Type == NAME_IntProperty && PropID == NAME_UInt64Property )
-				{
-					TConvertIntToIntProperty<int32, uint64>::Convert(Ar, Property, Data, Tag.ArrayIndex, *Tag.Name.ToString(), *GetName(), *Ar.GetArchiveName());
-					AdvanceProperty = true;
-					continue;
-				}
-				/* BEGIN - Up-converting uints to larger ints */
-				else if( Tag.Type == NAME_ByteProperty && PropID == NAME_Int16Property )
-				{
-					if (Tag.EnumName != NAME_None)
-					{
-						uint8 PreviousValue = ReadEnumAsUint8(Tag.EnumName, *Tag.Name.ToString());
-						CastChecked<UIntProperty>(Property)->SetPropertyValue_InContainer(Data, PreviousValue, Tag.ArrayIndex);
-					}
-					else
-					{
-						TConvertIntToIntProperty<uint8, int16>::Convert(Ar, Property, Data, Tag.ArrayIndex, *Tag.Name.ToString(), *GetName(), *Ar.GetArchiveName());
-					}
-
-					AdvanceProperty = true;
-					continue;
-				}
-				else if( Tag.Type == NAME_ByteProperty && PropID == NAME_IntProperty )
-				{
-					if (Tag.EnumName != NAME_None)
-					{
-						uint8 PreviousValue = ReadEnumAsUint8(Tag.EnumName, *Tag.Name.ToString());
-						CastChecked<UIntProperty>(Property)->SetPropertyValue_InContainer(Data, PreviousValue, Tag.ArrayIndex);
-					}
-					else
-					{
-						TConvertIntToIntProperty<uint8, int32>::Convert(Ar, Property, Data, Tag.ArrayIndex, *Tag.Name.ToString(), *GetName(), *Ar.GetArchiveName());
-					}
-
-					AdvanceProperty = true;
-					continue;
-				}
-				else if( Tag.Type == NAME_ByteProperty && PropID == NAME_Int64Property )
-				{
-					if (Tag.EnumName != NAME_None)
-					{
-						uint8 PreviousValue = ReadEnumAsUint8(Tag.EnumName, *Tag.Name.ToString());
-						CastChecked<UIntProperty>(Property)->SetPropertyValue_InContainer(Data, PreviousValue, Tag.ArrayIndex);
-					}
-					else
-					{
-						TConvertIntToIntProperty<uint8, int64>::Convert(Ar, Property, Data, Tag.ArrayIndex, *Tag.Name.ToString(), *GetName(), *Ar.GetArchiveName());
-					}
-
-					AdvanceProperty = true;
-					continue;
-				}
-				else if( Tag.Type == NAME_UInt16Property && PropID == NAME_IntProperty )
-				{
-					TConvertIntToIntProperty<uint16, int32>::Convert(Ar, Property, Data, Tag.ArrayIndex, *Tag.Name.ToString(), *GetName(), *Ar.GetArchiveName());
-					AdvanceProperty = true;
-					continue;
-				}
-				else if( Tag.Type == NAME_UInt16Property && PropID == NAME_Int64Property )
-				{
-					TConvertIntToIntProperty<uint16, int64>::Convert(Ar, Property, Data, Tag.ArrayIndex, *Tag.Name.ToString(), *GetName(), *Ar.GetArchiveName());
-					AdvanceProperty = true;
-					continue;
-				}
-				else if( Tag.Type == NAME_UInt32Property && PropID == NAME_Int64Property )
-				{
-					TConvertIntToIntProperty<uint32, int64>::Convert(Ar, Property, Data, Tag.ArrayIndex, *Tag.Name.ToString(), *GetName(), *Ar.GetArchiveName());
-					AdvanceProperty = true;
-					continue;
-				}
-				/* BEGIN - Down-converting ints to smaller ints */
-				else if( Tag.Type == NAME_Int64Property && PropID == NAME_Int8Property )
-				{
-					TConvertIntToIntProperty<int64, int8>::Convert(Ar, Property, Data, Tag.ArrayIndex, *Tag.Name.ToString(), *GetName(), *Ar.GetArchiveName());
-					AdvanceProperty = true;
-					continue;
-				}
-				else if( Tag.Type == NAME_Int64Property && PropID == NAME_Int16Property )
-				{
-					TConvertIntToIntProperty<int64, int16>::Convert(Ar, Property, Data, Tag.ArrayIndex, *Tag.Name.ToString(), *GetName(), *Ar.GetArchiveName());
-					AdvanceProperty = true;
-					continue;
-				}
-				else if( Tag.Type == NAME_Int64Property && PropID == NAME_IntProperty )
-				{
-					TConvertIntToIntProperty<int64, int32>::Convert(Ar, Property, Data, Tag.ArrayIndex, *Tag.Name.ToString(), *GetName(), *Ar.GetArchiveName());
-					AdvanceProperty = true;
-					continue;
-				}
-				else if( Tag.Type == NAME_IntProperty && PropID == NAME_Int8Property )
-				{
-					TConvertIntToIntProperty<int32, int8>::Convert(Ar, Property, Data, Tag.ArrayIndex, *Tag.Name.ToString(), *GetName(), *Ar.GetArchiveName());
-					AdvanceProperty = true;
-					continue;
-				}
-				else if( Tag.Type == NAME_IntProperty && PropID == NAME_Int16Property )
-				{
-					TConvertIntToIntProperty<int32, int16>::Convert(Ar, Property, Data, Tag.ArrayIndex, *Tag.Name.ToString(), *GetName(), *Ar.GetArchiveName());
-					AdvanceProperty = true;
-					continue;
-				}
-				else if( Tag.Type == NAME_Int16Property && PropID == NAME_Int8Property )
-				{
-					TConvertIntToIntProperty<int16, int8>::Convert(Ar, Property, Data, Tag.ArrayIndex, *Tag.Name.ToString(), *GetName(), *Ar.GetArchiveName());
-					AdvanceProperty = true;
-					continue;
-				}
-				/* BEGIN - Down-converting uints to smaller uints */
-				else if( Tag.Type == NAME_UInt64Property && PropID == NAME_ByteProperty )
-				{
-					UByteProperty* ByteProperty = static_cast<UByteProperty*>(Property);
-					if (UEnum* Enum = ByteProperty->GetIntPropertyEnum())
-					{
-						TConvertIntToEnumProperty<uint64>::Convert(Ar, ByteProperty, Enum, Data, Tag.ArrayIndex, *Tag.Name.ToString(), *GetName(), *Ar.GetArchiveName());
-					}
-					else
-					{
-						TConvertIntToIntProperty<uint64, uint8>::Convert(Ar, Property, Data, Tag.ArrayIndex, *Tag.Name.ToString(), *GetName(), *Ar.GetArchiveName());
-					}
-					AdvanceProperty = true;
-					continue;
-				}
-				else if( Tag.Type == NAME_UInt64Property && PropID == NAME_UInt16Property )
-				{
-					TConvertIntToIntProperty<uint64, uint16>::Convert(Ar, Property, Data, Tag.ArrayIndex, *Tag.Name.ToString(), *GetName(), *Ar.GetArchiveName());
-					AdvanceProperty = true;
-					continue;
-				}
-				else if( Tag.Type == NAME_UInt64Property && PropID == NAME_UInt32Property )
-				{
-					TConvertIntToIntProperty<uint64, uint32>::Convert(Ar, Property, Data, Tag.ArrayIndex, *Tag.Name.ToString(), *GetName(), *Ar.GetArchiveName());
-					AdvanceProperty = true;
-					continue;
-				}
-				else if( Tag.Type == NAME_UInt32Property && PropID == NAME_ByteProperty )
-				{
-					UByteProperty* ByteProperty = static_cast<UByteProperty*>(Property);
-					if (UEnum* Enum = ByteProperty->GetIntPropertyEnum())
-					{
-						TConvertIntToEnumProperty<uint32>::Convert(Ar, ByteProperty, Enum, Data, Tag.ArrayIndex, *Tag.Name.ToString(), *GetName(), *Ar.GetArchiveName());
-					}
-					else
-					{
-						TConvertIntToIntProperty<uint32, uint8>::Convert(Ar, Property, Data, Tag.ArrayIndex, *Tag.Name.ToString(), *GetName(), *Ar.GetArchiveName());
-					}
-					AdvanceProperty = true;
-					continue;
-				}
-				else if( Tag.Type == NAME_UInt32Property && PropID == NAME_UInt16Property )
-				{
-					TConvertIntToIntProperty<uint32, uint16>::Convert(Ar, Property, Data, Tag.ArrayIndex, *Tag.Name.ToString(), *GetName(), *Ar.GetArchiveName());
-					AdvanceProperty = true;
-					continue;
-				}
-				else if( Tag.Type == NAME_UInt16Property && PropID == NAME_ByteProperty )
-				{
-					UByteProperty* ByteProperty = static_cast<UByteProperty*>(Property);
-					if (UEnum* Enum = ByteProperty->GetIntPropertyEnum())
-					{
-						TConvertIntToEnumProperty<uint16>::Convert(Ar, ByteProperty, Enum, Data, Tag.ArrayIndex, *Tag.Name.ToString(), *GetName(), *Ar.GetArchiveName());
-					}
-					else
-					{
-						TConvertIntToIntProperty<uint16, uint8>::Convert(Ar, Property, Data, Tag.ArrayIndex, *Tag.Name.ToString(), *GetName(), *Ar.GetArchiveName());
-					}
-					AdvanceProperty = true;
-					continue;
-				}
-				/* BEGIN - Down-converting ints to smaller uints */
-				else if( Tag.Type == NAME_Int64Property && PropID == NAME_ByteProperty )
-				{
-					UByteProperty* ByteProperty = static_cast<UByteProperty*>(Property);
-					if (UEnum* Enum = ByteProperty->GetIntPropertyEnum())
-					{
-						TConvertIntToEnumProperty<int64>::Convert(Ar, ByteProperty, Enum, Data, Tag.ArrayIndex, *Tag.Name.ToString(), *GetName(), *Ar.GetArchiveName());
-					}
-					else
-					{
-						TConvertIntToIntProperty<int64, uint8>::Convert(Ar, Property, Data, Tag.ArrayIndex, *Tag.Name.ToString(), *GetName(), *Ar.GetArchiveName());
-					}
-					AdvanceProperty = true;
-					continue;
-				}
-				else if( Tag.Type == NAME_Int64Property && PropID == NAME_UInt16Property )
-				{
-					TConvertIntToIntProperty<int64, uint16>::Convert(Ar, Property, Data, Tag.ArrayIndex, *Tag.Name.ToString(), *GetName(), *Ar.GetArchiveName());
-					AdvanceProperty = true;
-					continue;
-				}
-				else if( Tag.Type == NAME_Int64Property && PropID == NAME_UInt32Property )
-				{
-					TConvertIntToIntProperty<int64, uint32>::Convert(Ar, Property, Data, Tag.ArrayIndex, *Tag.Name.ToString(), *GetName(), *Ar.GetArchiveName());
-					AdvanceProperty = true;
-					continue;
-				}
-				else if( Tag.Type == NAME_IntProperty && PropID == NAME_ByteProperty )
-				{
-					UByteProperty* ByteProperty = static_cast<UByteProperty*>(Property);
-					if (UEnum* Enum = ByteProperty->GetIntPropertyEnum())
-					{
-						TConvertIntToEnumProperty<int32>::Convert(Ar, ByteProperty, Enum, Data, Tag.ArrayIndex, *Tag.Name.ToString(), *GetName(), *Ar.GetArchiveName());
-					}
-					else
-					{
-						TConvertIntToIntProperty<int32, uint8>::Convert(Ar, Property, Data, Tag.ArrayIndex, *Tag.Name.ToString(), *GetName(), *Ar.GetArchiveName());
-					}
-					AdvanceProperty = true;
-					continue;
-				}
-				else if( Tag.Type == NAME_IntProperty && PropID == NAME_UInt16Property )
-				{
-					TConvertIntToIntProperty<int32, uint16>::Convert(Ar, Property, Data, Tag.ArrayIndex, *Tag.Name.ToString(), *GetName(), *Ar.GetArchiveName());
-					AdvanceProperty = true;
-					continue;
-				}
-				else if( Tag.Type == NAME_Int16Property && PropID == NAME_ByteProperty )
-				{
-					UByteProperty* ByteProperty = static_cast<UByteProperty*>(Property);
-					if (UEnum* Enum = ByteProperty->GetIntPropertyEnum())
-					{
-						TConvertIntToEnumProperty<int16>::Convert(Ar, ByteProperty, Enum, Data, Tag.ArrayIndex, *Tag.Name.ToString(), *GetName(), *Ar.GetArchiveName());
-					}
-					else
-					{
-						TConvertIntToIntProperty<int16, uint8>::Convert(Ar, Property, Data, Tag.ArrayIndex, *Tag.Name.ToString(), *GetName(), *Ar.GetArchiveName());
-					}
-					AdvanceProperty = true;
-					continue;
-				}
-				/* BEGIN - Down-converting uints to smaller ints */
-				else if( Tag.Type == NAME_UInt64Property && PropID == NAME_Int8Property )
-				{
-					TConvertIntToIntProperty<uint64, int8>::Convert(Ar, Property, Data, Tag.ArrayIndex, *Tag.Name.ToString(), *GetName(), *Ar.GetArchiveName());
-					AdvanceProperty = true;
-					continue;
-				}
-				else if( Tag.Type == NAME_UInt64Property && PropID == NAME_Int16Property )
-				{
-					TConvertIntToIntProperty<uint64, int16>::Convert(Ar, Property, Data, Tag.ArrayIndex, *Tag.Name.ToString(), *GetName(), *Ar.GetArchiveName());
-					AdvanceProperty = true;
-					continue;
-				}
-				else if( Tag.Type == NAME_UInt64Property && PropID == NAME_IntProperty )
-				{
-					TConvertIntToIntProperty<uint64, int32>::Convert(Ar, Property, Data, Tag.ArrayIndex, *Tag.Name.ToString(), *GetName(), *Ar.GetArchiveName());
-					AdvanceProperty = true;
-					continue;
-				}
-				else if( Tag.Type == NAME_UInt32Property && PropID == NAME_Int8Property )
-				{
-					TConvertIntToIntProperty<uint32, int8>::Convert(Ar, Property, Data, Tag.ArrayIndex, *Tag.Name.ToString(), *GetName(), *Ar.GetArchiveName());
-					AdvanceProperty = true;
-					continue;
-				}
-				else if( Tag.Type == NAME_UInt32Property && PropID == NAME_Int16Property )
-				{
-					TConvertIntToIntProperty<uint32, int16>::Convert(Ar, Property, Data, Tag.ArrayIndex, *Tag.Name.ToString(), *GetName(), *Ar.GetArchiveName());
-					AdvanceProperty = true;
-					continue;
-				}
-				else if( Tag.Type == NAME_UInt16Property && PropID == NAME_Int8Property )
-				{
-					TConvertIntToIntProperty<uint16, int8>::Convert(Ar, Property, Data, Tag.ArrayIndex, *Tag.Name.ToString(), *GetName(), *Ar.GetArchiveName());
-					AdvanceProperty = true;
-					continue;
-				}
-				/* BEGIN - Converting ints to same-sized uints */
-				else if( Tag.Type == NAME_Int8Property && PropID == NAME_ByteProperty )
-				{
-					UByteProperty* ByteProperty = static_cast<UByteProperty*>(Property);
-					if (UEnum* Enum = ByteProperty->GetIntPropertyEnum())
-					{
-						TConvertIntToEnumProperty<int8>::Convert(Ar, ByteProperty, Enum, Data, Tag.ArrayIndex, *Tag.Name.ToString(), *GetName(), *Ar.GetArchiveName());
-					}
-					else
-					{
-						TConvertIntToIntProperty<int8, uint8>::Convert(Ar, Property, Data, Tag.ArrayIndex, *Tag.Name.ToString(), *GetName(), *Ar.GetArchiveName());
-					}
-					AdvanceProperty = true;
-					continue;
-				}
-				else if( Tag.Type == NAME_Int16Property && PropID == NAME_UInt16Property )
-				{
-					TConvertIntToIntProperty<int16, uint16>::Convert(Ar, Property, Data, Tag.ArrayIndex, *Tag.Name.ToString(), *GetName(), *Ar.GetArchiveName());
-					AdvanceProperty = true;
-					continue;
-				}
-				else if( Tag.Type == NAME_IntProperty && PropID == NAME_UInt32Property )
-				{
-					TConvertIntToIntProperty<int32, uint32>::Convert(Ar, Property, Data, Tag.ArrayIndex, *Tag.Name.ToString(), *GetName(), *Ar.GetArchiveName());
-					AdvanceProperty = true;
-					continue;
-				}
-				else if( Tag.Type == NAME_Int64Property && PropID == NAME_UInt64Property )
-				{
-					TConvertIntToIntProperty<int64, uint64>::Convert(Ar, Property, Data, Tag.ArrayIndex, *Tag.Name.ToString(), *GetName(), *Ar.GetArchiveName());
-					AdvanceProperty = true;
-					continue;
-				}
-				/* BEGIN - Converting uints to same-sized ints */
-				else if( Tag.Type == NAME_ByteProperty && PropID == NAME_Int8Property )
-				{
-					if (Tag.EnumName != NAME_None)
-					{
-						uint8 PreviousValue = ReadEnumAsUint8(Tag.EnumName, *Tag.Name.ToString());
-						CastChecked<UIntProperty>(Property)->SetPropertyValue_InContainer(Data, PreviousValue, Tag.ArrayIndex);
-					}
-					else
-					{
-						TConvertIntToIntProperty<uint8, int8>::Convert(Ar, Property, Data, Tag.ArrayIndex, *Tag.Name.ToString(), *GetName(), *Ar.GetArchiveName());
-					}
-
-					AdvanceProperty = true;
-					continue;
-				}
-				else if( Tag.Type == NAME_UInt16Property && PropID == NAME_Int16Property )
-				{
-					TConvertIntToIntProperty<uint16, int16>::Convert(Ar, Property, Data, Tag.ArrayIndex, *Tag.Name.ToString(), *GetName(), *Ar.GetArchiveName());
-					AdvanceProperty = true;
-					continue;
-				}
-				else if( Tag.Type == NAME_UInt32Property && PropID == NAME_IntProperty )
-				{
-					TConvertIntToIntProperty<uint32, int32>::Convert(Ar, Property, Data, Tag.ArrayIndex, *Tag.Name.ToString(), *GetName(), *Ar.GetArchiveName());
-					AdvanceProperty = true;
-					continue;
-				}
-				else if( Tag.Type == NAME_UInt64Property && PropID == NAME_Int64Property )
-				{
-					TConvertIntToIntProperty<uint64, int64>::Convert(Ar, Property, Data, Tag.ArrayIndex, *Tag.Name.ToString(), *GetName(), *Ar.GetArchiveName());
-					AdvanceProperty = true;
-					continue;
-				}
-				/* END - Converting ints */
-				else
-				{
 					UE_LOG(LogClass, Warning, TEXT("Type mismatch in %s of %s - Previous (%s) Current(%s) for package:  %s"), *Tag.Name.ToString(), *GetName(), *Tag.Type.ToString(), *PropID.ToString(), *Ar.GetArchiveName() );
 				}
-			}
-			else if( Tag.Type == NAME_ArrayProperty && Tag.InnerType != NAME_None && Tag.InnerType != (ArrayInnerID = CastChecked<UArrayProperty>(Property)->Inner->GetID()) )
-			{
-				UArrayProperty* ArrayProperty = static_cast<UArrayProperty*>(Property);
-				void* ArrayPropertyData = ArrayProperty->ContainerPtrToValuePtr<void>(Data);
-
-				int32 ElementCount = 0;
-				Ar << ElementCount;
-
-				FScriptArrayHelper ScriptArrayHelper(ArrayProperty, ArrayPropertyData);
-				ScriptArrayHelper.EmptyAndAddValues(ElementCount);
-
-				// Convert properties from old type to new type automatically if types are compatible (array case)
-
-				if (Tag.InnerType == NAME_StrProperty && dynamic_cast<UTextProperty*>(ArrayProperty->Inner)) // Convert serialized string to text.
-				{ 
-					for(int32 i = 0; i < ElementCount; ++i)
-					{
-						FString str;
-						Ar << str;
-						FText Text = FText::FromString(str);
-						Text.TextData->PersistText();
-						Text.Flags |= ETextFlag::ConvertedProperty;
-						CastChecked<UTextProperty>(ArrayProperty->Inner)->SetPropertyValue(ScriptArrayHelper.GetRawPtr(i), Text);
-						AdvanceProperty = true;
-					}
-					continue;
-				}
-				else if( Tag.InnerType==NAME_TextProperty && dynamic_cast<UStrProperty*>(ArrayProperty->Inner) ) // Convert serialized text to string.
-				{ 
-					for(int32 i = 0; i < ElementCount; ++i)
-					{
-						FText Text;  
-						Ar << Text;
-						FString String = FTextInspector::GetSourceString(Text) ? *FTextInspector::GetSourceString(Text) : TEXT("");
-						static_cast<UStrProperty*>(ArrayProperty->Inner)->SetPropertyValue(ScriptArrayHelper.GetRawPtr(i), String);
-						AdvanceProperty = true;
-					}
-					continue; 
-				}
-				else if (Tag.InnerType == NAME_NameProperty && dynamic_cast<UTextProperty*>(ArrayProperty->Inner)) // Convert serialized name to text.
-				{ 
-					for(int32 i = 0; i < ElementCount; ++i)
-					{
-						FName Name;
-						Ar << Name;
-						FText Text = FText::FromName(Name);
-						Text.TextData->PersistText();
-						Text.Flags |= ETextFlag::ConvertedProperty;
-						CastChecked<UTextProperty>(ArrayProperty->Inner)->SetPropertyValue(ScriptArrayHelper.GetRawPtr(i), Text);
-						AdvanceProperty = true;
-					}
-					continue;
-				}
-				else if( Tag.InnerType==NAME_TextProperty && dynamic_cast<UNameProperty*>(ArrayProperty->Inner) ) // Convert serialized text to name.
-				{ 
-					for(int32 i = 0; i < ElementCount; ++i)
-					{
-						FText Text;  
-						Ar << Text;
-						FName Name = FTextInspector::GetSourceString(Text) ? FName(**FTextInspector::GetSourceString(Text)) : NAME_None;
-						static_cast<UNameProperty*>(ArrayProperty->Inner)->SetPropertyValue(ScriptArrayHelper.GetRawPtr(i), Name);
-						AdvanceProperty = true;
-					}
-					continue; 
-				}
-				else if ((Tag.InnerType == NAME_AssetObjectProperty || Tag.InnerType == NAME_AssetSubclassOfProperty) && (ArrayInnerID == NAME_ObjectProperty || ArrayInnerID == NAME_ClassProperty))
-				{
-					for (int32 i = 0; i < ElementCount; ++i)
-					{
-						// This property used to be a TAssetPtr<Foo> but is now a raw UObjectProperty Foo*, we can convert without loss of data
-						FAssetPtr PreviousValue;
-						Ar << PreviousValue;
-
-						// now copy the value into the object's address space
-						UObject* PreviousValueObj = PreviousValue.Get();
-						CastChecked<UObjectProperty>(ArrayProperty->Inner)->SetPropertyValue(ScriptArrayHelper.GetRawPtr(i), PreviousValueObj);
-
-						AdvanceProperty = true;
-					}
-					continue;
-				}
-				else if ((Tag.InnerType == NAME_ObjectProperty || Tag.InnerType == NAME_ClassProperty) && (ArrayInnerID == NAME_AssetObjectProperty || ArrayInnerID == NAME_AssetSubclassOfProperty))
-				{
-					for (int32 i = 0; i < ElementCount; ++i)
-					{
-						// This property used to be a raw UObjectProperty Foo* but is now a TAssetPtr<Foo>
-						UObject* PreviousValue = NULL;
-						Ar << PreviousValue;
-
-						// now copy the value into the object's address space
-						FAssetPtr PreviousValueAssetPtr(PreviousValue);
-						CastChecked<UAssetObjectProperty>(ArrayProperty->Inner)->SetPropertyValue(ScriptArrayHelper.GetRawPtr(i), PreviousValueAssetPtr);
 
-						AdvanceProperty = true;
-					}
-					continue;
-				}
-				// TODO: Implement SerializeFromMismatchedTag handling for arrays of structs
-				else
-				{
-					UE_LOG(LogClass, Warning, TEXT("Array Inner Type mismatch in %s of %s - Previous (%s) Current(%s) for package:  %s"), *Tag.Name.ToString(), *GetName(), *Tag.InnerType.ToString(), *ArrayInnerID.ToString(), *Ar.GetArchiveName() );
-				}
-			}
-			else if( Tag.Type==NAME_StructProperty && Tag.StructName!=CastChecked<UStructProperty>(Property)->Struct->GetFName() 
-				&& !CanSerializeFromStructWithDifferentName(Ar, Tag, CastChecked<UStructProperty>(Property)))
-			{
-				UE_LOG(LogClass, Warning, TEXT("Property %s of %s has a struct type mismatch (tag %s != prop %s) in package:  %s. If that struct got renamed, add an entry to ActiveStructRedirects."),
-					*Tag.Name.ToString(), *GetName(), *Tag.StructName.ToString(), *CastChecked<UStructProperty>(Property)->Struct->GetName(), *Ar.GetArchiveName() );
-			}
-			else if ( Tag.Type == NAME_ByteProperty && ( (Tag.EnumName == NAME_None && ExactCast<UByteProperty>(Property)->Enum != NULL) || 
-														(Tag.EnumName != NAME_None && ExactCast<UByteProperty>(Property)->Enum == NULL) ))
-			{
-				// a byte property gained or lost an enum
-				// attempt to convert it
-				uint8 PreviousValue;
-				if (Tag.EnumName == NAME_None)
-				{
-					// simply pretend the property still doesn't have an enum and serialize the single byte
-					Ar << PreviousValue;
-				}
-				else
-				{
-					// attempt to find the old enum and get the byte value from the serialized enum name
-					PreviousValue = ReadEnumAsUint8(Tag.EnumName, *Tag.Name.ToString());
-				}
-				
-				// now copy the value into the object's address space
-				CastChecked<UByteProperty>(Property)->SetPropertyValue_InContainer(Data, PreviousValue, Tag.ArrayIndex);
-				AdvanceProperty = true;
-				continue;
-			}
 			else
 			{
 				uint8* DestAddress = Property->ContainerPtrToValuePtr<uint8>(Data, Tag.ArrayIndex);  
@@ -1877,14 +1061,14 @@ void UStruct::SerializeTaggedProperties(FArchive& Ar, uint8* Data, UStruct* Defa
 				// This property is ok.			
 				Tag.SerializeTaggedProperty(Ar, Property, DestAddress, DefaultsFromParent);
 
-				AdvanceProperty = true;
+				bAdvanceProperty = true;
 				if (!Ar.IsCriticalError())
 				{
 					continue;
 				}
 			}
 
-			AdvanceProperty = false;
+			bAdvanceProperty = false;
 
 			// Skip unknown or bad property.
 			const int64 RemainingSize = Tag.Size - (Ar.Tell() - StartOfProperty);
diff --git a/Engine/Source/Runtime/CoreUObject/Private/UObject/FastReferenceCollector.h b/Engine/Source/Runtime/CoreUObject/Private/UObject/FastReferenceCollector.h
--- a/Engine/Source/Runtime/CoreUObject/Private/UObject/FastReferenceCollector.h
+++ b/Engine/Source/Runtime/CoreUObject/Private/UObject/FastReferenceCollector.h
@@ -359,6 +359,14 @@ private:
 						FSimpleObjectReferenceCollectorArchive CollectorArchive(CurrentObject, ReferenceCollector);
 						MapProperty->SerializeItem(CollectorArchive, Map, nullptr);
 					}
+					else if (ReferenceInfo.Type == GCRT_AddTSetReferencedObjects)
+					{
+						void*         Set = StackEntryData + ReferenceInfo.Offset;
+						USetProperty* SetProperty = (USetProperty*)TokenStream->ReadPointer(TokenStreamIndex);
+						TokenReturnCount = ReferenceInfo.ReturnCount;
+						FSimpleObjectReferenceCollectorArchive CollectorArchive(CurrentObject, ReferenceCollector);
+						SetProperty->SerializeItem(CollectorArchive, Set, nullptr);
+					}
 					else if (ReferenceInfo.Type == GCRT_EndOfPointer)
 					{
 						TokenReturnCount = ReferenceInfo.ReturnCount;
diff --git a/Engine/Source/Runtime/CoreUObject/Private/UObject/GarbageCollection.cpp b/Engine/Source/Runtime/CoreUObject/Private/UObject/GarbageCollection.cpp
--- a/Engine/Source/Runtime/CoreUObject/Private/UObject/GarbageCollection.cpp
+++ b/Engine/Source/Runtime/CoreUObject/Private/UObject/GarbageCollection.cpp
@@ -1447,6 +1447,18 @@ bool UMapProperty::ContainsObjectReference() const
 	return KeyProp->ContainsObjectReference() || ValueProp->ContainsObjectReference();
 }
 
+/**
+* Returns true if this property, or in the case of e.g. array or struct properties any sub- property, contains a
+* UObject reference.
+*
+* @return true if property (or sub- properties) contain a UObject reference, false otherwise
+*/
+bool USetProperty::ContainsObjectReference() const
+{
+	check(ElementProp);
+	return ElementProp->ContainsObjectReference();
+}
+
 /**
  * Returns true if this property, or in the case of e.g. array or struct properties any sub- property, contains a
  * UObject reference.
@@ -1507,6 +1519,13 @@ bool UMapProperty::ContainsWeakObjectReference() const
 	return KeyProp->ContainsWeakObjectReference() || ValueProp->ContainsWeakObjectReference();
 }
 
+// Returns true if this property contains a weak UObject reference.
+bool USetProperty::ContainsWeakObjectReference() const
+{
+	check(ElementProp);
+	return ElementProp->ContainsWeakObjectReference();
+}
+
 // Returns true if this property contains a weak UObject reference.
 bool UStructProperty::ContainsWeakObjectReference() const
 {
@@ -1669,6 +1688,19 @@ void UMapProperty::EmitReferenceInfo(UClass& OwnerClass, int32 BaseOffset)
 	}
 }
 
+/**
+* Emits tokens used by realtime garbage collection code to passed in OwnerClass' ReferenceTokenStream. The offset emitted is relative
+* to the passed in BaseOffset which is used by e.g. arrays of structs.
+*/
+void USetProperty::EmitReferenceInfo(UClass& OwnerClass, int32 BaseOffset)
+{
+	if (ContainsObjectReference())
+	{
+		OwnerClass.EmitObjectReference(BaseOffset + GetOffset_ForGC(), GetFName(), GCRT_AddTSetReferencedObjects);
+		OwnerClass.ReferenceTokenStream.EmitPointer((const void*)this);
+	}
+}
+
 
 /**
  * Emits tokens used by realtime garbage collection code to passed in ReferenceTokenStream. The offset emitted is relative
@@ -1892,6 +1924,7 @@ void FGCReferenceTokenStream::ReplaceOrAddAddReferencedObjectsCall(void (*AddRef
 				return;
 			}
 		case GCRT_AddTMapReferencedObjects:
+		case GCRT_AddTSetReferencedObjects:
 			{
 				// Skip pointer
 				TokenIndex += GNumTokensPerPointer;
diff --git a/Engine/Source/Runtime/CoreUObject/Private/UObject/Property.cpp b/Engine/Source/Runtime/CoreUObject/Private/UObject/Property.cpp
--- a/Engine/Source/Runtime/CoreUObject/Private/UObject/Property.cpp
+++ b/Engine/Source/Runtime/CoreUObject/Private/UObject/Property.cpp
@@ -6,6 +6,7 @@
 
 #include "CoreUObjectPrivate.h"
 #include "PropertyHelper.h"
+#include "PropertyTag.h"
 #include "StringAssetReference.h"
 #include "StringClassReference.h"
 
@@ -584,22 +585,28 @@ bool UProperty::ExportText_Direct
 bool UProperty::ShouldSerializeValue( FArchive& Ar ) const
 {
 	if (Ar.ShouldSkipProperty(this))
+	{
 		return false;
+	}
 
-	if (Ar.IsSaveGame() && !(PropertyFlags & CPF_SaveGame))
+	if (!(PropertyFlags & CPF_SaveGame) && Ar.IsSaveGame())
+	{
 		return false;
+	}
 
-	static uint64 SkipFlags = CPF_Transient | CPF_DuplicateTransient | CPF_NonPIEDuplicateTransient | CPF_NonTransactional | CPF_Deprecated | CPF_DevelopmentAssets;
+	const uint64 SkipFlags = CPF_Transient | CPF_DuplicateTransient | CPF_NonPIEDuplicateTransient | CPF_NonTransactional | CPF_Deprecated | CPF_DevelopmentAssets;
 	if (!(PropertyFlags & SkipFlags))
+	{
 		return true;
+	}
 
 	bool Skip =
 			((PropertyFlags & CPF_Transient) && Ar.IsPersistent() && !Ar.IsSerializingDefaults())
 		||	((PropertyFlags & CPF_DuplicateTransient) && (Ar.GetPortFlags() & PPF_Duplicate))
 		||	((PropertyFlags & CPF_NonPIEDuplicateTransient) && !(Ar.GetPortFlags() & PPF_DuplicateForPIE) && (Ar.GetPortFlags() & PPF_Duplicate))
-		||  (Ar.IsFilterEditorOnly() && IsEditorOnlyProperty())
 		||	((PropertyFlags & CPF_NonTransactional) && Ar.IsTransacting())
-		||	((PropertyFlags & CPF_Deprecated) && !Ar.HasAllPortFlags(PPF_UseDeprecatedProperties) && (Ar.IsSaving() || Ar.IsTransacting() || Ar.WantBinaryPropertySerialization()));
+		||	((PropertyFlags & CPF_Deprecated) && !Ar.HasAllPortFlags(PPF_UseDeprecatedProperties) && (Ar.IsSaving() || Ar.IsTransacting() || Ar.WantBinaryPropertySerialization()))
+		||  (IsEditorOnlyProperty() && Ar.IsFilterEditorOnly());
 
 	return !Skip;
 }
@@ -1186,6 +1193,44 @@ IMPLEMENT_CORE_INTRINSIC_CLASS(UProperty, UField,
 	}
 );
 
+uint8 UNumericProperty::ReadEnumAsUint8(FArchive& Ar, UStruct* DefaultsStruct, const FPropertyTag& Tag)
+{
+	//@warning: mirrors loading code in UByteProperty::SerializeItem()
+	FName EnumName;
+	Ar << EnumName;
+
+	UEnum* Enum = FindField<UEnum>(dynamic_cast<UClass*>(DefaultsStruct) ? static_cast<UClass*>(DefaultsStruct) : DefaultsStruct->GetTypedOuter<UClass>(), Tag.EnumName);
+	if (!Enum)
+	{
+		Enum = FindObject<UEnum>(ANY_PACKAGE, *Tag.EnumName.ToString(), true);
+	}
+
+	if (!Enum)
+	{
+		UE_LOG(LogClass, Warning, TEXT("Failed to find enum '%s' when converting property '%s' during property loading - setting to 0"), *Tag.EnumName.ToString(), *Tag.Name.ToString());
+		return 0;
+	}
+
+	Ar.Preload(Enum);
+
+	uint8 Result = Enum->GetValueByName(EnumName);
+	if (!Enum->IsValidEnumValue(Result))
+	{
+		UE_LOG(
+			LogClass,
+			Warning,
+			TEXT("Failed to find valid enum value '%s' for enum type '%s' when converting property '%s' during property loading - setting to '%s'"),
+			*EnumName.ToString(),
+			*Enum->GetName(),
+			*Tag.Name.ToString(),
+			*Enum->GetNameByValue(Enum->GetMaxEnumValue()).ToString()
+			);
+
+		return Enum->GetMaxEnumValue();
+	}
+
+	return Result;
+};
 
 const TCHAR* UNumericProperty::ImportText_Internal( const TCHAR* Buffer, void* Data, int32 PortFlags, UObject* Parent, FOutputDevice* ErrorText ) const
 {
diff --git a/Engine/Source/Runtime/CoreUObject/Private/UObject/PropertyArray.cpp b/Engine/Source/Runtime/CoreUObject/Private/UObject/PropertyArray.cpp
--- a/Engine/Source/Runtime/CoreUObject/Private/UObject/PropertyArray.cpp
+++ b/Engine/Source/Runtime/CoreUObject/Private/UObject/PropertyArray.cpp
@@ -490,6 +490,54 @@ bool UArrayProperty::SameType(const UProperty* Other) const
 	return Super::SameType(Other) && Inner && Inner->SameType(((UArrayProperty*)Other)->Inner);
 }
 
+bool UArrayProperty::ConvertFromType(const FPropertyTag& Tag, FArchive& Ar, uint8* Data, UStruct* DefaultsStruct, bool& bOutAdvanceProperty)
+{
+	// TODO: The ArrayProperty Tag really doesn't have adequate information for
+	// many types. This should probably all be moved in to ::SerializeItem
+	bOutAdvanceProperty = false;
+
+	if (Tag.Type == NAME_ArrayProperty && Tag.InnerType != NAME_None && Tag.InnerType != Inner->GetID())
+	{
+		void* ArrayPropertyData = ContainerPtrToValuePtr<void>(Data);
+
+		int32 ElementCount = 0;
+		Ar << ElementCount;
+
+		FScriptArrayHelper ScriptArrayHelper(this, ArrayPropertyData);
+		ScriptArrayHelper.EmptyAndAddValues(ElementCount);
+
+		// Convert properties from old type to new type automatically if types are compatible (array case)
+		if (ElementCount > 0)
+		{
+			FPropertyTag InnerPropertyTag;
+			InnerPropertyTag.Type = Tag.InnerType;
+			InnerPropertyTag.ArrayIndex = 0;
+
+			bool bDummyAdvance;
+			if (Inner->ConvertFromType(InnerPropertyTag, Ar, ScriptArrayHelper.GetRawPtr(0), DefaultsStruct, bDummyAdvance))
+			{
+				for (int32 i = 1; i < ElementCount; ++i)
+				{
+					verify(Inner->ConvertFromType(InnerPropertyTag, Ar, ScriptArrayHelper.GetRawPtr(i), DefaultsStruct, bDummyAdvance));
+				}
+				bOutAdvanceProperty = true;
+			}
+			// TODO: Implement SerializeFromMismatchedTag handling for arrays of structs
+			else
+			{
+				UE_LOG(LogClass, Warning, TEXT("Array Inner Type mismatch in %s of %s - Previous (%s) Current(%s) for package:  %s"), *Tag.Name.ToString(), *GetName(), *Tag.InnerType.ToString(), *Inner->GetID().ToString(), *Ar.GetArchiveName() );
+			}
+		}
+		else
+		{
+			bOutAdvanceProperty = true;
+		}
+		return true;
+	}
+
+	return false;
+}
+
 IMPLEMENT_CORE_INTRINSIC_CLASS(UArrayProperty, UProperty,
 	{
 		Class->EmitObjectReference(STRUCT_OFFSET(UArrayProperty, Inner), TEXT("Inner"));
diff --git a/Engine/Source/Runtime/CoreUObject/Private/UObject/PropertyAssetObject.cpp b/Engine/Source/Runtime/CoreUObject/Private/UObject/PropertyAssetObject.cpp
--- a/Engine/Source/Runtime/CoreUObject/Private/UObject/PropertyAssetObject.cpp
+++ b/Engine/Source/Runtime/CoreUObject/Private/UObject/PropertyAssetObject.cpp
@@ -131,6 +131,40 @@ const TCHAR* UAssetObjectProperty::ImportText_Internal( const TCHAR* InBuffer, v
 	return Buffer;
 }
 
+bool UAssetObjectProperty::ConvertFromType(const FPropertyTag& Tag, FArchive& Ar, uint8* Data, UStruct* DefaultsStruct, bool& bOutAdvanceProperty)
+{
+	bOutAdvanceProperty = true;
+
+	if (Tag.Type == NAME_ObjectProperty || Tag.Type == NAME_ClassProperty)
+	{
+		// This property used to be a raw UObjectProperty Foo* but is now a TAssetPtr<Foo>
+		UObject* PreviousValue = nullptr;
+		Ar << PreviousValue;
+
+		// now copy the value into the object's address space
+		FAssetPtr PreviousValueAssetPtr(PreviousValue);
+		SetPropertyValue_InContainer(Data, PreviousValueAssetPtr, Tag.ArrayIndex);
+
+		return true;
+	}
+	else if (Tag.Type == NAME_StructProperty)
+	{
+		// This property used to be a FStringAssetReference but is now a TAssetPtr<Foo>
+		FStringAssetReference PreviousValue;
+		// explicitly call Serialize to ensure that the various delegates needed for cooking are fired
+		PreviousValue.Serialize(Ar);
+
+		// now copy the value into the object's address space
+		FAssetPtr PreviousValueAssetPtr;
+		PreviousValueAssetPtr = PreviousValue;
+		SetPropertyValue_InContainer(Data, PreviousValueAssetPtr, Tag.ArrayIndex);
+
+		return true;
+	}
+
+	return false;
+}
+
 IMPLEMENT_CORE_INTRINSIC_CLASS(UAssetObjectProperty, UObjectPropertyBase,
 	{
 	}
diff --git a/Engine/Source/Runtime/CoreUObject/Private/UObject/PropertyBool.cpp b/Engine/Source/Runtime/CoreUObject/Private/UObject/PropertyBool.cpp
--- a/Engine/Source/Runtime/CoreUObject/Private/UObject/PropertyBool.cpp
+++ b/Engine/Source/Runtime/CoreUObject/Private/UObject/PropertyBool.cpp
@@ -2,6 +2,7 @@
 
 #include "CoreUObjectPrivate.h"
 #include "PropertyHelper.h"
+#include "PropertyTag.h"
 
 /*-----------------------------------------------------------------------------
 	UBoolProperty.
@@ -184,6 +185,81 @@ FString UBoolProperty::GetCPPMacroType( FString& ExtendedTypeText ) const
 	}
 	return TEXT("UBOOL32");
 }
+
+template<typename T>
+void LoadFromType(UBoolProperty* Property, const FPropertyTag& Tag, FArchive& Ar, uint8* Data)
+{
+	T IntValue;
+	Ar << IntValue;
+
+	if (IntValue != 0)
+	{
+#if !(UE_BUILD_SHIPPING || UE_BUILD_TEST)
+		if (IntValue != 1)
+		{
+			UE_LOG(LogClass, Log, TEXT("Loading %s property (%s) that is now a bool - value '%d', expecting 0 or 1. Value set to true."), *Tag.Type.ToString(), *Property->GetPathName(), IntValue);
+		}
+#endif
+		Property->SetPropertyValue_InContainer(Data, true, Tag.ArrayIndex);
+	}
+	else
+	{
+		Property->SetPropertyValue_InContainer(Data, false, Tag.ArrayIndex);
+	}
+}
+
+bool UBoolProperty::ConvertFromType(const FPropertyTag& Tag, FArchive& Ar, uint8* Data, UStruct* DefaultsStruct, bool& bOutAdvanceProperty)
+{
+	bOutAdvanceProperty = true;
+
+	if (Tag.Type == NAME_IntProperty)
+	{
+		LoadFromType<int32>(this, Tag, Ar, Data);
+	}
+	else if (Tag.Type == NAME_Int8Property)
+	{
+		LoadFromType<int8>(this, Tag, Ar, Data);
+	}
+	else if (Tag.Type == NAME_Int16Property)
+	{
+		LoadFromType<int16>(this, Tag, Ar, Data);
+	}
+	else if (Tag.Type == NAME_Int64Property)
+	{
+		LoadFromType<int64>(this, Tag, Ar, Data);
+	}
+	else if (Tag.Type == NAME_ByteProperty)
+	{
+		// if the byte property was an enum we won't allow a conversion to bool
+		if (Tag.EnumName == NAME_None)
+		{
+			LoadFromType<uint8>(this, Tag, Ar, Data);
+		}
+		else
+		{
+			bOutAdvanceProperty = false;
+		}
+	}
+	else if (Tag.Type == NAME_UInt16Property)
+	{
+		LoadFromType<uint16>(this, Tag, Ar, Data);
+	}
+	else if (Tag.Type == NAME_UInt32Property)
+	{
+		LoadFromType<uint32>(this, Tag, Ar, Data);
+	}
+	else if (Tag.Type == NAME_UInt64Property)
+	{
+		LoadFromType<uint64>(this, Tag, Ar, Data);
+	}
+	else
+	{
+		bOutAdvanceProperty = false;
+	}
+			
+	return bOutAdvanceProperty;
+}
+
 void UBoolProperty::ExportTextItem( FString& ValueStr, const void* PropertyValue, const void* DefaultValue, UObject* Parent, int32 PortFlags, UObject* ExportRootScope ) const
 {
 	check(FieldSize != 0);
diff --git a/Engine/Source/Runtime/CoreUObject/Private/UObject/PropertyByte.cpp b/Engine/Source/Runtime/CoreUObject/Private/UObject/PropertyByte.cpp
--- a/Engine/Source/Runtime/CoreUObject/Private/UObject/PropertyByte.cpp
+++ b/Engine/Source/Runtime/CoreUObject/Private/UObject/PropertyByte.cpp
@@ -1,6 +1,7 @@
 // Copyright 1998-2016 Epic Games, Inc. All Rights Reserved.
 
 #include "CoreUObjectPrivate.h"
+#include "PropertyTag.h"
 
 /*-----------------------------------------------------------------------------
 	UByteProperty.
@@ -71,7 +72,7 @@ void UByteProperty::Serialize( FArchive& Ar )
 {
 	Super::Serialize( Ar );
 	Ar << Enum;
-	if (Enum != NULL)
+	if (Enum != nullptr)
 	{
 		Ar.Preload(Enum);
 	}
@@ -125,6 +126,142 @@ FString UByteProperty::GetCPPType( FString* ExtendedTypeText/*=NULL*/, uint32 CP
 	return Super::GetCPPType(ExtendedTypeText, CPPExportFlags);
 }
 
+template <typename OldIntType>
+struct TConvertIntToEnumProperty
+{
+	static void Convert(FArchive& Ar, UByteProperty* Property, UEnum* Enum, void* Obj, const FPropertyTag& Tag)
+	{
+		OldIntType OldValue;
+		Ar << OldValue;
+
+		uint8 NewValue = OldValue;
+		if (OldValue > (OldIntType)TNumericLimits<uint8>::Max() || !Enum->IsValidEnumValue(NewValue))
+		{
+			UE_LOG(
+				LogClass,
+				Warning,
+				TEXT("Failed to find valid enum value '%d' for enum type '%s' when converting property '%s' during property loading - setting to '%s'"),
+				OldValue,
+				*Enum->GetName(),
+				*Property->GetName(),
+				*Enum->GetNameByValue(Enum->GetMaxEnumValue()).ToString()
+				);
+
+			NewValue = Enum->GetMaxEnumValue();
+		}
+
+		Property->SetPropertyValue_InContainer(Obj, NewValue, Tag.ArrayIndex);
+	}
+};
+
+bool UByteProperty::ConvertFromType(const FPropertyTag& Tag, FArchive& Ar, uint8* Data, UStruct* DefaultsStruct, bool& bOutAdvanceProperty)
+{
+	bOutAdvanceProperty = true;
+
+	if (Tag.Type == NAME_ByteProperty  && ((Tag.EnumName == NAME_None) != (Enum == nullptr)))
+	{
+		// a byte property gained or lost an enum
+		// attempt to convert it
+		uint8 PreviousValue;
+		if (Tag.EnumName == NAME_None)
+		{
+			// simply pretend the property still doesn't have an enum and serialize the single byte
+			Ar << PreviousValue;
+		}
+		else
+		{
+			// attempt to find the old enum and get the byte value from the serialized enum name
+			PreviousValue = ReadEnumAsUint8(Ar, DefaultsStruct, Tag);
+		}
+
+		// now copy the value into the object's address space
+		SetPropertyValue_InContainer(Data, PreviousValue, Tag.ArrayIndex);
+	}
+	else if (Tag.Type == NAME_Int8Property)
+	{
+		if (Enum)
+		{
+			TConvertIntToEnumProperty<int8>::Convert(Ar, this, Enum, Data, Tag);
+		}
+		else
+		{
+			ConvertFromInt<int8>(Ar, Data, Tag);
+		}
+	}
+	else if (Tag.Type == NAME_Int16Property)
+	{
+		if (Enum)
+		{
+			TConvertIntToEnumProperty<int16>::Convert(Ar, this, Enum, Data, Tag);
+		}
+		else
+		{
+			ConvertFromInt<int16>(Ar, Data, Tag);
+		}
+	}
+	else if (Tag.Type == NAME_IntProperty)
+	{
+		if (Enum)
+		{
+			TConvertIntToEnumProperty<int32>::Convert(Ar, this, Enum, Data, Tag);
+		}
+		else
+		{
+			ConvertFromInt<int32>(Ar, Data, Tag);
+		}
+	}
+	else if (Tag.Type == NAME_Int64Property)
+	{
+		if (Enum)
+		{
+			TConvertIntToEnumProperty<int64>::Convert(Ar, this, Enum, Data, Tag);
+		}
+		else
+		{
+			ConvertFromInt<int64>(Ar, Data, Tag);
+		}
+	}
+	else if (Tag.Type == NAME_UInt16Property)
+	{
+		if (Enum)
+		{
+			TConvertIntToEnumProperty<uint16>::Convert(Ar, this, Enum, Data, Tag);
+		}
+		else
+		{
+			ConvertFromInt<uint16>(Ar, Data, Tag);
+		}
+	}
+	else if (Tag.Type == NAME_UInt32Property)
+	{
+		if (Enum)
+		{
+			TConvertIntToEnumProperty<uint32>::Convert(Ar, this, Enum, Data, Tag);
+		}
+		else
+		{
+			ConvertFromInt<uint32>(Ar, Data, Tag);
+		}
+	}
+	else if (Tag.Type == NAME_UInt64Property)
+	{
+		if (Enum)
+		{
+			TConvertIntToEnumProperty<uint64>::Convert(Ar, this, Enum, Data, Tag);
+		}
+		else
+		{
+			ConvertFromInt<uint64>(Ar, Data, Tag);
+		}
+	}
+	else
+	{
+		bOutAdvanceProperty = false;
+	}
+
+	return bOutAdvanceProperty;
+}
+
 void UByteProperty::ExportTextItem( FString& ValueStr, const void* PropertyValue, const void* DefaultValue, UObject* Parent, int32 PortFlags, UObject* ExportRootScope ) const
 {
 	if (0 != (PortFlags & PPF_ExportCpp))
diff --git a/Engine/Source/Runtime/CoreUObject/Private/UObject/PropertyName.cpp b/Engine/Source/Runtime/CoreUObject/Private/UObject/PropertyName.cpp
--- a/Engine/Source/Runtime/CoreUObject/Private/UObject/PropertyName.cpp
+++ b/Engine/Source/Runtime/CoreUObject/Private/UObject/PropertyName.cpp
@@ -1,6 +1,7 @@
 // Copyright 1998-2016 Epic Games, Inc. All Rights Reserved.
 
 #include "CoreUObjectPrivate.h"
+#include "PropertyTag.h"
 
 /*-----------------------------------------------------------------------------
 	UNameProperty.
@@ -45,6 +46,32 @@ const TCHAR* UNameProperty::ImportText_Internal( const TCHAR* Buffer, void* Data
 	return Buffer;
 }
 
+bool UNameProperty::ConvertFromType(const FPropertyTag& Tag, FArchive& Ar, uint8* Data, UStruct* DefaultsStruct, bool& bOutAdvanceProperty)
+{
+	if (Tag.Type == NAME_StrProperty)
+	{
+		FString str;
+		Ar << str;
+		SetPropertyValue_InContainer(Data, FName(*str), Tag.ArrayIndex);
+		bOutAdvanceProperty = true;
+	}
+	// Convert serialized text to name.
+	else if (Tag.Type == NAME_TextProperty)
+	{ 
+		FText Text;  
+		Ar << Text;
+		const FName Name = FName(*Text.ToString());
+		SetPropertyValue_InContainer(Data, Name, Tag.ArrayIndex);
+		bOutAdvanceProperty = true; 
+	}
+	else
+	{
+		bOutAdvanceProperty = false;
+	}
+
+	return bOutAdvanceProperty;
+}
+
 IMPLEMENT_CORE_INTRINSIC_CLASS(UNameProperty, UProperty,
 	{
 	}
diff --git a/Engine/Source/Runtime/CoreUObject/Private/UObject/PropertyObject.cpp b/Engine/Source/Runtime/CoreUObject/Private/UObject/PropertyObject.cpp
--- a/Engine/Source/Runtime/CoreUObject/Private/UObject/PropertyObject.cpp
+++ b/Engine/Source/Runtime/CoreUObject/Private/UObject/PropertyObject.cpp
@@ -4,6 +4,7 @@
 #include "LinkerPlaceholderExportObject.h"
 #include "LinkerPlaceholderClass.h"
 #include "BlueprintSupport.h" // for IsDeferredDependencyPlaceholder()
+#include "PropertyTag.h"
 
 /*-----------------------------------------------------------------------------
 	UObjectProperty.
@@ -26,6 +27,24 @@ FString UObjectProperty::GetCPPMacroType( FString& ExtendedTypeText ) const
 	return TEXT("OBJECT");
 }
 
+bool UObjectProperty::ConvertFromType(const FPropertyTag& Tag, FArchive& Ar, uint8* Data, UStruct* DefaultsStruct, bool& bOutAdvanceProperty)
+{
+	if (Tag.Type == NAME_AssetObjectProperty || Tag.Type == NAME_AssetSubclassOfProperty)
+	{
+		// This property used to be a TAssetPtr<Foo> but is now a raw UObjectProperty Foo*, we can convert without loss of data
+		FAssetPtr PreviousValue;
+		Ar << PreviousValue;
+
+		// now copy the value into the object's address space
+		UObject* PreviousValueObj = PreviousValue.Get();
+		SetPropertyValue_InContainer(Data, PreviousValueObj, Tag.ArrayIndex);
+
+		return true;
+	}
+
+	return false;
+}
+
 void UObjectProperty::SerializeItem( FArchive& Ar, void* Value, void const* Defaults ) const
 {
 	UObject* ObjectValue = GetObjectPropertyValue(Value);
diff --git a/Engine/Source/Runtime/CoreUObject/Private/UObject/PropertySet.cpp b/Engine/Source/Runtime/CoreUObject/Private/UObject/PropertySet.cpp
new file mode 100644
--- /dev/null
+++ b/Engine/Source/Runtime/CoreUObject/Private/UObject/PropertySet.cpp
@@ -0,0 +1,720 @@
+// Copyright 1998-2016 Epic Games, Inc. All Rights Reserved.
+
+#include "CoreUObjectPrivate.h"
+#include "PropertyHelper.h"
+#include "ScopeExit.h"
+
+namespace UE4SetProperty_Private
+{
+	/**
+	 * Checks if any of the elements in the set compare equal to the one passed.
+	 *
+	 * @param  SetHelper  The set to search through.
+	 * @param  Index      The index in the set to start searching from.
+	 * @param  Num        The number of elements to compare.
+	 */
+	bool AnyEqual(const FScriptSetHelper& SetHelper, int32 Index, int32 Num, const uint8* ElementToCompare, uint32 PortFlags)
+	{
+		UProperty* ElementProp = SetHelper.GetElementProperty();
+
+		for (; Num; --Num)
+		{
+			while (!SetHelper.IsValidIndex(Index))
+			{
+				++Index;
+			}
+
+			if (ElementProp->Identical(SetHelper.GetElementPtr(Index), ElementToCompare, PortFlags))
+			{
+				return true;
+			}
+
+			++Index;
+		}
+
+		return false;
+	}
+
+	bool RangesContainSameAmountsOfVal(const FScriptSetHelper& SetHelperA, int32 IndexA, const FScriptSetHelper& SetHelperB, int32 IndexB, int32 Num, const uint8* ElementToCompare, uint32 PortFlags)
+	{
+		UProperty* ElementProp = SetHelperA.GetElementProperty();
+
+		// Ensure that both sets are the same type
+		check(ElementProp == SetHelperB.GetElementProperty());
+
+		int32 CountA = 0;
+		int32 CountB = 0;
+		for (;;)
+		{
+			if (Num == 0)
+			{
+				return CountA == CountB;
+			}
+
+			while (!SetHelperA.IsValidIndex(IndexA))
+			{
+				++IndexA;
+			}
+
+			while (!SetHelperB.IsValidIndex(IndexB))
+			{
+				++IndexB;
+			}
+
+			const uint8* ElementA = SetHelperA.GetElementPtr(IndexA);
+			const uint8* ElementB = SetHelperB.GetElementPtr(IndexB);
+			if (ElementProp->Identical(ElementA, ElementToCompare, PortFlags))
+			{
+				++CountA;
+			}
+
+			if (ElementProp->Identical(ElementB, ElementToCompare, PortFlags))
+			{
+				++CountB;
+			}
+
+			++IndexA;
+			++IndexB;
+			--Num;
+		}
+	}
+
+	bool IsPermutation(const FScriptSetHelper& SetHelperA, const FScriptSetHelper& SetHelperB, uint32 PortFlags)
+	{
+		UProperty* ElementProp = SetHelperA.GetElementProperty();
+
+		// Ensure that both maps are the same type
+		check(ElementProp == SetHelperB.GetElementProperty());
+
+		int32 Num = SetHelperA.Num();
+		if (Num != SetHelperB.Num())
+		{
+			return false;
+		}
+
+		// Skip over common initial sequence
+		int32 IndexA = 0;
+		int32 IndexB = 0;
+		for (;;)
+		{
+			if (Num == 0)
+			{
+				return true;
+			}
+
+			while (!SetHelperA.IsValidIndex(IndexA))
+			{
+				++IndexA;
+			}
+
+			while (!SetHelperB.IsValidIndex(IndexB))
+			{
+				++IndexB;
+			}
+
+			const uint8* ElementA = SetHelperA.GetElementPtr(IndexA);
+			const uint8* ElementB = SetHelperB.GetElementPtr(IndexB);
+			if (!ElementProp->Identical(ElementA, ElementB, PortFlags))
+			{
+				break;
+			}
+
+			++IndexA;
+			++IndexB;
+			--Num;
+		}
+
+		int32 FirstIndexA = IndexA;
+		int32 FirstIndexB = IndexB;
+		int32 FirstNum    = Num;
+		for (;;)
+		{
+			const uint8* ElementA = SetHelperA.GetElementPtr(IndexA);
+			if (!AnyEqual(SetHelperA, FirstIndexA, FirstNum - Num, ElementA, PortFlags) && !RangesContainSameAmountsOfVal(SetHelperA, IndexA, SetHelperB, IndexB, Num, ElementA, PortFlags))
+			{
+				return false;
+			}
+
+			--Num;
+			if (Num == 0)
+			{
+				return true;
+			}
+
+			while (!SetHelperA.IsValidIndex(IndexA))
+			{
+				++IndexA;
+			}
+
+			while (!SetHelperB.IsValidIndex(IndexB))
+			{
+				++IndexB;
+			}
+		}
+	}
+}
+
+USetProperty::USetProperty(const FObjectInitializer& ObjectInitializer, ECppProperty, int32 InOffset, uint64 InFlags)
+: USetProperty_Super(ObjectInitializer, EC_CppProperty, InOffset, InFlags)
+{
+	// This is expected to be set post-construction by AddCppProperty
+	ElementProp = nullptr;
+}
+
+void USetProperty::LinkInternal(FArchive& Ar)
+{
+	check(ElementProp);
+
+	if (FLinkerLoad* MyLinker = GetLinker())
+	{
+		MyLinker->Preload(this);
+	}
+	Ar.Preload(ElementProp);
+	ElementProp->Link(Ar);
+
+	const int32 ElementPropSize = ElementProp->GetSize();
+	const int32 ElementPropAlignment = ElementProp->GetMinAlignment();
+
+	SetLayout = FScriptSet::GetScriptLayout(ElementPropSize, ElementPropAlignment);
+
+	Super::LinkInternal(Ar);
+}
+
+bool USetProperty::Identical(const void* A, const void* B, uint32 PortFlags) const
+{
+	checkSlow(ElementProp);
+
+	FScriptSetHelper SetHelperA(this, A);
+
+	const int32 ANum = SetHelperA.Num();
+
+	if (!B)
+	{
+		return ANum == 0;
+	}
+
+	FScriptSetHelper SetHelperB(this, B);
+	if (ANum != SetHelperB.Num())
+	{
+		return false;
+	}
+
+	return UE4SetProperty_Private::IsPermutation(SetHelperA, SetHelperB, PortFlags);
+}
+
+void USetProperty::SerializeItem(FArchive& Ar, void* Value, const void* Defaults) const
+{
+	checkSlow(ElementProp);
+
+	// Ensure that the element property has been loaded before calling SerializeItem() on it
+	Ar.Preload(ElementProp);
+
+	FScriptSetHelper SetHelper(this, Value);
+
+	if (Ar.IsLoading())
+	{
+		if (Defaults)
+		{
+			CopyValuesInternal(Value, Defaults, 1);
+		}
+		else
+		{
+			SetHelper.EmptyElements();
+		}
+
+		uint8* TempElementStorage = nullptr;
+		ON_SCOPE_EXIT
+		{
+			if (TempElementStorage)
+			{
+				ElementProp->DestroyValue(TempElementStorage);
+				FMemory::Free(TempElementStorage);
+			}
+		};
+
+		// Delete any explicitly-removed elements
+		int32 NumElementsToRemove = 0;
+		Ar << NumElementsToRemove;
+		if (NumElementsToRemove)
+		{
+			TempElementStorage = (uint8*)FMemory::Malloc(SetLayout.Size);
+			ElementProp->InitializeValue(TempElementStorage);
+
+			for (; NumElementsToRemove; --NumElementsToRemove)
+			{
+				// Read key into temporary storage
+				ElementProp->SerializeItem(Ar, TempElementStorage);
+
+				// If the key is in the map, remove it
+				const int32 Found = SetHelper.FindElementIndex(TempElementStorage);
+				if (Found != INDEX_NONE)
+				{
+					SetHelper.RemoveAt_NeedsRehash(Found);
+				}
+			}
+		}
+
+		int32 Num = 0;
+		Ar << Num;
+
+		// Allocate temporary key space if we haven't allocated it already above
+		if (Num != 0 && !TempElementStorage)
+		{
+			TempElementStorage = (uint8*)FMemory::Malloc(SetLayout.Size);
+			ElementProp->InitializeValue(TempElementStorage);
+		}
+
+		// Read remaining items into container
+		for (; Num; --Num)
+		{
+			// Read key into temporary storage
+			ElementProp->SerializeItem(Ar, TempElementStorage);
+
+			// Add a new entry if the element doesn't currently exist in the set
+			if (SetHelper.FindElementIndex(TempElementStorage) == INDEX_NONE)
+			{
+				const int32 NewElementIndex = SetHelper.AddDefaultValue_Invalid_NeedsRehash();
+				uint8* NewElementPtr = SetHelper.GetElementPtrWithoutCheck(NewElementIndex);
+
+				// Copy over deserialized key from temporary storage
+				ElementProp->CopyCompleteValue_InContainer(NewElementPtr, TempElementStorage);
+			}
+		}
+
+		SetHelper.Rehash();
+	}
+	else
+	{
+		FScriptSetHelper DefaultsHelper(this, Defaults);
+
+		// Container for temporarily tracking some indices
+		TSet<int32> Indices;
+
+		// Determine how many keys are missing from the object
+		if (Defaults)
+		{
+			for (int32 Index = 0, Count = DefaultsHelper.Num(); Count; ++Index)
+			{
+				uint8* DefaultElementPtr = DefaultsHelper.GetElementPtrWithoutCheck(Index);
+
+				if (DefaultsHelper.IsValidIndex(Index))
+				{
+					if (SetHelper.FindElementIndex(DefaultElementPtr) == INDEX_NONE)
+					{
+						Indices.Add(Index);
+					}
+
+					--Count;
+				}
+			}
+		}
+
+		// Write out the removed elements
+		int32 RemovedElementsNum = Indices.Num();
+		Ar << RemovedElementsNum;
+		for (int32 Index : Indices)
+		{
+			ElementProp->SerializeItem(Ar, DefaultsHelper.GetElementPtr(Index));
+		}
+
+		// Write out added elements
+		if (Defaults)
+		{
+			Indices.Reset();
+			for (int32 Index = 0, Count = SetHelper.Num(); Count; ++Index)
+			{
+				if (SetHelper.IsValidIndex(Index))
+				{
+					uint8* ValueElement   = SetHelper.GetElementPtrWithoutCheck(Index);
+					uint8* DefaultElement = DefaultsHelper.FindElementPtr(ValueElement);
+
+					if (!DefaultElement)
+					{
+						Indices.Add(Index);
+					}
+
+					--Count;
+				}
+			}
+
+			// Write out differences from defaults
+			int32 Num = Indices.Num();
+			Ar << Num;
+			for (int32 Index : Indices)
+			{
+				uint8* ElementPtr = SetHelper.GetElementPtrWithoutCheck(Index);
+
+				ElementProp->SerializeItem(Ar, ElementPtr);
+			}
+		}
+		else
+		{
+			int32 Num = SetHelper.Num();
+			Ar << Num;
+			for (int32 Index = 0; Num; ++Index)
+			{
+				if (SetHelper.IsValidIndex(Index))
+				{
+					uint8* ElementPtr = SetHelper.GetElementPtrWithoutCheck(Index);
+
+					ElementProp->SerializeItem(Ar, ElementPtr);
+
+					--Num;
+				}
+			}
+		}
+	}
+}
+
+bool USetProperty::NetSerializeItem( FArchive& Ar, UPackageMap* Map, void* Data, TArray<uint8>* MetaData ) const
+{
+	UE_LOG( LogProperty, Fatal, TEXT( "Deprecated code path" ) );
+	return 1;
+}
+
+void USetProperty::Serialize( FArchive& Ar )
+{
+	Super::Serialize( Ar );
+	Ar << ElementProp;
+}
+
+void USetProperty::AddReferencedObjects(UObject* InThis, FReferenceCollector& Collector)
+{
+	USetProperty* This = CastChecked<USetProperty>(InThis);
+
+	Collector.AddReferencedObject(This->ElementProp, This);
+
+	Super::AddReferencedObjects(This, Collector);
+}
+
+FString USetProperty::GetCPPType(FString* ExtendedTypeText, uint32 CPPExportFlags) const
+{
+	checkSlow(ElementProp);
+
+	if (ExtendedTypeText)
+	{
+		FString ElementExtendedTypeText;
+		const FString ElementTypeText = ElementProp->GetCPPType(&ElementExtendedTypeText, CPPExportFlags & ~CPPF_ArgumentOrReturnValue); // we won't consider set elements to be "arguments or return values"
+
+		*ExtendedTypeText = FString::Printf(TEXT("<%s%s>"), *ElementTypeText, *ElementExtendedTypeText);
+	}
+
+	return TEXT("TSet");
+}
+
+FString USetProperty::GetCPPMacroType( FString& ExtendedTypeText ) const
+{
+	checkSlow(ElementProp);
+	ExtendedTypeText = FString::Printf(TEXT("%s"), *ElementProp->GetCPPType());
+	return TEXT("TSET");
+}
+
+void USetProperty::ExportTextItem(FString& ValueStr, const void* PropertyValue, const void* DefaultValue, UObject* Parent, int32 PortFlags, UObject* ExportRootScope) const
+{
+	if (0 != (PortFlags & PPF_ExportCpp))
+	{
+		ValueStr += TEXT("{}");
+		return;
+	}
+
+	checkSlow(ElementProp);
+
+	FScriptSetHelper SetHelper(this, PropertyValue);
+
+	if (SetHelper.Num() == 0)
+	{
+		ValueStr += TEXT("()");
+		return;
+	}
+
+	uint8* StructDefaults = nullptr;
+	if (UStructProperty* StructElementProp = Cast<UStructProperty>(ElementProp))
+	{
+		checkSlow(StructElementProp->Struct);
+
+		StructDefaults = (uint8*)FMemory::Malloc(SetLayout.Size);
+		ElementProp->InitializeValue(StructDefaults + SetLayout.ElementOffset);
+	}
+
+	ON_SCOPE_EXIT
+	{
+		if (StructDefaults)
+		{
+			ElementProp->DestroyValue(StructDefaults + SetLayout.ElementOffset);
+			FMemory::Free(StructDefaults);
+		}
+	};
+
+	FScriptSetHelper DefaultSetHelper(this, DefaultValue);
+
+	uint8* PropData = SetHelper.GetElementPtrWithoutCheck(0);
+	if (PortFlags & PPF_BlueprintDebugView)
+	{
+		int32 Index  = 0;
+		bool  bFirst = true;
+		for (int32 Count = SetHelper.Num(); Count; PropData += SetLayout.Size, ++Index)
+		{
+			if (SetHelper.IsValidIndex(Index))
+			{
+				if (bFirst)
+				{
+					bFirst = false;
+				}
+				else
+				{
+					ValueStr += TCHAR('\n');
+				}
+
+				// Always use struct defaults if the element is a struct, for symmetry with the import of array inner struct defaults
+				uint8* PropDefault = StructDefaults ? StructDefaults : DefaultValue ? DefaultSetHelper.FindElementPtr(PropData) : nullptr;
+
+				ElementProp->ExportTextItem(ValueStr, PropData, PropDefault, Parent, PortFlags | PPF_Delimited, ExportRootScope);
+
+				--Count;
+			}
+		}
+	}
+	else
+	{
+		int32 Index  = 0;
+		bool  bFirst = true;
+		for (int32 Count = SetHelper.Num(); Count; PropData += SetLayout.Size, ++Index)
+		{
+			if (SetHelper.IsValidIndex(Index))
+			{
+				if (bFirst)
+				{
+					ValueStr += TCHAR('(');
+					bFirst = false;
+				}
+				else
+				{
+					ValueStr += TCHAR(',');
+				}
+
+				ElementProp->ExportTextItem(ValueStr, PropData, nullptr, Parent, PortFlags | PPF_Delimited, ExportRootScope);
+
+
+				--Count;
+			}
+		}
+
+		ValueStr += TEXT(")");
+	}
+}
+
+const TCHAR* USetProperty::ImportText_Internal(const TCHAR* Buffer, void* Data, int32 PortFlags, UObject* Parent, FOutputDevice* ErrorText) const
+{
+	checkSlow(ElementProp);
+
+	FScriptSetHelper SetHelper(this, Data);
+	SetHelper.EmptyElements();
+
+	// If we export an empty array we export an empty string, so ensure that if we're passed an empty string
+	// we interpret it as an empty array.
+	if (*Buffer++ != TCHAR('('))
+	{
+		return nullptr;
+	}
+
+	SkipWhitespace(Buffer);
+	if (*Buffer == TCHAR(')'))
+	{
+		return Buffer + 1;
+	}
+
+	uint8* TempElementStorage = (uint8*)FMemory::Malloc(SetLayout.Size);
+	ElementProp->InitializeValue(TempElementStorage);
+
+	ON_SCOPE_EXIT
+	{
+		if (TempElementStorage)
+		{
+			ElementProp->DestroyValue(TempElementStorage);
+			FMemory::Free(TempElementStorage);
+		}
+	};
+
+	int32 Index = 0;
+	for (;;)
+	{
+		// Read key into temporary storage
+		Buffer = ElementProp->ImportText(Buffer, TempElementStorage, PortFlags | PPF_Delimited, Parent, ErrorText);
+		if (!Buffer)
+		{
+			return nullptr;
+		}
+
+		// If the key isn't in the map yet, add it
+		if (SetHelper.FindElementIndex(TempElementStorage) == INDEX_NONE)
+		{
+			const int32 NewElementIndex = SetHelper.AddDefaultValue_Invalid_NeedsRehash();
+			uint8* NewElementPtr = SetHelper.GetElementPtrWithoutCheck(NewElementIndex);
+
+			// Copy over imported key from temporary storage
+			ElementProp->CopyCompleteValue_InContainer(NewElementPtr, TempElementStorage);
+		}
+
+		// Parse the element
+
+		SkipWhitespace(Buffer);
+
+		switch (*Buffer++)
+		{
+		case TCHAR(')'):
+			SetHelper.Rehash();
+			return Buffer;
+
+		case TCHAR(','):
+			break;
+
+		default:
+			return nullptr;
+		}
+
+		++Index;
+	}
+}
+
+void USetProperty::AddCppProperty( UProperty* Property )
+{
+	check(!ElementProp);
+	check(Property);
+
+	ElementProp = Property;
+}
+
+void USetProperty::CopyValuesInternal(void* Dest, void const* Src, int32 Count) const
+{
+	check(Count == 1);
+
+	FScriptSetHelper SrcSetHelper (this, Src);
+	FScriptSetHelper DestSetHelper(this, Dest);
+
+	int32 Num = SrcSetHelper.Num();
+	DestSetHelper.EmptyElements(Num);
+
+	if (Num == 0)
+	{
+		return;
+	}
+
+	for (int32 SrcIndex = 0; Num; ++SrcIndex)
+	{
+		if (SrcSetHelper.IsValidIndex(SrcIndex))
+		{
+			const int32 DestIndex = DestSetHelper.AddDefaultValue_Invalid_NeedsRehash();
+
+			uint8* SrcData  = SrcSetHelper.GetElementPtrWithoutCheck(SrcIndex);
+			uint8* DestData = DestSetHelper.GetElementPtrWithoutCheck(DestIndex);
+
+			ElementProp->CopyCompleteValue_InContainer(DestData, SrcData);
+
+			--Num;
+		}
+	}
+
+	DestSetHelper.Rehash();
+}
+
+void USetProperty::ClearValueInternal(void* Data) const
+{
+	FScriptSetHelper SetHelper(this, Data);
+	SetHelper.EmptyElements();
+}
+
+void USetProperty::DestroyValueInternal(void* Data) const
+{
+	FScriptSetHelper SetHelper(this, Data);
+	SetHelper.EmptyElements();
+
+	//@todo UE4 potential double destroy later from this...would be ok for a script set, but still
+	((FScriptSet*)Data)->~FScriptSet();
+}
+
+bool USetProperty::PassCPPArgsByRef() const
+{
+	return true;
+}
+
+/**
+ * Creates new copies of components
+ * 
+ * @param	Data				pointer to the address of the instanced object referenced by this UComponentProperty
+ * @param	DefaultData			pointer to the address of the default value of the instanced object referenced by this UComponentProperty
+ * @param	Owner				the object that contains this property's data
+ * @param	InstanceGraph		contains the mappings of instanced objects and components to their templates
+ */
+void USetProperty::InstanceSubobjects(void* Data, void const* DefaultData, UObject* Owner, FObjectInstancingGraph* InstanceGraph)
+{
+	if (!Data)
+	{
+		return;
+	}
+
+	const bool bInstancedElement = ElementProp->ContainsInstancedObjectProperty();
+
+	if (!bInstancedElement)
+	{
+		return;
+	}
+
+	FScriptSetHelper SetHelper(this, Data);
+
+	if (DefaultData)
+	{
+		FScriptSetHelper DefaultSetHelper(this, DefaultData);
+		const int32 DefaultNum = DefaultSetHelper.Num();
+
+		for (int32 Index = 0, Num = SetHelper.Num(); Num; ++Index)
+		{
+			if (SetHelper.IsValidIndex(Index))
+			{
+				uint8* ElementPtr        = SetHelper.GetElementPtr(Index);
+				uint8* DefaultElementPtr = DefaultSetHelper.FindElementPtr(ElementPtr, Index);
+
+				ElementProp->InstanceSubobjects(ElementPtr, DefaultElementPtr, Owner, InstanceGraph);
+
+				--Num;
+			}
+		}
+	}
+	else
+	{
+		for (int32 Index = 0, Num = SetHelper.Num(); Num; ++Index)
+		{
+			if (SetHelper.IsValidIndex(Index))
+			{
+				uint8* ElementPtr = SetHelper.GetElementPtr(Index);
+
+				ElementProp->InstanceSubobjects(ElementPtr, nullptr, Owner, InstanceGraph);
+
+				--Num;
+			}
+		}
+	}
+}
+
+bool USetProperty::SameType(const UProperty* Other) const
+{
+	USetProperty* SetProp = (USetProperty*)Other;
+	return Super::SameType(Other) && ElementProp && ElementProp->SameType(SetProp->ElementProp);
+}
+
+IMPLEMENT_CORE_INTRINSIC_CLASS(USetProperty, UProperty,
+	{
+		Class->EmitObjectReference(STRUCT_OFFSET(USetProperty, ElementProp), TEXT("ElementProp"));
+
+		// Ensure that TArray and FScriptMap are interchangeable, as FScriptMap will be used to access a native array property
+		// from script that is declared as a TArray in C++.
+		static_assert(sizeof(FScriptSet) == sizeof(TMap<uint32, uint8>), "FScriptSet and TSet<uint32, uint8> must be interchangable.");
+	}
+);
+
+void FScriptSetHelper::Rehash()
+{
+	// Moved out-of-line to maybe fix a weird link error
+	Set->Rehash(SetLayout, [=](const void* Src) {
+		return ElementProp->GetValueTypeHash(Src);
+	});
+}
diff --git a/Engine/Source/Runtime/CoreUObject/Private/UObject/PropertyStr.cpp b/Engine/Source/Runtime/CoreUObject/Private/UObject/PropertyStr.cpp
--- a/Engine/Source/Runtime/CoreUObject/Private/UObject/PropertyStr.cpp
+++ b/Engine/Source/Runtime/CoreUObject/Private/UObject/PropertyStr.cpp
@@ -1,11 +1,31 @@
 // Copyright 1998-2016 Epic Games, Inc. All Rights Reserved.
 
 #include "CoreUObjectPrivate.h"
+#include "PropertyTag.h"
 
 /*-----------------------------------------------------------------------------
 	UStrProperty.
 -----------------------------------------------------------------------------*/
 
+bool UStrProperty::ConvertFromType(const FPropertyTag& Tag, FArchive& Ar, uint8* Data, UStruct* DefaultsStruct, bool& bOutAdvanceProperty)
+{
+	// Convert serialized text to string.
+	if (Tag.Type==NAME_TextProperty) 
+	{ 
+		FText Text;  
+		Ar << Text;
+		const FString String = FTextInspector::GetSourceString(Text) ? *FTextInspector::GetSourceString(Text) : TEXT("");
+		SetPropertyValue_InContainer(Data, String, Tag.ArrayIndex);
+		bOutAdvanceProperty = true;
+	}
+	else
+	{
+		bOutAdvanceProperty = false;
+	}
+
+	return bOutAdvanceProperty;
+}
+
 // Necessary to fix Compiler Error C2026 and C1091
 FString UStrProperty::ExportCppHardcodedText(const FString& InSource, const FString& Indent)
 {
diff --git a/Engine/Source/Runtime/CoreUObject/Private/UObject/PropertyStruct.cpp b/Engine/Source/Runtime/CoreUObject/Private/UObject/PropertyStruct.cpp
--- a/Engine/Source/Runtime/CoreUObject/Private/UObject/PropertyStruct.cpp
+++ b/Engine/Source/Runtime/CoreUObject/Private/UObject/PropertyStruct.cpp
@@ -429,6 +429,48 @@ bool UStructProperty::SameType(const UProperty* Other) const
 	return Super::SameType(Other) && (Struct == ((UStructProperty*)Other)->Struct);
 }
 
+bool UStructProperty::ConvertFromType(const FPropertyTag& Tag, FArchive& Ar, uint8* Data, UStruct* DefaultsStruct, bool& bOutAdvanceProperty)
+{
+	bOutAdvanceProperty = false;
+
+	auto CanSerializeFromStructWithDifferentName = [](const FArchive& InAr, const FPropertyTag& PropertyTag, const UStructProperty* StructProperty)
+	{
+		if (InAr.UE4Ver() < VER_UE4_STRUCT_GUID_IN_PROPERTY_TAG)
+		{
+			// Old Implementation
+			return StructProperty && !StructProperty->UseBinaryOrNativeSerialization(InAr);
+		}
+		return PropertyTag.StructGuid.IsValid() && StructProperty && StructProperty->Struct && (PropertyTag.StructGuid == StructProperty->Struct->GetCustomGuid());
+	};
+
+	if (Struct)
+	{
+		if ((Struct->StructFlags & STRUCT_SerializeFromMismatchedTag) && (Tag.Type != NAME_StructProperty || (Tag.StructName != Struct->GetFName())) && Struct->StructFlags & STRUCT_SerializeFromMismatchedTag)
+		{
+			UScriptStruct::ICppStructOps* CppStructOps = Struct->GetCppStructOps();
+			check(CppStructOps && CppStructOps->HasSerializeFromMismatchedTag()); // else should not have STRUCT_SerializeFromMismatchedTag
+			void* DestAddress = ContainerPtrToValuePtr<void>(Data, Tag.ArrayIndex);
+			if (CppStructOps->SerializeFromMismatchedTag(Tag, Ar, DestAddress))
+			{
+				bOutAdvanceProperty = true;
+			}
+			else
+			{
+				UE_LOG(LogClass, Warning, TEXT("SerializeFromMismatchedTag failed: Type mismatch in %s of %s - Previous (%s) Current(StructProperty) for package:  %s"), *Tag.Name.ToString(), *GetName(), *Tag.Type.ToString(), *Ar.GetArchiveName());
+			}
+			return true;
+		}
+		else if (Tag.Type == NAME_StructProperty && Tag.StructName != Struct->GetFName() && !CanSerializeFromStructWithDifferentName(Ar, Tag, this))
+		{
+			UE_LOG(LogClass, Warning, TEXT("Property %s of %s has a struct type mismatch (tag %s != prop %s) in package:  %s. If that struct got renamed, add an entry to ActiveStructRedirects."),
+				*Tag.Name.ToString(), *GetName(), *Tag.StructName.ToString(), *Struct->GetName(), *Ar.GetArchiveName());
+
+			return true;
+		}
+	}
+	return false;
+}
+
 IMPLEMENT_CORE_INTRINSIC_CLASS(UStructProperty, UProperty,
 	{
 		Class->EmitObjectReference(STRUCT_OFFSET(UStructProperty, Struct), TEXT("Struct"));
diff --git a/Engine/Source/Runtime/CoreUObject/Private/UObject/PropertyTag.cpp b/Engine/Source/Runtime/CoreUObject/Private/UObject/PropertyTag.cpp
new file mode 100644
--- /dev/null
+++ b/Engine/Source/Runtime/CoreUObject/Private/UObject/PropertyTag.cpp
@@ -0,0 +1,149 @@
+// Copyright 1998-2016 Epic Games, Inc. All Rights Reserved.
+
+#include "CoreUObjectPrivate.h"
+#include "PropertyTag.h"
+
+/*-----------------------------------------------------------------------------
+FPropertyTag
+-----------------------------------------------------------------------------*/
+
+// Constructors.
+FPropertyTag::FPropertyTag()
+	: Type      (NAME_None)
+	, BoolVal   (0)
+	, Name      (NAME_None)
+	, StructName(NAME_None)
+	, EnumName  (NAME_None)
+	, InnerType (NAME_None)
+	, Size      (0)
+	, ArrayIndex(INDEX_NONE)
+	, SizeOffset(INDEX_NONE)
+	, HasPropertyGuid(0)
+{}
+
+FPropertyTag::FPropertyTag( FArchive& InSaveAr, UProperty* Property, int32 InIndex, uint8* Value, uint8* Defaults )
+	: Type      (Property->GetID())
+	, BoolVal   (0)
+	, Name      (Property->GetFName())
+	, StructName(NAME_None)
+	, EnumName  (NAME_None)
+	, InnerType (NAME_None)
+	, Size      (0)
+	, ArrayIndex(InIndex)
+	, SizeOffset(INDEX_NONE)
+	, HasPropertyGuid(0)
+{
+	if (Property)
+	{
+		// Handle structs.
+		if (UStructProperty* StructProperty = Cast<UStructProperty>(Property))
+		{
+			StructName = StructProperty->Struct->GetFName();
+			StructGuid = StructProperty->Struct->GetCustomGuid();
+		}
+		else if (UByteProperty* ByteProp = Cast<UByteProperty>(Property))
+		{
+			if (ByteProp->Enum != NULL)
+			{
+				EnumName = ByteProp->Enum->GetFName();
+			}
+		}
+		else if (UArrayProperty* ArrayProp = Cast<UArrayProperty>(Property))
+		{
+			InnerType = ArrayProp->Inner->GetID();
+		}
+		else if (UBoolProperty* Bool = Cast<UBoolProperty>(Property))
+		{
+			BoolVal = Bool->GetPropertyValue(Value);
+		}
+	}
+}
+
+// Set optional property guid
+void FPropertyTag::SetPropertyGuid(const FGuid& InPropertyGuid)
+{
+	if (InPropertyGuid.IsValid())
+	{
+		PropertyGuid = InPropertyGuid;
+		HasPropertyGuid = true;
+	}
+}
+
+// Serializer.
+FArchive& operator<<( FArchive& Ar, FPropertyTag& Tag )
+{
+	// Name.
+	Ar << Tag.Name;
+	if ((Tag.Name == NAME_None) || !Tag.Name.IsValid())
+	{
+		return Ar;
+	}
+
+	Ar << Tag.Type;
+	if ( Ar.IsSaving() )
+	{
+		// remember the offset of the Size variable - UStruct::SerializeTaggedProperties will update it after the
+		// property has been serialized.
+		Tag.SizeOffset = Ar.Tell();
+	}
+	{
+		FArchive::FScopeSetDebugSerializationFlags S(Ar, DSF_IgnoreDiff);
+		Ar << Tag.Size << Tag.ArrayIndex;
+	}
+	// only need to serialize this for structs
+	if (Tag.Type == NAME_StructProperty)
+	{
+		Ar << Tag.StructName;
+		if (Ar.UE4Ver() >= VER_UE4_STRUCT_GUID_IN_PROPERTY_TAG)
+		{
+			Ar << Tag.StructGuid;
+		}
+	}
+	// only need to serialize this for bools
+	else if (Tag.Type == NAME_BoolProperty)
+	{
+		Ar << Tag.BoolVal;
+	}
+	// only need to serialize this for bytes
+	else if (Tag.Type == NAME_ByteProperty)
+	{
+		Ar << Tag.EnumName;
+	}
+	// only need to serialize this for arrays
+	else if (Tag.Type == NAME_ArrayProperty)
+	{
+		if (Ar.UE4Ver() >= VAR_UE4_ARRAY_PROPERTY_INNER_TAGS)
+		{
+			Ar << Tag.InnerType;
+		}
+	}
+	// Property tags to handle renamed blueprint properties effectively.
+	if (Ar.UE4Ver() >= VER_UE4_PROPERTY_GUID_IN_PROPERTY_TAG)
+	{
+		Ar << Tag.HasPropertyGuid;
+		if (Tag.HasPropertyGuid)
+		{
+			Ar << Tag.PropertyGuid;
+		}
+	}
+
+	return Ar;
+}
+
+// Property serializer.
+void FPropertyTag::SerializeTaggedProperty( FArchive& Ar, UProperty* Property, uint8* Value, uint8* Defaults )
+{
+	if (Property->GetClass() == UBoolProperty::StaticClass())
+	{
+		UBoolProperty* Bool = (UBoolProperty*)Property;
+		if (Ar.IsLoading())
+		{
+			Bool->SetPropertyValue(Value, BoolVal != 0);
+		}
+	}
+	else
+	{
+		FSerializedPropertyScope SerializedProperty(Ar, Property);
+		Property->SerializeItem( Ar, Value, Defaults );
+	}
+}
diff --git a/Engine/Source/Runtime/CoreUObject/Private/UObject/ReferenceChainSearch.cpp b/Engine/Source/Runtime/CoreUObject/Private/UObject/ReferenceChainSearch.cpp
--- a/Engine/Source/Runtime/CoreUObject/Private/UObject/ReferenceChainSearch.cpp
+++ b/Engine/Source/Runtime/CoreUObject/Private/UObject/ReferenceChainSearch.cpp
@@ -669,6 +669,20 @@ bool FReferenceChainSearch::ProcessObject( UObject* CurrentObject )
 				AddToReferenceList(ReferenceList, Ref);
 			}
 		}
+		else if( REFERENCE_INFO.Type == GCRT_AddTSetReferencedObjects )
+		{
+			void*         Set         = StackEntryData + REFERENCE_INFO.Offset;
+			USetProperty* SetProperty = (USetProperty*)TokenStream->ReadPointer( TokenStreamIndex );
+			TokenReturnCount = REFERENCE_INFO.ReturnCount;
+			FFindReferencerCollector ReferenceCollector(this, EReferenceType::SetProperty, (void*)SetProperty, CurrentObject);
+			FSimpleObjectReferenceCollectorArchive CollectorArchive(CurrentObject, ReferenceCollector);
+			SetProperty->SerializeItem(CollectorArchive, Set, nullptr);
+
+			for (const FReferenceChainLink& Ref : ReferenceCollector.References)
+			{
+				AddToReferenceList(ReferenceList, Ref);
+			}
+		}
 		else if (REFERENCE_INFO.Type == GCRT_EndOfPointer)
 		{
 			TokenReturnCount = REFERENCE_INFO.ReturnCount;
diff --git a/Engine/Source/Runtime/CoreUObject/Private/UObject/SavePackage.cpp b/Engine/Source/Runtime/CoreUObject/Private/UObject/SavePackage.cpp
--- a/Engine/Source/Runtime/CoreUObject/Private/UObject/SavePackage.cpp
+++ b/Engine/Source/Runtime/CoreUObject/Private/UObject/SavePackage.cpp
@@ -3519,6 +3519,7 @@ ESavePackageResult UPackage::Save(UPackage* InOuter, UObject* Base, EObjectFlags
 		// Untag all objects and names.
 		UnMarkAllObjects();
 
+		TArray<UObject*> CachedObjects;
 
 		// Size of serialized out package in bytes. This is before compression.
 		int32 PackageSize = INDEX_NONE;
@@ -3575,8 +3576,7 @@ ESavePackageResult UPackage::Save(UPackage* InOuter, UObject* Base, EObjectFlags
 
 				// Kick off any Precaching required for the target platform to save these objects
 				// only need to do this if we are cooking a different platform then the one which is currently running
-				// TODO: if save package is cancelled then call ClearCache on each object
-				TArray<UObject*> CachedObjects;
+				// TODO: if save package is canceled then call ClearCache on each object
 
 #if WITH_EDITOR
 				if ( bIsCooking )
@@ -5007,23 +5007,23 @@ ESavePackageResult UPackage::Save(UPackage* InOuter, UObject* Base, EObjectFlags
 
 
 				SlowTask.EnterProgressFrame();
+			}
 
-				// Route PostSaveRoot to allow e.g. the world to detach components for the persistent level that where
-				// attached in PreSaveRoot.
-				if( Base )
-				{
-					Base->PostSaveRoot( bCleanupIsRequired );
-				}
+			// Route PostSaveRoot to allow e.g. the world to detach components for the persistent level that were
+			// attached in PreSaveRoot.
+			if( Base )
+			{
+				Base->PostSaveRoot( bCleanupIsRequired );
+			}
 
-				SlowTask.EnterProgressFrame();
+			SlowTask.EnterProgressFrame();
 			
 #if WITH_EDITOR
-				for ( int CachedObjectIndex = 0; CachedObjectIndex < CachedObjects.Num(); ++CachedObjectIndex )
-				{
-					CachedObjects[CachedObjectIndex]->ClearCachedCookedPlatformData(TargetPlatform);
-				}
-#endif
+			for ( int CachedObjectIndex = 0; CachedObjectIndex < CachedObjects.Num(); ++CachedObjectIndex )
+			{
+				CachedObjects[CachedObjectIndex]->ClearCachedCookedPlatformData(TargetPlatform);
 			}
+#endif
 
 		}
 		if( Success == true )
diff --git a/Engine/Source/Runtime/CoreUObject/Private/UObject/UObjectBase.cpp b/Engine/Source/Runtime/CoreUObject/Private/UObject/UObjectBase.cpp
--- a/Engine/Source/Runtime/CoreUObject/Private/UObject/UObjectBase.cpp
+++ b/Engine/Source/Runtime/CoreUObject/Private/UObject/UObjectBase.cpp
@@ -105,15 +105,42 @@ void UObjectBase::CreateStatID() const
 	SCOPE_CYCLE_COUNTER(STAT_CreateStatID);
 
 	FString LongName;
+	LongName.Reserve(255);
+	TArray<UObjectBase const*, TInlineAllocator<24>> ClassChain;
+	
+	// Build class hierarchy
 	UObjectBase const* Target = this;
-	do 
+	do
 	{
-		LongName = Target->GetFName().GetPlainNameString() + (LongName.IsEmpty() ? TEXT("") : TEXT(".")) + LongName;
+		ClassChain.Add(Target);
 		Target = Target->GetOuter();
 	} while(Target);
+
+	// Start with class name
 	if (GetClass())
 	{
-		LongName = GetClass()->GetFName().GetPlainNameString() / LongName;
+		GetClass()->GetFName().GetDisplayNameEntry()->AppendNameToString(LongName);
+	}
+
+	// Now process from parent -> child so we can append strings more efficiently.
+	bool bFirstEntry = true;
+	for (int32 i = ClassChain.Num()-1; i >= 0; i--)
+	{
+		Target = ClassChain[i];
+		const FNameEntry* NameEntry = Target->GetFName().GetDisplayNameEntry();
+		if (bFirstEntry)
+		{
+			NameEntry->AppendNameToPathString(LongName);
+		}
+		else
+		{
+			if (!LongName.IsEmpty())
+			{
+				LongName += TEXT(".");
+			}
+			NameEntry->AppendNameToString(LongName);
+		}			
+		bFirstEntry = false;
 	}
 
 	StatID = FDynamicStats::CreateStatId<FStatGroup_STATGROUP_UObjects>( LongName );
diff --git a/Engine/Source/Runtime/CoreUObject/Private/UObject/UTextProperty.cpp b/Engine/Source/Runtime/CoreUObject/Private/UObject/UTextProperty.cpp
--- a/Engine/Source/Runtime/CoreUObject/Private/UObject/UTextProperty.cpp
+++ b/Engine/Source/Runtime/CoreUObject/Private/UObject/UTextProperty.cpp
@@ -3,6 +3,39 @@
 #include "CoreUObjectPrivate.h"
 #include "UObject/UTextProperty.h"
 #include "PropertyHelper.h"
+#include "PropertyTag.h"
+
+bool UTextProperty::ConvertFromType(const FPropertyTag& Tag, FArchive& Ar, uint8* Data, UStruct* DefaultsStruct, bool& bOutAdvanceProperty)
+{
+	// Convert serialized string to text.
+	if (Tag.Type==NAME_StrProperty) 
+	{ 
+		FString str;
+		Ar << str;
+		FText Text = FText::FromString(str);
+		Text.TextData->PersistText();
+		Text.Flags |= ETextFlag::ConvertedProperty;
+		SetPropertyValue_InContainer(Data, Text, Tag.ArrayIndex);
+		bOutAdvanceProperty = true;
+	}
+
+	// Convert serialized name to text.
+	else if (Tag.Type==NAME_NameProperty) 
+	{ 
+		FName Name;  
+		Ar << Name;
+		FText Text = FText::FromName(Name);
+		Text.Flags |= ETextFlag::ConvertedProperty;
+		SetPropertyValue_InContainer(Data, Text, Tag.ArrayIndex);
+		bOutAdvanceProperty = true;
+	}
+	else
+	{
+		bOutAdvanceProperty = false;
+	}
+
+	return bOutAdvanceProperty;
+}
 
 bool UTextProperty::Identical( const void* A, const void* B, uint32 PortFlags ) const
 {
diff --git a/Engine/Source/Runtime/CoreUObject/Public/UObject/GarbageCollection.h b/Engine/Source/Runtime/CoreUObject/Public/UObject/GarbageCollection.h
--- a/Engine/Source/Runtime/CoreUObject/Public/UObject/GarbageCollection.h
+++ b/Engine/Source/Runtime/CoreUObject/Public/UObject/GarbageCollection.h
@@ -29,6 +29,7 @@ enum EGCReferenceType
 	GCRT_AddStructReferencedObjects,
 	GCRT_AddReferencedObjects,
 	GCRT_AddTMapReferencedObjects,
+	GCRT_AddTSetReferencedObjects,
 	GCRT_EndOfPointer,
 	GCRT_EndOfStream,
 };
diff --git a/Engine/Source/Runtime/CoreUObject/Public/UObject/ObjectBase.h b/Engine/Source/Runtime/CoreUObject/Public/UObject/ObjectBase.h
--- a/Engine/Source/Runtime/CoreUObject/Public/UObject/ObjectBase.h
+++ b/Engine/Source/Runtime/CoreUObject/Public/UObject/ObjectBase.h
@@ -335,6 +335,7 @@ typedef uint64 EClassCastFlags;
 #define CASTCLASS_UDelegateFunction				DECLARE_UINT64(0x0000100000000000)
 #define CASTCLASS_UStaticMeshComponent			DECLARE_UINT64(0x0000200000000000)
 #define CASTCLASS_UMapProperty					DECLARE_UINT64(0x0000400000000000)
+#define CASTCLASS_USetProperty					DECLARE_UINT64(0x0000800000000000)
 
 #define CASTCLASS_AllFlags						DECLARE_UINT64(0xFFFFFFFFFFFFFFFF)
 
@@ -413,6 +414,7 @@ typedef uint64 EClassCastFlags;
 #define CPF_PropagateToArrayInner	(CPF_ExportObject | CPF_PersistentInstance | CPF_InstancedReference | CPF_ContainsInstancedReference | CPF_Config | CPF_EditConst | CPF_Deprecated | CPF_EditorOnly | CPF_AutoWeak | CPF_UObjectWrapper )
 #define CPF_PropagateToMapValue		(CPF_ExportObject | CPF_PersistentInstance | CPF_InstancedReference | CPF_ContainsInstancedReference | CPF_Config | CPF_EditConst | CPF_Deprecated | CPF_EditorOnly | CPF_AutoWeak | CPF_UObjectWrapper | CPF_Edit )
 #define CPF_PropagateToMapKey		(CPF_ExportObject | CPF_PersistentInstance | CPF_InstancedReference | CPF_ContainsInstancedReference | CPF_Config | CPF_EditConst | CPF_Deprecated | CPF_EditorOnly | CPF_AutoWeak | CPF_UObjectWrapper | CPF_Edit )
+#define CPF_PropagateToSetElement	(CPF_ExportObject | CPF_PersistentInstance | CPF_InstancedReference | CPF_ContainsInstancedReference | CPF_Config | CPF_EditConst | CPF_Deprecated | CPF_EditorOnly | CPF_AutoWeak | CPF_UObjectWrapper | CPF_Edit )
 
 /** the flags that should never be set on interface properties */
 #define CPF_InterfaceClearMask		(CPF_ExportObject|CPF_InstancedReference|CPF_ContainsInstancedReference)
