Copying //UE4/Dev-Framework to Dev-Main (//UE4/Dev-Main)

diff --git a/Engine/Source/Runtime/CoreUObject/Public/UObject/PropertyTag.h b/Engine/Source/Runtime/CoreUObject/Public/UObject/PropertyTag.h
--- a/Engine/Source/Runtime/CoreUObject/Public/UObject/PropertyTag.h
+++ b/Engine/Source/Runtime/CoreUObject/Public/UObject/PropertyTag.h
@@ -9,6 +9,8 @@
 
 #include "DebugSerializationFlags.h"
 
+class UProperty;
+
 /**
  *  A tag describing a class property, to aid in serialization.
  */
@@ -29,145 +31,17 @@ struct FPropertyTag
 	FGuid	PropertyGuid;
 
 	// Constructors.
-	FPropertyTag()
-		: Type      (NAME_None)
-		, BoolVal   (0)
-		, Name      (NAME_None)
-		, StructName(NAME_None)
-		, EnumName  (NAME_None)
-		, InnerType (NAME_None)
-		, Size      (0)
-		, ArrayIndex(INDEX_NONE)
-		, SizeOffset(INDEX_NONE)
-		, HasPropertyGuid(0)
-	{}
-
-	FPropertyTag( FArchive& InSaveAr, UProperty* Property, int32 InIndex, uint8* Value, uint8* Defaults )
-		: Type      (Property->GetID())
-		, BoolVal   (0)
-		, Name      (Property->GetFName())
-		, StructName(NAME_None)
-		, EnumName  (NAME_None)
-		, InnerType (NAME_None)
-		, Size      (0)
-		, ArrayIndex(InIndex)
-		, SizeOffset(INDEX_NONE)
-		, HasPropertyGuid(0)
-	{
-		if (Property)
-		{
-			// Handle structs.
-			if (UStructProperty* StructProperty = Cast<UStructProperty>(Property))
-			{
-				StructName = StructProperty->Struct->GetFName();
-				StructGuid = StructProperty->Struct->GetCustomGuid();
-			}
-			else if (UByteProperty* ByteProp = Cast<UByteProperty>(Property))
-			{
-				if (ByteProp->Enum != NULL)
-				{
-					EnumName = ByteProp->Enum->GetFName();
-				}
-			}
-			else if (UArrayProperty* ArrayProp = Cast<UArrayProperty>(Property))
-			{
-				InnerType = ArrayProp->Inner->GetID();
-			}
-			else if (UBoolProperty* Bool = Cast<UBoolProperty>(Property))
-			{
-				BoolVal = Bool->GetPropertyValue(Value);
-			}
-		}
-	}
+	FPropertyTag();
+	FPropertyTag( FArchive& InSaveAr, UProperty* Property, int32 InIndex, uint8* Value, uint8* Defaults );
 
 	// Set optional property guid
-	void SetPropertyGuid(const FGuid& InPropertyGuid)
-	{
-		if (InPropertyGuid.IsValid())
-		{
-			PropertyGuid = InPropertyGuid;
-			HasPropertyGuid = true;
-		}
-	}
+	void SetPropertyGuid(const FGuid& InPropertyGuid);
 
 	// Serializer.
-	friend FArchive& operator<<( FArchive& Ar, FPropertyTag& Tag )
-	{
-		// Name.
-		Ar << Tag.Name;
-		if ((Tag.Name == NAME_None) || !Tag.Name.IsValid())
-		{
-			return Ar;
-		}
-
-		Ar << Tag.Type;
-		if ( Ar.IsSaving() )
-		{
-			// remember the offset of the Size variable - UStruct::SerializeTaggedProperties will update it after the
-			// property has been serialized.
-			Tag.SizeOffset = Ar.Tell();
-		}
-		{
-			FArchive::FScopeSetDebugSerializationFlags S(Ar, DSF_IgnoreDiff);
-			Ar << Tag.Size << Tag.ArrayIndex;
-		}
-		// only need to serialize this for structs
-		if (Tag.Type == NAME_StructProperty)
-		{
-			Ar << Tag.StructName;
-			if (Ar.UE4Ver() >= VER_UE4_STRUCT_GUID_IN_PROPERTY_TAG)
-			{
-				Ar << Tag.StructGuid;
-			}
-		}
-		// only need to serialize this for bools
-		else if (Tag.Type == NAME_BoolProperty)
-		{
-			Ar << Tag.BoolVal;
-		}
-		// only need to serialize this for bytes
-		else if (Tag.Type == NAME_ByteProperty)
-		{
-			Ar << Tag.EnumName;
-		}
-		// only need to serialize this for arrays
-		else if (Tag.Type == NAME_ArrayProperty)
-		{
-			if (Ar.UE4Ver() >= VAR_UE4_ARRAY_PROPERTY_INNER_TAGS)
-			{
-				Ar << Tag.InnerType;
-			}
-		}
-		// Property tags to handle renamed blueprint properties effectively.
-		if (Ar.UE4Ver() >= VER_UE4_PROPERTY_GUID_IN_PROPERTY_TAG)
-		{
-			Ar << Tag.HasPropertyGuid;
-			if (Tag.HasPropertyGuid)
-			{
-				Ar << Tag.PropertyGuid;
-			}
-		}
-		
-		return Ar;
-	}
+	friend FArchive& operator<<( FArchive& Ar, FPropertyTag& Tag );
 
 	// Property serializer.
-	void SerializeTaggedProperty( FArchive& Ar, UProperty* Property, uint8* Value, uint8* Defaults )
-	{
-		if (Property->GetClass() == UBoolProperty::StaticClass())
-		{
-			UBoolProperty* Bool = (UBoolProperty*)Property;
-			if (Ar.IsLoading())
-			{
-				Bool->SetPropertyValue(Value, BoolVal != 0);
-			}
-		}
-		else
-		{
-			FSerializedPropertyScope SerializedProperty(Ar, Property);
-			Property->SerializeItem( Ar, Value, Defaults );
-		}
-	}
+	void SerializeTaggedProperty( FArchive& Ar, UProperty* Property, uint8* Value, uint8* Defaults );
 };
 
 #endif	// __UNPROPERTYTAG_H__
diff --git a/Engine/Source/Runtime/CoreUObject/Public/UObject/ReferenceChainSearch.h b/Engine/Source/Runtime/CoreUObject/Public/UObject/ReferenceChainSearch.h
--- a/Engine/Source/Runtime/CoreUObject/Public/UObject/ReferenceChainSearch.h
+++ b/Engine/Source/Runtime/CoreUObject/Public/UObject/ReferenceChainSearch.h
@@ -23,6 +23,8 @@ public:
 			ARO,
 			/** Object is referenced by a map */
 			MapProperty,
+			/** Object is referenced by a set */
+			SetProperty,
 		};
 	};
 
diff --git a/Engine/Source/Runtime/CoreUObject/Public/UObject/Stack.h b/Engine/Source/Runtime/CoreUObject/Public/UObject/Stack.h
--- a/Engine/Source/Runtime/CoreUObject/Public/UObject/Stack.h
+++ b/Engine/Source/Runtime/CoreUObject/Public/UObject/Stack.h
@@ -46,6 +46,7 @@ enum EPropertyType
 	CPT_AssetObjectReference,
 	CPT_Double,
 	CPT_Map,
+	CPT_Set,
 
 	// when you add new property types, make sure you add the corresponding entry
 	// in the PropertyTypeToNameMap array in ScriptCompiler.cpp!!
diff --git a/Engine/Source/Runtime/CoreUObject/Public/UObject/UTextProperty.h b/Engine/Source/Runtime/CoreUObject/Public/UObject/UTextProperty.h
--- a/Engine/Source/Runtime/CoreUObject/Public/UObject/UTextProperty.h
+++ b/Engine/Source/Runtime/CoreUObject/Public/UObject/UTextProperty.h
@@ -26,6 +26,7 @@ public:
 	}
 
 	// UProperty interface
+	virtual bool ConvertFromType(const FPropertyTag& Tag, FArchive& Ar, uint8* Data, UStruct* DefaultsStruct, bool& bOutAdvanceProperty) override;
 	virtual bool Identical( const void* A, const void* B, uint32 PortFlags ) const override;
 	virtual void SerializeItem( FArchive& Ar, void* Value, void const* Defaults ) const override;
 	virtual void ExportTextItem( FString& ValueStr, const void* PropertyValue, const void* DefaultValue, UObject* Parent, int32 PortFlags, UObject* ExportRootScope ) const override;
diff --git a/Engine/Source/Runtime/CoreUObject/Public/UObject/UnrealType.h b/Engine/Source/Runtime/CoreUObject/Public/UObject/UnrealType.h
--- a/Engine/Source/Runtime/CoreUObject/Public/UObject/UnrealType.h
+++ b/Engine/Source/Runtime/CoreUObject/Public/UObject/UnrealType.h
@@ -8,6 +8,7 @@
 
 #include "ObjectBase.h"
 #include "PropertyPortFlags.h"
+#include "PropertyTag.h"
 #include "Templates/IsTriviallyDestructible.h"
 
 COREUOBJECT_API DECLARE_LOG_CATEGORY_EXTERN(LogType, Log, All);
@@ -187,6 +188,20 @@ public:
 protected:
 	virtual void LinkInternal(FArchive& Ar);
 public:
+
+	/**
+	* Allows a property to implement backwards compatibility handling for tagged properties
+	* 
+	* @param	Tag			property tag of the loading data
+	* @param	Ar			the archive the data is being loaded from
+	* @param	Data		
+	* @param	DefaultsStruct 
+	* @param	bOutAdvanceProperty whether the property should be advanced and continue to next property or not
+	*
+	* @return	true if the function has handled the tag
+	*/
+	virtual bool ConvertFromType(const FPropertyTag& Tag, FArchive& Ar, uint8* Data, UStruct* DefaultsStruct, bool& bOutAdvanceProperty) { return false; }
+
 	/**
 	 * Determines whether the property values are identical.
 	 * 
@@ -1094,7 +1109,7 @@ class COREUOBJECT_API UNumericProperty : public UProperty
 	/** Return teh UEnum if this property is a UByteProperty with a non-null Enum **/
 	virtual UEnum* GetIntPropertyEnum() const
 	{
-		return NULL;
+		return nullptr;
 	}
 
 	/** 
@@ -1183,6 +1198,8 @@ class COREUOBJECT_API UNumericProperty : public UProperty
 		return FString();
 	}
 	// End of UNumericProperty interface
+
+	static uint8 ReadEnumAsUint8(FArchive& Ar, UStruct* DefaultsStruct, const FPropertyTag& Tag);
 };
 
 template<typename InTCppType>
@@ -1213,16 +1230,93 @@ public:
 	TProperty_Numeric(FVTableHelper& Helper) : Super(Helper) {};
 #endif // WITH_HOT_RELOAD_CTORS
 
-	FString GetCPPTypeForwardDeclaration() const override
+	virtual FString GetCPPTypeForwardDeclaration() const override
 	{
 		return FString();
 	}
 
 	// UProperty interface
-	uint32 GetValueTypeHashInternal(const void* Src) const override
+	virtual uint32 GetValueTypeHashInternal(const void* Src) const override
 	{
 		return GetTypeHash(*(const InTCppType*)Src);
 	}
+
+protected:
+	template <typename OldIntType>
+	void ConvertFromInt(FArchive& Ar, void* Obj, const FPropertyTag& Tag)
+	{
+		OldIntType OldValue;
+		Ar << OldValue;
+		TCppType NewValue = OldValue;
+		this->SetPropertyValue_InContainer(Obj, NewValue, Tag.ArrayIndex);
+
+		UE_CLOG(
+			(TIsSigned<OldIntType>::Value && !TIsSigned<TCppType>::Value && OldValue < 0) || (sizeof(TCppType) < sizeof(OldIntType) && (OldIntType)NewValue != OldValue),
+			LogClass,
+			Warning,
+			TEXT("Potential data loss during conversion of integer property %s of %s - was (%s) now (%s) - for package: %s"),
+			*this->GetName(),
+			*Ar.GetArchiveName(),
+			*LexicalConversion::ToString(OldValue),
+			*LexicalConversion::ToString(NewValue),
+			*Ar.GetArchiveName()
+			);
+	}
+
+public:
+	virtual bool ConvertFromType(const FPropertyTag& Tag, FArchive& Ar, uint8* Data, UStruct* DefaultsStruct, bool& bOutAdvanceProperty) override
+	{
+		if (Tag.Type == NAME_Int8Property)
+		{
+			ConvertFromInt<int8>(Ar, Data, Tag);
+			return true;
+		}
+		else if (Tag.Type == NAME_Int16Property)
+		{
+			ConvertFromInt<int16>(Ar, Data, Tag);
+			return true;
+		}
+		else if (Tag.Type == NAME_IntProperty)
+		{
+			ConvertFromInt<int32>(Ar, Data, Tag);
+			return true;
+		}
+		else if (Tag.Type == NAME_Int64Property)
+		{
+			ConvertFromInt<int64>(Ar, Data, Tag);
+			return true;
+		}
+		else if (Tag.Type == NAME_ByteProperty)
+		{
+			if (Tag.EnumName != NAME_None)
+			{
+				uint8 PreviousValue = this->ReadEnumAsUint8(Ar, DefaultsStruct, Tag);
+				this->SetPropertyValue_InContainer(Data, PreviousValue, Tag.ArrayIndex);
+			}
+			else
+			{
+				ConvertFromInt<int8>(Ar, Data, Tag);
+			}
+			return true;
+		}
+		else if (Tag.Type == NAME_UInt16Property)
+		{
+			ConvertFromInt<uint16>(Ar, Data, Tag);
+			return true;
+		}
+		else if (Tag.Type == NAME_UInt32Property)
+		{
+			ConvertFromInt<uint32>(Ar, Data, Tag);
+			return true;
+		}
+		else if (Tag.Type == NAME_UInt64Property)
+		{
+			ConvertFromInt<uint64>(Ar, Data, Tag);
+			return true;
+		}
+
+		return false;
+	}
 	// End of UProperty interface
 
 	// UNumericProperty interface.
@@ -1315,6 +1409,7 @@ class COREUOBJECT_API UByteProperty : public TProperty_Numeric<uint8>
 	virtual bool NetSerializeItem( FArchive& Ar, UPackageMap* Map, void* Data, TArray<uint8> * MetaData = NULL ) const override;
 	virtual void ExportTextItem( FString& ValueStr, const void* PropertyValue, const void* DefaultValue, UObject* Parent, int32 PortFlags, UObject* ExportRootScope ) const override;
 	virtual const TCHAR* ImportText_Internal( const TCHAR* Buffer, void* Data, int32 PortFlags, UObject* Parent, FOutputDevice* ErrorText ) const override;
+	virtual bool ConvertFromType(const FPropertyTag& Tag, FArchive& Ar, uint8* Data, UStruct* DefaultsStruct, bool& bOutAdvanceProperty) override;
 	// End of UProperty interface
 
 	// UNumericProperty interface.
@@ -1323,7 +1418,6 @@ class COREUOBJECT_API UByteProperty : public TProperty_Numeric<uint8>
 		return Enum;
 	}
 	// End of UNumericProperty interface
-
 };
 
 /*-----------------------------------------------------------------------------
@@ -1624,6 +1718,7 @@ public:
 	virtual void ClearValueInternal( void* Data ) const override;
 	virtual void InitializeValueInternal( void* Dest ) const override;
 	virtual int32 GetMinAlignment() const override;
+	virtual bool ConvertFromType(const FPropertyTag& Tag, FArchive& Ar, uint8* Data, UStruct* DefaultsStruct, bool& bOutAdvanceProperty) override;
 	// End of UProperty interface
 
 	// Emulate the CPP type API, see TPropertyTypeFundamentals
@@ -1931,6 +2026,7 @@ class COREUOBJECT_API UObjectProperty : public TUObjectPropertyBase<UObject*>
 	virtual void SerializeItem( FArchive& Ar, void* Value, void const* Defaults ) const override;
 	virtual void EmitReferenceInfo(UClass& OwnerClass, int32 BaseOffset) override;
 	virtual const TCHAR* ImportText_Internal(const TCHAR* Buffer, void* Data, int32 PortFlags, UObject* OwnerObject, FOutputDevice* ErrorText) const override;
+	virtual bool ConvertFromType(const FPropertyTag& Tag, FArchive& Ar, uint8* Data, UStruct* DefaultsStruct, bool& bOutAdvanceProperty) override;
 
 private:
 	virtual uint32 GetValueTypeHashInternal(const void* Src) const override
@@ -2066,6 +2162,7 @@ class COREUOBJECT_API UAssetObjectProperty : public TUObjectPropertyBase<FAssetP
 	virtual void SerializeItem( FArchive& Ar, void* Value, void const* Defaults ) const override;
 	virtual void ExportTextItem( FString& ValueStr, const void* PropertyValue, const void* DefaultValue, UObject* Parent, int32 PortFlags, UObject* ExportRootScope ) const override;
 	virtual const TCHAR* ImportText_Internal( const TCHAR* Buffer, void* Data, int32 PortFlags, UObject* OwnerObject, FOutputDevice* ErrorText ) const override;
+	virtual bool ConvertFromType(const FPropertyTag& Tag, FArchive& Ar, uint8* Data, UStruct* DefaultsStruct, bool& bOutAdvanceProperty) override;
 	// End of UProperty interface
 
 	// UObjectProperty interface
@@ -2341,6 +2438,7 @@ public:
 	// UProperty interface
 	virtual void ExportTextItem( FString& ValueStr, const void* PropertyValue, const void* DefaultValue, UObject* Parent, int32 PortFlags, UObject* ExportRootScope ) const override;
 	virtual const TCHAR* ImportText_Internal( const TCHAR* Buffer, void* Data, int32 PortFlags, UObject* OwnerObject, FOutputDevice* ErrorText ) const override;
+	virtual bool ConvertFromType(const FPropertyTag& Tag, FArchive& Ar, uint8* Data, UStruct* DefaultsStruct, bool& bOutAdvanceProperty) override;
 	virtual FString GetCPPTypeForwardDeclaration() const override
 	{
 		return FString();
@@ -2384,6 +2482,7 @@ public:
 	// UProperty interface
 	virtual void ExportTextItem( FString& ValueStr, const void* PropertyValue, const void* DefaultValue, UObject* Parent, int32 PortFlags, UObject* ExportRootScope ) const override;
 	virtual const TCHAR* ImportText_Internal( const TCHAR* Buffer, void* Data, int32 PortFlags, UObject* OwnerObject, FOutputDevice* ErrorText ) const override;
+	virtual bool ConvertFromType(const FPropertyTag& Tag, FArchive& Ar, uint8* Data, UStruct* DefaultsStruct, bool& bOutAdvanceProperty) override;
 	virtual FString GetCPPTypeForwardDeclaration() const override
 	{
 		return FString();
@@ -2459,6 +2558,7 @@ public:
 	virtual bool ContainsWeakObjectReference() const override;
 	virtual void EmitReferenceInfo(UClass& OwnerClass, int32 BaseOffset) override;
 	virtual bool SameType(const UProperty* Other) const override;
+	virtual bool ConvertFromType(const FPropertyTag& Tag, FArchive& Ar, uint8* Data, UStruct* DefaultsStruct, bool& bOutAdvanceProperty) override;
 	// End of UProperty interface
 
 	FString GetCPPTypeCustom(FString* ExtendedTypeText, uint32 CPPExportFlags, const FString& InnerTypeText, const FString& InInnerExtendedTypeText) const;
@@ -2512,6 +2612,53 @@ public:
 	// End of UProperty interface
 };
 
+// need to break this out a different type so that the DECLARE_CASTED_CLASS_INTRINSIC macro can digest the comma
+typedef TProperty<FScriptMap, UProperty> USetProperty_Super;
+
+class COREUOBJECT_API USetProperty : public USetProperty_Super
+{
+	DECLARE_CASTED_CLASS_INTRINSIC(USetProperty, USetProperty_Super, 0, TEXT("/Script/CoreUObject"), CASTCLASS_USetProperty)
+
+	// Properties representing the key type and value type of the contained pairs
+	UProperty*       ElementProp;
+	FScriptSetLayout SetLayout;
+
+public:
+	typedef USetProperty_Super::TTypeFundamentals TTypeFundamentals;
+	typedef TTypeFundamentals::TCppType TCppType;
+
+	USetProperty(const FObjectInitializer& ObjectInitializer, ECppProperty, int32 InOffset, uint64 InFlags);
+
+	// UObject interface
+	virtual void Serialize(FArchive& Ar) override;
+	static void AddReferencedObjects(UObject* InThis, FReferenceCollector& Collector);
+	// End of UObject interface
+
+	// UField interface
+	virtual void AddCppProperty(UProperty* Property) override;
+	// End of UField interface
+
+	// UProperty interface
+	virtual FString GetCPPMacroType(FString& ExtendedTypeText) const  override;
+	virtual FString GetCPPType(FString* ExtendedTypeText, uint32 CPPExportFlags) const override;
+	virtual void LinkInternal(FArchive& Ar) override;
+	virtual bool Identical(const void* A, const void* B, uint32 PortFlags) const override;
+	virtual void SerializeItem(FArchive& Ar, void* Value, void const* Defaults) const override;
+	virtual bool NetSerializeItem(FArchive& Ar, UPackageMap* Map, void* Data, TArray<uint8> * MetaData = NULL) const override;
+	virtual void ExportTextItem(FString& ValueStr, const void* PropertyValue, const void* DefaultValue, UObject* Parent, int32 PortFlags, UObject* ExportRootScope) const override;
+	virtual const TCHAR* ImportText_Internal(const TCHAR* Buffer, void* Data, int32 PortFlags, UObject* OwnerObject, FOutputDevice* ErrorText) const override;
+	virtual void CopyValuesInternal(void* Dest, void const* Src, int32 Count) const override;
+	virtual void ClearValueInternal(void* Data) const override;
+	virtual void DestroyValueInternal(void* Dest) const override;
+	virtual bool PassCPPArgsByRef() const override;
+	virtual void InstanceSubobjects(void* Data, void const* DefaultData, UObject* Owner, struct FObjectInstancingGraph* InstanceGraph) override;
+	virtual bool ContainsObjectReference() const override;
+	virtual bool ContainsWeakObjectReference() const override;
+	virtual void EmitReferenceInfo(UClass& OwnerClass, int32 BaseOffset) override;
+	virtual bool SameType(const UProperty* Other) const override;
+	// End of UProperty interface
+};
+
 /**
  * FScriptArrayHelper: Pseudo dynamic array. Used to work with array properties in a sensible way.
  **/
@@ -2910,11 +3057,11 @@ public:
 	}
 
 	/**
-	 * Static version of Num() used when you don't need to bother to construct a FScriptArrayHelper. Returns the number of elements in the array.
+	 * Static version of Num() used when you don't need to bother to construct a FScriptMapHelper. Returns the number of elements in the map.
 	 *
-	 * @param  Target  Pointer to the raw memory associated with a FScriptArray
+	 * @param  Target  Pointer to the raw memory associated with a FScriptMap
 	 *
-	 * @return The number of elements in the array.
+	 * @return The number of elements in the map.
 	 */
 	static FORCEINLINE int32 Num(const void* Target)
 	{
@@ -2924,11 +3071,11 @@ public:
 	}
 
 	/**
-	 * Returns a uint8 pointer to the pair in the array
+	 * Returns a uint8 pointer to the pair in the map
 	 *
 	 * @param  Index  index of the item to return a pointer to.
 	 *
-	 * @return Pointer to the pair, or nullptr if the array is empty.
+	 * @return Pointer to the pair, or nullptr if the map is empty.
 	 */
 	FORCEINLINE uint8* GetPairPtr(int32 Index)
 	{
@@ -2947,7 +3094,7 @@ public:
 	 *
 	 * @param  Index  index of the item to return a pointer to.
 	 *
-	 * @return Pointer to the pair, or nullptr if the array is empty.
+	 * @return Pointer to the pair, or nullptr if the map is empty.
 	 */
 	FORCEINLINE const uint8* GetPairPtr(int32 Index) const
 	{
@@ -2968,7 +3115,7 @@ public:
 
 	/**
 	 *	Remove all values from the map, calling destructors, etc as appropriate.
-	 *	@param Slack: used to presize the array for a subsequent add, to avoid reallocation.
+	 *	@param Slack: used to presize the set for a subsequent add, to avoid reallocation.
 	**/
 	void EmptyValues(int32 Slack = 0)
 	{
@@ -3209,7 +3356,7 @@ private:
 	 *
 	 * @param  Index  index of the item to return a pointer to.
 	 *
-	 * @return Pointer to the pair, or nullptr if the array is empty.
+	 * @return Pointer to the pair, or nullptr if the map is empty.
 	 */
 	FORCEINLINE uint8* GetPairPtrWithoutCheck(int32 Index)
 	{
@@ -3221,7 +3368,7 @@ private:
 	 *
 	 * @param  Index  index of the item to return a pointer to.
 	 *
-	 * @return Pointer to the pair, or nullptr if the array is empty.
+	 * @return Pointer to the pair, or nullptr if the map is empty.
 	 */
 	FORCEINLINE const uint8* GetPairPtrWithoutCheck(int32 Index) const
 	{
@@ -3244,6 +3391,348 @@ public:
 	}
 };
 
+/**
+* FScriptSetHelper: Pseudo dynamic Set. Used to work with Set properties in a sensible way.
+*/
+class FScriptSetHelper
+{
+	friend class USetProperty;
+
+public:
+	/**
+	* Constructor, brings together a property and an instance of the property located in memory
+	*
+	* @param  InProperty  The property associated with this memory
+	* @param  InSet       Pointer to raw memory that corresponds to this Set. This can be NULL, and sometimes is, but in that case almost all operations will crash.
+	*/
+	FORCEINLINE FScriptSetHelper(const USetProperty* InProperty, const void* InSet)
+		: ElementProp(InProperty->ElementProp)
+		, Set((FScriptSet*)InSet)  //@todo, we are casting away the const here
+		, SetLayout(InProperty->SetLayout)
+	{
+		check(ElementProp);
+	}
+
+	/**
+	* Index range check
+	*
+	* @param  Index  Index to check
+	*
+	* @return true if accessing this element is legal.
+	*/
+	FORCEINLINE bool IsValidIndex(int32 Index) const
+	{
+		return Set->IsValidIndex(Index);
+	}
+
+	/**
+	* Returns the number of elements in the set.
+	*
+	* @return The number of elements in the set.
+	*/
+	FORCEINLINE int32 Num() const
+	{
+		const int32 Result = Set->Num();
+		checkSlow(Result >= 0); 
+		return Result;
+	}
+
+	/**
+	* Returns the (non-inclusive) maximum index of elements in the set.
+	*
+	* @return The (non-inclusive) maximum index of elements in the set.
+	*/
+	FORCEINLINE int32 GetMaxIndex() const
+	{
+		const int32 Result = Set->GetMaxIndex();
+		checkSlow(Result >= Num());
+		return Result;
+	}
+
+	/**
+	* Static version of Num() used when you don't need to bother to construct a FScriptSetHelper. Returns the number of elements in the set.
+	*
+	* @param  Target  Pointer to the raw memory associated with a FScriptSet
+	*
+	* @return The number of elements in the set.
+	*/
+	static FORCEINLINE int32 Num(const void* Target)
+	{
+		const int32 Result = ((const FScriptSet*)Target)->Num();
+		checkSlow(Result >= 0); 
+		return Result;
+	}
+
+	/**
+	* Returns a uint8 pointer to the element in the set.
+	*
+	* @param  Index  index of the item to return a pointer to.
+	*
+	* @return Pointer to the element, or nullptr if the set is empty.
+	*/
+	FORCEINLINE uint8* GetElementPtr(int32 Index)
+	{
+		if (Num() == 0)
+		{
+			checkSlow(!Index);
+			return nullptr;
+		}
+
+		checkSlow(IsValidIndex(Index));
+		return (uint8*)Set->GetData(Index, SetLayout);
+	}
+
+	/**
+	* Returns a uint8 pointer to the element in the set.
+	*
+	* @param  Index  index of the item to return a pointer to.
+	*
+	* @return Pointer to the element, or nullptr if the set is empty.
+	*/
+	FORCEINLINE const uint8* GetElementPtr(int32 Index) const
+	{
+		return const_cast<FScriptSetHelper*>(this)->GetElementPtr(Index);
+	}
+
+	/**
+	* Add an uninitialized value to the end of the set.
+	*
+	* @return  The index of the added element.
+	*/
+	FORCEINLINE int32 AddUninitializedValue()
+	{
+		checkSlow(Num() >= 0);
+
+		return Set->AddUninitialized(SetLayout);
+	}
+
+	/**
+	*	Remove all values from the set, calling destructors, etc as appropriate.
+	*	@param Slack: used to presize the set for a subsequent add, to avoid reallocation.
+	**/
+	void EmptyElements(int32 Slack = 0)
+	{
+		checkSlow(Slack >= 0);
+
+		int32 OldNum = Num();
+		if (OldNum)
+		{
+			DestructItems(0, OldNum);
+		}
+		if (OldNum || Slack)
+		{
+			Set->Empty(Slack, SetLayout);
+		}
+	}
+
+	/**
+	* Adds a blank, constructed value to a given size.
+	* Note that this will create an invalid Set because all the keys will be default constructed, and the set needs rehashing.
+	*
+	* @return  The index of the first element added.
+	**/
+	int32 AddDefaultValue_Invalid_NeedsRehash()
+	{
+		checkSlow(Num() >= 0);
+
+		int32 Result = AddUninitializedValue();
+		ConstructItem(Result);
+
+		return Result;
+	}
+
+	/**
+	* Returns the property representing the element of the set
+	*/
+	UProperty* GetElementProperty() const
+	{
+		return ElementProp;
+	}
+
+	/**
+	* Removes an element at the specified index, destroying it.
+	* The set will be invalid until the next Rehash() call.
+	*
+	* @param  Index  The index of the element to remove.
+	*/
+	void RemoveAt_NeedsRehash(int32 Index, int32 Count = 1)
+	{
+		check(IsValidIndex(Index));
+
+		DestructItems(Index, Count);
+		for (; Count; ++Index)
+		{
+			if (IsValidIndex(Index))
+			{
+				Set->RemoveAt(Index, SetLayout);
+				--Count;
+			}
+		}
+	}
+
+	/**
+	* Rehashes the keys in the set.
+	* This function must be called to create a valid set.
+	*/
+	COREUOBJECT_API void Rehash();
+
+	/**
+	* Finds the index of an element in a set
+	*
+	* @param  ElementToFind		The address of an element to search for.
+	* @param  IndexHint         The index to start searching from.
+	*
+	* @return The index of an element found in SetHelper, or -1 if none was found.
+	*/
+	int32 FindElementIndex(const void* ElementToFind, int32 IndexHint = 0) const
+	{
+		const int32 SetMax = GetMaxIndex();
+		if (SetMax == 0)
+		{
+			return INDEX_NONE;
+		}
+
+		check(IndexHint >= 0 && IndexHint < SetMax);
+
+		UProperty* LocalKeyProp = this->ElementProp; // prevent aliasing in loop below
+
+		int32 Index = IndexHint;
+		for (;;)
+		{
+			if (IsValidIndex(Index))
+			{
+				const void* ElementToCheck = GetElementPtrWithoutCheck(Index);
+				if (LocalKeyProp->Identical(ElementToFind, ElementToCheck))
+				{
+					return Index;
+				}
+			}
+
+			++Index;
+			if (Index == SetMax)
+			{
+				Index = 0;
+			}
+
+			if (Index == IndexHint)
+			{
+				return INDEX_NONE;
+			}
+		}
+	}
+
+	/**
+	* Finds the pair in a map which matches the key in another pair.
+	*
+	* @param  PairWithKeyToFind  The address of a map pair which contains the key to search for.
+	* @param  IndexHint          The index to start searching from.
+	*
+	* @return A pointer to the found pair, or nullptr if none was found.
+	*/
+	FORCEINLINE uint8* FindElementPtr(const void* ElementToFind, int32 IndexHint = 0)
+	{
+		const int32 Index = FindElementIndex(ElementToFind, IndexHint);
+		uint8* Result = (Index >= 0 ? GetElementPtr(Index) : nullptr);
+		return Result;
+	}
+
+
+private:
+	/**
+	* Internal function to call into the property system to construct / initialize elements.
+	*
+	* @param  Index  First item to construct.
+	* @param  Count  Number of items to construct.
+	*/
+	void ConstructItem(int32 Index)
+	{
+		check(IsValidIndex(Index));
+
+		bool bZeroElement = !!(ElementProp->PropertyFlags & CPF_ZeroConstructor);
+		uint8* Dest = GetElementPtrWithoutCheck(Index);
+
+		if (bZeroElement)
+		{
+			// If any nested property needs zeroing, just pre-zero the whole space
+			FMemory::Memzero(Dest, SetLayout.Size);
+		}
+
+		if (!bZeroElement)
+		{
+			ElementProp->InitializeValue_InContainer(Dest);
+		}
+	}
+
+	/**
+	* Internal function to call into the property system to destruct elements.
+	*/
+	void DestructItems(int32 Index, int32 Count)
+	{
+		check(Index >= 0);
+		check(Count >= 0);
+
+		if (Count == 0)
+		{
+			return;
+		}
+
+		bool bDestroyElements = !(ElementProp->PropertyFlags & (CPF_IsPlainOldData | CPF_NoDestructor));
+
+		if (bDestroyElements)
+		{
+			uint32 Stride = SetLayout.Size;
+			uint8* ElementPtr = GetElementPtr(Index);
+
+			for (; Count; ++Index)
+			{
+				if (IsValidIndex(Index))
+				{
+					ElementProp->DestroyValue_InContainer(ElementPtr);
+					--Count;
+				}
+				ElementPtr += Stride;
+			}
+		}
+	}
+
+	/**
+	* Returns a uint8 pointer to the element in the array without checking the index.
+	*
+	* @param  Index  index of the item to return a pointer to.
+	*
+	* @return Pointer to the element, or nullptr if the array is empty.
+	*/
+	FORCEINLINE uint8* GetElementPtrWithoutCheck(int32 Index)
+	{
+		return (uint8*)Set->GetData(Index, SetLayout);
+	}
+
+	/**
+	* Returns a uint8 pointer to the element in the array without checking the index.
+	*
+	* @param  Index  index of the item to return a pointer to.
+	*
+	* @return Pointer to the pair, or nullptr if the array is empty.
+	*/
+	FORCEINLINE const uint8* GetElementPtrWithoutCheck(int32 Index) const
+	{
+		return const_cast<FScriptSetHelper*>(this)->GetElementPtrWithoutCheck(Index);
+	}
+
+public:
+	UProperty*       ElementProp;
+	FScriptSet*      Set;
+	FScriptSetLayout SetLayout;
+};
+
+class FScriptSetHelper_InContainer : public FScriptSetHelper
+{
+public:
+	FORCEINLINE FScriptSetHelper_InContainer(const USetProperty* InProperty, const void* InArray, int32 FixedArrayIndex=0)
+		:FScriptSetHelper(InProperty, InProperty->ContainerPtrToValuePtr<void>(InArray, FixedArrayIndex))
+	{
+	}
+};
 
 /*-----------------------------------------------------------------------------
 	UStructProperty.
@@ -3288,6 +3777,7 @@ public:
 	virtual bool ContainsWeakObjectReference() const override;
 	virtual void EmitReferenceInfo(UClass& OwnerClass, int32 BaseOffset) override;
 	virtual bool SameType(const UProperty* Other) const override;
+	virtual bool ConvertFromType(const FPropertyTag& Tag, FArchive& Ar, uint8* Data, UStruct* DefaultsStruct, bool& bOutAdvanceProperty) override;
 	// End of UProperty interface
 
 	static const TCHAR* ImportText_Static(UScriptStruct* InStruct, const FString& InName, const TCHAR* Buffer, void* Data, int32 PortFlags, UObject* OwnerObject, FOutputDevice* ErrorText);
diff --git a/Engine/Source/Runtime/Engine/Classes/AI/Navigation/NavigationData.h b/Engine/Source/Runtime/Engine/Classes/AI/Navigation/NavigationData.h
--- a/Engine/Source/Runtime/Engine/Classes/AI/Navigation/NavigationData.h
+++ b/Engine/Source/Runtime/Engine/Classes/AI/Navigation/NavigationData.h
@@ -941,8 +941,8 @@ protected:
 	uint32 bRegistered : 1;
 
 	/** was it generated for default agent (SupportedAgents[0]) */
-	uint32 bSupportsDefaultAgent:1;
-
+	uint32 bSupportsDefaultAgent : 1;
+	
 	DEPRECATED(4.12, "This flag is now deprecated, initial rebuild ignore should be handled by discarding dirty areas in UNavigationSystem::ConditionalPopulateNavOctree.")
 	uint32 bWantsUpdate:1;
 
@@ -989,7 +989,7 @@ FORCEINLINE bool FPathFindingResult::IsPartial() const
 	return (Result != ENavigationQueryResult::Error) && Path.IsValid() && Path->IsPartial();
 }
 
-FORCEINLINE void FNavigationPath::SetNavigationDataUsed(const ANavigationData* const NewData)
+FORCEINLINE void FNavigationPath::SetNavigationDataUsed(const ANavigationData* const NavData)
 {
-	NavigationDataUsed = NewData;
+	NavigationDataUsed = NavData;
 }
diff --git a/Engine/Source/Runtime/Engine/Classes/AI/Navigation/NavigationSystem.h b/Engine/Source/Runtime/Engine/Classes/AI/Navigation/NavigationSystem.h
--- a/Engine/Source/Runtime/Engine/Classes/AI/Navigation/NavigationSystem.h
+++ b/Engine/Source/Runtime/Engine/Classes/AI/Navigation/NavigationSystem.h
@@ -759,15 +759,6 @@ public:
 	FORCEINLINE void AddNavigationUpdateLock(uint8 Flags) { NavUpdateLockFlags |= Flags; }
 	FORCEINLINE void RemoveNavigationUpdateLock(uint8 Flags) { NavUpdateLockFlags &= ~Flags; }
 
-	DEPRECATED(4.5, "AreFakeComponentChangesBeingApplied is deprecated. Use IsNavigationUpdateLocked instead.")
-	bool AreFakeComponentChangesBeingApplied() { return IsNavigationUpdateLocked(); }
-	
-	DEPRECATED(4.5, "BeginFakeComponentChanges is deprecated. Use AddNavigationUpdateLock instead.")
-	void BeginFakeComponentChanges() { AddNavigationUpdateLock(ENavigationLockReason::Unknown); }
-	
-	DEPRECATED(4.5, "EndFakeComponentChanges is deprecated. Use RemoveNavigationUpdateLock instead.")
-	void EndFakeComponentChanges() { RemoveNavigationUpdateLock(ENavigationLockReason::Unknown); }
-
 	void UpdateLevelCollision(ULevel* InLevel);
 
 	virtual void OnEditorModeChanged(FEdMode* Mode, bool IsEntering);
@@ -938,18 +929,9 @@ protected:
 	/** Remove BSP collision data from navigation octree */
 	void RemoveLevelCollisionFromOctree(ULevel* Level);
 
-	DEPRECATED(4.8, "EnableAllGenerators is deprecated. Use AddNavigationBuildLock / RemoveNavigationBuildLock instead.")
-	void EnableAllGenerators(bool bEnable, bool bForce = false) {}
-
 	virtual void SpawnMissingNavigationData();
 
 private:
-	DEPRECATED(4.8, "NavigationBuildingLock is deprecated. Use AddNavigationBuildLock instead.")
-	void NavigationBuildingLock() { return AddNavigationBuildLock(ENavigationBuildLock::NoUpdateInEditor); }
-
-	DEPRECATED(4.8, "NavigationBuildingUnlock is deprecated. Use RemoveNavigationBuildLock instead.")
-	void NavigationBuildingUnlock(bool bForce = false) { RemoveNavigationBuildLock(ENavigationBuildLock::NoUpdateInEditor); }
-
 	// adds navigation bounds update request to a pending list
 	void AddNavigationBoundsUpdateRequest(const FNavigationBoundsUpdateRequest& UpdateRequest);
 
@@ -992,19 +974,6 @@ private:
 	// DEPRECATED
 	//----------------------------------------------------------------------//
 public:
-	DEPRECATED(4.8, "This version is deprecated. Please use the one taking reference to INavLinkCustomInterface rather than a pointer instead.")
-	void RegisterCustomLink(INavLinkCustomInterface* CustomLink);
-	DEPRECATED(4.8, "This version is deprecated. Please use the one taking reference to INavLinkCustomInterface rather than a pointer instead.")
-	void UnregisterCustomLink(INavLinkCustomInterface* CustomLink);
-	DEPRECATED(4.8, "GetRandomPointInRadius is deprecated. Use either GetRandomReachablePointInRadius or GetRandomPointInNavigableRadius")
-	UFUNCTION(BlueprintPure, Category = "AI|Navigation", meta = (WorldContext = "WorldContext", DeprecatedFunction, DeprecatedMessage = "GetRandomPointInRadius is deprecated. Use either GetRandomReachablePointInRadius or GetRandomPointInNavigableSpace"))
-	static FVector GetRandomPointInRadius(UObject* WorldContext, const FVector& Origin, float Radius, ANavigationData* NavData = NULL, TSubclassOf<UNavigationQueryFilter> FilterClass = NULL);
-	DEPRECATED(4.8, "GetRandomPointInRadius is deprecated. Use either GetRandomReachablePointInRadius or GetRandomPointInNavigableRadius")
-	bool GetRandomPointInRadius(const FVector& Origin, float Radius, FNavLocation& ResultLocation, ANavigationData* NavData = NULL, FSharedConstNavQueryFilter QueryFilter = NULL) const;
-	DEPRECATED(4.8, "GetRandomPoint is deprecated. Use either GetRandomReachablePointInRadius or GetRandomPointInNavigableRadius")
-	UFUNCTION(BlueprintPure, Category = "AI|Navigation", meta = (WorldContext = "WorldContext", DeprecatedFunction, DeprecatedMessage = "GetRandomPoint is deprecated. Use either GetRandomReachablePointInRadius or GetRandomPointInNavigableSpace"))
-	static FVector GetRandomPoint(UObject* WorldContext, ANavigationData* NavData = NULL, TSubclassOf<UNavigationQueryFilter> FilterClass = NULL);
-
 	DEPRECATED(4.11, "UpdateNavOctree is deprecated. Use UpdateActorInNavOctree")
 	static void UpdateNavOctree(AActor* Actor);
 	DEPRECATED(4.11, "UpdateNavOctree is deprecated. Use UpdateComponentInNavOctree")
diff --git a/Engine/Source/Runtime/Engine/Classes/AI/Navigation/RecastNavMeshDataChunk.h b/Engine/Source/Runtime/Engine/Classes/AI/Navigation/RecastNavMeshDataChunk.h
--- a/Engine/Source/Runtime/Engine/Classes/AI/Navigation/RecastNavMeshDataChunk.h
+++ b/Engine/Source/Runtime/Engine/Classes/AI/Navigation/RecastNavMeshDataChunk.h
@@ -29,6 +29,7 @@ struct FRecastTileData
 };
 
 class dtNavMesh;
+class FPImplRecastNavMesh;
 
 /** 
  * 
diff --git a/Engine/Source/Runtime/Engine/Classes/Animation/AnimBlueprintGeneratedClass.h b/Engine/Source/Runtime/Engine/Classes/Animation/AnimBlueprintGeneratedClass.h
--- a/Engine/Source/Runtime/Engine/Classes/Animation/AnimBlueprintGeneratedClass.h
+++ b/Engine/Source/Runtime/Engine/Classes/Animation/AnimBlueprintGeneratedClass.h
@@ -181,6 +181,10 @@ class ENGINE_API UAnimBlueprintGeneratedClass : public UBlueprintGeneratedClass,
 	UPROPERTY()
 	int32 RootAnimNodeIndex;
 
+	// Indices for each of the saved pose nodes that require updating, in the order they need to get updates.
+	UPROPERTY()
+	TArray<int32> OrderedSavedPoseIndices;
+
 	// The array of anim nodes; this is transient generated data (created during Link)
 	UStructProperty* RootAnimNodeProperty;
 	TArray<UStructProperty*> AnimNodeProperties;
@@ -205,6 +209,8 @@ public:
 
 	virtual const TArray<FName>& GetSyncGroupNames() const override { return SyncGroupNames; }
 
+	virtual const TArray<int32>& GetOrderedSavedPoseNodeIndices() const override { return OrderedSavedPoseIndices; }
+
 	virtual int32 GetSyncGroupIndex(FName SyncGroupName) const override { return SyncGroupNames.IndexOfByKey(SyncGroupName); }
 
 public:
diff --git a/Engine/Source/Runtime/Engine/Classes/Animation/AnimClassData.h b/Engine/Source/Runtime/Engine/Classes/Animation/AnimClassData.h
--- a/Engine/Source/Runtime/Engine/Classes/Animation/AnimClassData.h
+++ b/Engine/Source/Runtime/Engine/Classes/Animation/AnimClassData.h
@@ -26,6 +26,10 @@ public:
 	UPROPERTY()
 	int32 RootAnimNodeIndex;
 	
+	// Indices for each of the saved pose nodes that require updating, in the order they need to get updates.
+	UPROPERTY()
+	TArray<int32> OrderedSavedPoseIndices;
+
 	UPROPERTY()
 	UStructProperty* RootAnimNodeProperty;
 
@@ -49,6 +53,8 @@ public:
 
 	virtual UStructProperty* GetRootAnimNodeProperty() const override { return RootAnimNodeProperty; }
 
+	virtual const TArray<int32>& GetOrderedSavedPoseNodeIndices() const override { return OrderedSavedPoseIndices; }
+
 	virtual const TArray<UStructProperty*>& GetAnimNodeProperties() const override { return AnimNodeProperties; }
 
 	virtual const TArray<FName>& GetSyncGroupNames() const override { return SyncGroupNames; }
@@ -64,6 +70,7 @@ public:
 		AnimNotifies = AnimClass->GetAnimNotifies();
 		RootAnimNodeIndex = AnimClass->GetRootAnimNodeIndex();
 		RootAnimNodeProperty = AnimClass->GetRootAnimNodeProperty();
+		OrderedSavedPoseIndices = AnimClass->GetOrderedSavedPoseNodeIndices();
 		AnimNodeProperties = AnimClass->GetAnimNodeProperties();
 		SyncGroupNames = AnimClass->GetSyncGroupNames();
 	}
diff --git a/Engine/Source/Runtime/Engine/Classes/Animation/AnimClassInterface.h b/Engine/Source/Runtime/Engine/Classes/Animation/AnimClassInterface.h
--- a/Engine/Source/Runtime/Engine/Classes/Animation/AnimClassInterface.h
+++ b/Engine/Source/Runtime/Engine/Classes/Animation/AnimClassInterface.h
@@ -25,6 +25,7 @@ public:
 	virtual const TArray<FAnimNotifyEvent>& GetAnimNotifies() const = 0;
 	virtual const TArray<UStructProperty*>& GetAnimNodeProperties() const = 0;
 	virtual const TArray<FName>& GetSyncGroupNames() const = 0;
+	virtual const TArray<int32>& GetOrderedSavedPoseNodeIndices() const = 0;
 
 	virtual USkeleton* GetTargetSkeleton() const = 0;
 	virtual int32 GetRootAnimNodeIndex() const = 0;
diff --git a/Engine/Source/Runtime/Engine/Classes/Animation/AnimInstance.h b/Engine/Source/Runtime/Engine/Classes/Animation/AnimInstance.h
--- a/Engine/Source/Runtime/Engine/Classes/Animation/AnimInstance.h
+++ b/Engine/Source/Runtime/Engine/Classes/Animation/AnimInstance.h
@@ -28,10 +28,20 @@ struct FAnimNode_AssetPlayerBase;
 struct FAnimNode_Base;
 struct FAnimInstanceProxy;
 
+UENUM()
+enum class EAnimCurveType : uint8 
+{
+	EventCurve,
+	MaterialCurve, 
+	MorphTargetCurve, 
+	PoseCurve, 
+	// make sure to update MaxCurve 
+	MaxAnimCurveType
+};
+
 DECLARE_DELEGATE_OneParam(FOnMontageStarted, UAnimMontage*)
 DECLARE_DELEGATE_TwoParams(FOnMontageEnded, UAnimMontage*, bool /*bInterrupted*/)
 DECLARE_DELEGATE_TwoParams(FOnMontageBlendingOutStarted, UAnimMontage*, bool /*bInterrupted*/)
-
 /**
 * Delegate for when Montage is started
 */
@@ -59,32 +69,10 @@ DECLARE_DELEGATE_RetVal(bool, FCanTakeTransition);
 /** Delegate that native code can hook into to handle state entry/exit */
 DECLARE_DELEGATE_ThreeParams(FOnGraphStateChanged, const struct FAnimNode_StateMachine& /*Machine*/, int32 /*PrevStateIndex*/, int32 /*NextStateIndex*/);
 
-/** Enum for controlling which reference frame a controller is applied in. */
-UENUM()
-enum EBoneControlSpace
-{
-	/** Set absolute position of bone in world space. */
-	BCS_WorldSpace UMETA( DisplayName = "World Space" ),
-	/** Set position of bone in SkeletalMeshComponent's reference frame. */
-	BCS_ComponentSpace UMETA( DisplayName = "Component Space" ),
-	/** Set position of bone relative to parent bone. */
-	BCS_ParentBoneSpace UMETA( DisplayName = "Parent Bone Space" ),
-	/** Set position of bone in its own reference frame. */
-	BCS_BoneSpace UMETA( DisplayName = "Bone Space" ),
-	BCS_MAX,
-};
+/** Delegate that allows users to insert custom animation curve values - for now, it's only single, not sure how to make this to multi delegate and retrieve value sequentially, so */
+DECLARE_DELEGATE_OneParam(FOnAddCustomAnimationCurves, UAnimInstance*)
+
 
-/** Enum for specifying the source of a bone's rotation. */
-UENUM()
-enum EBoneRotationSource
-{
-	/** Don't change rotation at all. */
-	BRS_KeepComponentSpaceRotation UMETA(DisplayName = "No Change (Preserve Existing Component Space Rotation)"),
-	/** Keep forward direction vector relative to the parent bone. */
-	BRS_KeepLocalSpaceRotation UMETA(DisplayName = "Maintain Local Rotation Relative to Parent"),
-	/** Copy rotation of target to bone. */
-	BRS_CopyFromTarget UMETA(DisplayName = "Copy Target Rotation"),
-};
 
 USTRUCT()
 struct FA2Pose
@@ -155,37 +143,8 @@ private:
 	friend class FAnimationRuntime;
 };
 
-USTRUCT(BlueprintType)
-struct FPerBoneBlendWeight
-{
-	GENERATED_USTRUCT_BODY()
-
-	/** Source index of the buffer. */
-	UPROPERTY()
-	int32 SourceIndex;
 
-	UPROPERTY()
-	float BlendWeight;
 
-	FPerBoneBlendWeight()
-		: SourceIndex(0)
-		, BlendWeight(0.0f)
-	{
-	}
-};
-
-USTRUCT(BlueprintType)
-struct FPerBoneBlendWeights
-{
-	GENERATED_USTRUCT_BODY()
-
-	UPROPERTY()
-	TArray<FPerBoneBlendWeight> BoneBlendWeights;
-
-
-	FPerBoneBlendWeights() {}
-
-};
 /** Helper struct for Slot node pose evaluation. */
 USTRUCT()
 struct FSlotEvaluationPose
@@ -367,63 +326,6 @@ struct FMontageEvaluationState
 	bool bIsActive;
 };
 
-struct FGraphTraversalCounter
-{
-private:
-	int16 InternalCounter;
-
-public:
-	FGraphTraversalCounter()
-		: InternalCounter(INDEX_NONE)
-	{}
-
-	int16 Get() const
-	{
-		return InternalCounter;
-	}
-
-	void Increment()
-	{
-		InternalCounter++;
-
-		// Avoid wrapping over back to INDEX_NONE, as this means 'never been traversed'
-		if (InternalCounter == INDEX_NONE)
-		{
-			InternalCounter++;
-		}
-	}
-
-	void Reset()
-	{
-		InternalCounter = INDEX_NONE;
-	}
-
-	void SynchronizeWith(const FGraphTraversalCounter& InMasterCounter)
-	{
-		InternalCounter = InMasterCounter.Get();
-	}
-
-	bool IsSynchronizedWith(const FGraphTraversalCounter& InMasterCounter) const
-	{
-		return ((InternalCounter != INDEX_NONE) && (InternalCounter == InMasterCounter.Get()));
-	}
-
-	bool WasSynchronizedInTheLastFrame(const FGraphTraversalCounter& InMasterCounter) const
-	{
-		// Test if we're currently in sync with our master counter
-		if (IsSynchronizedWith(InMasterCounter))
-		{
-			return true;
-		}
-
-		// If not, test if the Master Counter is a frame ahead of us
-		FGraphTraversalCounter TestCounter(*this);
-		TestCounter.Increment();
-
-		return TestCounter.IsSynchronizedWith(InMasterCounter);
-	}
-};
-
 UCLASS(transient, Blueprintable, hideCategories=AnimInstance, BlueprintType)
 class ENGINE_API UAnimInstance : public UObject
 {
@@ -782,6 +684,14 @@ private:
 
 	/** Used to guard against recursive calls to UpdateAnimation */
 	bool bUpdatingAnimation;
+#if WITH_EDITOR
+	/** Delegate for custom animation curve addition */
+	TArray<FOnAddCustomAnimationCurves> OnAddAnimationCurves;
+public:
+	/** Add custom curve delegates */
+	void AddDelegate_AddCustomAnimationCurve(FOnAddCustomAnimationCurves& InOnAddCustomAnimationCurves);
+	void RemoveDelegate_AddCustomAnimationCurve(FOnAddCustomAnimationCurves& InOnAddCustomAnimationCurves);
+#endif // editor only for now
 public:
 
 	/** 
@@ -1128,8 +1038,12 @@ protected:
 	void TickSyncGroupWriteIndex();
 
 private:
-	/** Curve Values that are added to trigger in event**/
-	TMap<FName, float>	EventCurves;
+	/** This is if you want to separate to another array**/
+	TMap<FName, float>	AnimationCurves[(uint8)EAnimCurveType::MaxAnimCurveType];
+
+	/** Reset Animation Curves */
+	void ResetAnimationCurves();
+
 	/** Material parameters that we had been changing and now need to clear */
 	TArray<FName> MaterialParamatersToClear;
 
@@ -1146,16 +1060,11 @@ public:
 	/** Check whether we have active morph target curves */
 	bool HasMorphTargetCurves() const;
 
-	/** Access active morph target curves */
-	TMap<FName, float>& GetMorphTargetCurves();
-
-	/** Morph Target Curves that will be used for SkeletalMeshComponent **/
-	DEPRECATED(4.11, "This cannot be accessed directly, use UAnimInstance::GetMorphTargetCurves")
-	TMap<FName, float>	MorphTargetCurves;
-
-	/** Material Curves that will be used for SkeletalMeshComponent **/
-	DEPRECATED(4.11, "This cannot be accessed directly, use FAnimInstanceProxy::GetMaterialParameterCurves")
-	TMap<FName, float>	MaterialParameterCurves;
+	/** 
+	 * Retrieve animation curve list by Curve Flags, it will return list of {UID, value} 
+	 * It will clear the OutCurveList before adding
+	 */
+	void GetAnimationCurveList(int32 CurveFlags, TMap<FName, float>& OutCurveList) const;
 
 #if WITH_EDITORONLY_DATA
 	// Maximum playback position ever reached (only used when debugging in Persona)
@@ -1222,17 +1131,16 @@ public:
 
 	/** Given a machine index, record a state machine weight for this frame */
 	void RecordMachineWeight(const int32& InMachineClassIndex, const float& InMachineWeight);
-
-	/** Given a machine and state index, record a state weight for this frame */
-	void RecordStateWeight(const int32& InMachineClassIndex, const int32& InStateIndex, const float& InStateWeight);
-
-protected:
 	/** 
 	 * Add curve float data, using a curve name. External values should all be added using
 	 * The curve UID to the public version of this method
 	 */
 	void AddCurveValue(const FName& CurveName, float Value, int32 CurveTypeFlags);
 
+	/** Given a machine and state index, record a state weight for this frame */
+	void RecordStateWeight(const int32& InMachineClassIndex, const int32& InStateIndex, const float& InStateWeight);
+
+protected:
 #if WITH_EDITORONLY_DATA
 	// Returns true if a snapshot is being played back and the remainder of Update should be skipped.
 	bool UpdateSnapshotAndSkipRemainingUpdate();
@@ -1294,6 +1202,9 @@ protected:
 	/** Actually does the update work, can be called from a worker thread  */
 	void UpdateAnimationInternal_Concurrent(float DeltaSeconds, FAnimInstanceProxy& Proxy);
 
+	/** update animation curves to component */
+	void UpdateCurvesToComponents(USkeletalMeshComponent* Component);
+
 	/** Override point for derived classes to create their own proxy objects (allows custom allocation) */
 	virtual FAnimInstanceProxy* CreateAnimInstanceProxy();
 
diff --git a/Engine/Source/Runtime/Engine/Classes/Animation/AnimMontage.h b/Engine/Source/Runtime/Engine/Classes/Animation/AnimMontage.h
--- a/Engine/Source/Runtime/Engine/Classes/Animation/AnimMontage.h
+++ b/Engine/Source/Runtime/Engine/Classes/Animation/AnimMontage.h
@@ -7,11 +7,12 @@
  *
  */
 
-#include "AnimCompositeBase.h"
-#include "Animation/AnimInstance.h"
+#include "Animation/AnimCompositeBase.h"
+#include "AlphaBlend.h"
 #include "AnimMontage.generated.h"
 
 class UAnimMontage;
+class UAnimInstance;
 
 /**
  * Section data for each track. Reference of data will be stored in the child class for the way they want
@@ -279,43 +280,9 @@ private:
 	void InitializeBlend(const FAlphaBlend& InAlphaBlend);
 
 public:
-	FAnimMontageInstance()
-		: Montage(NULL)
-		, bPlaying(false)
-		, DefaultBlendTimeMultiplier(1.0f)
-		, bDidUseMarkerSyncThisTick(false)
-		, AnimInstance(NULL)
-		, InstanceID(INDEX_NONE)
-		, Position(0.f)
-		, PlayRate(1.f)
-		, bInterrupted(false)
-		, PreviousWeight(0.f)
-		, DeltaMoved(0.f)
-		, PreviousPosition(0.f)
-		, SyncGroupIndex(INDEX_NONE)
-		, MontageSyncLeader(NULL)
-		, MontageSyncUpdateFrameCounter(INDEX_NONE)
-	{
-	}
+	FAnimMontageInstance();
 
-	FAnimMontageInstance(UAnimInstance * InAnimInstance)
-		: Montage(NULL)
-		, bPlaying(false)
-		, DefaultBlendTimeMultiplier(1.0f)
-		, bDidUseMarkerSyncThisTick(false)
-		, AnimInstance(InAnimInstance)
-		, InstanceID(INDEX_NONE)
-		, Position(0.f)
-		, PlayRate(1.f)
-		, bInterrupted(false)
-		, PreviousWeight(0.f)	
-		, DeltaMoved(0.f)
-		, PreviousPosition(0.f)
-		, SyncGroupIndex(INDEX_NONE)
-		, MontageSyncLeader(NULL)
-		, MontageSyncUpdateFrameCounter(INDEX_NONE)
-	{
-	}
+	FAnimMontageInstance(UAnimInstance * InAnimInstance);
 
 	//~ Begin montage instance Interfaces
 	void Play(float InPlayRate = 1.f);
diff --git a/Engine/Source/Runtime/Engine/Classes/Animation/AnimNodeBase.h b/Engine/Source/Runtime/Engine/Classes/Animation/AnimNodeBase.h
--- a/Engine/Source/Runtime/Engine/Classes/Animation/AnimNodeBase.h
+++ b/Engine/Source/Runtime/Engine/Classes/Animation/AnimNodeBase.h
@@ -4,12 +4,12 @@
 
 #include "AnimBlueprint.h"
 #include "AnimBlueprintGeneratedClass.h"
-#include "AnimInstance.h"
 #include "AnimationRuntime.h"
 #include "BonePose.h"
 #include "AnimNodeBase.generated.h"
 
 struct FAnimInstanceProxy;
+class UAnimInstance;
 
 /** Base class for update/evaluate contexts */
 struct FAnimationBaseContext
@@ -267,9 +267,6 @@ private:
 	/** SaveCachePose Nodes */
 	TArray<FNodeDebugData> SaveCachePoseNodes;
 
-	/** Name of CachedPose if this is a CachePose node */
-	FName CachePoseName;
-
 public:
 	struct FFlattenedDebugData
 	{
@@ -298,7 +295,7 @@ public:
 
 	void AddDebugItem(FString DebugData, bool bPoseSource = false);
 	FNodeDebugData& BranchFlow(float BranchWeight, FString InNodeDescription = FString());
-	FNodeDebugData* GetCachePoseDebugData(FName InCachePoseName);
+	FNodeDebugData* GetCachePoseDebugData(float GlobalWeight);
 
 	template<class Type>
 	FString GetNodeName(Type* Node)
diff --git a/Engine/Source/Runtime/Engine/Classes/Animation/AnimNode_SaveCachedPose.h b/Engine/Source/Runtime/Engine/Classes/Animation/AnimNode_SaveCachedPose.h
new file mode 100644
--- /dev/null
+++ b/Engine/Source/Runtime/Engine/Classes/Animation/AnimNode_SaveCachedPose.h
@@ -0,0 +1,45 @@
+// Copyright 1998-2016 Epic Games, Inc. All Rights Reserved.
+
+#pragma once
+#include "Animation/AnimNodeBase.h"
+#include "AnimNode_SaveCachedPose.generated.h"
+
+USTRUCT()
+struct ENGINE_API FAnimNode_SaveCachedPose : public FAnimNode_Base
+{
+	GENERATED_USTRUCT_BODY()
+
+	UPROPERTY(EditAnywhere, BlueprintReadWrite, Category=Links)
+	FPoseLink Pose;
+
+	/** Intentionally not exposed, set by AnimBlueprintCompiler */
+	UPROPERTY()
+	FName CachePoseName;
+
+	UPROPERTY(Transient)
+	float GlobalWeight;
+
+protected:
+	FCompactHeapPose CachedPose;
+	FBlendedHeapCurve CachedCurve;
+
+	TArray<FAnimationUpdateContext> CachedUpdateContexts;
+
+	FGraphTraversalCounter InitializationCounter;
+	FGraphTraversalCounter CachedBonesCounter;
+	FGraphTraversalCounter UpdateCounter;
+	FGraphTraversalCounter EvaluationCounter;
+
+public:	
+	FAnimNode_SaveCachedPose();
+
+	// FAnimNode_Base interface
+	virtual void Initialize(const FAnimationInitializeContext& Context) override;
+	virtual void CacheBones(const FAnimationCacheBonesContext& Context) override;
+	virtual void Update(const FAnimationUpdateContext& Context) override;
+	virtual void Evaluate(FPoseContext& Output) override;
+	virtual void GatherDebugData(FNodeDebugData& DebugData) override;
+	// End of FAnimNode_Base interface
+
+	void PostGraphUpdate();
+};
diff --git a/Engine/Source/Runtime/Engine/Classes/Animation/AnimSequence.h b/Engine/Source/Runtime/Engine/Classes/Animation/AnimSequence.h
--- a/Engine/Source/Runtime/Engine/Classes/Animation/AnimSequence.h
+++ b/Engine/Source/Runtime/Engine/Classes/Animation/AnimSequence.h
@@ -373,7 +373,7 @@ class ENGINE_API UAnimSequence : public UAnimSequenceBase
 	 * The compression scheme that was most recently used to compress this animation.
 	 * May be NULL.
 	 */
-	UPROPERTY(Instanced, Category=Compression, EditAnywhere)
+	UPROPERTY(Category=Compression, VisibleAnywhere)
 	class UAnimCompress* CompressionScheme;
 #endif // WITH_EDITORONLY_DATA
 
@@ -668,6 +668,8 @@ public:
 	*/
 	void ExtractBoneTransform(const struct FRawAnimSequenceTrack& InRawAnimationTrack, FTransform& OutAtom, float Time) const;
 
+	void ExtractBoneTransform(const struct FRawAnimSequenceTrack& RawTrack, FTransform& OutAtom, int32 KeyIndex) const;
+
 	// End Transform related functions 
 
 	// Begin Memory related functions
@@ -907,7 +909,7 @@ private:
 	/**
 	 * Remap Tracks to New Skeleton
 	 */
-	void RemapTracksToNewSkeleton( USkeleton* NewSkeleton, bool bConvertSpaces );
+	virtual void RemapTracksToNewSkeleton( USkeleton* NewSkeleton, bool bConvertSpaces ) override;
 	/**
 	 * Remap NaN tracks from the RawAnimation data and recompress
 	 */	
diff --git a/Engine/Source/Runtime/Engine/Classes/Animation/AnimSequenceBase.h b/Engine/Source/Runtime/Engine/Classes/Animation/AnimSequenceBase.h
--- a/Engine/Source/Runtime/Engine/Classes/Animation/AnimSequenceBase.h
+++ b/Engine/Source/Runtime/Engine/Classes/Animation/AnimSequenceBase.h
@@ -186,5 +186,18 @@ public:
 
 protected:
 	template <typename DataType>
-	void VerifyCurveNames(USkeleton* Skeleton, const FName& NameContainer, TArray<DataType>& CurveList);
+	void VerifyCurveNames(USkeleton* Skeleton, const FName& NameContainer, TArray<DataType>& CurveList)
+	{
+		USkeleton* MySkeleton = GetSkeleton();
+
+		if (MySkeleton)
+		{
+			// since this is verify function that makes sure it exists after loaded
+			// we should add it if it doesn't exist
+			for (DataType& Curve : CurveList)
+			{
+				MySkeleton->VerifySmartName(NameContainer, Curve.Name);
+			}
+		}
+	}
 };
diff --git a/Engine/Source/Runtime/Engine/Classes/Animation/AnimSingleNodeInstance.h b/Engine/Source/Runtime/Engine/Classes/Animation/AnimSingleNodeInstance.h
--- a/Engine/Source/Runtime/Engine/Classes/Animation/AnimSingleNodeInstance.h
+++ b/Engine/Source/Runtime/Engine/Classes/Animation/AnimSingleNodeInstance.h
@@ -95,7 +95,9 @@ public:
 	/** Set New Asset - calls InitializeAnimation, for now we need MeshComponent **/
 	UFUNCTION(BlueprintCallable, Category="Animation")
 	virtual void SetAnimationAsset(UAnimationAsset* NewAsset, bool bIsLooping=true, float InPlayRate=1.f);
-
+	/** Set pose value */
+ 	UFUNCTION(BlueprintCallable, Category = "Animation")
+ 	void SetPreviewPoseOverride(const FName& PoseName, float Value);
 public:
 	/** AnimSequence specific **/
 	void StepForward();
diff --git a/Engine/Source/Runtime/Engine/Classes/Animation/AnimationAsset.h b/Engine/Source/Runtime/Engine/Classes/Animation/AnimationAsset.h
--- a/Engine/Source/Runtime/Engine/Classes/Animation/AnimationAsset.h
+++ b/Engine/Source/Runtime/Engine/Classes/Animation/AnimationAsset.h
@@ -785,12 +785,22 @@ public:
 	UPROPERTY(VisibleAnywhere, Instanced, Category = Thumbnail)
 	class UThumbnailInfo* ThumbnailInfo;
 
+	/** The default skeletal mesh to use when previewing this asset - this only applies when you open Persona using this asset*/
+	// @todo: note that this doesn't retarget right now
+	UPROPERTY(duplicatetransient, EditAnywhere, Category = Animation)
+	class UPoseAsset* PreviewPoseAsset;
+
 private:
 	/** The default skeletal mesh to use when previewing this asset - this only applies when you open Persona using this asset*/
 	UPROPERTY(duplicatetransient, AssetRegistrySearchable)
 	TAssetPtr<class USkeletalMesh> PreviewSkeletalMesh;
 #endif //WITH_EDITORONLY_DATA
 
+protected:
+#if WITH_EDITOR
+	virtual void RemapTracksToNewSkeleton(USkeleton* NewSkeleton, bool bConvertSpaces);
+#endif // WITH_EDITOR
+
 public:
 	class USkeleton* GetSkeleton() const { return Skeleton; }
 };
diff --git a/Engine/Source/Runtime/Engine/Classes/Animation/PoseAsset.h b/Engine/Source/Runtime/Engine/Classes/Animation/PoseAsset.h
new file mode 100644
--- /dev/null
+++ b/Engine/Source/Runtime/Engine/Classes/Animation/PoseAsset.h
@@ -0,0 +1,207 @@
+// Copyright 1998-2016 Epic Games, Inc. All Rights Reserved.
+
+#pragma once
+
+/**
+ * Class of pose asset that can evaluate pose by weights
+ *
+ */
+
+#include "AnimationAsset.h"
+#include "PoseAsset.generated.h"
+
+/** 
+ * Pose data 
+ * 
+ * This is one pose data structure
+ * This will let us blend poses quickly easily
+ * All poses within this asset should contain same number of tracks, 
+ * so that we can blend quickly
+ */
+
+USTRUCT()
+struct ENGINE_API FPoseData
+{
+	GENERATED_USTRUCT_BODY()
+
+	// local space pose, # of array match with # of Tracks
+	UPROPERTY()
+	TArray<FTransform>		LocalSpacePose;
+
+	// # of array match with # of Curves in PoseDataContainer
+ 	UPROPERTY()
+ 	TArray<float>			CurveData;
+};
+
+/**
+* Pose data container
+* 
+* Contains animation and curve for all poses
+*/
+USTRUCT()
+struct ENGINE_API FPoseDataContainer
+{
+	GENERATED_USTRUCT_BODY()
+
+private:
+	// pose names - horizontal data
+	// # of poses - there is no compression across tracks - 	
+	// unfortunately, tried TMap, but it is not great because it changes order whenever add/remove
+	// we need consistent array of names, so that it doesn't change orders
+	UPROPERTY()
+	TArray<FSmartName>						PoseNames;
+
+	UPROPERTY()
+	TArray<FPoseData>						Poses;
+
+	UPROPERTY()
+	TArray<FName>							Tracks;
+
+	// vertical data - the track names for bone position, and skeleton index 
+	UPROPERTY(transient)
+	TMap<FName, int32>						TrackMap;
+	
+	// curve meta data # of Curve UIDs should match with Poses.CurveValues.Num
+	UPROPERTY()
+	TArray<FAnimCurveBase>					Curves;
+
+	void Reset();
+
+	void AddOrUpdatePose(const FSmartName& InPoseName, const TArray<FTransform>& InlocalSpacePose, const TArray<float>& InCurveData);
+
+	// remove track if all poses has identity key
+	// we don't want to save transform per track (good for compression) because that will slow down blending
+	void Shrink(USkeleton* InSkeleton, FName& InRetargetSourceName);
+
+	bool InsertTrack(const FName& InTrackName, USkeleton* InSkeleton, FName& InRetargetSourceName);
+	bool FillUpDefaultPose(const FSmartName& InPoseName, USkeleton* InSkeleton, FName& InRetargetSourceName);
+	bool FillUpDefaultPose(FPoseData* PoseData, USkeleton* InSkeleton, FName& InRetargetSourceName);
+	FTransform GetDefaultTransform(int32 SkeletonIndex, USkeleton* InSkeleton, const FName& InRetargetSourceName) const;
+	FTransform GetDefaultTransform(const FName& InTrackName, USkeleton* InSkeleton, const FName& InRetargetSourceName) const;
+	void RenamePose(FSmartName OldPoseName, FSmartName NewPoseName);
+	bool DeletePose(FSmartName PoseName);
+	bool DeleteCurve(FSmartName CurveName);
+
+	FPoseData* FindPoseData(FSmartName PoseName);
+	FPoseData* FindOrAddPoseData(FSmartName PoseName);
+
+	int32 GetNumPoses() const { return Poses.Num();  }
+	bool Contains(FSmartName PoseName) const { return PoseNames.Contains(PoseName); }
+
+	bool IsValid() const { return PoseNames.Num() == Poses.Num() && Tracks.Num() == TrackMap.Num(); }
+	void GetPoseCurve(const FPoseData* PoseData, FBlendedCurve& OutCurve) const;
+	void ConvertToFullPose(int32 InBasePoseIndex, const TArray<FTransform>& InBasePose, const TArray<float>& InBaseCurve);
+	void ConvertToAdditivePose(int32 InBasePoseIndex, const TArray<FTransform>& InBasePose, const TArray<float>& InBaseCurve);
+
+	friend class UPoseAsset;
+};
+
+// @todo: add proper WITH_EDITOR
+UCLASS(MinimalAPI, BlueprintType)
+class UPoseAsset : public UAnimationAsset
+{
+	GENERATED_UCLASS_BODY()
+
+private:
+	/** Animation Pose Data*/
+	UPROPERTY()
+	struct FPoseDataContainer PoseContainer;
+
+	/** Total number of poses - redundant data but noce to see it in asset registry*/
+	UPROPERTY(Category=Animation, AssetRegistrySearchable, VisibleAnywhere)
+	int32 NumPoses;
+
+	UPROPERTY(Category = Animation, EditAnywhere)
+	bool bNormalizeWeight;
+
+	/** Whether or not Additive Pose or not - these are property that needs post process, so */
+	UPROPERTY()
+	bool bAdditivePose;
+
+	/** if -1, use ref pose */
+	UPROPERTY()
+	int32 BasePoseIndex;
+
+public: 
+	/** Base pose to use when retargeting */
+	UPROPERTY(Category=Animation, EditAnywhere)
+	FName RetargetSource;
+
+	/**
+	* Get Animation Pose from one pose of PoseIndex and with PoseWeight
+	* This returns OutPose and OutCurve of one pose of PoseIndex with PoseWeight
+	*
+	* @param	OutPose				Pose object to fill
+	* @param	InOutCurve			Curves to fill
+	* @param	PoseIndex			Index of Pose
+	* @param	PoseWeight			Weight of pose
+	*/
+	ENGINE_API void GetAnimationPose(struct FCompactPose& OutPose, FBlendedCurve& InOutCurve, const FAnimExtractContext& ExtractionContext) const;
+	ENGINE_API void GetBaseAnimationPose(struct FCompactPose& OutPose, FBlendedCurve& InOutCurve, const FAnimExtractContext& ExtractionContext) const;
+	virtual bool HasRootMotion() const { return false; }
+	virtual bool IsValidAdditive() const { return bAdditivePose; }
+
+	//Begin UObject Interface
+	virtual void PostLoad() override;
+	virtual void Serialize(FArchive& Ar) override;
+	//End UObject Interface
+
+public:
+	ENGINE_API const TArray<FSmartName> GetPoseNames() const;
+	ENGINE_API const TArray<FSmartName> GetCurveNames() const;
+
+	bool ContainsPose(const FSmartName& InPoseName) const { return PoseContainer.Contains(InPoseName); }
+	bool ContainsPose(const FName& InPoseName) const;
+
+	UFUNCTION(BlueprintCallable, Category = "Animation")
+	ENGINE_API void AddOrUpdatePose(const FSmartName& PoseName, USkeletalMeshComponent* MeshComponent);
+	ENGINE_API void AddOrUpdatePose(const FSmartName& PoseName, TArray<FName> TrackNames, TArray<FTransform>& LocalTransform);
+
+#if WITH_EDITOR
+	ENGINE_API void CreatePoseFromAnimation(class UAnimSequence* AnimSequence);
+#endif
+	ENGINE_API bool ModifyPoseName(FName OldPoseName, FName NewPoseName, const FSmartNameMapping::UID* NewUID);
+
+	ENGINE_API int32 DeletePoses(TArray<FName> PoseNamesToDelete);
+	ENGINE_API int32 DeleteCurves(TArray<FName> CurveNamesToDelete);
+	ENGINE_API bool ConvertSpace(bool bNewAdditivePose, int32 NewBasePoseInde);
+
+	ENGINE_API int32 GetBasePoseIndex() const { return BasePoseIndex;  }
+	ENGINE_API const int32 GetPoseIndexByName(const FName& InBasePoseName) const;
+	ENGINE_API const FName GetPoseNameByIndex(int32 InBasePoseIndex) const { return PoseContainer.PoseNames.IsValidIndex(InBasePoseIndex)? PoseContainer.PoseNames[InBasePoseIndex].DisplayName : NAME_None; }
+
+private: 
+	DECLARE_MULTICAST_DELEGATE(FOnPoseListChangedMulticaster)
+	FOnPoseListChangedMulticaster OnPoseListChanged;
+
+public:
+	typedef FOnPoseListChangedMulticaster::FDelegate FOnPoseListChanged;
+
+	/** Registers a delegate to be called after the preview animation has been changed */
+	FDelegateHandle RegisterOnPoseListChanged(const FOnPoseListChanged& Delegate)
+	{
+		return OnPoseListChanged.Add(Delegate);
+	}
+	/** Unregisters a delegate to be called after the preview animation has been changed */
+	void UnregisterOnPoseListChanged(FDelegateHandle Handle)
+	{
+		OnPoseListChanged.Remove(Handle);
+	}
+
+#if WITH_EDITOR
+protected:
+	virtual void RemapTracksToNewSkeleton(USkeleton* NewSkeleton, bool bConvertSpaces) override;
+#endif // WITH_EDITOR	
+
+private:
+
+	// this will do multiple things, it will add tracks and make sure it fix up all poses with it
+	// use same as retarget source system we have for animation
+	void CombineTracks(const TArray<FName>& NewTracks);
+
+	bool ConvertToFullPose();
+	bool ConvertToAdditivePose(int32 NewBasePoseIndex);
+
+	bool GetBasePoseTransform(TArray<FTransform>& OutBasePose, TArray<float>& OutCurve) const;
+	void RecacheTrackmap();
+};
diff --git a/Engine/Source/Runtime/Engine/Classes/Animation/Skeleton.h b/Engine/Source/Runtime/Engine/Classes/Animation/Skeleton.h
--- a/Engine/Source/Runtime/Engine/Classes/Animation/Skeleton.h
+++ b/Engine/Source/Runtime/Engine/Classes/Animation/Skeleton.h
@@ -8,6 +8,7 @@
 #pragma once
 #include "PreviewAssetAttachComponent.h"
 #include "SmartName.h"
+#include "ReferenceSkeleton.h"
 #include "Skeleton.generated.h"
 
 class UAnimSequence;
@@ -334,9 +335,12 @@ public:
 	ENGINE_API void RemoveSlotGroup(const FName& InSlotName);
 	ENGINE_API void RenameSlotName(const FName& OldName, const FName& NewName);
 
+	////////////////////////////////////////////////////////////////////////////
+	// Smart Name Interfaces
+	////////////////////////////////////////////////////////////////////////////
 	// Adds a new name to the smart name container and modifies the skeleton so it can be saved
 	// return bool - Whether a name was added (false if already present)
-	ENGINE_API bool AddSmartNameAndModify(FName ContainerName, FName NewName, FSmartNameMapping::UID& NewUid);
+	ENGINE_API bool AddSmartNameAndModify(FName ContainerName, FName NewDisplayName, FSmartName& NewName);
 
 	// Renames a smartname in the specified container and modifies the skeleton
 	// return bool - Whether the rename was sucessful
@@ -348,12 +352,26 @@ public:
 	// Removes smartnames from the specified container and modifies the skeleton
 	ENGINE_API void RemoveSmartnamesAndModify(FName ContainerName, const TArray<FSmartNameMapping::UID>& Uids);
 
-	// Get or add a smartname container with the given name
-	ENGINE_API const FSmartNameMapping* GetOrAddSmartNameContainer(FName ContainerName);
+	// quick wrapper function for Find UID by name, if not found, it will return FSmartNameMapping::MaxUID
+	ENGINE_API FSmartNameMapping::UID GetUIDByName(FName ContainerName, FName Name);
+	ENGINE_API bool GetSmartNameByUID(FName ContainerName, FSmartNameMapping::UID UID, FSmartName& OutSmartName);
+	ENGINE_API bool GetSmartNameByName(FName ContainerName, FName InName, FSmartName& OutSmartName);
+
+	// Adds a new name to the smart name container and modifies the skeleton so it can be saved
+	// return bool - Whether a name was added (false if already present)
+	ENGINE_API bool RenameSmartName(FName ContainerName, const FSmartNameMapping::UID& Uid, FName NewName);
 
 	// Get or add a smartname container with the given name
-	ENGINE_API const FSmartNameMapping* GetSmartNameContainer(FName ContainerName) const;
+	ENGINE_API const FSmartNameMapping* GetSmartNameContainer(const FName& ContainerName) const;
 
+	// make sure the smart name has valid UID and so on
+	ENGINE_API void VerifySmartName(const FName&  ContainerName, FSmartName& InOutSmartName);
+	ENGINE_API void VerifySmartNames(const FName&  ContainerName, TArray<FSmartName>& InOutSmartNames);
+private:
+	// Get or add a smartname container with the given name
+	FSmartNameMapping* GetOrAddSmartNameContainer(const FName& ContainerName);
+	bool VerifySmartNameInternal(const FName&  ContainerName, FSmartName& InOutSmartName);
+	bool FillSmartNameByDisplayName(FSmartNameMapping* Mapping, const FName& DisplayName, FSmartName& OutSmartName);
 #if WITH_EDITORONLY_DATA
 private:
 	/** The default skeletal mesh to use when previewing this skeleton */
@@ -524,13 +542,13 @@ public:
 	 *	Anybody modifying BoneTree will corrupt animation data, so will need to make sure it's not modifiable outside of Skeleton
 	 *	You can add new BoneNode but you can't modify current list. The index will be referenced by Animation data.
 	 */
-	const TArray<FBoneNode> & GetBoneTree()	const
+	const TArray<FBoneNode>& GetBoneTree()	const
 	{ 
 		return BoneTree;	
 	}
 
 	// @todo document
-	const TArray<FTransform> & GetRefLocalPoses( FName RetargetSource = NAME_None ) const 
+	const TArray<FTransform>& GetRefLocalPoses( FName RetargetSource = NAME_None ) const 
 	{
 		if ( RetargetSource != NAME_None ) 
 		{
diff --git a/Engine/Source/Runtime/Engine/Classes/Animation/SmartName.h b/Engine/Source/Runtime/Engine/Classes/Animation/SmartName.h
--- a/Engine/Source/Runtime/Engine/Classes/Animation/SmartName.h
+++ b/Engine/Source/Runtime/Engine/Classes/Animation/SmartName.h
@@ -4,6 +4,8 @@
 
 #include "SmartName.generated.h"
 
+struct FSmartName;
+
 USTRUCT()
 struct ENGINE_API FSmartNameMapping
 {
@@ -23,14 +25,14 @@ struct ENGINE_API FSmartNameMapping
 	// @param Name - The name to add/get
 	// @param OUT OutUid - The UID of the newly created or retrieved name
 	// @return bool - true if the name was added, false if it existed (OutUid will be correctly set still in this case)
-	bool AddOrFindName(FName Name, UID& OutUid);
+	bool AddOrFindName(FName Name, UID& OutUid, FGuid& OutGuid);
 
 	// Get a name from the mapping
 	// @param Uid - UID of the name to retrieve
 	// @param OUT OutName - Retrieved name
 	// @return bool - true if name existed and OutName is valid
 	bool GetName(const UID& Uid, FName& OutName) const;
-
+	bool GetNameByGuid(const FGuid& Guid, FName& OutName) const;
 	// Fill an array with all used UIDs
 	// @param Array - Array to fill
 	void FillUidArray(TArray<UID>& Array) const;
@@ -68,16 +70,20 @@ struct ENGINE_API FSmartNameMapping
 	// Get the number of names currently stored in this container
 	int32 GetNumNames() const;
 
+	// Find Or Add Smart Names
+	bool FindOrAddSmartName(FName Name, FSmartName& OutName);
+	bool FindSmartName(FName Name, FSmartName& OutName) const;
+	bool FindSmartNameByUID(FSmartNameMapping::UID UID, FSmartName& OutName) const;
+
 	// Serialize this to the provided archive; required for TMap serialization
 	void Serialize(FArchive& Ar);
 
 	friend FArchive& operator<<(FArchive& Ar, FSmartNameMapping& Elem);
 
 private:
-
 	UID NextUid;				// The next UID to use 
-	TMap<UID, FName> UidMap;	// Mapping of UIDs and names
-	TArray<UID> FreeList;		// Freelist of UIDs from removing/moving names around. Populated on load
+	TMap<UID, FName> UidMap;	// Mapping of UIDs and names. This is transient and built upon serialize - this is run-time data, and searching from UID to FName is important
+	TMap<FName, FGuid> GuidMap;	// Mapping of GUIDs and names. This is the one serializing
 };
 
 USTRUCT()
@@ -99,8 +105,74 @@ struct ENGINE_API FSmartNameContainer
 	/** Only restricted classes can access the protected interface */
 	friend class USkeleton;
 protected:
-	FSmartNameMapping* GetContainerInternal(FName ContainerName);
+	FSmartNameMapping* GetContainerInternal(const FName& ContainerName);
 
 private:
 	TMap<FName, FSmartNameMapping> NameMappings;	// List of smartname mappings
+};
+
+USTRUCT()
+struct ENGINE_API FSmartName
+{
+
+	GENERATED_USTRUCT_BODY();
+
+	// name 
+	UPROPERTY()
+	FName DisplayName;
+
+	// UID - for faster access
+	FSmartNameMapping::UID	UID;
+
+	UPROPERTY()
+	FGuid	Guid;
+
+	FSmartName()
+		: DisplayName(NAME_None)
+		, UID(FSmartNameMapping::MaxUID)
+	{}
+
+	FSmartName(const FName& InDisplayName, FSmartNameMapping::UID InUID, const FGuid& InGuid)
+		: DisplayName(InDisplayName)
+		, UID(InUID)
+		, Guid(InGuid)
+	{}
+
+	bool operator==(FSmartName const& Other) const
+	{
+		return (DisplayName == Other.DisplayName && UID == Other.UID && Guid == Other.Guid);
+	}
+	bool operator!=(const FSmartName& Other) const
+	{
+		return !(*this == Other);
+	}
+	/**
+	* Serialize the SmartName
+	*
+	* @param Ar	Archive to serialize to
+	*
+	* @return True if the container was serialized
+	*/
+	bool Serialize(FArchive& Ar);
+
+	friend FArchive& operator<<(FArchive& Ar, FSmartName& P)
+	{
+		P.Serialize(Ar);
+		return Ar;
+	}
+
+	bool IsValid() const
+	{
+		return UID != FSmartNameMapping::MaxUID && Guid.IsValid(); 
+	}
+};
+
+template<>
+struct TStructOpsTypeTraits<FSmartName> : public TStructOpsTypeTraitsBase
+{
+	enum
+	{
+		WithSerializer = true,
+		WithIdenticalViaEquality = true
+	};
 };
\ No newline at end of file
diff --git a/Engine/Source/Runtime/Engine/Classes/Components/ActorComponent.h b/Engine/Source/Runtime/Engine/Classes/Components/ActorComponent.h
--- a/Engine/Source/Runtime/Engine/Classes/Components/ActorComponent.h
+++ b/Engine/Source/Runtime/Engine/Classes/Components/ActorComponent.h
@@ -325,10 +325,22 @@ public:
 	/** Returns true if we are replicating and not authorative */
 	bool	IsNetSimulating() const;
 
+	/** Get the network role of the Owner, or ROLE_None if there is no owner. */
 	ENetRole GetOwnerRole() const;
 
+	/**
+	 * Get the network mode (dedicated server, client, standalone, etc) for this component.
+	 * @see IsNetMode()
+	 */
 	ENetMode GetNetMode() const;
 
+	/**
+	* Test whether net mode is the given mode.
+	* In optimized non-editor builds this can be more efficient than GetNetMode()
+	* because it can check the static build flags without considering PIE.
+	*/
+	bool IsNetMode(ENetMode Mode) const;
+
 private:
 
 	/** 
@@ -340,6 +352,9 @@ private:
 	// If WorldPrivate isn't set this will determine the world from outers
 	UWorld* GetWorld_Uncached() const;
 
+	/** Private version without inlining that does *not* check Dedicated server build flags (which should already have been done). */
+	ENetMode InternalGetNetMode() const;
+
 protected:
 
 	/** "Trigger" related function. Return true if it should activate **/
@@ -785,4 +800,38 @@ FORCEINLINE_DEBUGGABLE class AActor* UActorComponent::GetOwner() const
 FORCEINLINE bool UActorComponent::CanEverAffectNavigation() const
 {
 	return bCanEverAffectNavigation;
-}
\ No newline at end of file
+}
+
+FORCEINLINE_DEBUGGABLE bool UActorComponent::IsNetSimulating() const
+{
+	return GetIsReplicated() && GetOwnerRole() != ROLE_Authority;
+}
+
+FORCEINLINE_DEBUGGABLE ENetMode UActorComponent::GetNetMode() const
+{
+	// IsRunningDedicatedServer() is a compile-time check in optimized non-editor builds.
+	if (IsRunningDedicatedServer())
+	{
+		return NM_DedicatedServer;
+	}
+
+	return InternalGetNetMode();
+}
+
+FORCEINLINE_DEBUGGABLE bool UActorComponent::IsNetMode(ENetMode Mode) const
+{
+#if UE_EDITOR
+	// Editor builds are special because of PIE, which can run a dedicated server without the app running with -server.
+	return GetNetMode() == Mode;
+#else
+	// IsRunningDedicatedServer() is a compile-time check in optimized non-editor builds.
+	if (Mode == NM_DedicatedServer)
+	{
+		return IsRunningDedicatedServer();
+	}
+	else
+	{
+		return !IsRunningDedicatedServer() && (InternalGetNetMode() == Mode);
+	}
+#endif // UE_EDITOR
+}
diff --git a/Engine/Source/Runtime/Engine/Classes/Components/ChildActorComponent.h b/Engine/Source/Runtime/Engine/Classes/Components/ChildActorComponent.h
--- a/Engine/Source/Runtime/Engine/Classes/Components/ChildActorComponent.h
+++ b/Engine/Source/Runtime/Engine/Classes/Components/ChildActorComponent.h
@@ -78,6 +78,7 @@ public:
 	virtual void OnRegister() override;
 	virtual void OnUnregister() override;
 	virtual FActorComponentInstanceData* GetComponentInstanceData() const override;
+	virtual void BeginPlay() override;
 	//~ End ActorComponent Interface.
 
 	/** Apply the component instance data to the child actor component */
diff --git a/Engine/Source/Runtime/Engine/Classes/Components/PoseableMeshComponent.h b/Engine/Source/Runtime/Engine/Classes/Components/PoseableMeshComponent.h
--- a/Engine/Source/Runtime/Engine/Classes/Components/PoseableMeshComponent.h
+++ b/Engine/Source/Runtime/Engine/Classes/Components/PoseableMeshComponent.h
@@ -1,7 +1,7 @@
 // Copyright 1998-2016 Epic Games, Inc. All Rights Reserved.
 
 #pragma once
-#include "Animation/AnimationAsset.h"
+#include "Components/SkinnedMeshComponent.h"
 #include "PoseableMeshComponent.generated.h"
 
 /**
diff --git a/Engine/Source/Runtime/Engine/Classes/Components/PrimitiveComponent.h b/Engine/Source/Runtime/Engine/Classes/Components/PrimitiveComponent.h
--- a/Engine/Source/Runtime/Engine/Classes/Components/PrimitiveComponent.h
+++ b/Engine/Source/Runtime/Engine/Classes/Components/PrimitiveComponent.h
@@ -1253,15 +1253,30 @@ public:
 	virtual FBodyInstance* GetBodyInstance(FName BoneName = NAME_None, bool bGetWelded = true) const;
 
 	/** 
-	 * returns Distance to closest Body Instance surface. 
+	 * returns The square of the distance to closest Body Instance surface. 
 	 *
 	 * @param Point				World 3D vector
+	 * @param OutSquaredDistance The squared distance to closest Body Instance surface. 0 if inside of the body
 	 * @param OutPointOnBody	Point on the surface of collision closest to Point
 	 * 
-	 * @return		Success if returns > 0.f, if returns 0.f, it is either not convex or inside of the point
-	 *				If returns < 0.f, this primitive does not have collsion
+	 * @return		true if a distance to the body was found and OutDistanceSquared has been populated
 	 */
-	virtual float GetDistanceToCollision(const FVector& Point, FVector& ClosestPointOnCollision) const;
+	virtual bool GetSquaredDistanceToCollision(const FVector& Point, float& OutSquaredDistance, FVector& OutClosestPointOnCollision) const;
+
+	/** 
+	* returns Distance to closest Body Instance surface. 
+	*
+	* @param Point				World 3D vector
+	* @param OutPointOnBody	Point on the surface of collision closest to Point
+	* 
+	* @return		Success if returns > 0.f, if returns 0.f, point is inside the geometry
+	*				If returns < 0.f, this primitive does not have collsion or if geometry is not supported
+	*/	
+	float GetDistanceToCollision(const FVector& Point, FVector& ClosestPointOnCollision) const 
+	{
+		float DistanceSqr = -1.f;
+		return (GetSquaredDistanceToCollision(Point, DistanceSqr, ClosestPointOnCollision) ? FMath::Sqrt(DistanceSqr) : -1.f);
+	}
 
 	/**
 	* Returns the distance and closest point to the collision surface.
diff --git a/Engine/Source/Runtime/Engine/Classes/Components/SkeletalMeshComponent.h b/Engine/Source/Runtime/Engine/Classes/Components/SkeletalMeshComponent.h
--- a/Engine/Source/Runtime/Engine/Classes/Components/SkeletalMeshComponent.h
+++ b/Engine/Source/Runtime/Engine/Classes/Components/SkeletalMeshComponent.h
@@ -4,8 +4,6 @@
 
 #include "Interfaces/Interface_CollisionDataProvider.h"
 #include "Components/SkinnedMeshComponent.h"
-#include "PhysicsEngine/PhysicsConstraintComponent.h"
-#include "SkeletalMeshTypes.h"
 #include "AnimCurveTypes.h"
 #include "ClothSimData.h"
 #include "SingleAnimationPlayData.h"
@@ -18,6 +16,7 @@ struct FEngineShowFlags;
 struct FConvexVolume;
 struct FClothingAssetData;
 struct FRootMotionMovementParams;
+struct FApexClothCollisionVolumeData;
 
 DECLARE_MULTICAST_DELEGATE(FOnSkelMeshPhysicsCreatedMultiCast);
 typedef FOnSkelMeshPhysicsCreatedMultiCast::FDelegate FOnSkelMeshPhysicsCreated;
@@ -813,12 +812,15 @@ public:
 	#endif // WITH_CLOTH_COLLISION_DETECTION
 #endif // WITH_APEX_CLOTHING
 
+private:
 	/** 
 	 * Morph Target Curves. This will override AnimInstance MorphTargetCurves
 	 * if same curve is found
 	 **/
 	TMap<FName, float>	MorphTargetCurves;
 
+public:
+	const TMap<FName, float>& GetMorphTargetCurves() const { return MorphTargetCurves;  }
 	// 
 	// Animation
 	//
@@ -952,7 +954,7 @@ public:
 	virtual bool IsGravityEnabled() const override;
 	virtual void OnComponentCollisionSettingsChanged() override;
 	virtual void SetPhysMaterialOverride(UPhysicalMaterial* NewPhysMaterial) override;
-	virtual float GetDistanceToCollision(const FVector& Point, FVector& ClosestPointOnCollision) const override;
+	virtual bool GetSquaredDistanceToCollision(const FVector& Point, float& OutSquaredDistance, FVector& OutClosestPointOnCollision) const override;
 
 
 	/** 
@@ -1452,7 +1454,7 @@ public:
 	float LastPoseTickTime;
 
 	/** Checked whether we have already ticked the pose this frame */
-	bool PoseTickedThisFrame() const { return LastPoseTickTime == GetWorld()->TimeSeconds; }
+	bool PoseTickedThisFrame() const;
 
 	/** Take extracted RootMotion and convert it from local space to world space. */
 	FTransform ConvertLocalRootMotionToWorld(const FTransform& InTransform);
diff --git a/Engine/Source/Runtime/Engine/Classes/Components/SkinnedMeshComponent.h b/Engine/Source/Runtime/Engine/Classes/Components/SkinnedMeshComponent.h
--- a/Engine/Source/Runtime/Engine/Classes/Components/SkinnedMeshComponent.h
+++ b/Engine/Source/Runtime/Engine/Classes/Components/SkinnedMeshComponent.h
@@ -9,7 +9,7 @@
 // Forward declarations
 //
 class FSkeletalMeshResource;
-struct FSkelMeshChunk;
+struct FSkelMeshSection;
 class FSkeletalMeshVertexBuffer;
 
 DECLARE_DELEGATE_OneParam(FOnAnimUpdateRateParamsCreated, FAnimUpdateRateParameters*)
@@ -927,7 +927,7 @@ private:
 	 * Simple, CPU evaluation of a vertex's skinned position helper function
 	 */
 	template <bool bExtraBoneInfluencesT, bool bCachedMatrices>
-	FVector GetTypedSkinnedVertexPosition(const FSkelMeshChunk& Chunk, const FSkeletalMeshVertexBuffer& VertexBufferGPUSkin, int32 VertIndex, bool bSoftVertex, const TArray<FMatrix> & RefToLocals = TArray<FMatrix>()) const;
+	FVector GetTypedSkinnedVertexPosition(const FSkelMeshSection& Section, const FSkeletalMeshVertexBuffer& VertexBufferGPUSkin, int32 VertIndex, const TArray<FMatrix> & RefToLocals = TArray<FMatrix>()) const;
 
 	/**
 	 * Gets called when register, verifies all ActiveMorphTargets is still valid 
diff --git a/Engine/Source/Runtime/Engine/Classes/Distributions/DistributionFloat.h b/Engine/Source/Runtime/Engine/Classes/Distributions/DistributionFloat.h
--- a/Engine/Source/Runtime/Engine/Classes/Distributions/DistributionFloat.h
+++ b/Engine/Source/Runtime/Engine/Classes/Distributions/DistributionFloat.h
@@ -173,6 +173,7 @@ public:
 #endif	// WITH_EDITOR
 	virtual bool NeedsLoadForClient() const override;
 	virtual bool NeedsLoadForServer() const override;
+	virtual bool NeedsLoadForEditorGame() const override;
 	virtual void Serialize(FArchive& Ar) override;
 	/** End UObject interface */
 };
diff --git a/Engine/Source/Runtime/Engine/Classes/Distributions/DistributionVector.h b/Engine/Source/Runtime/Engine/Classes/Distributions/DistributionVector.h
--- a/Engine/Source/Runtime/Engine/Classes/Distributions/DistributionVector.h
+++ b/Engine/Source/Runtime/Engine/Classes/Distributions/DistributionVector.h
@@ -217,6 +217,7 @@ public:
 #endif // WITH_EDITOR
 	virtual bool NeedsLoadForClient() const override;
 	virtual bool NeedsLoadForServer() const override;
+	virtual bool NeedsLoadForEditorGame() const override;
 	virtual void Serialize(FArchive& Ar) override;
 	/** End UObject interface */
 
diff --git a/Engine/Source/Runtime/Engine/Classes/Engine/EngineTypes.h b/Engine/Source/Runtime/Engine/Classes/Engine/EngineTypes.h
--- a/Engine/Source/Runtime/Engine/Classes/Engine/EngineTypes.h
+++ b/Engine/Source/Runtime/Engine/Classes/Engine/EngineTypes.h
@@ -2770,7 +2770,7 @@ enum class ERotatorQuantization : uint8
   * and velocity.Z is commonly zero (most position replications are for walking pawns). 
   */
 USTRUCT()
-struct FRepMovement
+struct ENGINE_API FRepMovement
 {
 	GENERATED_USTRUCT_BODY()
 
@@ -3310,6 +3310,8 @@ struct FComponentSocketDescription
 	}
 };
 
+/** Dynamic delegate to use by components that want to route the broken-event into blueprints */
+DECLARE_DYNAMIC_MULTICAST_DELEGATE_OneParam(FConstraintBrokenSignature, int32, ConstraintIndex);
 
 // ANGULAR DOF
 UENUM()
@@ -3325,6 +3327,17 @@ enum EAngularConstraintMotion
 	ACM_MAX,
 };
 
+/** Enum to indicate which frame we want. */
+UENUM()
+namespace EConstraintFrame
+{
+	enum Type
+	{
+		Frame1,
+		Frame2
+	};
+}
+
 
 /**
  * Structure for file paths that are displayed in the UI.
diff --git a/Engine/Source/Runtime/Engine/Classes/Engine/LocalPlayer.h b/Engine/Source/Runtime/Engine/Classes/Engine/LocalPlayer.h
--- a/Engine/Source/Runtime/Engine/Classes/Engine/LocalPlayer.h
+++ b/Engine/Source/Runtime/Engine/Classes/Engine/LocalPlayer.h
@@ -8,8 +8,10 @@
 #include "SlateCore.h"
 #include "Reply.h"
 #include "Player.h"
+#include "Engine/GameViewportClient.h"
 #include "LocalPlayer.generated.h"
 
+
 #define INVALID_CONTROLLERID 255
 
 class FUniqueNetId;
diff --git a/Engine/Source/Runtime/Engine/Classes/Engine/World.h b/Engine/Source/Runtime/Engine/Classes/Engine/World.h
--- a/Engine/Source/Runtime/Engine/Classes/Engine/World.h
+++ b/Engine/Source/Runtime/Engine/Classes/Engine/World.h
@@ -2573,14 +2573,31 @@ public:
 	 * @param NetDriverName the name of the net driver being asked for
 	 * @return a pointer to the net driver or NULL if the named driver is not found
 	 */
-	UNetDriver* GetNetDriver() const
+	FORCEINLINE_DEBUGGABLE UNetDriver* GetNetDriver() const
 	{
 		return NetDriver;
 	}
 
-	/** Returns the net mode this world is running under */
+	/**
+	 * Returns the net mode this world is running under.
+	 * @see IsNetMode()
+	 */
 	ENetMode GetNetMode() const;
 
+	/**
+	* Test whether net mode is the given mode.
+	* In optimized non-editor builds this can be more efficient than GetNetMode()
+	* because it can check the static build flags without considering PIE.
+	*/
+	bool IsNetMode(ENetMode Mode) const;
+
+private:
+
+	/** Private version without inlining that does *not* check Dedicated server build flags (which should already have been done). */
+	ENetMode InternalGetNetMode() const;
+
+public:
+
 #if WITH_EDITOR
 	/** Attempts to derive the net mode from PlayInSettings for PIE*/
 	ENetMode AttemptDeriveFromPlayInSettings() const;
@@ -2999,4 +3016,32 @@ FORCEINLINE_DEBUGGABLE bool UWorld::ComponentSweepMulti(TArray<struct FHitResult
 	return ComponentSweepMulti(OutHits, PrimComp, Start, End, Rot.Quaternion(), Params);
 }
 
+FORCEINLINE_DEBUGGABLE ENetMode UWorld::GetNetMode() const
+{
+	// IsRunningDedicatedServer() is a compile-time check in optimized non-editor builds.
+	if (IsRunningDedicatedServer())
+	{
+		return NM_DedicatedServer;
+	}
+
+	return InternalGetNetMode();
+}
+
+FORCEINLINE_DEBUGGABLE bool UWorld::IsNetMode(ENetMode Mode) const
+{
+#if UE_EDITOR
+	// Editor builds are special because of PIE, which can run a dedicated server without the app running with -server.
+	return GetNetMode() == Mode;
+#else
+	// IsRunningDedicatedServer() is a compile-time check in optimized non-editor builds.
+	if (Mode == NM_DedicatedServer)
+	{
+		return IsRunningDedicatedServer();
+	}
+	else
+	{
+		return !IsRunningDedicatedServer() && (InternalGetNetMode() == Mode);
+	}
+#endif
+}
 
diff --git a/Engine/Source/Runtime/Engine/Classes/GameFramework/Actor.h b/Engine/Source/Runtime/Engine/Classes/GameFramework/Actor.h
--- a/Engine/Source/Runtime/Engine/Classes/GameFramework/Actor.h
+++ b/Engine/Source/Runtime/Engine/Classes/GameFramework/Actor.h
@@ -2144,10 +2144,18 @@ public:
 	virtual bool DestroyNetworkActorHandled();
 
 	/**
-	 * Gets the net mode for this actor, indicating whether it is a client or server (including standalone/not networked).
+	 * Get the network mode (dedicated server, client, standalone, etc) for this actor.
+	 * @see IsNetMode()
 	 */
 	ENetMode GetNetMode() const;
 
+	/**
+	* Test whether net mode is the given mode.
+	* In optimized non-editor builds this can be more efficient than GetNetMode()
+	* because it can check the static build flags without considering PIE.
+	*/
+	bool IsNetMode(ENetMode Mode) const;
+
 	class UNetDriver * GetNetDriver() const;
 
 	/** Puts actor in dormant networking state */
@@ -2854,6 +2862,9 @@ private:
 	// Helper that already assumes the Hit info is reversed, and avoids creating a temp FHitResult if possible.
 	void InternalDispatchBlockingHit(UPrimitiveComponent* MyComp, UPrimitiveComponent* OtherComp, bool bSelfMoved, FHitResult const& Hit);
 
+	/** Private version without inlining that does *not* check Dedicated server build flags (which should already have been done). */
+	ENetMode InternalGetNetMode() const;
+
 	friend struct FMarkActorIsBeingDestroyed;
 	friend struct FActorParentComponentSetter;
 };
@@ -2967,6 +2978,36 @@ FORCEINLINE_DEBUGGABLE ENetRole AActor::GetRemoteRole() const
 	return RemoteRole;
 }
 
+FORCEINLINE_DEBUGGABLE ENetMode AActor::GetNetMode() const
+{
+	// IsRunningDedicatedServer() is a compile-time check in optimized non-editor builds.
+	if (IsRunningDedicatedServer())
+	{
+		return NM_DedicatedServer;
+	}
+
+	return InternalGetNetMode();
+}
+
+FORCEINLINE_DEBUGGABLE bool AActor::IsNetMode(ENetMode Mode) const
+{
+#if UE_EDITOR
+	// Editor builds are special because of PIE, which can run a dedicated server without the app running with -server.
+	return GetNetMode() == Mode;
+#else
+	// IsRunningDedicatedServer() is a compile-time check in optimized non-editor builds.
+	if (Mode == NM_DedicatedServer)
+	{
+		return IsRunningDedicatedServer();
+	}
+	else
+	{
+		return !IsRunningDedicatedServer() && (InternalGetNetMode() == Mode);
+	}
+#endif
+}
+
+
 FORCEINLINE_DEBUGGABLE void AActor::SetNetUpdateTime(float NewUpdateTime)
 {
 	NetUpdateTime = NewUpdateTime;
diff --git a/Engine/Source/Runtime/Engine/Classes/Materials/Material.h b/Engine/Source/Runtime/Engine/Classes/Materials/Material.h
--- a/Engine/Source/Runtime/Engine/Classes/Materials/Material.h
+++ b/Engine/Source/Runtime/Engine/Classes/Materials/Material.h
@@ -1202,18 +1202,18 @@ public:
 	 * Backs up all material shaders to memory through serialization, organized by FMaterialShaderMap. 
 	 * This will also clear all FMaterialShaderMap references to FShaders.
 	 */
-	ENGINE_API static void BackupMaterialShadersToMemory(TMap<FMaterialShaderMap*, TScopedPointer<TArray<uint8> > >& ShaderMapToSerializedShaderData);
+	ENGINE_API static void BackupMaterialShadersToMemory(TMap<class FMaterialShaderMap*, TScopedPointer<TArray<uint8> > >& ShaderMapToSerializedShaderData);
 
 	/** 
 	 * Recreates FShaders for FMaterialShaderMap's from the serialized data.  Shader maps may not be complete after this due to changes in the shader keys.
 	 */
-	ENGINE_API static void RestoreMaterialShadersFromMemory(const TMap<FMaterialShaderMap*, TScopedPointer<TArray<uint8> > >& ShaderMapToSerializedShaderData);
+	ENGINE_API static void RestoreMaterialShadersFromMemory(const TMap<class FMaterialShaderMap*, TScopedPointer<TArray<uint8> > >& ShaderMapToSerializedShaderData);
 
 	/** Builds a map from UMaterialInterface name to the shader maps that are needed for rendering on the given platform. */
 	ENGINE_API static void CompileMaterialsForRemoteRecompile(
 		const TArray<UMaterialInterface*>& MaterialsToCompile,
 		EShaderPlatform ShaderPlatform, 
-		TMap<FString, TArray<TRefCountPtr<FMaterialShaderMap> > >& OutShaderMaps);
+		TMap<FString, TArray<TRefCountPtr<class FMaterialShaderMap> > >& OutShaderMaps);
 
 	/**
 	 * Add an expression node that represents a parameter to the list of material parameters.
diff --git a/Engine/Source/Runtime/Engine/Classes/Particles/Location/ParticleModuleLocationSkelVertSurface.h b/Engine/Source/Runtime/Engine/Classes/Particles/Location/ParticleModuleLocationSkelVertSurface.h
--- a/Engine/Source/Runtime/Engine/Classes/Particles/Location/ParticleModuleLocationSkelVertSurface.h
+++ b/Engine/Source/Runtime/Engine/Classes/Particles/Location/ParticleModuleLocationSkelVertSurface.h
@@ -160,5 +160,5 @@ class ENGINE_API UParticleModuleLocationSkelVertSurface : public UParticleModule
 private:
 	/** Helper function for concrete types. */
 	template<bool bExtraBoneInfluencesT>
-	bool VertInfluencedByActiveBoneTyped(bool bSoftVertex, FStaticLODModel& Model, const FSkelMeshChunk& Chunk, int32 VertIndex, USkeletalMeshComponent* InSkelMeshComponent, FModuleLocationVertSurfaceInstancePayload* InstancePayload, int32* OutBoneIndex);
+	bool VertInfluencedByActiveBoneTyped(FStaticLODModel& Model, const FSkelMeshSection& Section, int32 VertIndex, USkeletalMeshComponent* InSkelMeshComponent, FModuleLocationVertSurfaceInstancePayload* InstancePayload, int32* OutBoneIndex);
 };
diff --git a/Engine/Source/Runtime/Engine/Classes/Particles/TypeData/ParticleModuleTypeDataGpu.h b/Engine/Source/Runtime/Engine/Classes/Particles/TypeData/ParticleModuleTypeDataGpu.h
--- a/Engine/Source/Runtime/Engine/Classes/Particles/TypeData/ParticleModuleTypeDataGpu.h
+++ b/Engine/Source/Runtime/Engine/Classes/Particles/TypeData/ParticleModuleTypeDataGpu.h
@@ -403,11 +403,11 @@ class UParticleModuleTypeDataGpu : public UParticleModuleTypeDataBase
 	GENERATED_UCLASS_BODY()
 
 	/** Information for runtime simulation. */
-	UPROPERTY(transient)
+	UPROPERTY()
 	struct FGPUSpriteEmitterInfo EmitterInfo;
 
 	/** Data used to initialize runtime resources. */
-	UPROPERTY(transient)
+	UPROPERTY()
 	struct FGPUSpriteResourceData ResourceData;
 
 	/** TEMP: How much to stretch sprites based on camera motion blur. */
diff --git a/Engine/Source/Runtime/Engine/Classes/PhysicsEngine/BodyInstance.h b/Engine/Source/Runtime/Engine/Classes/PhysicsEngine/BodyInstance.h
--- a/Engine/Source/Runtime/Engine/Classes/PhysicsEngine/BodyInstance.h
+++ b/Engine/Source/Runtime/Engine/Classes/PhysicsEngine/BodyInstance.h
@@ -1015,11 +1015,24 @@ public:
 	/**
 	 * Get distance to the body surface if available
 	 * It is only valid if BodyShape is convex
-	 * If point is inside or shape is not convex, it will return 0.f
+	 * If point is inside distance it will be 0
+	 * Returns false if geometry is not supported
 	 *
 	 * @param Point				Point in world space
+	 * @param OutDistanceSquared How far from the instance the point is. 0 if inside the shape
 	 * @param OutPointOnBody	Point on the surface of body closest to Point
+	 * @return true if a distance to the body was found and OutDistanceSquared has been populated
 	 */
+	bool GetSquaredDistanceToBody(const FVector& Point, float& OutDistanceSquared, FVector& OutPointOnBody) const;
+
+	/**
+	* Get the square of the distance to the body surface if available
+	* It is only valid if BodyShape is convex
+	* If point is inside or shape is not convex, it will return 0.f
+	*
+	* @param Point				Point in world space
+	* @param OutPointOnBody	Point on the surface of body closest to Point
+	*/
 	float GetDistanceToBody(const FVector& Point, FVector& OutPointOnBody) const;
 
 	/** 
diff --git a/Engine/Source/Runtime/Engine/Classes/PhysicsEngine/ConstraintInstance.h b/Engine/Source/Runtime/Engine/Classes/PhysicsEngine/ConstraintInstance.h
--- a/Engine/Source/Runtime/Engine/Classes/PhysicsEngine/ConstraintInstance.h
+++ b/Engine/Source/Runtime/Engine/Classes/PhysicsEngine/ConstraintInstance.h
@@ -26,17 +26,6 @@ enum ELinearConstraintMotion
 	LCM_MAX,
 };
 
-/** Enum to indicate which frame we want. */
-UENUM()
-namespace EConstraintFrame
-{
-	enum Type
-	{
-		Frame1,
-		Frame2
-	};
-}
-
 UENUM()
 namespace EAngularDriveMode
 {
diff --git a/Engine/Source/Runtime/Engine/Classes/PhysicsEngine/PhysicsConstraintComponent.h b/Engine/Source/Runtime/Engine/Classes/PhysicsEngine/PhysicsConstraintComponent.h
--- a/Engine/Source/Runtime/Engine/Classes/PhysicsEngine/PhysicsConstraintComponent.h
+++ b/Engine/Source/Runtime/Engine/Classes/PhysicsEngine/PhysicsConstraintComponent.h
@@ -3,9 +3,6 @@
 #include "PhysicsEngine/ConstraintInstance.h"
 #include "PhysicsConstraintComponent.generated.h"
 
-/** Dynamic delegate to use by components that want to route the broken-event into blueprints */
-DECLARE_DYNAMIC_MULTICAST_DELEGATE_OneParam(FConstraintBrokenSignature, int32, ConstraintIndex);
-
 /**
  *	This is effectively a joint that allows you to connect 2 rigid bodies together. You can create different types of joints using the various parameters of this component.
  */
diff --git a/Engine/Source/Runtime/Engine/Private/AI/Navigation/NavLinkCustomComponent.cpp b/Engine/Source/Runtime/Engine/Private/AI/Navigation/NavLinkCustomComponent.cpp
--- a/Engine/Source/Runtime/Engine/Private/AI/Navigation/NavLinkCustomComponent.cpp
+++ b/Engine/Source/Runtime/Engine/Private/AI/Navigation/NavLinkCustomComponent.cpp
@@ -9,6 +9,7 @@
 #include "AI/Navigation/NavLinkCustomComponent.h"
 #include "AI/Navigation/NavAreas/NavArea_Default.h"
 #include "AI/Navigation/NavAreas/NavArea_Null.h"
+#include "AI/NavigationOctree.h"
 
 UNavLinkCustomComponent::UNavLinkCustomComponent(const FObjectInitializer& ObjectInitializer) : Super(ObjectInitializer)
 {
diff --git a/Engine/Source/Runtime/Engine/Private/AI/Navigation/NavLinkProxy.cpp b/Engine/Source/Runtime/Engine/Private/AI/Navigation/NavLinkProxy.cpp
--- a/Engine/Source/Runtime/Engine/Private/AI/Navigation/NavLinkProxy.cpp
+++ b/Engine/Source/Runtime/Engine/Private/AI/Navigation/NavLinkProxy.cpp
@@ -11,6 +11,7 @@
 #include "AI/Navigation/NavLinkRenderingComponent.h"
 #include "NavigationSystemHelpers.h"
 #include "VisualLogger/VisualLogger.h"
+#include "AI/NavigationOctree.h"
 
 ANavLinkProxy::ANavLinkProxy(const FObjectInitializer& ObjectInitializer) : Super(ObjectInitializer)
 {
diff --git a/Engine/Source/Runtime/Engine/Private/AI/Navigation/NavLinkRenderingComponent.cpp b/Engine/Source/Runtime/Engine/Private/AI/Navigation/NavLinkRenderingComponent.cpp
--- a/Engine/Source/Runtime/Engine/Private/AI/Navigation/NavLinkRenderingComponent.cpp
+++ b/Engine/Source/Runtime/Engine/Private/AI/Navigation/NavLinkRenderingComponent.cpp
@@ -9,6 +9,7 @@
 #endif
 #include "AI/Navigation/NavLinkHostInterface.h"
 #include "AI/Navigation/NavAreas/NavArea.h"
+#include "AI/Navigation/RecastNavMesh.h"
 
 //----------------------------------------------------------------------//
 // UNavLinkRenderingComponent
diff --git a/Engine/Source/Runtime/Engine/Private/AI/Navigation/NavTestRenderingComponent.cpp b/Engine/Source/Runtime/Engine/Private/AI/Navigation/NavTestRenderingComponent.cpp
--- a/Engine/Source/Runtime/Engine/Private/AI/Navigation/NavTestRenderingComponent.cpp
+++ b/Engine/Source/Runtime/Engine/Private/AI/Navigation/NavTestRenderingComponent.cpp
@@ -5,6 +5,7 @@
 #include "AI/Navigation/NavTestRenderingComponent.h"
 #include "DynamicMeshBuilder.h"
 #include "AI/Navigation/NavigationTestingActor.h"
+#include "AI/Navigation/RecastNavMesh.h"
 #include "Debug/DebugDrawService.h"
 
 static const FColor NavMeshRenderColor_OpenSet(255,128,0,255);
diff --git a/Engine/Source/Runtime/Engine/Private/AI/Navigation/NavigationPath.cpp b/Engine/Source/Runtime/Engine/Private/AI/Navigation/NavigationPath.cpp
--- a/Engine/Source/Runtime/Engine/Private/AI/Navigation/NavigationPath.cpp
+++ b/Engine/Source/Runtime/Engine/Private/AI/Navigation/NavigationPath.cpp
@@ -5,6 +5,7 @@
 #include "AI/Navigation/NavigationTypes.h"
 #include "AI/Navigation/RecastNavMesh.h"
 #include "AI/Navigation/NavigationPath.h"
+#include "AI/Navigation/NavAreas/NavArea.h"
 #include "Debug/DebugDrawService.h"
 
 #define DEBUG_DRAW_OFFSET 0
@@ -900,10 +901,11 @@ void FNavMeshPath::DebugDraw(const ANavigationData* NavData, FColor PathColor, U
 	const ARecastNavMesh* RecastNavMesh = Cast<const ARecastNavMesh>(NavData);		
 	const TArray<FNavigationPortalEdge>& Edges = (const_cast<FNavMeshPath*>(this))->GetPathCorridorEdges();	
 	const int32 CorridorEdgesCount = Edges.Num();
+	const UWorld* World = NavData->GetWorld();
 
 	for (int32 EdgeIndex = 0; EdgeIndex < CorridorEdgesCount; ++EdgeIndex)
 	{
-		DrawDebugLine(NavData->GetWorld(), Edges[EdgeIndex].Left + NavigationDebugDrawing::PathOffset, Edges[EdgeIndex].Right + NavigationDebugDrawing::PathOffset
+		DrawDebugLine(World, Edges[EdgeIndex].Left + NavigationDebugDrawing::PathOffset, Edges[EdgeIndex].Right + NavigationDebugDrawing::PathOffset
 			, FColor::Blue, bPersistent, /*LifeTime*/-1.f, /*DepthPriority*/0
 			, /*Thickness*/NavigationDebugDrawing::PathLineThickness);
 	}
