Copying //UE4/Dev-Blueprints to //UE4/Dev-Main (Source: //UE4/Dev-Blueprints @ 3042808)

diff --git a/Engine/Source/Runtime/CoreUObject/Private/Blueprint/BlueprintSupport.cpp b/Engine/Source/Runtime/CoreUObject/Private/Blueprint/BlueprintSupport.cpp
--- a/Engine/Source/Runtime/CoreUObject/Private/Blueprint/BlueprintSupport.cpp
+++ b/Engine/Source/Runtime/CoreUObject/Private/Blueprint/BlueprintSupport.cpp
@@ -1826,9 +1826,21 @@ bool FObjectInitializer::InitNonNativeProperty(UProperty* Property, UObject* Dat
 FObjectInitializer* FDeferredObjInitializerTracker::Add(const FObjectInitializer& DeferringInitializer)
 {
 	UObject* InitingObj = DeferringInitializer.GetObj();
-	DEFERRED_DEPENDENCY_CHECK((InitingObj != nullptr) && InitingObj->HasAnyFlags(RF_ClassDefaultObject));
+	DEFERRED_DEPENDENCY_CHECK(InitingObj != nullptr);
+	const bool bIsSubObjTemplate = InitingObj && InitingObj->HasAnyFlags(RF_InheritableComponentTemplate);
 
-	UClass* LoadClass = (InitingObj != nullptr) ? InitingObj->GetClass() : nullptr;
+	UClass* LoadClass = nullptr;
+	if (bIsSubObjTemplate)
+	{
+		LoadClass = Cast<UClass>(InitingObj->GetOuter());
+	}
+	else if (InitingObj != nullptr)
+	{
+		DEFERRED_DEPENDENCY_CHECK(InitingObj->HasAnyFlags(RF_ClassDefaultObject));
+		LoadClass = InitingObj->GetClass();
+	}
+	
+	FObjectInitializer* DeferredCopy = nullptr;
 	if (LoadClass != nullptr)
 	{
 		FDeferredObjInitializerTracker& ThreadInst = FDeferredObjInitializerTracker::Get();
@@ -1840,18 +1852,35 @@ FObjectInitializer* FDeferredObjInitializerTracker::Add(const FObjectInitializer
 		UObject* SuperCDO = SuperClass->GetDefaultObject(/*bCreateIfNeeded =*/false);
 		DEFERRED_DEPENDENCY_CHECK( SuperCDO && (SuperCDO->HasAnyFlags(RF_NeedLoad) ||
 			(SuperClass->GetLinker() && SuperClass->GetLinker()->IsBlueprintFinalizationPending()) || IsCdoDeferred(SuperClass)) );
+
+		FLinkerLoad* ClassLinker = LoadClass->GetLinker();
+		DEFERRED_DEPENDENCY_CHECK((bIsSubObjTemplate && ClassLinker->IsBlueprintFinalizationPending()) || 
+			(ClassLinker->LoadFlags & LOAD_DeferDependencyLoads) != 0);
 #endif // USE_DEFERRED_DEPENDENCY_CHECK_VERIFICATION_TESTS
 		
-		auto& InitializersCache = ThreadInst.DeferredInitializers;
-		DEFERRED_DEPENDENCY_CHECK(InitializersCache.Find(LoadClass) == nullptr); // did we try to init the CDO twice?
-		DEFERRED_DEPENDENCY_CHECK(LoadClass->GetLinker()->LoadFlags & LOAD_DeferDependencyLoads);
-
-		// NOTE: we copy the FObjectInitializer, because it is most likely being destroyed
-		FObjectInitializer& DeferredCopy = InitializersCache.Add(LoadClass, DeferringInitializer);
+		if (bIsSubObjTemplate)
+		{
+#if USE_DEFERRED_DEPENDENCY_CHECK_VERIFICATION_TESTS
+			TArray<FObjectInitializer*> DeferredSubObjInitializers;
+			ThreadInst.DeferredSubObjInitializers.MultiFindPointer(LoadClass, DeferredSubObjInitializers);
+			// check to make sure that we haven't already added this one for deferral
+			for (FObjectInitializer* SubObjInitializer : DeferredSubObjInitializers)
+			{
+				DEFERRED_DEPENDENCY_CHECK(SubObjInitializer->GetObj() != DeferringInitializer.GetObj());
+			}
+#endif // USE_DEFERRED_DEPENDENCY_CHECK_VERIFICATION_TESTS
+			DeferredCopy = &ThreadInst.DeferredSubObjInitializers.Add(LoadClass, DeferringInitializer);
+		}
+		else
+		{
+			auto& InitializersCache = ThreadInst.DeferredInitializers;
+			DEFERRED_DEPENDENCY_CHECK(InitializersCache.Find(LoadClass) == nullptr); // did we try to init the CDO twice?
 
-		return &DeferredCopy;
+			// NOTE: we copy the FObjectInitializer, because it is most likely being destroyed
+			DeferredCopy = &InitializersCache.Add(LoadClass, DeferringInitializer);
+		}
 	}
-	return nullptr;
+	return DeferredCopy;
 }
 
 FObjectInitializer* FDeferredObjInitializerTracker::Find(UClass* LoadClass)
@@ -1867,29 +1896,67 @@ bool FDeferredObjInitializerTracker::IsCdoDeferred(UClass* LoadClass)
 
 bool FDeferredObjInitializerTracker::DeferSubObjectPreload(UObject* SubObject)
 {
-	DEFERRED_DEPENDENCY_CHECK(SubObject->HasAnyFlags(RF_DefaultSubObject));
-
-	UObject* CdoOuter = SubObject->GetOuter();
-	UClass* OuterClass = CdoOuter->GetClass();
+	bool bDeferralNeeded = false;
+	// if this is an "InheritableComponentTemplate" we expect it's outer to be
+	// the class it belongs to, and know that this overrides a inherited 
+	// component (belonging to the owner's super)
+	const bool bIsComponentOverride = SubObject->HasAnyFlags(RF_InheritableComponentTemplate);
 
-	FDeferredObjInitializerTracker& ThreadInst = FDeferredObjInitializerTracker::Get();
-	if (IsCdoDeferred(OuterClass) && (OuterClass != ThreadInst.ResolvingClass))
+	UClass* OwningClass = nullptr;
+	if (bIsComponentOverride)
 	{
-		DEFERRED_DEPENDENCY_CHECK(CdoOuter->HasAnyFlags(RF_ClassDefaultObject));
+		OwningClass = Cast<UClass>(SubObject->GetOuter());
+		DEFERRED_DEPENDENCY_CHECK(OwningClass != nullptr);
+	}
+	else
+	{
+		DEFERRED_DEPENDENCY_CHECK(SubObject->HasAnyFlags(RF_DefaultSubObject));
 
-		UObject* SubObjTemplate = SubObject->GetArchetype();
-		// check that this is an inherited sub-object (that it is defined on the 
-		// parent class)... we only need to defer Preload() for inherited 
-		// components because they're what gets filled out in the CDO's InitSubobjectProperties()
-		if (SubObjTemplate && (SubObjTemplate->GetOuter() != CdoOuter))
-		{
-			
-			ThreadInst.DeferredSubObjects.AddUnique(OuterClass, SubObject);
+		UObject* SubObjOuter = SubObject->GetOuter();
+		OwningClass = SubObjOuter->GetClass();
+		// NOTE: the outer of a DSO may not be a CDO like we want, it could 
+		//       be something like a component template; right now we ignore 
+		//       those cases (IsCdoDeferred() below will reject this), but if 
+		//       this case proves to be a problem, then we may need to look up 
+		//       the outer chain, or see if the outer sub-obj is deferred itself
+	}
 
-			return true;
+	FDeferredObjInitializerTracker& ThreadInst = FDeferredObjInitializerTracker::Get();	
+	if (IsCdoDeferred(OwningClass) && (OwningClass != ThreadInst.ResolvingClass))
+	{
+		// no need to check the archetype, we know that this is overriding a
+		// super component - we need to defer Preload()
+		if (bIsComponentOverride)
+		{
+			ThreadInst.DeferredSubObjects.AddUnique(OwningClass, SubObject);
+			bDeferralNeeded = true;
 		}
+		else
+		{
+			DEFERRED_DEPENDENCY_CHECK(SubObject->GetOuter()->HasAnyFlags(RF_ClassDefaultObject));
+
+			UObject* SubObjTemplate = SubObject->GetArchetype();
+			// check that this is an inherited sub-object (that it is defined on the 
+			// parent class)... we only need to defer Preload() for inherited 
+			// components because they're what gets filled out in the CDO's InitSubobjectProperties()
+			if (SubObjTemplate && (SubObjTemplate->GetOuter() != SubObject->GetOuter()))
+			{
+				ThreadInst.DeferredSubObjects.AddUnique(OwningClass, SubObject);
+				bDeferralNeeded = true;
+			}
+		}		
 	}
-	return false;
+
+#if USE_DEFERRED_DEPENDENCY_CHECK_VERIFICATION_TESTS
+	TArray<FObjectInitializer*> DeferredSubObjInitializers;
+	ThreadInst.DeferredSubObjInitializers.MultiFindPointer(OwningClass, DeferredSubObjInitializers);
+	// check to make sure that we haven't already added this one for deferral
+	for (FObjectInitializer* SubObjInitializer : DeferredSubObjInitializers)
+	{
+		DEFERRED_DEPENDENCY_CHECK(bDeferralNeeded || SubObjInitializer->GetObj() != SubObject);
+	}
+#endif // USE_DEFERRED_DEPENDENCY_CHECK_VERIFICATION_TESTS
+	return bDeferralNeeded;
 }
 
 void FDeferredObjInitializerTracker::Remove(UClass* LoadClass)
@@ -1898,6 +1965,7 @@ void FDeferredObjInitializerTracker::Remove(UClass* LoadClass)
 	ThreadInst.DeferredInitializers.Remove(LoadClass);
 	ThreadInst.DeferredSubObjects.Remove(LoadClass);
 	ThreadInst.SuperClassMap.RemoveSingle(LoadClass->GetSuperClass(), LoadClass);
+	ThreadInst.DeferredSubObjInitializers.Remove(LoadClass);
 }
 
 bool FDeferredObjInitializerTracker::ResolveDeferredInitialization(UClass* LoadClass)
@@ -1917,6 +1985,12 @@ bool FDeferredObjInitializerTracker::ResolveDeferredInitialization(UClass* LoadC
 
 		return true;
 	}
+	else
+	{
+		// make sure we're not missing sub-obj initializers that would have been 
+		// ran in the above if case
+		DEFERRED_DEPENDENCY_CHECK(FDeferredObjInitializerTracker::Get().DeferredSubObjInitializers.Find(LoadClass) == nullptr);
+	}
 	return false;
 }
 
@@ -1926,11 +2000,20 @@ void FDeferredObjInitializerTracker::ResolveDeferredSubObjects(UObject* CDO)
 	UClass* LoadClass = CDO->GetClass();
 
 	FDeferredObjInitializerTracker& ThreadInst = FDeferredObjInitializerTracker::Get();
-	TArray<UObject*> DeferredSubObjects;
-	ThreadInst.DeferredSubObjects.MultiFind(LoadClass, DeferredSubObjects);
-
 	// guard to keep sub-object preloads from ending up back in DeferSubObjectPreload()
 	TGuardValue<UClass*> ResolvingGuard(ThreadInst.ResolvingClass, LoadClass);
+	
+	TArray<FObjectInitializer*> DeferredSubObjInitializers;
+	ThreadInst.DeferredSubObjInitializers.MultiFindPointer(LoadClass, DeferredSubObjInitializers);
+
+	for (FObjectInitializer* DeferredInitializer : DeferredSubObjInitializers)
+	{
+		FScriptIntegrationObjectHelper::PostConstructInitObject(*DeferredInitializer);
+	}
+	ThreadInst.DeferredSubObjInitializers.Remove(LoadClass);
+	
+	TArray<UObject*> DeferredSubObjects;
+	ThreadInst.DeferredSubObjects.MultiFind(LoadClass, DeferredSubObjects);
 
 	FLinkerLoad* ClassLinker = LoadClass->GetLinker();
 	DEFERRED_DEPENDENCY_CHECK(ClassLinker != nullptr);
@@ -1943,24 +2026,9 @@ void FDeferredObjInitializerTracker::ResolveDeferredSubObjects(UObject* CDO)
 		// skipped on account of the deferred CDO initialization
 		for (UObject* SubObj : DeferredSubObjects)
 		{
-			// no longer need to handle this case, because we're no longer 
-			// attempting to recreate the CDO object in 
-			// FLinkerLoad::ResolveDeferredExports() (when the super class has
-			// been regenerated)
-// 			if (SubObj->GetOuter() != CDO)
-// 			{
-// 				// we may have had to recreate the CDO (because its super has had 
-// 				// its layout altered through regeneration); see 
-// 				// FLinkerLoad::ResolveDeferredExports() for more context
-// 				//
-// 				// @TODO: Are we sure that DeferredSubObjects would have all the
-// 				//        sub-objects that we need to over?
-// 				// 
-// 				// @TODO: Are we sure these are the right rename clags to use?
-// 				const ERenameFlags RenameFlags = (REN_DontCreateRedirectors | REN_ForceNoResetLoaders | REN_NonTransactional | REN_DoNotDirty);
-// 				SubObj->Rename(nullptr, CDO, RenameFlags);
-// 			}
-			DEFERRED_DEPENDENCY_CHECK(SubObj->GetOuter() == CDO);
+			DEFERRED_DEPENDENCY_CHECK( (SubObj->GetOuter() == CDO && SubObj->HasAnyFlags(RF_DefaultSubObject)) || 
+				(SubObj->GetOuter() == LoadClass && SubObj->HasAnyFlags(RF_InheritableComponentTemplate)) );
+			
 			ClassLinker->Preload(SubObj);
 		}
 	}
diff --git a/Engine/Source/Runtime/CoreUObject/Private/Serialization/AsyncLoading.cpp b/Engine/Source/Runtime/CoreUObject/Private/Serialization/AsyncLoading.cpp
--- a/Engine/Source/Runtime/CoreUObject/Private/Serialization/AsyncLoading.cpp
+++ b/Engine/Source/Runtime/CoreUObject/Private/Serialization/AsyncLoading.cpp
@@ -2261,17 +2261,36 @@ EAsyncPackageState::Type FAsyncPackage::PostLoadDeferredObjects(double InTickSta
 		// Clear async loading flags (we still want RF_Async, but EInternalObjectFlags::AsyncLoading can be cleared)
 		for (UObject* Object : DeferredFinalizeObjects)
 		{
-			Object->AtomicallyClearInternalFlags(EInternalObjectFlags::AsyncLoading);
+			if (Object)
+			{
+				Object->AtomicallyClearInternalFlags(EInternalObjectFlags::AsyncLoading);
+			}
+
+			// CDO need special handling, no matter if it's listed in DeferredFinalizeObjects or created here for DynamicClass
+			UObject* CDOToHandle = nullptr;
+
 			// Dynamic Class doesn't require/use pre-loading (or post-loading). 
 			// The CDO is created at this point, because now it's safe to solve cyclic dependencies.
 			if (UDynamicClass* DynamicClass = Cast<UDynamicClass>(Object))
 			{
-				DynamicClass->GetDefaultObject(true);
+				UObject* OldCDO = DynamicClass->GetDefaultObject(false);
+				UObject* NewCDO = DynamicClass->GetDefaultObject(true);
+				const bool bCDOWasJustCreated = (OldCDO != NewCDO);
+				if (bCDOWasJustCreated && (NewCDO != nullptr))
+				{
+					NewCDO->AtomicallyClearInternalFlags(EInternalObjectFlags::AsyncLoading);
+					CDOToHandle = NewCDO;
+				}
+			}
+			else
+			{
+				CDOToHandle = ((Object != nullptr) && Object->HasAnyFlags(RF_ClassDefaultObject)) ? Object : nullptr;
 			}
 
-			if(Object->HasAnyFlags(RF_ClassDefaultObject))
+			// Clear AsyncLoading in CDO's subobjects.
+			if(CDOToHandle != nullptr)
 			{
-				Object->GetDefaultSubobjects(CDODefaultSubobjects);
+				CDOToHandle->GetDefaultSubobjects(CDODefaultSubobjects);
 				for (UObject* SubObject : CDODefaultSubobjects)
 				{
 					if (SubObject && SubObject->HasAnyInternalFlags(EInternalObjectFlags::AsyncLoading))
diff --git a/Engine/Source/Runtime/CoreUObject/Private/UObject/Class.cpp b/Engine/Source/Runtime/CoreUObject/Private/UObject/Class.cpp
--- a/Engine/Source/Runtime/CoreUObject/Private/UObject/Class.cpp
+++ b/Engine/Source/Runtime/CoreUObject/Private/UObject/Class.cpp
@@ -4692,6 +4692,15 @@ void UDynamicClass::PurgeClass(bool bRecompilingOnLoad)
 	AnimClassImplementation = nullptr;
 }
 
+UObject* UDynamicClass::FindArchetype(UClass* ArchetypeClass, const FName ArchetypeName) const
+{
+	UDynamicClass* ThisClass = const_cast<UDynamicClass*>(this);
+	UObject* Archetype = static_cast<UObject*>(FindObjectWithOuter(ThisClass, ArchetypeClass, ArchetypeName));
+	const UClass* SuperClass = GetSuperClass();
+	return Archetype ? Archetype :
+		(SuperClass ? SuperClass->FindArchetype(ArchetypeClass, ArchetypeName) : nullptr);
+}
+
 IMPLEMENT_CORE_INTRINSIC_CLASS(UDynamicClass, UClass,
 {
 	Class->ClassAddReferencedObjects = &UDynamicClass::AddReferencedObjects;
diff --git a/Engine/Source/Runtime/CoreUObject/Private/UObject/LinkerLoad.cpp b/Engine/Source/Runtime/CoreUObject/Private/UObject/LinkerLoad.cpp
--- a/Engine/Source/Runtime/CoreUObject/Private/UObject/LinkerLoad.cpp
+++ b/Engine/Source/Runtime/CoreUObject/Private/UObject/LinkerLoad.cpp
@@ -3283,7 +3283,7 @@ void FLinkerLoad::Preload( UObject* Object )
 			// its owner has had a chance to initialize itself (because, as part
 			// of CDO initialization, inherited sub-objects get filled in with 
 			// values inherited from the super)
-			else if (Object->HasAnyFlags(RF_DefaultSubObject) && FDeferredObjInitializerTracker::DeferSubObjectPreload(Object))
+			else if (Object->HasAnyFlags(RF_DefaultSubObject|RF_InheritableComponentTemplate) && FDeferredObjInitializerTracker::DeferSubObjectPreload(Object))
 			{
 				// don't worry, FDeferredObjInitializerTracker::DeferSubObjectPreload() 
 				// should have cached this object, and it will run Preload() on 
diff --git a/Engine/Source/Runtime/CoreUObject/Private/UObject/ScriptCore.cpp b/Engine/Source/Runtime/CoreUObject/Private/UObject/ScriptCore.cpp
--- a/Engine/Source/Runtime/CoreUObject/Private/UObject/ScriptCore.cpp
+++ b/Engine/Source/Runtime/CoreUObject/Private/UObject/ScriptCore.cpp
@@ -1147,11 +1147,13 @@ void UObject::ProcessEvent( UFunction* Function, void* Parms )
 		// Don't instrument native events that have not been implemented/overridden in script (BP). These events will not have had any profiler data generated for them at compile time.
 		if (!Function->HasAnyFunctionFlags(FUNC_Native) || Function->GetOuter()->GetClass()->HasAnyClassFlags(CLASS_CompiledFromBlueprint))
 		{
-			const EScriptInstrumentation::Type EventType = Function->HasAnyFunctionFlags(FUNC_Event | FUNC_BlueprintEvent) ? EScriptInstrumentation::Event : EScriptInstrumentation::ResumeEvent;
-			EScriptInstrumentationEvent EventInstrumentationInfo(EventType, this, Function->GetFName());
-			FBlueprintCoreDelegates::InstrumentScriptEvent(EventInstrumentationInfo);
-
-			bInstrumentScriptEvent = true;
+			if (Function->HasAnyFunctionFlags(FUNC_Event|FUNC_BlueprintEvent))
+			{
+				// Don't handle latent actions here, let the latent action manager handle them.
+				EScriptInstrumentationEvent EventInstrumentationInfo(EScriptInstrumentation::Event, this, Function->GetFName());
+				FBlueprintCoreDelegates::InstrumentScriptEvent(EventInstrumentationInfo);
+				bInstrumentScriptEvent = true;
+			}
 		}
 	}
 #endif
diff --git a/Engine/Source/Runtime/CoreUObject/Private/UObject/UObjectGlobals.cpp b/Engine/Source/Runtime/CoreUObject/Private/UObject/UObjectGlobals.cpp
--- a/Engine/Source/Runtime/CoreUObject/Private/UObject/UObjectGlobals.cpp
+++ b/Engine/Source/Runtime/CoreUObject/Private/UObject/UObjectGlobals.cpp
@@ -2531,41 +2531,93 @@ FObjectInitializer::~FObjectInitializer()
 
 #if USE_CIRCULAR_DEPENDENCY_LOAD_DEFERRING
 	bool bIsPostConstructInitDeferred = false;
-	if (bIsCDO && (ObjectArchetype != nullptr) && !FBlueprintSupport::IsDeferredCDOInitializationDisabled())
-	{
-		UClass* ArchetypeClass = ObjectArchetype->GetClass();
-		const bool bSuperCDONeedsLoad = ObjectArchetype->HasAnyFlags(RF_NeedLoad) ||
-			(ArchetypeClass->GetLinker() && ArchetypeClass->GetLinker()->IsBlueprintFinalizationPending()) ||
-			FDeferredObjInitializerTracker::IsCdoDeferred(ArchetypeClass);
+	if (!FBlueprintSupport::IsDeferredCDOInitializationDisabled())
+	{		
+		UClass* BlueprintClass = nullptr;
+		// since "InheritableComponentTemplate"s are not default sub-objects, 
+		// they won't be fixed up by the owner's FObjectInitializer (CDO 
+		// FObjectInitializers will init default sub-object properties, copying  
+		// from the super's DSOs) - this means that we need to separately defer 
+		// init'ing these sub-objects when their archetype hasn't been loaded 
+		// yet (it is possible that the archetype isn't even correct, as the 
+		// super's sub-object hasn't even been created yet; in this case the 
+		// component's CDO is used, which is probably wrong)
+		if (Obj->HasAnyFlags(RF_InheritableComponentTemplate))
+		{
+			BlueprintClass = Cast<UClass>(Obj->GetOuter());
+#if USE_DEFERRED_DEPENDENCY_CHECK_VERIFICATION_TESTS
+			check(BlueprintClass != nullptr);
+#endif // USE_DEFERRED_DEPENDENCY_CHECK_VERIFICATION_TESTS
+		}
+		else if (bIsCDO && !Class->IsNative())
+		{
+			BlueprintClass = Class;
+#if USE_DEFERRED_DEPENDENCY_CHECK_VERIFICATION_TESTS
+			check(Class->HasAnyClassFlags(CLASS_CompiledFromBlueprint));
+#endif // USE_DEFERRED_DEPENDENCY_CHECK_VERIFICATION_TESTS
+		}
 
-		// if this is a blueprint CDO that derives from another blueprint, and 
-		// that parent (archetype) CDO isn't fully serialized
-		if (!Class->IsNative() && !ArchetypeClass->IsNative() && bSuperCDONeedsLoad)
+		if (BlueprintClass != nullptr)
 		{
-			FLinkerLoad* ClassLinker = Class->GetLinker();
-			if ((ClassLinker != nullptr) && (ClassLinker->LoadFlags & LOAD_DeferDependencyLoads) != 0x00)
-			{
 #if USE_DEFERRED_DEPENDENCY_CHECK_VERIFICATION_TESTS
-				// make sure we haven't already deferred this once, if we have
-				// then something is destroying this one prematurely 
-				check(bIsDeferredInitializer == false);
+			check(!BlueprintClass->IsNative());
+#endif // USE_DEFERRED_DEPENDENCY_CHECK_VERIFICATION_TESTS
 
-				for (const FSubobjectsToInit::FSubobjectInit& SubObjInfo : ComponentInits.SubobjectInits)
+			UClass* SuperClass = BlueprintClass->GetSuperClass();
+			if (SuperClass && !SuperClass->IsNative())
+			{
+				UObject* SuperBpCDO = nullptr;
+				// if this is a CDO (then we know/assume the archetype is the 
+				// CDO from the super class), use the ObjectArchetype for the 
+				// SuperBpCDO (because the SuperClass may have a REINST CDO 
+				// cached currently)
+				if (bIsCDO)
 				{
-					check(!SubObjInfo.Subobject->HasAnyFlags(RF_NeedLoad));
-				}
+					SuperBpCDO = ObjectArchetype;
+					SuperClass = ObjectArchetype->GetClass();
+
+#if USE_DEFERRED_DEPENDENCY_CHECK_VERIFICATION_TESTS
+					check(ObjectArchetype->HasAnyFlags(RF_ClassDefaultObject));
 #endif // USE_DEFERRED_DEPENDENCY_CHECK_VERIFICATION_TESTS
+				}
+				else
+				{
+					SuperBpCDO = SuperClass->GetDefaultObject(/*bCreateIfNeeded =*/false);
+				}
+
+				const bool bSuperLoadPending = FDeferredObjInitializerTracker::IsCdoDeferred(SuperClass) ||
+					(SuperBpCDO && SuperBpCDO->HasAnyFlags(RF_NeedLoad)) ||
+					(SuperClass->GetLinker() && SuperClass->GetLinker()->IsBlueprintFinalizationPending());
+
+				FLinkerLoad* ObjLinker = BlueprintClass->GetLinker();
+				const bool bIsBpClassSerializing    = ObjLinker && (ObjLinker->LoadFlags & LOAD_DeferDependencyLoads);
+				const bool bIsResolvingDeferredObjs = BlueprintClass->HasAnyFlags(RF_LoadCompleted) &&
+					ObjLinker && ObjLinker->IsBlueprintFinalizationPending();
 
-				// makes a copy of this and saves it off, to be ran later
-				if (FObjectInitializer* DeferredCopy = FDeferredObjInitializerTracker::Add(*this))
+				if (bSuperLoadPending && (bIsBpClassSerializing || bIsResolvingDeferredObjs))
 				{
-					bIsPostConstructInitDeferred = true;
-					DeferredCopy->bIsDeferredInitializer = true;
+#if USE_DEFERRED_DEPENDENCY_CHECK_VERIFICATION_TESTS
+					// make sure we haven't already deferred this once, if we have
+					// then something is destroying this one prematurely 
+					check(bIsDeferredInitializer == false);
 
-					// make sure this wasn't mistakenly pushed into ObjectInitializers
-					// (the copy constructor should have been what was invoked, 
-					// which doesn't push to ObjectInitializers)
-					check(FUObjectThreadContext::Get().TopInitializer() != DeferredCopy);
+					for (const FSubobjectsToInit::FSubobjectInit& SubObjInfo : ComponentInits.SubobjectInits)
+					{
+						check(!SubObjInfo.Subobject->HasAnyFlags(RF_NeedLoad));
+					}
+#endif // USE_DEFERRED_DEPENDENCY_CHECK_VERIFICATION_TESTS
+					
+					// makes a copy of this and saves it off, to be ran later
+					if (FObjectInitializer* DeferredCopy = FDeferredObjInitializerTracker::Add(*this))
+					{
+						bIsPostConstructInitDeferred = true;
+						DeferredCopy->bIsDeferredInitializer = true;
+
+						// make sure this wasn't mistakenly pushed into ObjectInitializers
+						// (the copy constructor should have been what was invoked, 
+						// which doesn't push to ObjectInitializers)
+						check(FUObjectThreadContext::Get().TopInitializer() != DeferredCopy);
+					}
 				}
 			}
 		}
@@ -2596,25 +2648,33 @@ void FObjectInitializer::PostConstructInit()
 	UClass* SuperClass = Class->GetSuperClass();
 
 #if USE_CIRCULAR_DEPENDENCY_LOAD_DEFERRING
-	// if this is a deferred initializer (implying that it's for a CDO), and the 
-	// ObjectArchetype (super CDO) has since been regenerated, then we cannot
-	// reliably initialize and instance properties (if we were to use the 
-	// regenerated ObjectArchetype, the property layouts could differ and that's
-	// not viable for InitProperties())
-	// 
-	// However, this case is handled in FLinkerLoad::ResolveDeferredExports(), 
-	// where we forcefully recreate and separately init this CDO (hence it being 
-	// moved to the transient package) 
 	if (bIsDeferredInitializer)
 	{
-		const bool bSuperHasBeenRegenerated = SuperClass->HasAnyClassFlags(CLASS_NewerVersionExists);
+		const bool bIsDeferredSubObject = Obj->HasAnyFlags(RF_InheritableComponentTemplate);
+		if (bIsDeferredSubObject)
+		{
+			// when this sub-object was created it's archetype object (the 
+			// super's sub-obj) may not have been created yet (thanks cyclic 
+			// dependencies). in that scenario, the component class's CDO would  
+			// have been used in its place; now that we're resolving the defered 
+			// sub-obj initialization we should try to update the archetype
+			if (ObjectArchetype->HasAnyFlags(RF_ClassDefaultObject))
+			{
+				ObjectArchetype = UObject::GetArchetypeFromRequiredInfo(Class, Obj->GetOuter(), Obj->GetFName(), Obj->GetFlags());
+				// NOTE: this may still be the component class's CDO (like when 
+				// a component was removed from the super, without resaving the child)
+			}			
+		}
+
+		UClass* ArchetypeClass = ObjectArchetype->GetClass();
+		const bool bSuperHasBeenRegenerated = ArchetypeClass->HasAnyClassFlags(CLASS_NewerVersionExists);
 #if USE_DEFERRED_DEPENDENCY_CHECK_VERIFICATION_TESTS
-		check(bIsCDO);
+		check(bIsCDO || bIsDeferredSubObject);
 		check(ObjectArchetype->GetOutermost() != GetTransientPackage());
-		check(ObjectArchetype->GetClass() == SuperClass && !bSuperHasBeenRegenerated);
+		check(!bIsCDO || (ArchetypeClass == SuperClass && !bSuperHasBeenRegenerated));
 #endif // USE_DEFERRED_DEPENDENCY_CHECK_VERIFICATION_TESTS
 
-		if ( !ensureMsgf(!bSuperHasBeenRegenerated, TEXT("The super class for %s has been regenerated, we cannot properly initialize inherited properties, as the class layout may have changed."), *Obj->GetName()) )
+		if ( !ensureMsgf(!bSuperHasBeenRegenerated, TEXT("The archetype for %s has been regenerated, we cannot properly initialize inherited properties, as the class layout may have changed."), *Obj->GetName()) )
 		{
 			// attempt to complete initialization/instancing as best we can, but
 			// it would not be surprising if our CDO was improperly initialized 
diff --git a/Engine/Source/Runtime/CoreUObject/Public/Blueprint/BlueprintSupport.h b/Engine/Source/Runtime/CoreUObject/Public/Blueprint/BlueprintSupport.h
--- a/Engine/Source/Runtime/CoreUObject/Public/Blueprint/BlueprintSupport.h
+++ b/Engine/Source/Runtime/CoreUObject/Public/Blueprint/BlueprintSupport.h
@@ -41,7 +41,7 @@ struct FBlueprintSupport
 	static bool IsDeferredCDOInitializationDisabled();
 
 	/** Tells if the specified object is one of the many flavors of FLinkerPlaceholderBase that we have. */
-	static bool IsDeferredDependencyPlaceholder(UObject* LoadedObj);
+	COREUOBJECT_API static bool IsDeferredDependencyPlaceholder(UObject* LoadedObj);
 
 	/** Not a particularly fast function. Mostly intended for validation in debug builds. */
 	static bool IsInBlueprintPackage(UObject* LoadedObj);
@@ -161,8 +161,8 @@ struct IBlueprintNativeCodeGenCore
 
 /** 
  * A helper struct for storing FObjectInitializers that were not run on 
- * Blueprint CDO's post-construction (presumably because that CDO's super had 
- * not been fully serialized yet). 
+ * Blueprint CDO's (or CDO sub-objects) post-construction (presumably because 
+ * that CDO's super had not been fully serialized yet). 
  * 
  * This was designed to hold onto FObjectInitializers until a later point, when 
  * they can properly be ran (presumably in FLinkerLoad::ResolveDeferredExports,
@@ -206,6 +206,8 @@ private:
 	UClass* ResolvingClass;
 	/** Tracks sub-classes that have had their CDO deferred as a result of the super not being fully serialized */
 	TMultiMap<UClass*, UClass*> SuperClassMap;
+	/** A map that tracks deferred FObjectInitializers for a class/CDO's sub-objects (not DSOs, but things like component override templates)*/
+	TMultiMap<UClass*, FObjectInitializer> DeferredSubObjInitializers;
 };
 
 
diff --git a/Engine/Source/Runtime/CoreUObject/Public/UObject/Class.h b/Engine/Source/Runtime/CoreUObject/Public/UObject/Class.h
--- a/Engine/Source/Runtime/CoreUObject/Public/UObject/Class.h
+++ b/Engine/Source/Runtime/CoreUObject/Public/UObject/Class.h
@@ -2524,6 +2524,7 @@ public:
 
 	// UClass interface
 	virtual void PurgeClass(bool bRecompilingOnLoad) override;
+	virtual UObject* FindArchetype(UClass* ArchetypeClass, const FName ArchetypeName) const override;
 
 	/** Misc objects owned by the class. */
 	TArray<UObject*> MiscConvertedSubobjects;
diff --git a/Engine/Source/Runtime/CoreUObject/Public/UObject/Script.h b/Engine/Source/Runtime/CoreUObject/Public/UObject/Script.h
--- a/Engine/Source/Runtime/CoreUObject/Public/UObject/Script.h
+++ b/Engine/Source/Runtime/CoreUObject/Public/UObject/Script.h
@@ -357,6 +357,7 @@ public:
 		, ContextObject(InContextObject)
 		, EventName(NAME_None)
 		, StackFramePtr(&InStackFrame)
+		, LatentLinkId(INDEX_NONE)
 	{
 	}
 
@@ -365,6 +366,16 @@ public:
 		, ContextObject(InContextObject)
 		, EventName(InEventName)
 		, StackFramePtr(nullptr)
+		, LatentLinkId(INDEX_NONE)
+	{
+	}
+
+	EScriptInstrumentationEvent(EScriptInstrumentation::Type InEventType, const UObject* InContextObject, FName InEventName, const int32 LinkId)
+		: EventType(InEventType)
+		, ContextObject(InContextObject)
+		, EventName(InEventName)
+		, StackFramePtr(nullptr)
+		, LatentLinkId(LinkId)
 	{
 	}
 
@@ -374,6 +385,7 @@ public:
 	bool IsStackFrameValid() const { return StackFramePtr != nullptr; }
 	FName GetEventName() const { return EventName; }
 	const FFrame& GetStackFrame() const { return *StackFramePtr; }
+	int32 GetLatentLinkId() const { return LatentLinkId; }
 
 protected:
 
@@ -381,6 +393,7 @@ protected:
 	const UObject* ContextObject;
 	const FName EventName;
 	const struct FFrame* StackFramePtr;
+	const int32 LatentLinkId;
 };
 
 // Blueprint core runtime delegates
diff --git a/Engine/Source/Runtime/Engine/Classes/Components/ChildActorComponent.h b/Engine/Source/Runtime/Engine/Classes/Components/ChildActorComponent.h
--- a/Engine/Source/Runtime/Engine/Classes/Components/ChildActorComponent.h
+++ b/Engine/Source/Runtime/Engine/Classes/Components/ChildActorComponent.h
@@ -14,6 +14,7 @@ public:
 	virtual ~FChildActorComponentInstanceData();
 
 	virtual void ApplyToComponent(UActorComponent* Component, const ECacheApplyPhase CacheApplyPhase) override;
+	virtual void AddReferencedObjects(FReferenceCollector& Collector) override;
 
 	// The name of the spawned child actor so it (attempts to) remain constant across construction script reruns
 	FName ChildActorName;
@@ -70,6 +71,7 @@ public:
 	virtual void BeginDestroy() override;
 	virtual void GetLifetimeReplicatedProps(TArray<FLifetimeProperty>& OutLifetimeProps) const override;
 	virtual void PostRepNotifies() override;
+	static void AddReferencedObjects(UObject* InThis, FReferenceCollector& Collector);
 	//~ End Object Interface.
 
 	//~ Begin ActorComponent Interface.
diff --git a/Engine/Source/Runtime/Engine/Classes/EdGraph/EdGraphNode.h b/Engine/Source/Runtime/Engine/Classes/EdGraph/EdGraphNode.h
--- a/Engine/Source/Runtime/Engine/Classes/EdGraph/EdGraphNode.h
+++ b/Engine/Source/Runtime/Engine/Classes/EdGraph/EdGraphNode.h
@@ -240,12 +240,18 @@ public:
 	/** Destroys the specified pin, does not modify its owning pin's Pins list */
 	static void DestroyPin(UEdGraphPin* Pin);
 
-	/** Find a pin on this node with the supplied name */
-	UEdGraphPin* FindPin(const FString& PinName) const;
+	/** Find a pin on this node with the supplied name and optional direction */
+	UEdGraphPin* FindPin(const FString& PinName, const EEdGraphPinDirection Direction = EGPD_MAX) const;
 
-	/** Find a pin on this node with the supplied name and assert if it is not present */
-	UEdGraphPin* FindPinChecked(const FString& PinName) const;
+	/** Find a pin on this node with the supplied name and optional direction and assert if it is not present */
+	UEdGraphPin* FindPinChecked(const FString& PinName, const EEdGraphPinDirection Direction = EGPD_MAX) const;
 	
+	/** Find the pin on this node with the supplied guid */
+	UEdGraphPin* FindPinById(const FGuid PinId) const;
+
+	/** Find the pin on this node with the supplied guid and assert if it is not present */
+	UEdGraphPin* FindPinByIdChecked(const FGuid PinId) const;
+
 	/** Find a pin on this node with the supplied name and remove it, returns TRUE if successful */
 	bool RemovePin(UEdGraphPin* Pin);
 
diff --git a/Engine/Source/Runtime/Engine/Classes/EdGraph/EdGraphPin.h b/Engine/Source/Runtime/Engine/Classes/EdGraph/EdGraphPin.h
--- a/Engine/Source/Runtime/Engine/Classes/EdGraph/EdGraphPin.h
+++ b/Engine/Source/Runtime/Engine/Classes/EdGraph/EdGraphPin.h
@@ -397,9 +397,6 @@ public:
 	/** Changes the owning node. This will remove the pin from the old owning node's pin list and add itself to the new node's pin list. */
 	ENGINE_API void SetOwningNode(UEdGraphNode* NewOwningNode);
 
-	/** Marks the pin as 'trashed'. Removed the pin from the Owning Node's Pins list */
-	ENGINE_API void InvalidateAndTrash();
-
 	/** Marks the pin as 'trashed'. *Does not* remove the pin from the Owning Node's Pins list */
 	ENGINE_API void MarkPendingKill();
 
@@ -429,18 +426,15 @@ public:
 	/** This needs to be called if you want to use pin data within PostEditUndo */
 	ENGINE_API static void ResolveAllPinReferences();
 private:
-	/** 
-		Unsafe helper function used to deal with odd design of transaction buffer, 
-		which immediately applies changes to some records before saving state of others 
-	  */
-	void QueueForDestruction();
-
 	/** Private Constructor. Create pins using CreatePin since all pin instances are managed by TSharedPtr. */
 	UEdGraphPin(UEdGraphNode* InOwningNode, const FGuid& PinGuid);
 
 	/** Backward compatibility code to populate this pin with data from the supplied deprecated UEdGraphPin. */
 	void InitFromDeprecatedPin(class UEdGraphPin_Deprecated* DeprecatedPin);
 
+	/** Helper function for common destruction logic */
+	void DestroyImpl(bool bClearLinks);
+
 	bool Serialize(FArchive& Ar);
 
 	// Helper functions
diff --git a/Engine/Source/Runtime/Engine/Classes/Engine/Blueprint.h b/Engine/Source/Runtime/Engine/Classes/Engine/Blueprint.h
--- a/Engine/Source/Runtime/Engine/Classes/Engine/Blueprint.h
+++ b/Engine/Source/Runtime/Engine/Classes/Engine/Blueprint.h
@@ -326,6 +326,10 @@ class ENGINE_API UBlueprint : public UBlueprintCore
 	UPROPERTY(EditAnywhere, Category=ClassOptions, AdvancedDisplay)
 	uint32 bGenerateConstClass : 1;
 
+	/** Whether or not this blueprint's class is a abstract class or not.  Should set CLASS_Abstract in the KismetCompiler. */
+	UPROPERTY(EditAnywhere, Category = ClassOptions, AdvancedDisplay)
+	uint32 bGenerateAbstractClass : 1;
+
 	/**shows up in the content browser when the blueprint is hovered */
 	UPROPERTY(EditAnywhere, Category=BlueprintOptions, meta=(MultiLine=true))
 	FString BlueprintDescription;
diff --git a/Engine/Source/Runtime/Engine/Classes/Engine/BlueprintGeneratedClass.h b/Engine/Source/Runtime/Engine/Classes/Engine/BlueprintGeneratedClass.h
--- a/Engine/Source/Runtime/Engine/Classes/Engine/BlueprintGeneratedClass.h
+++ b/Engine/Source/Runtime/Engine/Classes/Engine/BlueprintGeneratedClass.h
@@ -197,6 +197,9 @@ protected:
 	//   It does *not* contain the wire tracepoint placed after the impure function call
 	TArray<struct FNodeToCodeAssociation> DebugNodeLineNumbers;
 
+	// List of entry points that contributed to the ubergraph
+	TMap<FName, FEdGraphPinReference> DelegatePins;
+
 	// Acceleration structure for execution wire highlighting at runtime
 	TMap<TWeakObjectPtr<UFunction>, FDebuggingInfoForSingleFunction> PerFunctionLineNumbers;
 
@@ -406,6 +409,13 @@ public:
 		PerFuncInfo.SourcePinToLineNumbersMap.Add(SourcePin, CodeOffset);
 	}
 
+	const TMap<FName, FEdGraphPinReference>& GetCompilerGeneratedEvents() const { return DelegatePins; }
+
+	void RegisterPinToEventName(UEdGraphPin const* SourcePin, const FName FunctionName)
+	{
+		DelegatePins.Add(FunctionName) = SourcePin;
+	}
+
 	// Registers an association between an object (pin or node typically) and an associated class member property
 	void RegisterClassPropertyAssociation(class UObject* TrueSourceObject, class UProperty* AssociatedProperty)
 	{
diff --git a/Engine/Source/Runtime/Engine/Private/Actor.cpp b/Engine/Source/Runtime/Engine/Private/Actor.cpp
--- a/Engine/Source/Runtime/Engine/Private/Actor.cpp
+++ b/Engine/Source/Runtime/Engine/Private/Actor.cpp
@@ -267,6 +267,12 @@ void AActor::AddReferencedObjects(UObject* InThis, FReferenceCollector& Collecto
 {
 	AActor* This = CastChecked<AActor>(InThis);
 	Collector.AddReferencedObjects(This->OwnedComponents);
+#if WITH_EDITOR
+	if (This->CurrentTransactionAnnotation.IsValid())
+	{
+		This->CurrentTransactionAnnotation->AddReferencedObjects(Collector);
+	}
+#endif
 	Super::AddReferencedObjects(InThis, Collector);
 }
 
diff --git a/Engine/Source/Runtime/Engine/Private/Components/ChildActorComponent.cpp b/Engine/Source/Runtime/Engine/Private/Components/ChildActorComponent.cpp
--- a/Engine/Source/Runtime/Engine/Private/Components/ChildActorComponent.cpp
+++ b/Engine/Source/Runtime/Engine/Private/Components/ChildActorComponent.cpp
@@ -180,6 +180,28 @@ void FChildActorComponentInstanceData::ApplyToComponent(UActorComponent* Compone
 	CastChecked<UChildActorComponent>(Component)->ApplyComponentInstanceData(this, CacheApplyPhase);
 }
 
+void FChildActorComponentInstanceData::AddReferencedObjects(FReferenceCollector& Collector)
+{
+	FSceneComponentInstanceData::AddReferencedObjects(Collector);
+
+	if (ComponentInstanceData)
+	{
+		ComponentInstanceData->AddReferencedObjects(Collector);
+	}
+}
+
+void UChildActorComponent::AddReferencedObjects(UObject* InThis, FReferenceCollector& Collector)
+{
+	UChildActorComponent* This = CastChecked<UChildActorComponent>(InThis);
+
+	if (This->CachedInstanceData)
+	{
+		This->CachedInstanceData->AddReferencedObjects(Collector);
+	}
+
+	Super::AddReferencedObjects(InThis, Collector);
+}
+
 void UChildActorComponent::BeginDestroy()
 {
 	Super::BeginDestroy();
diff --git a/Engine/Source/Runtime/Engine/Private/EdGraph/EdGraphNode.cpp b/Engine/Source/Runtime/Engine/Private/EdGraph/EdGraphNode.cpp
--- a/Engine/Source/Runtime/Engine/Private/EdGraph/EdGraphNode.cpp
+++ b/Engine/Source/Runtime/Engine/Private/EdGraph/EdGraphNode.cpp
@@ -95,11 +95,11 @@ UEdGraphPin* UEdGraphNode::CreatePin(EEdGraphPinDirection Dir, const FString& Pi
 	return CreatePin(Dir, PinType, PinName, Index);
 }
 
-UEdGraphPin* UEdGraphNode::FindPin(const FString& PinName) const
+UEdGraphPin* UEdGraphNode::FindPin(const FString& PinName, const EEdGraphPinDirection Direction) const
 {
 	for(int32 PinIdx=0; PinIdx<Pins.Num(); PinIdx++)
 	{
-		if( Pins[PinIdx]->PinName == PinName )
+		if( Pins[PinIdx]->PinName == PinName && (Direction == EGPD_MAX || Direction == Pins[PinIdx]->Direction))
 		{
 			return Pins[PinIdx];
 		}
@@ -108,9 +108,29 @@ UEdGraphPin* UEdGraphNode::FindPin(const FString& PinName) const
 	return NULL;
 }
 
-UEdGraphPin* UEdGraphNode::FindPinChecked(const FString& PinName) const
+UEdGraphPin* UEdGraphNode::FindPinChecked(const FString& PinName, const EEdGraphPinDirection Direction) const
 {
-	UEdGraphPin* Result = FindPin(PinName);
+	UEdGraphPin* Result = FindPin(PinName, Direction);
+	check(Result != NULL);
+	return Result;
+}
+
+UEdGraphPin* UEdGraphNode::FindPinById(const FGuid PinId) const
+{
+	for (int32 PinIdx = 0; PinIdx < Pins.Num(); PinIdx++)
+	{
+		if (Pins[PinIdx]->PinId == PinId)
+		{
+			return Pins[PinIdx];
+		}
+	}
+
+	return NULL;
+}
+
+UEdGraphPin* UEdGraphNode::FindPinByIdChecked(const FGuid PinId) const
+{
+	UEdGraphPin* Result = FindPinById(PinId);
 	check(Result != NULL);
 	return Result;
 }
diff --git a/Engine/Source/Runtime/Engine/Private/EdGraph/EdGraphPin.cpp b/Engine/Source/Runtime/Engine/Private/EdGraph/EdGraphPin.cpp
--- a/Engine/Source/Runtime/Engine/Private/EdGraph/EdGraphPin.cpp
+++ b/Engine/Source/Runtime/Engine/Private/EdGraph/EdGraphPin.cpp
@@ -6,6 +6,7 @@
 #include "BlueprintUtilities.h"
 #include "Tickable.h"
 #if WITH_EDITOR
+#include "Editor/EditorEngine.h"
 #include "Editor/UnrealEd/Public/Kismet2/BlueprintEditorUtils.h"
 #include "SlateBasics.h"
 #include "ScopedTransaction.h"
@@ -167,10 +168,11 @@ struct FUnresolvedPinData
 	UEdGraphPin* ReferencingPin;
 	int32 ArrayIdx;
 	EPinResolveType ResolveType;
+	bool bResolveSymmetrically;
 
 	FUnresolvedPinData() : ReferencingPin(nullptr), ArrayIdx(INDEX_NONE), ResolveType(EPinResolveType::OwningNode) {}
-	FUnresolvedPinData(UEdGraphPin* InReferencingPin, EPinResolveType InResolveType, int32 InArrayIdx)
-		: ReferencingPin(InReferencingPin), ArrayIdx(InArrayIdx), ResolveType(InResolveType)
+	FUnresolvedPinData(UEdGraphPin* InReferencingPin, EPinResolveType InResolveType, int32 InArrayIdx, bool bInResolveSymmetrically = false)
+		: ReferencingPin(InReferencingPin), ArrayIdx(InArrayIdx), ResolveType(InResolveType), bResolveSymmetrically(bInResolveSymmetrically)
 	{}
 };
 
@@ -458,11 +460,11 @@ const FString UEdGraphPin::GetLinkInfoString( const FString& InFunctionName, con
 {
 #if WITH_EDITOR
 	const FString FromPinName = PinName;
-	const UEdGraphNode* FromPinNode = GetOwningNode();
+	const UEdGraphNode* FromPinNode = GetOwningNodeUnchecked();
 	const FString FromPinNodeName = (FromPinNode != nullptr) ? FromPinNode->GetNodeTitle(ENodeTitleType::ListView).ToString() : TEXT("Unknown");
 	
 	const FString ToPinName = InToPin->PinName;
-	const UEdGraphNode* ToPinNode = InToPin->GetOwningNode();
+	const UEdGraphNode* ToPinNode = InToPin->GetOwningNodeUnchecked();
 	const FString ToPinNodeName = (ToPinNode != nullptr) ? ToPinNode->GetNodeTitle(ENodeTitleType::ListView).ToString() : TEXT("Unknown");
 	const FString LinkInfo = FString::Printf( TEXT("UEdGraphPin::%s Pin '%s' on node '%s' %s '%s' on node '%s'"), *InFunctionName, *ToPinName, *ToPinNodeName, *InInfoData, *FromPinName, *FromPinNodeName);
 	return LinkInfo;
@@ -522,21 +524,6 @@ void UEdGraphPin::SetOwningNode(UEdGraphNode* NewOwningNode)
 	}
 }
 
-int32 NumSubPins = 0;
-
-void UEdGraphPin::InvalidateAndTrash()
-{
-	check(!bWasTrashed);
-	PinsToDelete.Add(this);
-	BreakAllPinLinks();
-	SetOwningNode(nullptr);
-	NumSubPins += SubPins.Num();
-	SubPins.Empty();
-	ParentPin = nullptr;
-	ReferencePassThroughConnection = nullptr;
-	bWasTrashed = true;
-}
-
 UEdGraphPin* UEdGraphPin::CreatePinFromDeprecatedPin(UEdGraphPin_Deprecated* DeprecatedPin)
 {
 	if (!DeprecatedPin)
@@ -975,29 +962,7 @@ void UEdGraphPin::MarkPendingKill()
 {
 	if (!bWasTrashed)
 	{
-		PinsToDelete.Add(this);
-#if WITH_EDITOR
-		if (!GIsTransacting)
-		{
-#endif //WITH_EDITOR
-			BreakAllPinLinks();
-			if (ParentPin)
-			{
-				ParentPin->SubPins.Remove(this);
-			}
-#if WITH_EDITOR
-		}
-		else
-		{
-			LinkedTo.Empty();
-		}
-#endif //WITH_EDITOR
-		OwningNode = nullptr;
-		NumSubPins = 0;
-		SubPins.Empty();
-		ParentPin = nullptr;
-		ReferencePassThroughConnection = nullptr;
-		bWasTrashed = true;
+		DestroyImpl(true);
 	}
 }
 
@@ -1014,11 +979,6 @@ void UEdGraphPin::Purge()
 	PinDeletionQueue->Tick(0.f);
 }
 
-void UEdGraphPin::QueueForDestruction()
-{
-	PinsToDelete.Add(this);
-}
-
 UEdGraphPin::UEdGraphPin(UEdGraphNode* InOwningNode, const FGuid& PinIdGuid)
 	: OwningNode(InOwningNode)
 	, PinId(PinIdGuid)
@@ -1230,6 +1190,31 @@ void UEdGraphPin::InitFromDeprecatedPin(class UEdGraphPin_Deprecated* Deprecated
 	}
 }
 
+void UEdGraphPin::DestroyImpl(bool bClearLinks)
+{
+#if WITH_EDITOR
+	ensure(!GIsTransacting || bClearLinks == false);
+#endif
+	PinsToDelete.Add(this);
+	if (bClearLinks)
+	{
+		BreakAllPinLinks();
+		if (ParentPin)
+		{
+			ParentPin->SubPins.Remove(this);
+		}
+	}
+	else
+	{
+		LinkedTo.Empty();
+	}
+	OwningNode = nullptr;
+	SubPins.Empty();
+	ParentPin = nullptr;
+	ReferencePassThroughConnection = nullptr;
+	bWasTrashed = true;
+}
+
 bool UEdGraphPin::Serialize(FArchive& Ar)
 {
 	// These properties are in every pin and are unlikely to be removed, so they are native serialized for speed.
@@ -1337,6 +1322,7 @@ void UEdGraphPin::EnableAllConnectedNodes(UEdGraphNode* InNode)
 
 void UEdGraphPin::ResolveReferencesToPin(UEdGraphPin* Pin)
 {
+	check(!Pin->bWasTrashed);
 	FPinResolveId ResolveId(Pin->PinId, Pin->OwningNode);
 	TArray<FUnresolvedPinData>* UnresolvedPinData = PinHelpers::UnresolvedPins.Find(ResolveId);
 	if (UnresolvedPinData)
@@ -1347,7 +1333,17 @@ void UEdGraphPin::ResolveReferencesToPin(UEdGraphPin* Pin)
 			switch (PinData.ResolveType)
 			{
 			case EPinResolveType::LinkedTo:
-				ReferencingPin->LinkedTo[PinData.ArrayIdx] = Pin;
+				if (ensure(ReferencingPin->LinkedTo.Num() > PinData.ArrayIdx))
+				{
+					ReferencingPin->LinkedTo[PinData.ArrayIdx] = Pin;
+				}
+				if (PinData.bResolveSymmetrically)
+				{
+					if (ensure(!Pin->LinkedTo.Contains(ReferencingPin)))
+					{
+						Pin->LinkedTo.Add(ReferencingPin);
+					}
+				}
 #if WITH_EDITOR
 				// When in the middle of a transaction the LinkedTo lists will be in an 
 				// indeterminate state. After PostEditUndo runs we could validate LinkedTo
@@ -1400,10 +1396,27 @@ void UEdGraphPin::SerializePinArray(FArchive& Ar, TArray<UEdGraphPin*>& ArrayRef
 		SerializePin(Ar, PinRef, PinIdx, RequestingPin, ResolveType, OldPins);
 	}
 
-	// free unused pins:
+	// free unused pins, this only happens when loading and we have already allocated pin data, which currently
+	// means we're serializing for undo/redo:
 	for (UEdGraphPin* Pin : OldPins)
 	{
-		Pin->QueueForDestruction();
+#if WITH_EDITOR
+		// More complexity to handle asymmetry in the transaction buffer. If our peer node is not in the transaction
+		// then we need to take ownership of the entire connection and clear both LinkedTo Arrays:
+		extern UNREALED_API UEditorEngine* GEditor;
+		for (int32 I = 0; I < Pin->LinkedTo.Num(); ++I )
+		{
+			UEdGraphPin* Peer = Pin->LinkedTo[I];
+			// PeerNode will be null if the pin we were linked to was already thrown away:
+			UEdGraphNode* PeerNode = Peer->GetOwningNodeUnchecked();
+			if (PeerNode && !GEditor->Trans->IsObjectTransacting(PeerNode))
+			{
+				Pin->BreakLinkTo(Peer);
+				--I;
+			}
+		}
+#endif// WITH_EDITOR
+		Pin->DestroyImpl(false);
 	}
 }
 
@@ -1440,23 +1453,64 @@ bool UEdGraphPin::SerializePin(FArchive& Ar, UEdGraphPin*& PinRef, int32 ArrayId
 		{
 			if (Ar.IsLoading())
 			{
-				UEdGraphPin** ExistingPin = 
+				UEdGraphPin** ExistingPin = LocalOwningNode->Pins.FindByPredicate([&PinGuid](const UEdGraphPin* Pin) { return Pin && Pin->PinId == PinGuid; });
+				if (ExistingPin
 #if WITH_EDITOR
-					// When transacting we're in a volatile state, and we cannot trust any pins we might find, so we defer all 
-					// pin resolving:
-					Ar.IsTransacting() ? nullptr : 
+					// When transacting we're in a volatile state, and we cannot trust any pins we might find:
+					&& !Ar.IsTransacting()
 #endif//WITH_EDITOR
-					LocalOwningNode->Pins.FindByPredicate([&PinGuid](const UEdGraphPin* Pin) { return Pin && Pin->PinId == PinGuid; });
-				if (ExistingPin)
+					)
 				{
 					PinRef = *ExistingPin;
+					check(!PinRef->bWasTrashed);
 				}
 				else
 				{
 					check(RequestingPin);
-					TArray<FUnresolvedPinData>& UnresolvedPinData = PinHelpers::UnresolvedPins.FindOrAdd(FPinResolveId(PinGuid, LocalOwningNode));
-					UnresolvedPinData.Add(FUnresolvedPinData(RequestingPin, ResolveType, ArrayIdx));
-					bRetVal = false;
+#if WITH_EDITOR
+					extern UNREALED_API UEditorEngine* GEditor;
+					if(Ar.IsTransacting() && !GEditor->Trans->IsObjectTransacting(LocalOwningNode))
+					{
+						/* 
+							Two possibilities:
+
+							1. This transaction did not alter this node's LinkedTo list, and also did not
+								alter ExistingPin's
+							2. This transaction has altered this node's LinkedTo list, but did not alter
+								ExistingPin's
+
+							In case 2 we need to make sure that the LinkedTo connection is symmetrical, but
+							in the case of 1 we have to assume that there is already a connection in ExistingPin.
+							If ExistingPin's LinkedTo list *was* mutated IsObjectTransacting must return true
+							for the owning node.
+						*/
+
+						check(ResolveType == EPinResolveType::LinkedTo);
+
+						check(ExistingPin && *ExistingPin); // the transaction buffer is corrupt
+						FGuid RequestingPinId = RequestingPin->PinId;
+						UEdGraphPin** LinkedTo = (*ExistingPin)->LinkedTo.FindByPredicate([&RequestingPinId](const UEdGraphPin* Pin) { return Pin && Pin->PinId == RequestingPinId; });
+						if (LinkedTo)
+						{
+							// case 1:
+							check(*LinkedTo == RequestingPin);
+							PinRef = *ExistingPin;
+						}
+						else
+						{
+							// case 2:
+							TArray<FUnresolvedPinData>& UnresolvedPinData = PinHelpers::UnresolvedPins.FindOrAdd(FPinResolveId(PinGuid, LocalOwningNode));
+							UnresolvedPinData.Add(FUnresolvedPinData(RequestingPin, ResolveType, ArrayIdx, true));
+							bRetVal = false;
+						}
+					}
+					else
+#endif // WITH_EDITOR
+					{
+						TArray<FUnresolvedPinData>& UnresolvedPinData = PinHelpers::UnresolvedPins.FindOrAdd(FPinResolveId(PinGuid, LocalOwningNode));
+						UnresolvedPinData.Add(FUnresolvedPinData(RequestingPin, ResolveType, ArrayIdx));
+						bRetVal = false;
+					}
 				}
 			}
 		}
@@ -1468,6 +1522,7 @@ bool UEdGraphPin::SerializePin(FArchive& Ar, UEdGraphPin*& PinRef, int32 ArrayId
 				if (PinToReuse)
 				{
 					PinRef = *PinToReuse;
+					check(!PinRef->bWasTrashed);
 					OldPins.RemoveAtSwap(PinToReuse - OldPins.GetData());
 				}
 				else
@@ -1477,6 +1532,7 @@ bool UEdGraphPin::SerializePin(FArchive& Ar, UEdGraphPin*& PinRef, int32 ArrayId
 			}
 
 			PinRef->Serialize(Ar);
+			check(!PinRef->bWasTrashed);
 		}
 	}
 
diff --git a/Engine/Source/Runtime/Engine/Private/LatentActionManager.cpp b/Engine/Source/Runtime/Engine/Private/LatentActionManager.cpp
--- a/Engine/Source/Runtime/Engine/Private/LatentActionManager.cpp
+++ b/Engine/Source/Runtime/Engine/Private/LatentActionManager.cpp
@@ -169,7 +169,20 @@ void FLatentActionManager::TickLatentActionForObject(float DeltaTime, FActionLis
 
 				if (UFunction* ExecutionFunction = CallbackTarget->FindFunction(LinkInfo.ExecutionFunction))
 				{
-					CallbackTarget->ProcessEvent(ExecutionFunction, &(LinkInfo.LinkID));
+					const bool bInstrumentLatentCall = CallbackTarget->GetClass()->HasInstrumentation();
+					if (CallbackTarget->GetClass()->HasInstrumentation())
+					{
+						// Place event markers around the latent blueprint call so we can match up the entry points.
+						EScriptInstrumentationEvent LatentEventStartInfo(EScriptInstrumentation::ResumeEvent, CallbackTarget, ExecutionFunction->GetFName(), LinkInfo.LinkID);
+						FBlueprintCoreDelegates::InstrumentScriptEvent(LatentEventStartInfo);
+						CallbackTarget->ProcessEvent(ExecutionFunction, &(LinkInfo.LinkID));
+						EScriptInstrumentationEvent LatentEventStopInfo(EScriptInstrumentation::Stop, CallbackTarget, NAME_None);
+						FBlueprintCoreDelegates::InstrumentScriptEvent(LatentEventStopInfo);
+					}
+					else
+					{
+						CallbackTarget->ProcessEvent(ExecutionFunction, &(LinkInfo.LinkID));
+					}
 				}
 				else
 				{
diff --git a/Engine/Source/Runtime/GameplayAbilities/Private/GameplayCueManager.cpp b/Engine/Source/Runtime/GameplayAbilities/Private/GameplayCueManager.cpp
--- a/Engine/Source/Runtime/GameplayAbilities/Private/GameplayCueManager.cpp
+++ b/Engine/Source/Runtime/GameplayAbilities/Private/GameplayCueManager.cpp
@@ -398,6 +398,50 @@ void UGameplayCueManager::LoadObjectLibrary_Internal()
 	InitObjectLibraries(LoadedPaths, GameplayCueNotifyActorObjectLibrary, GameplayCueNotifyStaticObjectLibrary, OnLoadDelegate);
 }
 
+static void SearchDynamicClassCues(const FName PropertyName, const TArray<FString>& Paths, TArray<FGameplayCueReferencePair>& CuesToAdd, TArray<FStringAssetReference>& AssetsToLoad)
+{
+	// Iterate over all Dynamic Classes (nativized Blueprints). Search for ones with GameplayCueName tag.
+
+	IGameplayTagsModule& GameplayTagsModule = IGameplayTagsModule::Get();
+	TMap<FName, FDynamicClassStaticData>& DynamicClassMap = GetDynamicClassMap();
+	for (auto PairIter : DynamicClassMap)
+	{
+		const FName* FoundGameplayTag = PairIter.Value.SelectedSearchableValues.Find(PropertyName);
+		if (!FoundGameplayTag)
+		{
+			continue;
+		}
+
+		const FString ClassPath = PairIter.Key.ToString();
+		for (const FString& Path : Paths)
+		{
+			const bool PathContainsClass = ClassPath.StartsWith(Path); // TODO: is it enough?
+			if (!PathContainsClass)
+			{
+				continue;
+			}
+
+			ABILITY_LOG(Log, TEXT("GameplayCueManager Found a Dynamic Class: %s / %s"), *FoundGameplayTag->ToString(), *ClassPath);
+
+			FGameplayTag  GameplayCueTag = GameplayTagsModule.GetGameplayTagsManager().RequestGameplayTag(*FoundGameplayTag, false);
+			if (GameplayCueTag.IsValid())
+			{
+				FStringAssetReference StringRef(ClassPath); // TODO: is there any translation needed?
+				ensure(StringRef.IsValid());
+
+				CuesToAdd.Add(FGameplayCueReferencePair(GameplayCueTag, StringRef));
+				AssetsToLoad.Add(StringRef);
+			}
+			else
+			{
+				ABILITY_LOG(Warning, TEXT("Found GameplayCue tag %s in Dynamic Class %s but there is no corresponding tag in the GameplayTagManager."), *FoundGameplayTag->ToString(), *ClassPath);
+			}
+
+			break;
+		}
+	}
+}
+
 void UGameplayCueManager::InitObjectLibraries(TArray<FString> Paths, UObjectLibrary* ActorObjectLibrary, UObjectLibrary* StaticObjectLibrary, FOnGameplayCueNotifySetLoaded OnLoadDelegate, FShouldLoadGCNotifyDelegate ShouldLoadDelegate)
 {
 	DECLARE_SCOPE_CYCLE_COUNTER(TEXT("Loading Library"), STAT_ObjectLibrary, STATGROUP_LoadTime);
@@ -453,50 +497,9 @@ void UGameplayCueManager::InitObjectLibraries(TArray<FString> Paths, UObjectLibr
 	BuildCuesToAddToGlobalSet(ActorAssetDatas, GET_MEMBER_NAME_CHECKED(AGameplayCueNotify_Actor, GameplayCueName), CuesToAdd, AssetsToLoad, ShouldLoadDelegate);
 	BuildCuesToAddToGlobalSet(StaticAssetDatas, GET_MEMBER_NAME_CHECKED(UGameplayCueNotify_Static, GameplayCueName), CuesToAdd, AssetsToLoad, ShouldLoadDelegate);
 
-	{
-		// Iterate over all Dynamic Classes (nativized Blueprints). Search for ones with GameplayCueName tag.
-
-		IGameplayTagsModule& GameplayTagsModule = IGameplayTagsModule::Get();
-		const FName PropertyName = GET_MEMBER_NAME_CHECKED(AGameplayCueNotify_Actor, GameplayCueName);
-		check(PropertyName == GET_MEMBER_NAME_CHECKED(UGameplayCueNotify_Static, GameplayCueName));
-		TMap<FName, FDynamicClassStaticData>& DynamicClassMap = GetDynamicClassMap();
-		for (auto PairIter : DynamicClassMap)
-		{
-			const FName* FoundGameplayTag = PairIter.Value.SelectedSearchableValues.Find(PropertyName);
-			if (!FoundGameplayTag)
-			{
-				continue;
-			}
-
-			const FString ClassPath = PairIter.Key.ToString();
-			for (FString& Path : Paths)
-			{
-				const bool PathContainsClass = ClassPath.StartsWith(Path); // TODO: is it enough?
-				if (!PathContainsClass)
-				{
-					continue;
-				}
-
-				ABILITY_LOG(Log, TEXT("GameplayCueManager Found a Dynamic Class: %s / %s"), *FoundGameplayTag->ToString(), *ClassPath);
-
-				FGameplayTag  GameplayCueTag = GameplayTagsModule.GetGameplayTagsManager().RequestGameplayTag(*FoundGameplayTag, false);
-				if (GameplayCueTag.IsValid())
-				{
-					FStringAssetReference StringRef(ClassPath); // TODO: is there any translation needed?
-					ensure(StringRef.IsValid());
-
-					CuesToAdd.Add(FGameplayCueReferencePair(GameplayCueTag, StringRef));
-					AssetsToLoad.Add(StringRef);
-				}
-				else
-				{
-					ABILITY_LOG(Warning, TEXT("Found GameplayCue tag %s in Dynamic Class %s but there is no corresponding tag in the GameplayTagManager."), *FoundGameplayTag->ToString(), *ClassPath);
-				}
-
-				break;
-			}
-		}
-	}
+	const FName PropertyName = GET_MEMBER_NAME_CHECKED(AGameplayCueNotify_Actor, GameplayCueName);
+	check(PropertyName == GET_MEMBER_NAME_CHECKED(UGameplayCueNotify_Static, GameplayCueName));
+	SearchDynamicClassCues(PropertyName, Paths, CuesToAdd, AssetsToLoad);
 
 	// Add these cues to the global set
 	check(GlobalCueSet);
