Copying //UE4/Dev-Core to //UE4/Dev-Main (Source: //UE4/Dev-Core @ 3049602)

diff --git a/Engine/Source/Runtime/Core/Public/Templates/Function.h b/Engine/Source/Runtime/Core/Public/Templates/Function.h
--- a/Engine/Source/Runtime/Core/Public/Templates/Function.h
+++ b/Engine/Source/Runtime/Core/Public/Templates/Function.h
@@ -184,7 +184,7 @@ namespace UE4Function_Private
 	struct TFunctionRefCaller;
 
 	/**
-	 * A class which is used to instantiate the code needed to assert when called - used for null bindings.
+	 * A class which is used to instantiate the code needed to assert when called - used for unset bindings.
 	 */
 	template <typename FuncType>
 	struct TFunctionRefAsserter;
@@ -225,66 +225,6 @@ namespace UE4Function_Private
 	{
 	}
 
-	template <typename FuncType, typename CallableType>
-	struct TFunctionRefBaseCommon
-	{
-		explicit TFunctionRefBaseCommon(ENoInit)
-		{
-			// Not really designed to be initialized directly, but want to be explicit about that.
-		}
-
-		template <typename FunctorType>
-		void Set(FunctorType* Functor)
-		{
-			Callable = &UE4Function_Private::TFunctionRefCaller<FunctorType, FuncType>::Call;
-
-			#if ENABLE_TFUNCTIONREF_VISUALIZATION
-				// We placement new over the top of the same object each time.  This is illegal,
-				// but it ensures that the vptr is set correctly for the bound type, and so is
-				// visualizable.  We never depend on the state of this object at runtime, so it's
-				// ok.
-				new ((void*)&DebugPtrStorage) UE4Function_Private::TDebugHelper<FunctorType>;
-				DebugPtrStorage.Ptr = (void*)Functor;
-			#endif
-		}
-
-		void CopyAndReseat(const TFunctionRefBaseCommon& Other, void* Functor)
-		{
-			Callable = Other.Callable;
-
-			#if ENABLE_TFUNCTIONREF_VISUALIZATION
-				// Use Memcpy to copy the other DebugPtrStorage, including vptr (because we don't know the bound type
-				// here), and then reseat the underlying pointer.  Possibly even more evil than the Set code.
-				FMemory::Memcpy(&DebugPtrStorage, &Other.DebugPtrStorage, sizeof(DebugPtrStorage));
-				DebugPtrStorage.Ptr = Functor;
-			#endif
-		}
-
-		void Unset()
-		{
-			Callable = &UE4Function_Private::TFunctionRefAsserter<FuncType>::Call;
-		}
-
-		CallableType* GetCallable() const
-		{
-			return Callable;
-		}
-
-	private:
-		// A pointer to a function which invokes the call operator on the callable object
-		CallableType* Callable;
-
-		#if ENABLE_TFUNCTIONREF_VISUALIZATION
-			// To help debug visualizers
-			UE4Function_Private::TDebugHelper<void> DebugPtrStorage;
-		#endif
-	};
-
-	/**
-	 * Switch on the existence of variadics.  Once all our supported compilers support variadics, a lot of this code
-	 * can be collapsed into FFunctionRefCaller.  They're currently separated out to minimize the amount of workarounds
-	 * needed.
-	 */
 	template <typename Functor, typename Ret, typename... ParamTypes>
 	struct TFunctionRefCaller<Functor, Ret (ParamTypes...)>
 	{
@@ -311,7 +251,7 @@ namespace UE4Function_Private
 	{
 		static Ret Call(void* Obj, ParamTypes&...)
 		{
-			checkf(false, TEXT("Attempting to call a null TFunction!"));
+			checkf(false, TEXT("Attempting to call an unbound TFunction!"));
 
 			// This doesn't need to be valid, because it'll never be reached, but it does at least need to compile.
 			return FakeCall((Ret*)Obj);
@@ -319,26 +259,59 @@ namespace UE4Function_Private
 	};
 
 	template <typename DerivedType, typename Ret, typename... ParamTypes>
-	struct TFunctionRefBase<DerivedType, Ret (ParamTypes...)> : TFunctionRefBaseCommon<Ret (ParamTypes...), Ret (void*, ParamTypes&...)>
+	struct TFunctionRefBase<DerivedType, Ret (ParamTypes...)>
 	{
-		typedef TFunctionRefBaseCommon<Ret (ParamTypes...), Ret (void*, ParamTypes&...)> Super;
-
 		explicit TFunctionRefBase(ENoInit)
-			: Super(NoInit)
 		{
+			// Not really designed to be initialized directly, but we want to be explicit about that.
+		}
+
+		Ret operator()(ParamTypes... Params) const
+		{
+			const DerivedType* Derived = static_cast<const DerivedType*>(this);
+			return Callable(Derived->GetPtr(), Params...);
 		}
 
 		template <typename FunctorType>
-		explicit TFunctionRefBase(FunctorType* Functor)
-			: Super(Functor)
+		void Set(FunctorType* Functor)
 		{
+			Callable = &UE4Function_Private::TFunctionRefCaller<FunctorType, Ret (ParamTypes...)>::Call;
+
+			#if ENABLE_TFUNCTIONREF_VISUALIZATION
+				// We placement new over the top of the same object each time.  This is illegal,
+				// but it ensures that the vptr is set correctly for the bound type, and so is
+				// visualizable.  We never depend on the state of this object at runtime, so it's
+				// ok.
+				new ((void*)&DebugPtrStorage) UE4Function_Private::TDebugHelper<FunctorType>;
+				DebugPtrStorage.Ptr = (void*)Functor;
+			#endif
 		}
 
-		Ret operator()(ParamTypes... Params) const
+		void CopyAndReseat(const TFunctionRefBase& Other, void* Functor)
 		{
-			const DerivedType* Derived = static_cast<const DerivedType*>(this);
-			return this->GetCallable()(Derived->GetPtr(), Params...);
+			Callable = Other.Callable;
+
+			#if ENABLE_TFUNCTIONREF_VISUALIZATION
+				// Use Memcpy to copy the other DebugPtrStorage, including vptr (because we don't know the bound type
+				// here), and then reseat the underlying pointer.  Possibly even more evil than the Set code.
+				FMemory::Memcpy(&DebugPtrStorage, &Other.DebugPtrStorage, sizeof(DebugPtrStorage));
+				DebugPtrStorage.Ptr = Functor;
+			#endif
+		}
+
+		void Unset()
+		{
+			Callable = &UE4Function_Private::TFunctionRefAsserter<Ret (ParamTypes...)>::Call;
 		}
+
+	private:
+		// A pointer to a function which invokes the call operator on the callable object
+		Ret (*Callable)(void*, ParamTypes&...);
+
+		#if ENABLE_TFUNCTIONREF_VISUALIZATION
+			// To help debug visualizers
+			UE4Function_Private::TDebugHelper<void> DebugPtrStorage;
+		#endif
 	};
 }
 
@@ -456,7 +429,7 @@ public:
 
 	// We delete the assignment operators because we don't want it to be confused with being related to
 	// regular C++ reference assignment - i.e. calling the assignment operator of whatever the reference
-	// is bound to - because that's not what TFunctionRef does, or is it even capable of doing that.
+	// is bound to - because that's not what TFunctionRef does, nor is it even capable of doing that.
 	#if PLATFORM_COMPILER_HAS_DEFAULTED_FUNCTIONS
 
 		#if !ENABLE_TFUNCTIONREF_VISUALIZATION
@@ -490,15 +463,6 @@ private:
 		Super::Set(Functor);
 	}
 
-	/**
-	 * 'Nulls' the TFunctionRef.
-	 */
-	void Unset()
-	{
-		Ptr = nullptr;
-		Super::Unset();
-	}
-
 	/**
 	 * Copies another TFunctionRef and rebinds it to another object of the same type which was originally bound.
 	 * Only intended to be used by TFunction's copy constructor/assignment operator.
@@ -647,7 +611,7 @@ public:
 	}
 
 	/**
-	 * Nullptr assignment operator.
+	 * Nullptr assignment operator - unbinds any bound function.
 	 */
 	TFunction& operator=(TYPE_OF_NULLPTR)
 	{
@@ -675,17 +639,11 @@ public:
 	/**
 	 * Tests if the TFunction is callable.
 	 */
-	FORCEINLINE_EXPLICIT_OPERATOR_BOOL() const
+	FORCEINLINE explicit operator bool() const
 	{
 		return !!Storage.GetBoundObject();
 	}
 
-	/**
-	 * TFunctions are not usefully comparable, but we need to define operators for FORCEINLINE_EXPLICIT_OPERATOR_BOOL to work
-	 * in the case where explicit operator bool is not supported by the compiler.
-	 */
-	SAFE_BOOL_OPERATORS(TFunction)
-
 private:
 	/**
 	 * Returns a pointer to the callable object - needed by TFunctionRefBase.
diff --git a/Engine/Source/Runtime/Core/Public/Templates/HasOperatorEquals.h b/Engine/Source/Runtime/Core/Public/Templates/HasOperatorEquals.h
new file mode 100644
--- /dev/null
+++ b/Engine/Source/Runtime/Core/Public/Templates/HasOperatorEquals.h
@@ -0,0 +1,59 @@
+// Copyright 1998-2016 Epic Games, Inc. All Rights Reserved.
+
+#pragma once
+
+namespace UE4HasOperatorEquals_Private
+{
+	struct FNotSpecified {};
+
+	template <typename T>
+	struct FReturnValueCheck
+	{
+		static char (&Func())[2];
+	};
+
+	template <>
+	struct FReturnValueCheck<FNotSpecified>
+	{
+		static char (&Func())[1];
+	};
+
+	template <typename T>
+	FNotSpecified operator==(const T&, const T&);
+
+	template <typename T>
+	FNotSpecified operator!=(const T&, const T&);
+
+	template <typename T>
+	const T& Make();
+
+	template <typename T>
+	struct Equals
+	{
+		enum { Value = sizeof(FReturnValueCheck<decltype(Make<T>() == Make<T>())>::Func()) == sizeof(char[2]) };
+	};
+
+	template <typename T>
+	struct NotEquals
+	{
+		enum { Value = sizeof(FReturnValueCheck<decltype(Make<T>() != Make<T>())>::Func()) == sizeof(char[2]) };
+	};
+}
+
+/**
+ * Traits class which tests if a type has an operator== overload.
+ */
+template <typename T>
+struct THasOperatorEquals
+{
+	enum { Value = UE4HasOperatorEquals_Private::Equals<T>::Value };
+};
+
+/**
+ * Traits class which tests if a type has an operator!= overload.
+ */
+template <typename T>
+struct THasOperatorNotEquals
+{
+	enum { Value = UE4HasOperatorEquals_Private::NotEquals<T>::Value };
+};
diff --git a/Engine/Source/Runtime/Core/Public/UObject/AutoPointer.h b/Engine/Source/Runtime/Core/Public/UObject/AutoPointer.h
--- a/Engine/Source/Runtime/Core/Public/UObject/AutoPointer.h
+++ b/Engine/Source/Runtime/Core/Public/UObject/AutoPointer.h
@@ -31,13 +31,7 @@ public:
 		return (const T*)TBASE::Get();
 	}
 
-#if PLATFORM_COMPILER_HAS_EXPLICIT_OPERATORS
 	FORCEINLINE explicit operator bool() const
-#else
-	// Can't use FORCEINLINE_EXPLICIT_OPERATOR_BOOL() due to already having an implicit pointer cast
-	// (which breaks the safe bool idiom it uses when PLATFORM_COMPILER_HAS_EXPLICIT_OPERATORS is 0)
-	FORCEINLINE operator bool() const
-#endif
 	{
 		return TBASE::Get() != nullptr;
 	}
diff --git a/Engine/Source/Runtime/Core/Public/WinRT/WinRTPlatform.h b/Engine/Source/Runtime/Core/Public/WinRT/WinRTPlatform.h
--- a/Engine/Source/Runtime/Core/Public/WinRT/WinRTPlatform.h
+++ b/Engine/Source/Runtime/Core/Public/WinRT/WinRTPlatform.h
@@ -53,7 +53,6 @@ typedef FWinRTTypes FPlatformTypes;
 #if _MSC_VER < 1900
 	#define PLATFORM_COMPILER_HAS_DEFAULTED_FUNCTIONS		0
 #endif
-#define PLATFORM_COMPILER_HAS_EXPLICIT_OPERATORS			1
 
 [REDACTED_EMAIL]: Fixup once sockets are supported
 #define PLATFORM_SUPPORTS_MESSAGEBUS						0
diff --git a/Engine/Source/Runtime/Core/Public/Windows/WIndowsPlatform.h b/Engine/Source/Runtime/Core/Public/Windows/WIndowsPlatform.h
--- a/Engine/Source/Runtime/Core/Public/Windows/WIndowsPlatform.h
+++ b/Engine/Source/Runtime/Core/Public/Windows/WIndowsPlatform.h
@@ -61,7 +61,6 @@ typedef FWindowsPlatformTypes FPlatformTypes;
 #if _MSC_VER < 1900
 	#define PLATFORM_COMPILER_HAS_DEFAULTED_FUNCTIONS		0
 #endif
-#define PLATFORM_COMPILER_HAS_EXPLICIT_OPERATORS		1
 #define PLATFORM_COMPILER_HAS_TCHAR_WMAIN					1
 
 #define PLATFORM_RHITHREAD_DEFAULT_BYPASS					WITH_EDITOR
diff --git a/Engine/Source/Runtime/CoreUObject/Private/UObject/Enum.cpp b/Engine/Source/Runtime/CoreUObject/Private/UObject/Enum.cpp
--- a/Engine/Source/Runtime/CoreUObject/Private/UObject/Enum.cpp
+++ b/Engine/Source/Runtime/CoreUObject/Private/UObject/Enum.cpp
@@ -252,7 +252,7 @@ void UEnum::AddNamesToMasterList()
 		{
 			AllEnumNames.Add(Kvp.Key, this);
 		}
-		else if (Enum != this)
+		else if (Enum != this && Enum->GetOutermost() != GetTransientPackage())
 		{
 			UE_LOG(LogEnum, Warning, TEXT("Enum name collision: '%s' is in both '%s' and '%s'"), *Kvp.Key.ToString(), *GetPathName(), *Enum->GetPathName());
 		}
diff --git a/Engine/Source/Runtime/CoreUObject/Private/UObject/MetaData.cpp b/Engine/Source/Runtime/CoreUObject/Private/UObject/MetaData.cpp
--- a/Engine/Source/Runtime/CoreUObject/Private/UObject/MetaData.cpp
+++ b/Engine/Source/Runtime/CoreUObject/Private/UObject/MetaData.cpp
@@ -408,11 +408,6 @@ bool UMetaData::NeedsLoadForClient() const
 	return false; 
 }
 
-bool UMetaData::NeedsLoadForServer() const
-{ 
-	return false; 
-}
-
 bool UMetaData::NeedsLoadForEditorGame() const
 {
 	return true;
diff --git a/Engine/Source/Runtime/CoreUObject/Private/UObject/Obj.cpp b/Engine/Source/Runtime/CoreUObject/Private/UObject/Obj.cpp
--- a/Engine/Source/Runtime/CoreUObject/Private/UObject/Obj.cpp
+++ b/Engine/Source/Runtime/CoreUObject/Private/UObject/Obj.cpp
@@ -488,6 +488,68 @@ void UObject::PostEditUndo(TSharedPtr<ITransactionObjectAnnotation> TransactionA
 
 #endif // WITH_EDITOR
 
+class FCachedClassExclusionData
+{
+public:
+
+	FCachedClassExclusionData(const TCHAR* InSectionName)
+	{
+		check(GConfig != nullptr && GConfig->IsReadyForUse());
+
+		TArray<FString> ConfigData;
+		GConfig->GetArray(TEXT("Core.System"), InSectionName, ConfigData, GEngineIni);
+
+		for (TObjectIterator<UClass> ClassIt; ClassIt; ++ClassIt)
+		{
+			UClass* Class = *ClassIt;
+			if (ConfigData.Contains(Class->GetName()))
+			{
+				ExcludedClasses.Add(Class);
+			}
+		}
+
+		// Now gather all the classes which are derived from the specified ones
+		TSet<UClass*> DerivedClasses;
+
+		for (TObjectIterator<UClass> ClassIt; ClassIt; ++ClassIt)
+		{
+			UClass* Class = *ClassIt;
+
+			for (const UClass* SpecifiedClass : ExcludedClasses)
+			{
+				if (Class->IsChildOf(SpecifiedClass))
+				{
+					if (!ExcludedClasses.Contains(Class))
+					{
+						DerivedClasses.Add(Class);
+					}
+					
+					break;
+				}
+			}
+		}
+
+		// Add them together
+		ExcludedClasses.Append(DerivedClasses);
+	}
+
+	bool IsExcluded(UClass* InClass) const 
+	{
+		return ExcludedClasses.Contains(InClass);
+	}
+
+private:
+
+	TSet<UClass*> ExcludedClasses;
+
+};
+
+bool UObject::NeedsLoadForServer() const
+{
+	static const FCachedClassExclusionData CachedClassExclusionData(TEXT("ClassesExcludedForServer"));
+	return !CachedClassExclusionData.IsExcluded(GetClass());
+}
+
 bool UObject::CanCreateInCurrentContext(UObject* Template)
 {
 	check(Template);
diff --git a/Engine/Source/Runtime/CoreUObject/Private/UObject/ObjectRedirector.cpp b/Engine/Source/Runtime/CoreUObject/Private/UObject/ObjectRedirector.cpp
--- a/Engine/Source/Runtime/CoreUObject/Private/UObject/ObjectRedirector.cpp
+++ b/Engine/Source/Runtime/CoreUObject/Private/UObject/ObjectRedirector.cpp
@@ -42,11 +42,6 @@ bool UObjectRedirector::NeedsLoadForClient() const
 	return false;
 }
 
-bool UObjectRedirector::NeedsLoadForServer() const
-{
-	return false;
-}
-
 bool UObjectRedirector::NeedsLoadForEditorGame() const
 {
 	return true;
diff --git a/Engine/Source/Runtime/CoreUObject/Private/UObject/SavePackage.cpp b/Engine/Source/Runtime/CoreUObject/Private/UObject/SavePackage.cpp
--- a/Engine/Source/Runtime/CoreUObject/Private/UObject/SavePackage.cpp
+++ b/Engine/Source/Runtime/CoreUObject/Private/UObject/SavePackage.cpp
@@ -4343,7 +4343,7 @@ ESavePackageResult UPackage::Save(UPackage* InOuter, UObject* Base, EObjectFlags
 #if WITH_EDITOR
 							// If we still didn't find index, maybe it was a duplicate export which got removed.
 							// Check if we have a redirect to original.
-							if (DependencyIndex.IsNull())
+							if (DependencyIndex.IsNull() && DuplicateRedirects.Contains(DependentObject))
 							{
 								DependencyIndex = ExportToIndexMap.FindRef(DuplicateRedirects[DependentObject]);
 							}
diff --git a/Engine/Source/Runtime/CoreUObject/Private/UObject/UObjectArray.cpp b/Engine/Source/Runtime/CoreUObject/Private/UObject/UObjectArray.cpp
--- a/Engine/Source/Runtime/CoreUObject/Private/UObject/UObjectArray.cpp
+++ b/Engine/Source/Runtime/CoreUObject/Private/UObject/UObjectArray.cpp
@@ -14,7 +14,7 @@ FUObjectArray::FUObjectArray()
 : ObjFirstGCIndex(0)
 , ObjLastNonGCIndex(INDEX_NONE)
 , MaxObjectsNotConsideredByGC(0)
-, OpenForDisregardForGC(!IS_PROGRAM)
+, OpenForDisregardForGC(!HACK_HEADER_GENERATOR)
 , MasterSerialNumber(START_SERIAL_NUMBER)
 {
 	FCoreDelegates::GetObjectArrayForDebugVisualizersDelegate().BindStatic(GetObjectArrayForDebugVisualizers);
diff --git a/Engine/Source/Runtime/CoreUObject/Private/UObject/UObjectBase.cpp b/Engine/Source/Runtime/CoreUObject/Private/UObject/UObjectBase.cpp
--- a/Engine/Source/Runtime/CoreUObject/Private/UObject/UObjectBase.cpp
+++ b/Engine/Source/Runtime/CoreUObject/Private/UObject/UObjectBase.cpp
@@ -493,18 +493,6 @@ void UObjectCompiledInDeferEnum(class UEnum *(*InRegister)(), const TCHAR* Packa
 
 class UEnum *GetStaticEnum(class UEnum *(*InRegister)(), UObject* EnumOuter, const TCHAR* EnumName)
 {
-#if WITH_HOT_RELOAD
-	if (GIsHotReload)
-	{
-		UEnum* ReturnEnum = FindObject<UEnum>(EnumOuter, EnumName);
-		if (ReturnEnum)
-		{
-			UE_LOG(LogClass, Log, TEXT( "%s HotReload."), EnumName);
-			return ReturnEnum;
-		}
-		UE_LOG(LogClass, Log, TEXT("Could not find existing enum %s for HotReload. Assuming new"), EnumName);
-	}
-#endif
 	return (*InRegister)();
 }
 
@@ -666,20 +654,23 @@ void UObjectCompiledInDefer(UClass *(*InRegister)(), UClass *(*InStaticClass)(),
 /** Register all loaded classes */
 void UClassRegisterAllCompiledInClasses()
 {
-	for (auto Class : GetDeferredClassRegistration())
+	TArray<FFieldCompiledInInfo*>& DeferredClassRegistration = GetDeferredClassRegistration();
+	for (const FFieldCompiledInInfo* Class : DeferredClassRegistration)
 	{
 		Class->Register();
 	}
-	GetDeferredClassRegistration().Empty();
+	DeferredClassRegistration.Empty();
 }
 
 #if WITH_HOT_RELOAD
 /** Re-instance all existing classes that have changed during hot-reload */
 void UClassReplaceHotReloadClasses()
 {
+	TArray<FFieldCompiledInInfo*>& HotReloadClasses = GetHotReloadClasses();
+
 	if (FCoreUObjectDelegates::RegisterClassForHotReloadReinstancingDelegate.IsBound())
 	{
-		for (auto Class : GetHotReloadClasses())
+		for (const FFieldCompiledInInfo* Class : HotReloadClasses)
 		{
 			check(Class->OldClass);
 
@@ -694,7 +685,7 @@ void UClassReplaceHotReloadClasses()
 	}
 
 	FCoreUObjectDelegates::ReinstanceHotReloadedClassesDelegate.ExecuteIfBound();
-	GetHotReloadClasses().Empty();
+	HotReloadClasses.Empty();
 }
 
 /**
@@ -708,21 +699,26 @@ static void UClassGenerateCDODuplicatesForHotReload()
 		return;
 	}
 
-	for (FObjectIterator It(UClass::StaticClass()); It; ++It)
+	const TArray<FFieldCompiledInInfo*>& HotReloadClasses = GetHotReloadClasses();
+	TMap<UObject*, UObject*> DuplicatedCDOMap = GetDuplicatedCDOMap();
+	UPackage* TransientPackage = GetTransientPackage();
+
+	for (UClass* Class : TObjectRange<UClass>())
 	{
-		UClass* Class = (UClass*)*It;
+		UObject* CDO = Class->GetDefaultObject();
 
-		for (auto* HotReloadedClass : GetHotReloadClasses())
+		for (const FFieldCompiledInInfo* HotReloadedClass : HotReloadClasses)
 		{
 			if (!HotReloadedClass->bHasChanged && Class->IsChildOf(HotReloadedClass->OldClass))
 			{
 				GIsDuplicatingClassForReinstancing = true;
-				UObject* DupCDO = (UObject*)StaticDuplicateObject(
-					Class->GetDefaultObject(), GetTransientPackage(),
-					MakeUniqueObjectName(GetTransientPackage(), Class, TEXT("HOTRELOAD_CDO_DUPLICATE"))
-					);
+
+				FName UniqueName = MakeUniqueObjectName(TransientPackage, Class, TEXT("HOTRELOAD_CDO_DUPLICATE"));
+				UObject* DupCDO = StaticDuplicateObject(CDO, TransientPackage, UniqueName);
+
 				GIsDuplicatingClassForReinstancing = false;
-				GetDuplicatedCDOMap().Add(Class->GetDefaultObject(), DupCDO);
+
+				DuplicatedCDOMap.Add(CDO, DupCDO);
 			}
 		}
 	}
@@ -734,20 +730,19 @@ static void UClassGenerateCDODuplicatesForHotReload()
  */
 static void UObjectLoadAllCompiledInDefaultProperties()
 {
-	const bool bHaveRegistrants = GetDeferredCompiledInRegistration().Num() != 0;
+	TArray<UClass *(*)()>& DeferredCompiledInRegistration = GetDeferredCompiledInRegistration();
+
+	const bool bHaveRegistrants = DeferredCompiledInRegistration.Num() != 0;
 	if( bHaveRegistrants )
 	{
-		TArray<class UClass *(*)()> PendingRegistrants;
 		TArray<UClass*> NewClasses;
-		PendingRegistrants = GetDeferredCompiledInRegistration();
-		GetDeferredCompiledInRegistration().Empty();
-		for(int32 RegistrantIndex = 0;RegistrantIndex < PendingRegistrants.Num();++RegistrantIndex)
+		TArray<UClass* (*)()> PendingRegistrants = MoveTemp(DeferredCompiledInRegistration);
+		for (UClass* (*Registrant)() : PendingRegistrants)
 		{
-			NewClasses.Add(PendingRegistrants[RegistrantIndex]());
+			NewClasses.Add(Registrant());
 		}
-		for (int32 Index = 0; Index < NewClasses.Num(); Index++)
+		for (UClass* Class : NewClasses)
 		{
-			UClass* Class = NewClasses[Index];
 			Class->GetDefaultObject();
 		}
 		FFeedbackContext& ErrorsFC = UClass::GetDefaultPropertiesFeedbackContext();
@@ -774,53 +769,44 @@ static void UObjectLoadAllCompiledInDefaultProperties()
  */
 static void UObjectLoadAllCompiledInStructs()
 {
+	TArray<FPendingEnumRegistrant>& DeferredCompiledInEnumRegistration = GetDeferredCompiledInEnumRegistration();
+	TArray<FPendingStructRegistrant>& DeferredCompiledInStructRegistration = GetDeferredCompiledInStructRegistration();
+
 	// Load Enums first
-	const bool bHaveEnumRegistrants = GetDeferredCompiledInEnumRegistration().Num() != 0;
-	if( bHaveEnumRegistrants )
+	if( DeferredCompiledInEnumRegistration.Num() )
 	{
-		TArray<FPendingEnumRegistrant> PendingRegistrants;
-		PendingRegistrants = GetDeferredCompiledInEnumRegistration();
-		GetDeferredCompiledInEnumRegistration().Empty();
+		TArray<FPendingEnumRegistrant> PendingRegistrants = MoveTemp(DeferredCompiledInEnumRegistration);
 		
-		for (auto& EnumRegistrant : PendingRegistrants)
+		for (const FPendingEnumRegistrant& EnumRegistrant : PendingRegistrants)
 		{
 			// Make sure the package exists in case it does not contain any UObjects
-			CreatePackage(NULL, EnumRegistrant.PackageName);
+			CreatePackage(nullptr, EnumRegistrant.PackageName);
 		}
 
-		for (auto& EnumRegistrant : PendingRegistrants)
+		for (const FPendingEnumRegistrant& EnumRegistrant : PendingRegistrants)
 		{
 			EnumRegistrant.RegisterFn();
 		}
 	}
 
 	// Load Structs
-	const bool bHaveRegistrants = GetDeferredCompiledInStructRegistration().Num() != 0;
-	if( bHaveRegistrants )
+	if( DeferredCompiledInStructRegistration.Num() )
 	{
-		TArray<FPendingStructRegistrant> PendingRegistrants;
-		PendingRegistrants = GetDeferredCompiledInStructRegistration();
-		GetDeferredCompiledInStructRegistration().Empty();
+		TArray<FPendingStructRegistrant> PendingRegistrants = MoveTemp(DeferredCompiledInStructRegistration);
 
-		for (auto& StructRegistrant : PendingRegistrants)
+		for (const FPendingStructRegistrant& StructRegistrant : PendingRegistrants)
 		{
 			// Make sure the package exists in case it does not contain any UObjects or UEnums
-			CreatePackage(NULL, StructRegistrant.PackageName);
+			CreatePackage(nullptr, StructRegistrant.PackageName);
 		}
 
-		for (auto& StructRegistrant : PendingRegistrants)
+		for (const FPendingStructRegistrant& StructRegistrant : PendingRegistrants)
 		{
 			StructRegistrant.RegisterFn();
 		}
 	}
 }
 
-bool AnyNewlyLoadedUObjects()
-{
-	return GFirstPendingRegistrant != NULL || GetDeferredCompiledInRegistration().Num() || GetDeferredCompiledInStructRegistration().Num() || GetDeferredCompiledInEnumRegistration().Num();
-}
-
-
 void ProcessNewlyLoadedUObjects()
 {
 	DECLARE_SCOPE_CYCLE_COUNTER(TEXT("ProcessNewlyLoadedUObjects"), STAT_ProcessNewlyLoadedUObjects, STATGROUP_ObjectVerbose);
@@ -830,13 +816,17 @@ void ProcessNewlyLoadedUObjects()
 #endif
 	UClassRegisterAllCompiledInClasses();
 
+	const TArray<UClass* (*)()>& DeferredCompiledInRegistration = GetDeferredCompiledInRegistration();
+	const TArray<FPendingStructRegistrant>& DeferredCompiledInStructRegistration = GetDeferredCompiledInStructRegistration();
+	const TArray<FPendingEnumRegistrant>& DeferredCompiledInEnumRegistration = GetDeferredCompiledInEnumRegistration();
+
 	bool bNewUObjects = false;
-	while( AnyNewlyLoadedUObjects() )
+	while( GFirstPendingRegistrant || DeferredCompiledInRegistration.Num() || DeferredCompiledInStructRegistration.Num() || DeferredCompiledInEnumRegistration.Num() )
 	{
 		bNewUObjects = true;
 		UObjectProcessRegistrants();
 		UObjectLoadAllCompiledInStructs();
-		UObjectLoadAllCompiledInDefaultProperties();		
+		UObjectLoadAllCompiledInDefaultProperties();
 	}
 #if WITH_HOT_RELOAD
 	UClassReplaceHotReloadClasses();
diff --git a/Engine/Source/Runtime/CoreUObject/Public/Serialization/BulkData.h b/Engine/Source/Runtime/CoreUObject/Public/Serialization/BulkData.h
--- a/Engine/Source/Runtime/CoreUObject/Public/Serialization/BulkData.h
+++ b/Engine/Source/Runtime/CoreUObject/Public/Serialization/BulkData.h
@@ -102,13 +102,11 @@ private:
 			return Ptr;
 		}
 
-		FORCEINLINE_EXPLICIT_OPERATOR_BOOL() const
+		FORCEINLINE explicit operator bool() const
 		{
 			return bAllocated;
 		}
 
-		SAFE_BOOL_OPERATORS(FAllocatedPtr)
-
 		void Reallocate(int32 Count, int32 Alignment = DEFAULT_ALIGNMENT)
 		{
 			if (Count)
diff --git a/Engine/Source/Runtime/CoreUObject/Public/UObject/CoreObject.h b/Engine/Source/Runtime/CoreUObject/Public/UObject/CoreObject.h
--- a/Engine/Source/Runtime/CoreUObject/Public/UObject/CoreObject.h
+++ b/Engine/Source/Runtime/CoreUObject/Public/UObject/CoreObject.h
@@ -679,7 +679,6 @@ public:
 	// UObject interface
 	virtual void Serialize(FArchive& Ar) override;
 	virtual bool NeedsLoadForClient() const override;
-	virtual bool NeedsLoadForServer() const override;
 	virtual bool NeedsLoadForEditorGame() const override;
 	virtual bool IsAsset() const override { return false; }
 	// End of UObject interface
diff --git a/Engine/Source/Runtime/CoreUObject/Public/UObject/UObject.h b/Engine/Source/Runtime/CoreUObject/Public/UObject/UObject.h
--- a/Engine/Source/Runtime/CoreUObject/Public/UObject/UObject.h
+++ b/Engine/Source/Runtime/CoreUObject/Public/UObject/UObject.h
@@ -348,10 +348,7 @@ public:
 	 *
 	 * @return	true if this object should not be loaded on servers
 	 */
-	virtual bool NeedsLoadForServer() const 
-	{ 
-		return true; 
-	}
+	virtual bool NeedsLoadForServer() const;
 
 	/**
 	 * Called during saving to determine the load flags to save with the object.
diff --git a/Engine/Source/Runtime/CoreUObject/Public/UObject/UObjectBase.h b/Engine/Source/Runtime/CoreUObject/Public/UObject/UObjectBase.h
--- a/Engine/Source/Runtime/CoreUObject/Public/UObject/UObjectBase.h
+++ b/Engine/Source/Runtime/CoreUObject/Public/UObject/UObjectBase.h
@@ -370,9 +370,6 @@ COREUOBJECT_API class UEnum *GetStaticEnum(class UEnum *(*InRegister)(), UObject
 COREUOBJECT_API class UScriptStruct* FindExistingStructIfHotReloadOrDynamic(UObject* Outer, const TCHAR* StructName, SIZE_T Size, uint32 Crc, bool bIsDynamic);
 COREUOBJECT_API class UEnum* FindExistingEnumIfHotReloadOrDynamic(UObject* Outer, const TCHAR* EnumName, SIZE_T Size, uint32 Crc, bool bIsDynamic);
 
-/** @return	True if there are any newly-loaded UObjects that are waiting to be registered by calling ProcessNewlyLoadedUObjects() */
-COREUOBJECT_API bool AnyNewlyLoadedUObjects();
-
 /** Must be called after a module has been loaded that contains UObject classes */
 COREUOBJECT_API void ProcessNewlyLoadedUObjects();
 
diff --git a/Engine/Source/Runtime/DatabaseSupport/Public/Database.h b/Engine/Source/Runtime/DatabaseSupport/Public/Database.h
--- a/Engine/Source/Runtime/DatabaseSupport/Public/Database.h
+++ b/Engine/Source/Runtime/DatabaseSupport/Public/Database.h
@@ -168,10 +168,8 @@ public:
 			RecordSet->MoveToNext();	
 		}
 
-		SAFE_BOOL_OPERATORS(TIterator)
-
 		/** Conversion to "bool" returning true if the iterator is valid. */
-		FORCEINLINE_EXPLICIT_OPERATOR_BOOL() const
+		FORCEINLINE explicit operator bool() const
 		{ 
 			return !RecordSet->IsAtEnd(); 
 		}
diff --git a/Engine/Source/Runtime/Engine/Classes/AI/Navigation/NavMeshRenderingComponent.h b/Engine/Source/Runtime/Engine/Classes/AI/Navigation/NavMeshRenderingComponent.h
--- a/Engine/Source/Runtime/Engine/Classes/AI/Navigation/NavMeshRenderingComponent.h
+++ b/Engine/Source/Runtime/Engine/Classes/AI/Navigation/NavMeshRenderingComponent.h
@@ -144,9 +144,6 @@ class ENGINE_API UNavMeshRenderingComponent : public UPrimitiveComponent
 	GENERATED_UCLASS_BODY()
 
 public:
-	//~ Begin UObject Interface
-	virtual bool NeedsLoadForServer() const override;
-	//~ End UObject Interface
 	
 	//~ Begin UPrimitiveComponent Interface
 	virtual FPrimitiveSceneProxy* CreateSceneProxy() override;
diff --git a/Engine/Source/Runtime/Engine/Classes/Components/ReflectionCaptureComponent.h b/Engine/Source/Runtime/Engine/Classes/Components/ReflectionCaptureComponent.h
--- a/Engine/Source/Runtime/Engine/Classes/Components/ReflectionCaptureComponent.h
+++ b/Engine/Source/Runtime/Engine/Classes/Components/ReflectionCaptureComponent.h
@@ -168,10 +168,6 @@ class UReflectionCaptureComponent : public USceneComponent
 	virtual void BeginDestroy() override;
 	virtual bool IsReadyForFinishDestroy() override;
 	virtual void FinishDestroy() override;
-	virtual bool NeedsLoadForServer() const override
-	{
-		return false;
-	}
 	//~ End UObject Interface
 
 private:
diff --git a/Engine/Source/Runtime/Engine/Classes/Materials/MaterialExpression.h b/Engine/Source/Runtime/Engine/Classes/Materials/MaterialExpression.h
--- a/Engine/Source/Runtime/Engine/Classes/Materials/MaterialExpression.h
+++ b/Engine/Source/Runtime/Engine/Classes/Materials/MaterialExpression.h
@@ -61,8 +61,8 @@ struct FMaterialAttributesInput : public FExpressionInput
 
 #endif
 
-/** Struct that represents an expression's output. */
 #if !CPP      //noexport struct
+/** Struct that represents an expression's output. */
 USTRUCT(noexport)
 struct FExpressionOutput
 {
@@ -193,7 +193,6 @@ class ENGINE_API UMaterialExpression : public UObject
 	virtual bool Modify( bool bAlwaysMarkDirty=true ) override;
 	virtual void Serialize( FArchive& Ar ) override;
 	virtual bool NeedsLoadForClient() const override;
-	virtual bool NeedsLoadForServer() const override;
 	virtual bool NeedsLoadForEditorGame() const override
 	{
 		return true;
diff --git a/Engine/Source/Runtime/Engine/Classes/Matinee/MatineeActorCameraAnim.h b/Engine/Source/Runtime/Engine/Classes/Matinee/MatineeActorCameraAnim.h
--- a/Engine/Source/Runtime/Engine/Classes/Matinee/MatineeActorCameraAnim.h
+++ b/Engine/Source/Runtime/Engine/Classes/Matinee/MatineeActorCameraAnim.h
@@ -23,10 +23,5 @@ class AMatineeActorCameraAnim : public AMatineeActor
 	{ 
 		return false; 
 	}
-
-	virtual bool NeedsLoadForServer() const override
-	{ 
-		return false;
-	}
 	//~ End UObject Interface
 };
\ No newline at end of file
diff --git a/Engine/Source/Runtime/Engine/Classes/Particles/ParticleEmitter.h b/Engine/Source/Runtime/Engine/Classes/Particles/ParticleEmitter.h
--- a/Engine/Source/Runtime/Engine/Classes/Particles/ParticleEmitter.h
+++ b/Engine/Source/Runtime/Engine/Classes/Particles/ParticleEmitter.h
@@ -240,12 +240,6 @@ class UParticleEmitter : public UObject
 	// @todo document
 	ENGINE_API FName& GetEmitterName();
 
-	// Ignore this object in dedicated server builds
-	virtual bool NeedsLoadForServer() const override
-	{
-		return false;
-	}
-
 	// @todo document
 	virtual	void						SetLODCount(int32 LODCount);
 
diff --git a/Engine/Source/Runtime/Engine/Classes/Particles/ParticleLODLevel.h b/Engine/Source/Runtime/Engine/Classes/Particles/ParticleLODLevel.h
--- a/Engine/Source/Runtime/Engine/Classes/Particles/ParticleLODLevel.h
+++ b/Engine/Source/Runtime/Engine/Classes/Particles/ParticleLODLevel.h
@@ -83,12 +83,6 @@ class UParticleLODLevel : public UObject
 	// @todo document
 	virtual bool	GenerateFromLODLevel(UParticleLODLevel* SourceLODLevel, float Percentage, bool bGenerateModuleData = true);
 
-	// Ignore this object in dedicated server builds
-	virtual bool NeedsLoadForServer() const override
-	{
-		return false;
-	}
-
 	/**
 	 *	CalculateMaxActiveParticleCount
 	 *	Determine the maximum active particles that could occur with this emitter.
diff --git a/Engine/Source/Runtime/Engine/Classes/Particles/ParticleModule.h b/Engine/Source/Runtime/Engine/Classes/Particles/ParticleModule.h
--- a/Engine/Source/Runtime/Engine/Classes/Particles/ParticleModule.h
+++ b/Engine/Source/Runtime/Engine/Classes/Particles/ParticleModule.h
@@ -536,12 +536,6 @@ class UParticleModule : public UObject
 	/** Returns whether this module is used in any GPU emitters. */
 	ENGINE_API bool IsUsedInGPUEmitter()const;
 
-	/** Ignore this object in dedicated server builds */
-	virtual bool NeedsLoadForServer() const override
-	{
-		return false;
-	}
-
 #if WITH_EDITOR
 	virtual void PostLoadSubobjects( FObjectInstancingGraph* OuterInstanceGraph ) override;
 
diff --git a/Engine/Source/Runtime/Engine/Private/AVIWriter.cpp b/Engine/Source/Runtime/Engine/Private/AVIWriter.cpp
--- a/Engine/Source/Runtime/Engine/Private/AVIWriter.cpp
+++ b/Engine/Source/Runtime/Engine/Private/AVIWriter.cpp
@@ -16,6 +16,10 @@ typedef TCHAR* PTCHAR;
 #pragma warning(push)
 #pragma warning(disable : 4263) // 'function' : member function does not override any base class virtual member function
 #pragma warning(disable : 4264) // 'virtual_function' : no override available for virtual member function from base 'cla
+#if USING_CODE_ANALYSIS
+	#pragma warning(disable:6509) // Invalid annotation: 'return' cannot be referenced in some contexts
+	#pragma warning(disable:6101) // Returning uninitialized memory '*lpdwExitCode'.  A successful path through the function does not set the named _Out_ parameter.
+#endif
 #include <streams.h>
 #pragma warning(pop)
 
@@ -127,6 +131,7 @@ IBaseFilter* FindEncodingFilter(const FString& Name)
 	while (EnumIterator->Next(1, &Moniker, nullptr) == S_OK)
 	{
 		IPropertyBag* Properties = nullptr;
+		CA_SUPPRESS(6387);
 		if (FAILED(Moniker->BindToStorage(0, 0, IID_IPropertyBag, (void**)&Properties)))
 		{
 			Moniker->Release();
@@ -146,6 +151,7 @@ IBaseFilter* FindEncodingFilter(const FString& Name)
 		Properties->Release();
 
 		IBaseFilter* Filter = nullptr;
+		CA_SUPPRESS(6387);
 		if (bUseThisEncoder && Moniker->BindToObject(nullptr, nullptr, IID_IBaseFilter, (void**)&Filter) == S_OK)
 		{
 			Filter->AddRef();
@@ -252,6 +258,7 @@ public:
 			if (!EncodingFilter)
 			{
 				// Attempt to use a default encoder
+				CA_SUPPRESS(6031);
 				CoCreateInstance(CLSID_MJPGEnc, NULL, CLSCTX_INPROC, IID_IBaseFilter, (void**)&EncodingFilter);
 				Graph->AddFilter( EncodingFilter, TEXT("Encoder") );
 			}
@@ -265,6 +272,7 @@ public:
 		}
 		
 		IBaseFilter *pMux;
+		CA_SUPPRESS(6031);
 		CoCreateInstance(CLSID_AviDest, NULL, CLSCTX_INPROC, IID_IBaseFilter, (void**)&pMux);
 		hr = Graph->AddFilter(pMux, TEXT("AVI Mux"));
 		if (FAILED(hr)) 
@@ -276,6 +284,7 @@ public:
 		}
 
 		IBaseFilter *FileWriter;
+		CA_SUPPRESS(6031);
 		CoCreateInstance(CLSID_FileWriter, NULL, CLSCTX_INPROC, IID_IBaseFilter, (void**)&FileWriter);
 		hr = Graph->AddFilter(FileWriter, TEXT("File Writer"));
 		if (FAILED(hr)) 
diff --git a/Engine/Source/Runtime/Engine/Private/Blueprint.cpp b/Engine/Source/Runtime/Engine/Private/Blueprint.cpp
--- a/Engine/Source/Runtime/Engine/Private/Blueprint.cpp
+++ b/Engine/Source/Runtime/Engine/Private/Blueprint.cpp
@@ -788,10 +788,12 @@ UWorld* UBlueprint::GetWorldBeingDebugged()
 
 void UBlueprint::GetAssetRegistryTags(TArray<FAssetRegistryTag>& OutTags) const
 {
-	UClass* GenClass = Cast<UClass>(GeneratedClass);
-	if ( GenClass && GenClass->GetDefaultObject() )
+	if (UClass* GenClass = Cast<UClass>(GeneratedClass))
 	{
-		GenClass->GetDefaultObject()->GetAssetRegistryTags(OutTags);
+		if (UObject* CDO = GenClass->GetDefaultObject())
+		{
+			CDO->GetAssetRegistryTags(OutTags);
+		}
 	}
 
 	Super::GetAssetRegistryTags(OutTags);
diff --git a/Engine/Source/Runtime/Engine/Private/Components/ActorComponent.cpp b/Engine/Source/Runtime/Engine/Private/Components/ActorComponent.cpp
--- a/Engine/Source/Runtime/Engine/Private/Components/ActorComponent.cpp
+++ b/Engine/Source/Runtime/Engine/Private/Components/ActorComponent.cpp
@@ -129,7 +129,8 @@ void UActorComponent::PostInitProperties()
 {
 	Super::PostInitProperties();
 
-	if (OwnerPrivate)
+	// Instance components will be added during the owner's initialization
+	if (OwnerPrivate && CreationMethod != EComponentCreationMethod::Instance)
 	{
 		OwnerPrivate->AddOwnedComponent(this);
 	}
diff --git a/Engine/Source/Runtime/Engine/Private/Components/NavMeshRenderingComponent.cpp b/Engine/Source/Runtime/Engine/Private/Components/NavMeshRenderingComponent.cpp
--- a/Engine/Source/Runtime/Engine/Private/Components/NavMeshRenderingComponent.cpp
+++ b/Engine/Source/Runtime/Engine/Private/Components/NavMeshRenderingComponent.cpp
@@ -1181,12 +1181,6 @@ UNavMeshRenderingComponent::UNavMeshRenderingComponent(const FObjectInitializer&
 	bCollectNavigationData = false;
 }
 
-bool UNavMeshRenderingComponent::NeedsLoadForServer() const
-{
-	// This rendering component is not needed in dedicated server builds
-	return false;
-}
-
 bool UNavMeshRenderingComponent::IsNavigationShowFlagSet(const UWorld* World)
 {
 	bool bShowNavigation = false;
diff --git a/Engine/Source/Runtime/Engine/Private/EngineUtils.cpp b/Engine/Source/Runtime/Engine/Private/EngineUtils.cpp
--- a/Engine/Source/Runtime/Engine/Private/EngineUtils.cpp
+++ b/Engine/Source/Runtime/Engine/Private/EngineUtils.cpp
@@ -458,7 +458,7 @@ void FConsoleOutputDevice::Serialize(const TCHAR* Text, ELogVerbosity::Type Verb
 
 		if(CVar)
 		{
-			int MinVerbosity = CVar->GetValueOnGameThread();
+			int MinVerbosity = CVar->GetValueOnAnyThread(true);
 
 			if((int)Verbosity <= MinVerbosity)
 			{
diff --git a/Engine/Source/Runtime/Engine/Private/Materials/MaterialExpressions.cpp b/Engine/Source/Runtime/Engine/Private/Materials/MaterialExpressions.cpp
--- a/Engine/Source/Runtime/Engine/Private/Materials/MaterialExpressions.cpp
+++ b/Engine/Source/Runtime/Engine/Private/Materials/MaterialExpressions.cpp
@@ -452,11 +452,6 @@ void UMaterialExpression::Serialize( FArchive& Ar )
 #endif
 }
 
-bool UMaterialExpression::NeedsLoadForServer() const
-{
-	return false; 
-}
-
 bool UMaterialExpression::NeedsLoadForClient() const
 {
 	// Expressions that reference texture objects need to be cooked
diff --git a/Engine/Source/Runtime/Engine/Private/SubDSurfaceComponent.cpp b/Engine/Source/Runtime/Engine/Private/SubDSurfaceComponent.cpp
--- a/Engine/Source/Runtime/Engine/Private/SubDSurfaceComponent.cpp
+++ b/Engine/Source/Runtime/Engine/Private/SubDSurfaceComponent.cpp
@@ -41,10 +41,23 @@
 
 	#include <far/stencilTableFactory.h>
 
+#if defined(_MSC_VER) && USING_CODE_ANALYSIS
+	#pragma warning(push)
+	#pragma warning(disable : 6011)  // Dereferencing NULL pointer 'X'
+	#pragma warning(disable : 6308)  // 'realloc' might return null pointer: assigning null pointer to 'X', which is passed as an argument to 'realloc', will cause the original memory block to be leaked.
+	#pragma warning(disable : 6385)  // Reading invalid data from 'X':  the readable size is 'Y' bytes, but 'Z' bytes may be read.
+	#pragma warning(disable : 6386)  // Buffer overrun while writing to 'X':  the writable size is 'Y' bytes, but 'Z' bytes might be written.
+	#pragma warning(disable : 28182) // Dereferencing NULL pointer. 'X' contains the same NULL value as 'Y' did.
+#endif
+
 // for non manifold
 	#include <hbr/mesh.h>
 	#include <hbr/catmark.h>
 
+#if defined(_MSC_VER) && USING_CODE_ANALYSIS
+	#pragma warning(pop)
+#endif
+
 #undef and
 #undef not
 #undef M_PI
diff --git a/Engine/Source/Runtime/Engine/Private/UnrealEngine.cpp b/Engine/Source/Runtime/Engine/Private/UnrealEngine.cpp
--- a/Engine/Source/Runtime/Engine/Private/UnrealEngine.cpp
+++ b/Engine/Source/Runtime/Engine/Private/UnrealEngine.cpp
@@ -6054,6 +6054,86 @@ bool UEngine::PerformError(const TCHAR* Cmd, FOutputDevice& Ar)
 			);
 		return true;
 	}
+	if (FParse::Command(&Cmd, TEXT("TWOTHREADSCRASH")))
+	{
+		class FThreadPoolCrash : public IQueuedWork
+		{
+		private:
+			double CrashDelay;
+		public:
+			FThreadPoolCrash(double InCrashDelay)
+				: CrashDelay(InCrashDelay)
+			{
+			}
+			void Abandon()
+			{
+			}
+			void DoThreadedWork()
+			{
+				double CrashTime = FPlatformTime::Seconds() + CrashDelay;
+				do 
+				{
+					if (FPlatformTime::Seconds() >= CrashTime)
+					{
+						UE_LOG(LogEngine, Warning, TEXT("Printed warning to log."));
+						SetCrashType(ECrashType::Debug);
+						UE_LOG(LogEngine, Fatal, TEXT("Crashing the worker thread at your request"));
+						break;
+					}
+					else
+					{
+						FPlatformProcess::Sleep(0);
+					}
+				} while (true);
+			}
+		};
+
+		UE_LOG(LogEngine, Warning, TEXT("Queuing two tasks to crash."));
+		GThreadPool->AddQueuedWork(new FThreadPoolCrash(0.100));
+		GThreadPool->AddQueuedWork(new FThreadPoolCrash(0.110));
+
+		return true;
+	}
+	else if (FParse::Command(&Cmd, TEXT("TWOTHREADSGPF")))
+	{
+		class FThreadPoolCrash : public IQueuedWork
+		{
+		private:
+			double CrashDelay;
+		public:
+			FThreadPoolCrash(double InCrashDelay)
+				: CrashDelay(InCrashDelay)
+			{
+			}
+			void Abandon()
+			{
+			}
+			void DoThreadedWork()
+			{
+				double CrashTime = FPlatformTime::Seconds() + CrashDelay;
+				do
+				{
+					if (FPlatformTime::Seconds() >= CrashTime)
+					{
+						UE_LOG(LogEngine, Warning, TEXT("Printed warning to log."));
+						SetCrashType(ECrashType::Debug);
+						*(int32 *)3 = 123;
+						break;
+					}
+					else
+					{
+						FPlatformProcess::Sleep(0);
+					}
+				} while (true);
+			}
+		};
+
+		UE_LOG(LogEngine, Warning, TEXT("Queuing two tasks to crash."));
+		GThreadPool->AddQueuedWork(new FThreadPoolCrash(0.100));
+		GThreadPool->AddQueuedWork(new FThreadPoolCrash(0.110));
+
+		return true;
+	}
 	else if (FParse::Command(&Cmd, TEXT("THREADENSURE")))
 	{
 		struct FThread
diff --git a/Engine/Source/Runtime/Engine/Public/kDOP.h b/Engine/Source/Runtime/Engine/Public/kDOP.h
--- a/Engine/Source/Runtime/Engine/Public/kDOP.h
+++ b/Engine/Source/Runtime/Engine/Public/kDOP.h
@@ -83,14 +83,11 @@ struct TTraversalHistory
 		}
 	}
 
-	TTraversalHistory AddNode(KDOP_IDX_TYPE NewNodeIndex)
+	TTraversalHistory AddNode(KDOP_IDX_TYPE NewNodeIndex) const
 	{
 		TTraversalHistory NewHistory;
-		for (int32 NodeIndex = NodeHistoryLength - 1; NodeIndex > 0; NodeIndex--)
-		{
-			// Move all the indices toward the end of the array one element
-			NewHistory.Nodes[NodeIndex] = Nodes[NodeIndex - 1];
-		}
+		// Move all the indices toward the end of the array one element
+		CopyAssignItems(NewHistory.Nodes + 1, Nodes, NodeHistoryLength - 1);
 		// Insert the new node at the beginning of the array
 		NewHistory.Nodes[0] = NewNodeIndex;
 		return NewHistory;
diff --git a/Engine/Source/Runtime/GameplayAbilities/Public/GameplayEffectUIData.h b/Engine/Source/Runtime/GameplayAbilities/Public/GameplayEffectUIData.h
--- a/Engine/Source/Runtime/GameplayAbilities/Public/GameplayEffectUIData.h
+++ b/Engine/Source/Runtime/GameplayAbilities/Public/GameplayEffectUIData.h
@@ -12,9 +12,4 @@ UCLASS(Blueprintable, Abstract, EditInlineNew, CollapseCategories)
 class GAMEPLAYABILITIES_API UGameplayEffectUIData : public UObject
 {
 	GENERATED_UCLASS_BODY()
-
-	virtual bool NeedsLoadForServer() const override
-	{
-		return false;
-	}
 };
\ No newline at end of file
diff --git a/Engine/Source/Runtime/JsonUtilities/Private/JsonObjectConverter.cpp b/Engine/Source/Runtime/JsonUtilities/Private/JsonObjectConverter.cpp
--- a/Engine/Source/Runtime/JsonUtilities/Private/JsonObjectConverter.cpp
+++ b/Engine/Source/Runtime/JsonUtilities/Private/JsonObjectConverter.cpp
@@ -407,9 +407,17 @@ bool ConvertScalarJsonValueToUProperty(TSharedPtr<FJsonValue> JsonValue, UProper
 				// this value's not really meaningful from json serialization (since we don't know timezone) but handle it anyway since we're handling the other keywords
 				DateTimeOut = FDateTime::UtcNow();
 			}
-			else if (!FDateTime::ParseIso8601(*DateString, DateTimeOut))
+			else if (FDateTime::ParseIso8601(*DateString, DateTimeOut))
 			{
-				UE_LOG(LogJson, Error, TEXT("JsonValueToUProperty - Unable to import FDateTime from Iso8601 String for property %s"), *Property->GetNameCPP());
+				// ok
+			}
+			else if (FDateTime::Parse(DateString, DateTimeOut))
+			{
+				// ok
+			}
+			else
+			{
+				UE_LOG(LogJson, Error, TEXT("JsonValueToUProperty - Unable to import FDateTime for property %s"), *Property->GetNameCPP());
 				return false;
 			}
 		}
diff --git a/Engine/Source/Runtime/Landscape/Private/LandscapeEdit.cpp b/Engine/Source/Runtime/Landscape/Private/LandscapeEdit.cpp
--- a/Engine/Source/Runtime/Landscape/Private/LandscapeEdit.cpp
+++ b/Engine/Source/Runtime/Landscape/Private/LandscapeEdit.cpp
@@ -3302,7 +3302,7 @@ void ALandscape::PostEditImport()
 	{
 		for (ALandscape* Landscape : TActorRange<ALandscape>(GetWorld()))
 		{
-			if (Landscape != this && !Landscape->HasAnyFlags(RF_BeginDestroyed) && Landscape->LandscapeGuid == LandscapeGuid)
+			if (Landscape && Landscape != this && !Landscape->HasAnyFlags(RF_BeginDestroyed) && Landscape->LandscapeGuid == LandscapeGuid)
 			{
 				// Copy/Paste case, need to generate new GUID
 				LandscapeGuid = FGuid::NewGuid();
@@ -4817,17 +4817,20 @@ struct FMobileLayerAllocation
 	{
 	}
 
-	friend bool operator<(const FMobileLayerAllocation& lhs, const FMobileLayerAllocation& rhs)
+	friend bool operator<(const FMobileLayerAllocation& Lhs, const FMobileLayerAllocation& Rhs)
 	{
-		if (!lhs.Allocation.LayerInfo && !rhs.Allocation.LayerInfo) return false; // equally broken :P
-		if (!lhs.Allocation.LayerInfo && rhs.Allocation.LayerInfo) return false; // broken layers sort to the end
-		if (!rhs.Allocation.LayerInfo && lhs.Allocation.LayerInfo) return true;
+		ULandscapeLayerInfoObject* LhsLayerInfo = Lhs.Allocation.LayerInfo;
+		ULandscapeLayerInfoObject* RhsLayerInfo = Rhs.Allocation.LayerInfo;
 
-		if (lhs.Allocation.LayerInfo == ALandscapeProxy::VisibilityLayer && rhs.Allocation.LayerInfo != ALandscapeProxy::VisibilityLayer) return true; // visibility layer to the front
-		if (rhs.Allocation.LayerInfo == ALandscapeProxy::VisibilityLayer && lhs.Allocation.LayerInfo != ALandscapeProxy::VisibilityLayer) return false;
+		if (!LhsLayerInfo && !RhsLayerInfo) return false; // equally broken :P
+		if (!LhsLayerInfo && RhsLayerInfo) return false; // broken layers sort to the end
+		if (!RhsLayerInfo && LhsLayerInfo) return true;
 
-		if (lhs.Allocation.LayerInfo->bNoWeightBlend && !rhs.Allocation.LayerInfo->bNoWeightBlend) return false; // non-blended layers sort to the end
-		if (rhs.Allocation.LayerInfo->bNoWeightBlend && !lhs.Allocation.LayerInfo->bNoWeightBlend) return true;
+		if (LhsLayerInfo == ALandscapeProxy::VisibilityLayer && RhsLayerInfo != ALandscapeProxy::VisibilityLayer) return true; // visibility layer to the front
+		if (RhsLayerInfo == ALandscapeProxy::VisibilityLayer && LhsLayerInfo != ALandscapeProxy::VisibilityLayer) return false;
+
+		if (LhsLayerInfo->bNoWeightBlend && !RhsLayerInfo->bNoWeightBlend) return false; // non-blended layers sort to the end
+		if (RhsLayerInfo->bNoWeightBlend && !LhsLayerInfo->bNoWeightBlend) return true;
 
 		// TODO: If we want to support cleanly decaying a pc landscape for mobile
 		// we should probably add other sort criteria, e.g. coverage
diff --git a/Engine/Source/Runtime/Launch/Private/LaunchEngineLoop.cpp b/Engine/Source/Runtime/Launch/Private/LaunchEngineLoop.cpp
--- a/Engine/Source/Runtime/Launch/Private/LaunchEngineLoop.cpp
+++ b/Engine/Source/Runtime/Launch/Private/LaunchEngineLoop.cpp
@@ -24,6 +24,7 @@
 
 #if WITH_COREUOBJECT
 	#include "Internationalization/PackageLocalizationManager.h"
+	#include "CoreUObject.h"
 #endif
 
 #if WITH_EDITOR
@@ -1652,12 +1653,16 @@ int32 FEngineLoop::PreInit( const TCHAR* CmdLine )
 		// At least one startup package failed to load, return 1 to indicate an error
 		return 1;
 	}
+#endif // WITH_ENGINE
 
+#if WITH_COREUOBJECT
 	if (GUObjectArray.IsOpenForDisregardForGC())
 	{
 		GUObjectArray.CloseDisregardForGC();
 	}
+#endif // WITH_COREUOBJECT
 
+#if WITH_ENGINE
 	if (UOnlineEngineInterface::Get()->IsLoaded())
 	{
 		SetIsServerForOnlineSubsystemsDelegate(FQueryIsRunningServer::CreateStatic(&IsServerDelegateForOSS));
diff --git a/Engine/Source/Runtime/NetworkFileSystem/Private/NetworkFileServerConnection.cpp b/Engine/Source/Runtime/NetworkFileSystem/Private/NetworkFileServerConnection.cpp
--- a/Engine/Source/Runtime/NetworkFileSystem/Private/NetworkFileServerConnection.cpp
+++ b/Engine/Source/Runtime/NetworkFileSystem/Private/NetworkFileServerConnection.cpp
@@ -858,7 +858,25 @@ bool FNetworkFileServerClientConnection::ProcessGetFileList( FArchive& In, FArch
 
 			FString ConnectedContentFolder = ContentFolder;
 			ConnectedContentFolder.ReplaceInline(*LocalEngineDir, *ConnectedEngineDir);
-			int32 ReplaceCount = ConnectedContentFolder.ReplaceInline(*LocalGameDir, *ConnectedGameDir);
+
+			int32 ReplaceCount = 0;
+
+			// If one path is relative and the other isn't, convert both to absolute paths before trying to replace
+			if (FPaths::IsRelative(LocalGameDir) != FPaths::IsRelative(ConnectedContentFolder))
+			{
+				FString AbsoluteLocalGameDir = FPaths::ConvertRelativePathToFull(LocalGameDir);
+				FString AbsoluteConnectedContentFolder = FPaths::ConvertRelativePathToFull(ConnectedContentFolder);
+				ReplaceCount = AbsoluteConnectedContentFolder.ReplaceInline(*AbsoluteLocalGameDir, *ConnectedGameDir);
+				if (ReplaceCount > 0)
+				{
+					ConnectedContentFolder = AbsoluteConnectedContentFolder;
+				}
+			}
+			else
+			{
+				ReplaceCount = ConnectedContentFolder.ReplaceInline(*LocalGameDir, *ConnectedGameDir);
+			}
+			
 			if (ReplaceCount == 0)
 			{
 				int32 GameDirOffset = ConnectedContentFolder.Find(ConnectedGameDir, ESearchCase::IgnoreCase, ESearchDir::FromEnd);
diff --git a/Engine/Source/Runtime/PakFile/Public/IPlatformFilePak.h b/Engine/Source/Runtime/PakFile/Public/IPlatformFilePak.h
--- a/Engine/Source/Runtime/PakFile/Public/IPlatformFilePak.h
+++ b/Engine/Source/Runtime/PakFile/Public/IPlatformFilePak.h
@@ -552,9 +552,8 @@ public:
 			return *this; 
 		}
 
-		SAFE_BOOL_OPERATORS(FFileIterator)
 		/** conversion to "bool" returning true if the iterator is valid. */
-		FORCEINLINE_EXPLICIT_OPERATOR_BOOL() const
+		FORCEINLINE explicit operator bool() const
 		{ 
 			return !!IndexIt; 
 		}
diff --git a/Engine/Source/Runtime/SandboxFile/Private/IPlatformFileSandboxWrapper.cpp b/Engine/Source/Runtime/SandboxFile/Private/IPlatformFileSandboxWrapper.cpp
--- a/Engine/Source/Runtime/SandboxFile/Private/IPlatformFileSandboxWrapper.cpp
+++ b/Engine/Source/Runtime/SandboxFile/Private/IPlatformFileSandboxWrapper.cpp
@@ -129,6 +129,15 @@ bool FSandboxPlatformFile::Initialize(IPlatformFile* Inner, const TCHAR* CmdLine
 	return !!LowerLevel;
 }
 
+const FString& FSandboxPlatformFile::GetGameSandboxDirectoryName()
+{
+	if (GameSandboxDirectoryName.IsEmpty())
+	{
+		GameSandboxDirectoryName = FString::Printf(TEXT("%s/"), FApp::GetGameName());
+	}
+	return GameSandboxDirectoryName;
+}
+
 FString FSandboxPlatformFile::ConvertToSandboxPath( const TCHAR* Filename ) const
 {
 	// Mostly for the malloc profiler to flush the data.
@@ -316,14 +325,24 @@ FString FSandboxPlatformFile::ConvertToAbsolutePathForExternalAppForWrite( const
 }
 
 const FString& FSandboxPlatformFile::GetAbsolutePathToGameDirectory()
+{
+	if (AbsolutePathToGameDirectory.IsEmpty())
+	{
+		// Strip game directory, keep just to path to the game directory which could simply be the root dir (but not always).
+		AbsolutePathToGameDirectory = FPaths::GetPath(GetAbsoluteGameDirectory());
+	}
+	return AbsolutePathToGameDirectory;
+}
+
+const FString& FSandboxPlatformFile::GetAbsoluteGameDirectory()
 {
 	if (AbsoluteGameDirectory.IsEmpty())
 	{
 		AbsoluteGameDirectory = FPaths::GetProjectFilePath();
 		UE_CLOG(AbsoluteGameDirectory.IsEmpty(), SandboxFile, Fatal, TEXT("SandboxFileWrapper tried to access project path before it was set."));
 		AbsoluteGameDirectory = FPaths::ConvertRelativePathToFull(AbsoluteGameDirectory);
-		// Strip .uproject filename and game directory, keep just to path to the game directory which could simply be the root dir (but not always).
-		AbsoluteGameDirectory = FPaths::GetPath(FPaths::GetPath(AbsoluteGameDirectory));
+		// Strip .uproject filename
+		AbsoluteGameDirectory = FPaths::GetPath(AbsoluteGameDirectory);
 	}
 	return AbsoluteGameDirectory;
 }
diff --git a/Engine/Source/Runtime/SandboxFile/Public/IPlatformFileSandboxWrapper.h b/Engine/Source/Runtime/SandboxFile/Public/IPlatformFileSandboxWrapper.h
--- a/Engine/Source/Runtime/SandboxFile/Public/IPlatformFileSandboxWrapper.h
+++ b/Engine/Source/Runtime/SandboxFile/Public/IPlatformFileSandboxWrapper.h
@@ -58,12 +58,16 @@ class SANDBOXFILE_API FSandboxPlatformFile : public IPlatformFile
 	IPlatformFile*		LowerLevel;
 	/** Absolute path to the sandbox directory */
 	FString			SandboxDirectory;
+	/** Name of the game's sandbox directory */
+	FString			GameSandboxDirectoryName;
 	/** Relative path to root directory. Cached for faster access */
 	FString			RelativeRootDirectory;
 	/** Absolute path to root directory. Cached for faster access */
 	FString			AbsoluteRootDirectory;
-	/** Absolute path to game directory. Cached for faster access */
+	/** Absolute game directory. Cached for faster access */
 	FString			AbsoluteGameDirectory;
+	/** Absolute path to game directory. Cached for faster access */
+	FString			AbsolutePathToGameDirectory;
 	/** Access to any file (in unreal standard form) matching this is not allowed */
 	TArray<FString>		FileExclusionWildcards;
 	/** Access to any directory (in unreal standard form) matching this is not allowed */
@@ -216,12 +220,18 @@ public:
 		return SandboxDirectory;
 	}
 
+	/** Returns the name of the sandbox directory for the game's content */
+	const FString& GetGameSandboxDirectoryName();
+
 	/** Returns absolute root directory */
 	const FString& GetAbsoluteRootDirectory() const
 	{
 		return AbsoluteRootDirectory;
 	}
 
+	/** Returns absolute game directory */
+	const FString& GetAbsoluteGameDirectory();
+
 	/** Returns absolute path to game directory (without the game directory itself) */
 	const FString& GetAbsolutePathToGameDirectory();
 
@@ -459,7 +469,8 @@ public:
 				}
 				else
 				{
-					LocalFilename = SandboxFile.GetAbsolutePathToGameDirectory() / LocalFilename;
+					LocalFilename = LocalFilename.Mid(SandboxFile.GetGameSandboxDirectoryName().Len());
+					LocalFilename = SandboxFile.GetAbsoluteGameDirectory() / LocalFilename;
 				}
 			}
 			else
diff --git a/Engine/Source/Runtime/SlateCore/Private/Stats/SlateStats.cpp b/Engine/Source/Runtime/SlateCore/Private/Stats/SlateStats.cpp
--- a/Engine/Source/Runtime/SlateCore/Private/Stats/SlateStats.cpp
+++ b/Engine/Source/Runtime/SlateCore/Private/Stats/SlateStats.cpp
@@ -141,6 +141,9 @@ void FSlateStatCycleCounter::EndFrame(double CurrentTime)
 						}
 					}
 				}
+
+				check(PathCurrentDepth >= 0 && PathCurrentDepth < kMaxDepth);
+
 				// increment the ordinal at this depth
 				++Path[PathCurrentDepth];
 				FCString::Sprintf(PathStrCurrent, TEXT("%4d"), Path[PathCurrentDepth]);
diff --git a/Engine/Source/Runtime/UnrealAudio/Private/UnrealAudioDeviceFormat.cpp b/Engine/Source/Runtime/UnrealAudio/Private/UnrealAudioDeviceFormat.cpp
--- a/Engine/Source/Runtime/UnrealAudio/Private/UnrealAudioDeviceFormat.cpp
+++ b/Engine/Source/Runtime/UnrealAudio/Private/UnrealAudioDeviceFormat.cpp
@@ -99,7 +99,7 @@ namespace UAudio
 			{
 				for (uint32 ChannelIndex = 0; ChannelIndex < ConvertInfo.NumChannels; ++ChannelIndex)
 				{
-					int32 Value = FromBufferPtr[ChannelIndex];
+					int64 Value = FromBufferPtr[ChannelIndex];
 					ToBufferPtr[ChannelIndex] = (ToType)(Value >> ShiftAmount);
 				}
 				FromBufferPtr += ConvertInfo.FromChannels;
@@ -112,7 +112,7 @@ namespace UAudio
 			{
 				for (uint32 ChannelIndex = 0; ChannelIndex < ConvertInfo.NumChannels; ++ChannelIndex)
 				{
-					int32 Value = FromBufferPtr[ChannelIndex];
+					int64 Value = FromBufferPtr[ChannelIndex];
 					ToBufferPtr[ChannelIndex] = (ToType)(Value << ShiftAmount);
 				}
 				FromBufferPtr += ConvertInfo.FromChannels;
diff --git a/Engine/Source/Runtime/Windows/UnrealAudioXAudio2/Private/UnrealAudioDeviceXAudio2.cpp b/Engine/Source/Runtime/Windows/UnrealAudioXAudio2/Private/UnrealAudioDeviceXAudio2.cpp
--- a/Engine/Source/Runtime/Windows/UnrealAudioXAudio2/Private/UnrealAudioDeviceXAudio2.cpp
+++ b/Engine/Source/Runtime/Windows/UnrealAudioXAudio2/Private/UnrealAudioDeviceXAudio2.cpp
@@ -315,7 +315,7 @@ namespace UAudio
 			WAVEFORMATEXTENSIBLE* WaveFormatExtensible = (WAVEFORMATEXTENSIBLE*)WaveFormatEx;
 			DeviceInfo.Speakers.Reset();
 			uint32 ChanCount = 0;
-			for (uint32 SpeakerTypeIndex = 0; SpeakerTypeIndex < ESpeaker::SPEAKER_TYPE_COUNT, ChanCount < DeviceInfo.NumChannels; ++SpeakerTypeIndex)
+			for (uint32 SpeakerTypeIndex = 0; SpeakerTypeIndex < ESpeaker::SPEAKER_TYPE_COUNT && ChanCount < DeviceInfo.NumChannels; ++SpeakerTypeIndex)
 			{
 				if (WaveFormatExtensible->dwChannelMask & SpeakerMaskMapping[SpeakerTypeIndex].XAudio2SpeakerTypeFlag)
 				{
diff --git a/Engine/Source/Runtime/XmlParser/Private/FastXml.cpp b/Engine/Source/Runtime/XmlParser/Private/FastXml.cpp
--- a/Engine/Source/Runtime/XmlParser/Private/FastXml.cpp
+++ b/Engine/Source/Runtime/XmlParser/Private/FastXml.cpp
@@ -413,6 +413,7 @@ protected:
 								ErrorMessage = LOCTEXT( "TooManyAttributes", "Encountered too many attributes in a single element for this parser to handle" );
 								return false;
 							}
+							check(AttributeCount >= 0);
 							Attributes[ AttributeCount ] = Buffer;
 							Buffer = NextSeparator( Buffer );  // scan up to a space, or an equal
 							if( *Buffer )
