Fixed compression flag being consumed by wrong parameter

diff --git a/Engine/Source/Runtime/CoreUObject/Private/Serialization/BulkData.cpp b/Engine/Source/Runtime/CoreUObject/Private/Serialization/BulkData.cpp
--- a/Engine/Source/Runtime/CoreUObject/Private/Serialization/BulkData.cpp
+++ b/Engine/Source/Runtime/CoreUObject/Private/Serialization/BulkData.cpp
@@ -1158,7 +1158,7 @@ void FUntypedBulkData::SerializeBulkData( FArchive& Ar, void* Data )
 		// Serialize data compressed.
 		if( BulkDataFlags & BULKDATA_SerializeCompressed )
 		{
-			Ar.SerializeCompressed( Data, GetBulkDataSize(), GetDecompressionFlags(), !!(BulkDataFlags&BULKDATA_SerializeCompressedBitWindow) );
+			Ar.SerializeCompressed( Data, GetBulkDataSize(), GetDecompressionFlags(), false, !!(BulkDataFlags&BULKDATA_SerializeCompressedBitWindow) );
 		}
 		// Uncompressed/ regular serialization.
 		else
@@ -1185,7 +1185,7 @@ void FUntypedBulkData::SerializeBulkData( FArchive& Ar, void* Data )
 				SerializedData.AddUninitialized( GetBulkDataSize() );
 
 				// Serialize data with passed in archive and compress.
-				Ar.SerializeCompressed( SerializedData.GetData(), SerializedData.Num(), GetDecompressionFlags(), !!(BulkDataFlags&BULKDATA_SerializeCompressedBitWindow) );
+				Ar.SerializeCompressed( SerializedData.GetData(), SerializedData.Num(), GetDecompressionFlags(), false, !!(BulkDataFlags&BULKDATA_SerializeCompressedBitWindow) );
 				
 				// Initialize memory reader with uncompressed data array and propagate forced byte swapping
 				FMemoryReader MemoryReader( SerializedData, true );
@@ -1205,7 +1205,7 @@ void FUntypedBulkData::SerializeBulkData( FArchive& Ar, void* Data )
 				SerializeElements(MemoryWriter, Data);
 
 				// Serialize data with passed in archive and compress.
-				Ar.SerializeCompressed( SerializedData.GetData(), SerializedData.Num(), GetDecompressionFlags(), !!(BulkDataFlags&BULKDATA_SerializeCompressedBitWindow) );
+				Ar.SerializeCompressed( SerializedData.GetData(), SerializedData.Num(), GetDecompressionFlags(), false, !!(BulkDataFlags&BULKDATA_SerializeCompressedBitWindow) );
 			}
 		}
 		// Uncompressed/ regular serialization.
