Copying //UE4/Dev-Build to //UE4/Dev-Main (Source: //UE4/Dev-Build @ 3092544)

diff --git a/Engine/Source/Runtime/CoreUObject/Public/UObject/NoExportTypes.h b/Engine/Source/Runtime/CoreUObject/Public/UObject/NoExportTypes.h
new file mode 100644
--- /dev/null
+++ b/Engine/Source/Runtime/CoreUObject/Public/UObject/NoExportTypes.h
@@ -0,0 +1,1118 @@
+// Copyright 1998-2016 Epic Games, Inc. All Rights Reserved.
+
+
+
+//=============================================================================
+// Unreal base structures.
+
+// Temporary mirrors of C++ structs, used mainly as forward declarations for the core object module
+// to avoid including the full engine source. In most cases the full class definition is in another
+// file and is noted as such. More complete documentation will generally be found in those files.
+
+#pragma once
+#if !CPP      //noexport class
+
+/** String search case used in UnrealString.h */
+UENUM()
+namespace ESearchCase
+{
+	enum Type
+	{
+		CaseSensitive,
+		IgnoreCase,
+	};
+}
+
+/** String search dir used in UnrealString.h */
+UENUM()
+namespace ESearchDir
+{
+	enum Type
+	{
+		FromStart,
+		FromEnd,
+	};
+}
+
+/**
+ * Enum that defines how the log times are to be displayed (mirrored from OutputDevice.h).
+ */
+UENUM()
+namespace ELogTimes
+{
+	enum Type
+	{
+		/** Do not display log timestamps. */
+		None UMETA(DisplayName = "None"),
+
+		/** Display log timestamps in UTC. */
+		UTC UMETA(DisplayName = "UTC"),
+
+		/** Display log timestamps in seconds elapsed since GStartTime. */
+		SinceGStartTime UMETA(DisplayName = "Time since application start"),
+	};
+}
+
+/** Generic axis enum (mirrored for native use in Axis.h). */
+UENUM()
+namespace EAxis
+{
+	enum Type
+	{
+		None,
+		X,
+		Y,
+		Z
+	};
+}
+
+// Interpolation data types.
+UENUM()
+enum EInterpCurveMode
+{
+	CIM_Linear UMETA(DisplayName="Linear"),
+	CIM_CurveAuto UMETA(DisplayName="Curve Auto"),
+	CIM_Constant UMETA(DisplayName="Constant"),
+	CIM_CurveUser UMETA(DisplayName="Curve User"),
+	CIM_CurveBreak UMETA(DisplayName="Curve Break"),
+	CIM_CurveAutoClamped UMETA(DisplayName="Curve Auto Clamped"),
+	CIM_MAX,
+};
+
+// @warning:	When you update this, you must add an entry to GPixelFormats(see RenderUtils.cpp)
+// @warning:	When you update this, you must add an entries to PixelFormat.h, usually just copy the generated section on the header into EPixelFormat
+// @warning:	The *Tools DLLs will also need to be recompiled if the ordering is changed, but should not need code changes.
+
+UENUM()
+enum EPixelFormat
+{
+	PF_Unknown,
+	PF_A32B32G32R32F,
+	PF_B8G8R8A8,
+	/** UNORM red (0..1) */
+	PF_G8,
+	PF_G16,
+	PF_DXT1,
+	PF_DXT3,
+	PF_DXT5,
+	PF_UYVY,
+	/** A RGB FP format with platform-specific implementation, for use with render targets. */
+	PF_FloatRGB,
+	/** A RGBA FP format with platform-specific implementation, for use with render targets. */
+	PF_FloatRGBA,
+	/** A depth+stencil format with platform-specific implementation, for use with render targets. */
+	PF_DepthStencil,
+	/** A depth format with platform-specific implementation, for use with render targets. */
+	PF_ShadowDepth,
+	PF_R32_FLOAT,
+	PF_G16R16,
+	PF_G16R16F,
+	PF_G16R16F_FILTER,
+	PF_G32R32F,
+	PF_A2B10G10R10,
+	PF_A16B16G16R16,
+	PF_D24,
+	PF_R16F,
+	PF_R16F_FILTER,
+	PF_BC5,
+	/** SNORM red, green (-1..1). Not supported on all RHI e.g. Metal */
+	PF_V8U8,
+	PF_A1,
+	/** A low precision floating point format. */
+	PF_FloatR11G11B10,
+	PF_A8,
+	PF_R32_UINT,
+	PF_R32_SINT,
+	PF_PVRTC2,
+	PF_PVRTC4,
+	PF_R16_UINT,
+	PF_R16_SINT,
+	PF_R16G16B16A16_UINT,
+	PF_R16G16B16A16_SINT,
+	PF_R5G6B5_UNORM,
+	PF_R8G8B8A8,
+	/** Only used for legacy loading; do NOT use! */
+	PF_A8R8G8B8,
+	/** High precision single channel block compressed, equivalent to a single channel BC5, 8 bytes per 4x4 block. */
+	PF_BC4,
+	/** UNORM red, green (0..1). */
+	PF_R8G8,
+	/** ATITC format. */
+	PF_ATC_RGB,
+	/** ATITC format. */
+	PF_ATC_RGBA_E,
+	/** ATITC format. */
+	PF_ATC_RGBA_I,
+	/** Used for creating SRVs to alias a DepthStencil buffer to read Stencil.  Don't use for creating textures. */
+	PF_X24_G8,
+	PF_ETC1,
+	PF_ETC2_RGB,
+	PF_ETC2_RGBA,
+	PF_R32G32B32A32_UINT,
+	PF_R16G16_UINT,
+	/** 8.00 bpp */
+	PF_ASTC_4x4,
+	/** 3.56 bpp */
+	PF_ASTC_6x6,
+	/** 2.00 bpp */
+	PF_ASTC_8x8,
+	/** 1.28 bpp */
+	PF_ASTC_10x10,
+	/** 0.89 bpp */
+	PF_ASTC_12x12,
+	PF_BC6H,
+	PF_BC7,
+	PF_R8_UINT,
+	PF_MAX,
+};
+
+UENUM()
+namespace EMouseCursor
+{
+	enum Type
+	{
+		/** Causes no mouse cursor to be visible. */
+		None,
+
+		/** Default cursor (arrow). */
+		Default,
+
+		/** Text edit beam. */
+		TextEditBeam,
+
+		/** Resize horizontal. */
+		ResizeLeftRight,
+
+		/** Resize vertical. */
+		ResizeUpDown,
+
+		/** Resize diagonal. */
+		ResizeSouthEast,
+
+		/** Resize other diagonal. */
+		ResizeSouthWest,
+
+		/** MoveItem. */
+		CardinalCross,
+
+		/** Target Cross. */
+		Crosshairs,
+
+		/** Hand cursor. */
+		Hand,
+
+		/** Grab Hand cursor. */
+		GrabHand,
+
+		/** Grab Hand cursor closed. */
+		GrabHandClosed,
+
+		/** a circle with a diagonal line through it. */
+		SlashedCircle,
+
+		/** Eye-dropper cursor for picking colors. */
+		EyeDropper,
+	};
+}
+
+/** A set of numerical unit types supported by the engine. Mirrored from UnitConversion.h */
+UENUM(BlueprintType)
+enum class EUnit
+{
+	/** Scalar distance/length unit. */
+
+	Micrometers, Millimeters, Centimeters, Meters, Kilometers,
+	Inches, Feet, Yards, Miles,
+	Lightyears,
+
+
+	/** Angular units */
+
+	Degrees, Radians,
+
+
+	/** Speed units */
+
+	MetersPerSecond, KilometersPerHour, MilesPerHour,
+
+
+	/** Temperature units */
+
+	Celsius, Farenheit, Kelvin,
+
+
+	/** Mass units */
+
+	Micrograms, Milligrams, Grams, Kilograms, MetricTons,
+	Ounces, Pounds, Stones,
+
+
+	/** Force units */
+
+	Newtons, PoundsForce, KilogramsForce,
+
+
+	/** Frequency units */
+
+	Hertz, Kilohertz, Megahertz, Gigahertz, RevolutionsPerMinute,
+
+
+	/** Data Size units */
+	
+	Bytes, Kilobytes, Megabytes, Gigabytes, Terabytes,
+
+
+	/** Luminous flux units */
+	
+	Lumens,
+
+
+	/** Time units */
+	
+	Milliseconds, Seconds, Minutes, Hours, Days, Months, Years,
+
+	/** Arbitrary multiplier */
+	
+	Multiplier,
+
+
+	/** Symbolic entry, not specifiable on meta data. */
+	
+	Unspecified
+};
+
+/** A globally unique identifier. */
+USTRUCT(immutable, noexport, BlueprintType)
+struct FGuid
+{
+	UPROPERTY(EditAnywhere, SaveGame, Category=Guid)
+	int32 A;
+
+	UPROPERTY(EditAnywhere, SaveGame, Category=Guid)
+	int32 B;
+
+	UPROPERTY(EditAnywhere, SaveGame, Category=Guid)
+	int32 C;
+
+	UPROPERTY(EditAnywhere, SaveGame, Category=Guid)
+	int32 D;
+};
+
+/**
+ * A point or direction FVector in 3d space.
+ * The full C++ class is located here: Engine\Source\Runtime\Core\Public\Math\Vector.h
+ */
+USTRUCT(immutable, noexport, BlueprintType, meta=(HasNativeMake="Engine.KismetMathLibrary.MakeVector", HasNativeBreak="Engine.KismetMathLibrary.BreakVector"))
+struct FVector
+{
+	UPROPERTY(EditAnywhere, BlueprintReadWrite, Category=Vector, SaveGame)
+	float X;
+
+	UPROPERTY(EditAnywhere, BlueprintReadWrite, Category=Vector, SaveGame)
+	float Y;
+
+	UPROPERTY(EditAnywhere, BlueprintReadWrite, Category=Vector, SaveGame)
+	float Z;
+};
+
+/**
+* A 4-D homogeneous vector.
+* The full C++ class is located here: Engine\Source\Runtime\Core\Public\Math\Vector4.h
+*/
+USTRUCT(immutable, noexport)
+struct FVector4
+{
+	UPROPERTY(EditAnywhere, BlueprintReadWrite, Category=Vector4, SaveGame)
+	float X;
+
+	UPROPERTY(EditAnywhere, BlueprintReadWrite, Category=Vector4, SaveGame)
+	float Y;
+
+	UPROPERTY(EditAnywhere, BlueprintReadWrite, Category=Vector4, SaveGame)
+	float Z;
+
+	UPROPERTY(EditAnywhere, BlueprintReadWrite, Category=Vector4, SaveGame)
+	float W;
+
+};
+
+/**
+* A point or direction FVector in 2d space.
+* The full C++ class is located here: Engine\Source\Runtime\Core\Public\Math\Vector2D.h
+*/
+USTRUCT(immutable, noexport, BlueprintType, meta=(HasNativeMake="Engine.KismetMathLibrary.MakeVector2D", HasNativeBreak="Engine.KismetMathLibrary.BreakVector2D"))
+struct FVector2D
+{
+	UPROPERTY(EditAnywhere, BlueprintReadWrite, Category=Vector2D, SaveGame)
+	float X;
+
+	UPROPERTY(EditAnywhere, BlueprintReadWrite, Category=Vector2D, SaveGame)
+	float Y;
+
+};
+
+
+
+USTRUCT(immutable, noexport)
+struct FTwoVectors
+{
+	UPROPERTY(EditAnywhere, BlueprintReadWrite, Category=TwoVectors, SaveGame)
+	FVector v1;
+
+	UPROPERTY(EditAnywhere, BlueprintReadWrite, Category=TwoVectors, SaveGame)
+	FVector v2;
+
+};
+
+
+/**
+ * A plane definition in 3D space.
+ * The full C++ class is located here: Engine\Source\Runtime\Core\Public\Math\Plane.h
+ */
+USTRUCT(immutable, noexport)
+struct FPlane : public FVector
+{
+	UPROPERTY(EditAnywhere, BlueprintReadWrite, Category=Plane, SaveGame)
+	float W;
+
+};
+
+
+/**
+ * An orthogonal rotation in 3d space.
+ * The full C++ class is located here: Engine\Source\Runtime\Core\Public\Math\Rotator.h
+ */
+USTRUCT(immutable, noexport, BlueprintType, meta=(HasNativeMake="Engine.KismetMathLibrary.MakeRotator", HasNativeBreak="Engine.KismetMathLibrary.BreakRotator"))
+struct FRotator
+{
+	/** Pitch (degrees) around Y axis */
+	UPROPERTY(EditAnywhere, BlueprintReadWrite, Category=Rotator, SaveGame, meta=(DisplayName="Y"))
+	float Pitch;
+
+	/** Yaw (degrees) around Z axis */
+	UPROPERTY(EditAnywhere, BlueprintReadWrite, Category=Rotator, SaveGame, meta=(DisplayName="Z"))
+	float Yaw;
+
+	/** Roll (degrees) around X axis */
+	UPROPERTY(EditAnywhere, BlueprintReadWrite, Category=Rotator, SaveGame, meta=(DisplayName="X"))
+	float Roll;
+
+};
+
+
+/**
+ * Quaternion.
+ * The full C++ class is located here: Engine\Source\Runtime\Core\Public\Math\Quat.h
+ */
+USTRUCT(immutable, noexport)
+struct FQuat
+{
+	UPROPERTY(EditAnywhere, BlueprintReadWrite, Category=Quat, SaveGame)
+	float X;
+
+	UPROPERTY(EditAnywhere, BlueprintReadWrite, Category=Quat, SaveGame)
+	float Y;
+
+	UPROPERTY(EditAnywhere, BlueprintReadWrite, Category=Quat, SaveGame)
+	float Z;
+
+	UPROPERTY(EditAnywhere, BlueprintReadWrite, Category=Quat, SaveGame)
+	float W;
+
+};
+
+
+/**
+ * A packed normal.
+ * The full C++ class is located here: Engine\Source\Runtime\RenderCore\Public\PackedNormal.h
+ */
+USTRUCT(immutable, noexport)
+struct FPackedNormal
+{
+	UPROPERTY(EditAnywhere, BlueprintReadWrite, Category=PackedNormal, SaveGame)
+	uint8 X;
+
+	UPROPERTY(EditAnywhere, BlueprintReadWrite, Category=PackedNormal, SaveGame)
+	uint8 Y;
+
+	UPROPERTY(EditAnywhere, BlueprintReadWrite, Category=PackedNormal, SaveGame)
+	uint8 Z;
+
+	UPROPERTY(EditAnywhere, BlueprintReadWrite, Category=PackedNormal, SaveGame)
+	uint8 W;
+
+};
+
+/**
+* A packed basis vector.
+* The full C++ class is located here: Engine\Source\Runtime\RenderCore\Public\PackedNormal.h
+*/
+USTRUCT(immutable, noexport)
+struct FPackedRGB10A2N
+{
+	UPROPERTY(EditAnywhere, BlueprintReadWrite, Category = PackedBasis, SaveGame)
+	int32 Packed;
+};
+
+/**
+* A packed vector.
+* The full C++ class is located here: Engine\Source\Runtime\RenderCore\Public\PackedNormal.h
+*/
+USTRUCT(immutable, noexport)
+struct FPackedRGBA16N
+{
+	UPROPERTY(EditAnywhere, BlueprintReadWrite, Category = PackedNormal, SaveGame)
+	int32 XY;
+
+	UPROPERTY(EditAnywhere, BlueprintReadWrite, Category = PackedNormal, SaveGame)
+	int32 ZW;
+};
+
+/**
+ * Screen coordinates.
+ * The full C++ class is located here: Engine\Source\Runtime\Core\Public\Math\IntPoint.h
+ */
+USTRUCT(immutable, noexport)
+struct FIntPoint
+{
+	UPROPERTY(EditAnywhere, BlueprintReadWrite, Category=IntPoint, SaveGame)
+	int32 X;
+
+	UPROPERTY(EditAnywhere, BlueprintReadWrite, Category=IntPoint, SaveGame)
+	int32 Y;
+
+};
+
+/**
+ * An integer vector in 3D space.
+ * The full C++ class is located here: Engine\Source\Runtime\Core\Public\Math\IntVector.h
+ */
+USTRUCT(immutable, noexport)
+struct FIntVector
+{
+	UPROPERTY(EditAnywhere, BlueprintReadWrite, Category=IntVector, SaveGame)
+	int32 X;
+
+	UPROPERTY(EditAnywhere, BlueprintReadWrite, Category=IntVector, SaveGame)
+	int32 Y;
+
+	UPROPERTY(EditAnywhere, BlueprintReadWrite, Category=IntVector, SaveGame)
+	int32 Z;
+};
+
+
+/**
+ * A Color (BGRA).
+ * The full C++ class is located here: Engine\Source\Runtime\Core\Public\Math\Color.h
+ */
+USTRUCT(immutable, noexport)
+struct FColor
+{
+	UPROPERTY(EditAnywhere, BlueprintReadWrite, Category=Color, SaveGame, meta=(ClampMin="0", ClampMax="255"))
+	uint8 B;
+
+	UPROPERTY(EditAnywhere, BlueprintReadWrite, Category=Color, SaveGame, meta=(ClampMin="0", ClampMax="255"))
+	uint8 G;
+
+	UPROPERTY(EditAnywhere, BlueprintReadWrite, Category=Color, SaveGame, meta=(ClampMin="0", ClampMax="255"))
+	uint8 R;
+
+	UPROPERTY(EditAnywhere, BlueprintReadWrite, Category=Color, SaveGame, meta=(ClampMin="0", ClampMax="255"))
+	uint8 A;
+
+};
+
+
+/**
+ * A linear color.
+ * The full C++ class is located here: Engine\Source\Runtime\Core\Public\Math\Color.h
+ */
+USTRUCT(immutable, noexport, BlueprintType)
+struct FLinearColor
+{
+	UPROPERTY(EditAnywhere, BlueprintReadWrite, Category=LinearColor, SaveGame)
+	float R;
+
+	UPROPERTY(EditAnywhere, BlueprintReadWrite, Category=LinearColor, SaveGame)
+	float G;
+
+	UPROPERTY(EditAnywhere, BlueprintReadWrite, Category=LinearColor, SaveGame)
+	float B;
+
+	UPROPERTY(EditAnywhere, BlueprintReadWrite, Category=LinearColor, SaveGame)
+	float A;
+
+};
+
+
+/**
+ * A bounding box.
+ * The full C++ class is located here: Engine\Source\Runtime\Core\Public\Math\Box.h
+ */
+USTRUCT(immutable, noexport, BlueprintType)
+struct FBox
+{
+	UPROPERTY(EditAnywhere, BlueprintReadWrite, Category=Box, SaveGame)
+	FVector Min;
+
+	UPROPERTY(EditAnywhere, BlueprintReadWrite, Category=Box, SaveGame)
+	FVector Max;
+
+	UPROPERTY()
+	uint8 IsValid;
+
+};
+
+/**
+ * A rectangular 2D Box.
+ * The full C++ class is located here: Engine\Source\Runtime\Core\Public\Math\Box2D.h
+ */
+USTRUCT(immutable, noexport, BlueprintType)
+struct FBox2D
+{
+	UPROPERTY(EditAnywhere, BlueprintReadWrite, Category=Box2D, SaveGame)
+	FVector2D Min;
+
+	UPROPERTY(EditAnywhere, BlueprintReadWrite, Category=Box2D, SaveGame)
+	FVector2D Max;
+
+	UPROPERTY()
+	uint8 IsValid;
+
+};
+
+
+/**
+ * A bounding box and bounding sphere with the same origin.
+ * The full C++ class is located here : Engine\Source\Runtime\Core\Public\Math\BoxSphereBounds.h
+ */
+USTRUCT(noexport)
+struct FBoxSphereBounds
+{
+	UPROPERTY(EditAnywhere, BlueprintReadWrite, Category=BoxSphereBounds, SaveGame)
+	FVector Origin;
+
+	UPROPERTY(EditAnywhere, BlueprintReadWrite, Category=BoxSphereBounds, SaveGame)
+	FVector BoxExtent;
+
+	UPROPERTY(EditAnywhere, BlueprintReadWrite, Category=BoxSphereBounds, SaveGame)
+	float SphereRadius;
+
+};
+
+/**
+ * Structure for arbitrarily oriented boxes (i.e. not necessarily axis-aligned).
+ * The full C++ class is located here: Engine\Source\Runtime\Core\Public\Math\OrientedBox.h
+ */
+USTRUCT(immutable, noexport)
+struct FOrientedBox
+{
+	UPROPERTY(EditAnywhere, BlueprintReadWrite, Category=OrientedBox, SaveGame)
+	FVector Center;
+
+	UPROPERTY(EditAnywhere, BlueprintReadWrite, Category=OrientedBox, SaveGame)
+	FVector AxisX;
+	
+	UPROPERTY(EditAnywhere, BlueprintReadWrite, Category=OrientedBox, SaveGame)
+	FVector AxisY;
+	
+	UPROPERTY(EditAnywhere, BlueprintReadWrite, Category=OrientedBox, SaveGame)
+	FVector AxisZ;
+
+	UPROPERTY(EditAnywhere, BlueprintReadWrite, Category=OrientedBox, SaveGame)
+	float ExtentX;
+	
+	UPROPERTY(EditAnywhere, BlueprintReadWrite, Category=OrientedBox, SaveGame)
+	float ExtentY;
+
+	UPROPERTY(EditAnywhere, BlueprintReadWrite, Category=OrientedBox, SaveGame)
+	float ExtentZ;
+};
+
+/*
+ * A 4x4 matrix.
+ * The full C++ class is located here: Engine\Source\Runtime\Core\Public\Math\Matrix.h
+ */
+USTRUCT(immutable, noexport)
+struct FMatrix
+{
+	UPROPERTY(EditAnywhere, BlueprintReadWrite, Category=Matrix, SaveGame)
+	FPlane XPlane;
+
+	UPROPERTY(EditAnywhere, BlueprintReadWrite, Category=Matrix, SaveGame)
+	FPlane YPlane;
+
+	UPROPERTY(EditAnywhere, BlueprintReadWrite, Category=Matrix, SaveGame)
+	FPlane ZPlane;
+
+	UPROPERTY(EditAnywhere, BlueprintReadWrite, Category=Matrix, SaveGame)
+	FPlane WPlane;
+
+};
+
+
+
+USTRUCT(noexport)
+struct FInterpCurvePointFloat
+{
+	UPROPERTY(EditAnywhere, BlueprintReadWrite, Category=InterpCurvePointFloat)
+	float InVal;
+
+	UPROPERTY(EditAnywhere, BlueprintReadWrite, Category=InterpCurvePointFloat)
+	float OutVal;
+
+	UPROPERTY(EditAnywhere, BlueprintReadWrite, Category=InterpCurvePointFloat)
+	float ArriveTangent;
+
+	UPROPERTY(EditAnywhere, BlueprintReadWrite, Category=InterpCurvePointFloat)
+	float LeaveTangent;
+
+	UPROPERTY(EditAnywhere, BlueprintReadWrite, Category=InterpCurvePointFloat)
+	TEnumAsByte<enum EInterpCurveMode> InterpMode;
+
+};
+
+
+
+USTRUCT(noexport)
+struct FInterpCurveFloat
+{
+	UPROPERTY(EditAnywhere, BlueprintReadWrite, Category=InterpCurveFloat)
+	TArray<FInterpCurvePointFloat> Points;
+
+	UPROPERTY(EditAnywhere, BlueprintReadWrite, Category=InterpCurveFloat)
+	bool bIsLooped;
+
+	UPROPERTY(EditAnywhere, BlueprintReadWrite, Category=InterpCurveFloat)
+	float LoopKeyOffset;
+};
+
+
+
+USTRUCT(noexport)
+struct FInterpCurvePointVector2D
+{
+	UPROPERTY(EditAnywhere, BlueprintReadWrite, Category=InterpCurvePointVector2D)
+	float InVal;
+
+	UPROPERTY(EditAnywhere, BlueprintReadWrite, Category=InterpCurvePointVector2D)
+	FVector2D OutVal;
+
+	UPROPERTY(EditAnywhere, BlueprintReadWrite, Category=InterpCurvePointVector2D)
+	FVector2D ArriveTangent;
+
+	UPROPERTY(EditAnywhere, BlueprintReadWrite, Category=InterpCurvePointVector2D)
+	FVector2D LeaveTangent;
+
+	UPROPERTY(EditAnywhere, BlueprintReadWrite, Category=InterpCurvePointVector2D)
+	TEnumAsByte<enum EInterpCurveMode> InterpMode;
+
+};
+
+
+
+USTRUCT(noexport)
+struct FInterpCurveVector2D
+{
+	UPROPERTY(EditAnywhere, BlueprintReadWrite, Category=InterpCurveVector2D)
+	TArray<FInterpCurvePointVector2D> Points;
+
+	UPROPERTY(EditAnywhere, BlueprintReadWrite, Category=InterpCurveVector2D)
+	bool bIsLooped;
+
+	UPROPERTY(EditAnywhere, BlueprintReadWrite, Category=InterpCurveVector2D)
+	float LoopKeyOffset;
+};
+
+
+
+USTRUCT(noexport)
+struct FInterpCurvePointVector
+{
+	UPROPERTY(EditAnywhere, BlueprintReadWrite, Category=InterpCurvePointVector)
+	float InVal;
+
+	UPROPERTY(EditAnywhere, BlueprintReadWrite, Category=InterpCurvePointVector)
+	FVector OutVal;
+
+	UPROPERTY(EditAnywhere, BlueprintReadWrite, Category=InterpCurvePointVector)
+	FVector ArriveTangent;
+
+	UPROPERTY(EditAnywhere, BlueprintReadWrite, Category=InterpCurvePointVector)
+	FVector LeaveTangent;
+
+	UPROPERTY(EditAnywhere, BlueprintReadWrite, Category=InterpCurvePointVector)
+	TEnumAsByte<enum EInterpCurveMode> InterpMode;
+
+};
+
+
+
+USTRUCT(noexport)
+struct FInterpCurveVector
+{
+	UPROPERTY(EditAnywhere, BlueprintReadWrite, Category=InterpCurveVector)
+	TArray<FInterpCurvePointVector> Points;
+
+	UPROPERTY(EditAnywhere, BlueprintReadWrite, Category=InterpCurveVector)
+	bool bIsLooped;
+
+	UPROPERTY(EditAnywhere, BlueprintReadWrite, Category=InterpCurveVector)
+	float LoopKeyOffset;
+};
+
+
+
+USTRUCT(noexport)
+struct FInterpCurvePointQuat
+{
+	UPROPERTY(EditAnywhere, BlueprintReadWrite, Category=InterpCurvePointQuat)
+	float InVal;
+
+	UPROPERTY(EditAnywhere, BlueprintReadWrite, Category=InterpCurvePointQuat)
+	FQuat OutVal;
+
+	UPROPERTY(EditAnywhere, BlueprintReadWrite, Category=InterpCurvePointQuat)
+	FQuat ArriveTangent;
+
+	UPROPERTY(EditAnywhere, BlueprintReadWrite, Category=InterpCurvePointQuat)
+	FQuat LeaveTangent;
+
+	UPROPERTY(EditAnywhere, BlueprintReadWrite, Category=InterpCurvePointQuat)
+	TEnumAsByte<enum EInterpCurveMode> InterpMode;
+
+};
+
+
+
+USTRUCT(noexport)
+struct FInterpCurveQuat
+{
+	UPROPERTY(EditAnywhere, BlueprintReadWrite, Category=InterpCurveQuat)
+	TArray<FInterpCurvePointQuat> Points;
+
+	UPROPERTY(EditAnywhere, BlueprintReadWrite, Category=InterpCurveQuat)
+	bool bIsLooped;
+
+	UPROPERTY(EditAnywhere, BlueprintReadWrite, Category=InterpCurveQuat)
+	float LoopKeyOffset;
+};
+
+
+
+USTRUCT(noexport)
+struct FInterpCurvePointTwoVectors
+{
+	UPROPERTY(EditAnywhere, BlueprintReadWrite, Category=InterpCurvePointTwoVectors)
+	float InVal;
+
+	UPROPERTY(EditAnywhere, BlueprintReadWrite, Category=InterpCurvePointTwoVectors)
+	FTwoVectors OutVal;
+
+	UPROPERTY(EditAnywhere, BlueprintReadWrite, Category=InterpCurvePointTwoVectors)
+	FTwoVectors ArriveTangent;
+
+	UPROPERTY(EditAnywhere, BlueprintReadWrite, Category=InterpCurvePointTwoVectors)
+	FTwoVectors LeaveTangent;
+
+	UPROPERTY(EditAnywhere, BlueprintReadWrite, Category=InterpCurvePointTwoVectors)
+	TEnumAsByte<enum EInterpCurveMode> InterpMode;
+
+};
+
+
+
+USTRUCT(noexport)
+struct FInterpCurveTwoVectors
+{
+	UPROPERTY(EditAnywhere, BlueprintReadWrite, Category=InterpCurveTwoVectors)
+	TArray<FInterpCurvePointTwoVectors> Points;
+
+	UPROPERTY(EditAnywhere, BlueprintReadWrite, Category=InterpCurveTwoVectors)
+	bool bIsLooped;
+
+	UPROPERTY(EditAnywhere, BlueprintReadWrite, Category=InterpCurveTwoVectors)
+	float LoopKeyOffset;
+};
+
+
+
+USTRUCT(noexport)
+struct FInterpCurvePointLinearColor
+{
+	UPROPERTY(EditAnywhere, BlueprintReadWrite, Category=InterpCurvePointLinearColor)
+	float InVal;
+
+	UPROPERTY(EditAnywhere, BlueprintReadWrite, Category=InterpCurvePointLinearColor)
+	FLinearColor OutVal;
+
+	UPROPERTY(EditAnywhere, BlueprintReadWrite, Category=InterpCurvePointLinearColor)
+	FLinearColor ArriveTangent;
+
+	UPROPERTY(EditAnywhere, BlueprintReadWrite, Category=InterpCurvePointLinearColor)
+	FLinearColor LeaveTangent;
+
+	UPROPERTY(EditAnywhere, BlueprintReadWrite, Category=InterpCurvePointLinearColor)
+	TEnumAsByte<enum EInterpCurveMode> InterpMode;
+
+};
+
+
+
+USTRUCT(noexport)
+struct FInterpCurveLinearColor
+{
+	UPROPERTY(EditAnywhere, BlueprintReadWrite, Category=InterpCurveLinearColor)
+	TArray<FInterpCurvePointLinearColor> Points;
+
+	UPROPERTY(EditAnywhere, BlueprintReadWrite, Category=InterpCurveLinearColor)
+	bool bIsLooped;
+
+	UPROPERTY(EditAnywhere, BlueprintReadWrite, Category=InterpCurveLinearColor)
+	float LoopKeyOffset;
+};
+
+
+/**
+ * Transform composed of Quat/Translation/Scale.
+ * The full C++ class is located here: Engine\Source\Runtime\Core\Public\Math\Transform.h
+ */
+USTRUCT(noexport, BlueprintType, meta=(HasNativeMake="Engine.KismetMathLibrary.MakeTransform", HasNativeBreak="Engine.KismetMathLibrary.BreakTransform"))
+struct FTransform
+{
+	UPROPERTY(EditAnywhere, BlueprintReadWrite, Category=Transform, SaveGame)
+	FQuat Rotation;
+
+	UPROPERTY(EditAnywhere, BlueprintReadWrite, Category=Transform, SaveGame)
+	FVector Translation;
+
+	UPROPERTY(EditAnywhere, BlueprintReadWrite, Category=Transform, SaveGame, meta=(MakeStructureDefaultValue = "1,1,1"))
+	FVector Scale3D;
+
+};
+
+
+/**
+ * Thread-safe RNG.
+ * The full C++ class is located here: Engine\Source\Runtime\Core\Public\Math\RandomStream.h
+ */
+USTRUCT(noexport, BlueprintType, meta = (HasNativeMake = "Engine.KismetMathLibrary.MakeRandomStream", HasNativeBreak = "Engine.KismetMathLibrary.BreakRandomStream"))
+struct FRandomStream
+{
+public:
+	UPROPERTY(EditAnywhere, BlueprintReadWrite, Category=RandomStream, SaveGame)
+	int32 InitialSeed;
+	
+	UPROPERTY()
+	int32 Seed;
+};
+
+
+// A date/time value.
+
+USTRUCT(immutable, noexport, BlueprintType, meta=(HasNativeMake="Engine.KismetMathLibrary.MakeDateTime", HasNativeBreak="Engine.KismetMathLibrary.BreakDateTime"))
+struct FDateTime
+{
+	int64 Ticks;
+};
+
+
+// A time span value.
+
+USTRUCT(immutable, noexport, BlueprintType, meta=(HasNativeMake="Engine.KismetMathLibrary.MakeTimespan", HasNativeBreak="Engine.KismetMathLibrary.BreakTimespan"))
+struct FTimespan
+{
+	int64 Ticks;
+};
+
+
+// A string asset reference
+
+USTRUCT(noexport, meta=(HasNativeMake="Engine.BlueprintFunctionLibrary.MakeStringAssetReference"))
+struct FStringAssetReference
+{
+	UPROPERTY(EditAnywhere, BlueprintReadWrite, Category=StringAssetReference)
+	FString AssetLongPathname;
+};
+
+// A string class reference
+
+USTRUCT(noexport)
+struct FStringClassReference : public FStringAssetReference
+{
+};
+
+// A struct used as stub for deleted ones.
+
+USTRUCT(noexport)
+struct FFallbackStruct  
+{
+};
+
+UENUM()
+namespace ERangeBoundTypes
+{
+	/**
+	* Enumerates the valid types of range bounds.
+	*/
+	enum Type
+	{
+		/**
+		* The range excludes the bound.
+		*/
+		Exclusive,
+
+		/**
+		* The range includes the bound.
+		*/
+		Inclusive,
+
+		/**
+		* The bound is open.
+		*/
+		Open
+	};
+}
+
+// A float range bound
+
+USTRUCT(noexport)
+struct FFloatRangeBound
+{
+	UPROPERTY(EditAnywhere, Category=Range)
+	TEnumAsByte<ERangeBoundTypes::Type> Type;
+
+	UPROPERTY(EditAnywhere, Category=Range)
+	float Value;
+};
+
+// A float range
+
+USTRUCT(noexport)
+struct FFloatRange
+{
+	UPROPERTY(EditAnywhere, Category=Range)
+	FFloatRangeBound LowerBound;
+
+	UPROPERTY(EditAnywhere, Category=Range)
+	FFloatRangeBound UpperBound;
+};
+
+// An int32 range bound
+
+USTRUCT(noexport)
+struct FInt32RangeBound
+{
+	UPROPERTY(EditAnywhere, Category = Range)
+	TEnumAsByte<ERangeBoundTypes::Type> Type;
+
+	UPROPERTY(EditAnywhere, Category = Range)
+	int32 Value;
+};
+
+// An int32 range
+
+USTRUCT(noexport)
+struct FInt32Range
+{
+	UPROPERTY(EditAnywhere, Category = Range)
+	FInt32RangeBound LowerBound;
+
+	UPROPERTY(EditAnywhere, Category = Range)
+	FInt32RangeBound UpperBound;
+};
+
+// A float interval
+
+USTRUCT(noexport)
+struct FFloatInterval
+{
+	UPROPERTY(EditAnywhere, Category=Interval)
+	float Min;
+
+	UPROPERTY(EditAnywhere, Category=Interval)
+	float Max;
+};
+
+// An int32 interval
+
+USTRUCT(noexport)
+struct FInt32Interval
+{
+	UPROPERTY(EditAnywhere, Category=Interval)
+	int32 Min;
+
+	UPROPERTY(EditAnywhere, Category=Interval)
+	int32 Max;
+};
+
+//=============================================================================
+/**
+ * Object: The base class all objects.
+ * This is a built-in Unreal class and it shouldn't be modified by mod authors.
+ * The full C++ class is located here: Engine\Source\Runtime\CoreUObject\Public\UObject\UObject.h
+ */
+//=============================================================================
+
+UCLASS(abstract, noexport)
+class UObject
+{
+	GENERATED_BODY()
+public:
+
+	/**
+	 * Default UObject constructor.
+	 */
+	UObject(const FObjectInitializer& ObjectInitializer = FObjectInitializer::Get());
+#if WITH_HOT_RELOAD_CTORS
+	/** DO NOT USE. This constructor is for internal usage only for hot-reload purposes. */
+	UObject(FVTableHelper& Helper);
+#endif // WITH_HOT_RELOAD_CTORS
+
+	//=============================================================================
+	// K2 support functions.
+	
+	/**
+	 * Executes some portion of the ubergraph.
+	 *
+	 * @param	EntryPoint	The entry point to start code execution at.
+	 */
+	UFUNCTION(BlueprintImplementableEvent, meta=(BlueprintInternalUseOnly = "true"))
+	void ExecuteUbergraph(int32 EntryPoint);
+};
+
+#endif
+
+#if CPP
+
+// Include the real definitions of the noexport classes above, to allow the generated cpp file to compile.
+
+#include "Misc/Guid.h"
+#include "Misc/DateTime.h"
+#include "Misc/Timespan.h"
+#include "Misc/StringAssetReference.h"
+#include "Misc/StringClassReference.h"
+
+#include "Math/InterpCurvePoint.h"
+#include "Math/UnitConversion.h"
+#include "Math/Vector.h"
+#include "Math/Vector4.h"
+#include "Math/Vector2D.h"
+#include "Math/TwoVectors.h"
+#include "Math/Plane.h"
+#include "Math/Rotator.h"
+#include "Math/Quat.h"
+#include "Math/IntPoint.h"
+#include "Math/IntVector.h"
+#include "Math/Color.h"
+#include "Math/Box.h"
+#include "Math/Box2D.h"
+#include "Math/BoxSphereBounds.h"
+#include "Math/OrientedBox.h"
+#include "Math/Matrix.h"
+#include "Math/ScalarRegister.h"
+#include "Math/RandomStream.h"
+#include "Math/RangeBound.h"
+#include "Math/Interval.h"
+
+#endif
diff --git a/Engine/Source/Runtime/CoreUObject/Public/UObject/Object.h b/Engine/Source/Runtime/CoreUObject/Public/UObject/Object.h
new file mode 100644
--- /dev/null
+++ b/Engine/Source/Runtime/CoreUObject/Public/UObject/Object.h
@@ -0,0 +1,1196 @@
+// Copyright 1998-2016 Epic Games, Inc. All Rights Reserved.
+
+#pragma once
+
+#include "UObjectBaseUtility.h"
+
+DECLARE_LOG_CATEGORY_EXTERN(LogObj, Log, All);
+
+
+namespace ECastCheckedType
+{
+	enum Type
+	{
+		NullAllowed,
+		NullChecked
+	};
+};
+
+/** Passed to GetResourceSize() to indicate which resource size should be returned.*/
+namespace EResourceSizeMode
+{
+	enum Type
+	{
+		/** Only exclusive resource size */
+		Exclusive,
+		/** Resource size of the object and all of its references */
+		Inclusive,
+	};
+};
+
+//
+// The base class of all objects.
+//
+class COREUOBJECT_API UObject : public UObjectBaseUtility
+{
+	// Declarations.
+	DECLARE_CLASS(UObject,UObject,CLASS_Abstract|CLASS_NoExport,CASTCLASS_None,TEXT("/Script/CoreUObject"),NO_API)
+	DEFINE_DEFAULT_OBJECT_INITIALIZER_CONSTRUCTOR_CALL(UObject)
+#if WITH_HOT_RELOAD_CTORS
+	static UObject* __VTableCtorCaller(FVTableHelper& Helper)
+	{
+		return new (EC_InternalUseOnlyConstructor, (UObject*)GetTransientPackage(), NAME_None, RF_NeedLoad | RF_ClassDefaultObject | RF_TagGarbageTemp) UObject(Helper);
+	}
+#endif // WITH_HOT_RELOAD_CTORS
+
+	typedef UObject WithinClass;
+	static const TCHAR* StaticConfigName() {return TEXT("Engine");}
+
+	// Constructors and destructors.
+	UObject();
+	UObject(const FObjectInitializer& ObjectInitializer);
+	UObject( EStaticConstructor, EObjectFlags InFlags );
+#if WITH_HOT_RELOAD_CTORS
+	/** DO NOT USE. This constructor is for internal usage only for hot-reload purposes. */
+	UObject(FVTableHelper& Helper);
+#endif // WITH_HOT_RELOAD_CTORS
+
+	static void StaticRegisterNativesUObject() 
+	{
+	}
+
+	UObject* CreateDefaultSubobject(FName SubobjectFName, UClass* ReturnType, UClass* ClassToCreateByDefault, bool bIsRequired, bool bAbstract, bool bIsTransient);
+
+	/**
+	* Create a component or subobject only to be used with the editor.
+	* @param	TReturnType					class of return type, all overrides must be of this type
+	* @param	Outer						outer to construct the subobject in
+	* @param	SubobjectName				name of the new component
+	* @param bTransient		true if the component is being assigned to a transient property
+	*/
+	template<class TReturnType>
+	TReturnType* CreateEditorOnlyDefaultSubobject(FName SubobjectName, bool bTransient = false)
+	{
+		UClass* ReturnType = TReturnType::StaticClass();
+		return static_cast<TReturnType*>(CreateEditorOnlyDefaultSubobjectImpl(SubobjectName, ReturnType, bTransient));
+	}
+
+	/**
+	* Create a component or subobject
+	* @param	TReturnType					class of return type, all overrides must be of this type
+	* @param	Outer						outer to construct the subobject in
+	* @param	SubobjectName				name of the new component
+	* @param bTransient		true if the component is being assigned to a transient property
+	*/
+	template<class TReturnType>
+	TReturnType* CreateDefaultSubobject(FName SubobjectName, bool bTransient = false)
+	{
+		UClass* ReturnType = TReturnType::StaticClass();
+		return static_cast<TReturnType*>(CreateDefaultSubobject(SubobjectName, ReturnType, ReturnType, /*bIsRequired =*/ true, /*bIsAbstract =*/ false, bTransient));
+	}
+	
+	/**
+	* Create a component or subobject
+	* @param TReturnType class of return type, all overrides must be of this type
+	* @param TClassToConstructByDefault class to construct by default
+	* @param Outer outer to construct the subobject in
+	* @param SubobjectName name of the new component
+	* @param bTransient		true if the component is being assigned to a transient property
+	*/
+	template<class TReturnType, class TClassToConstructByDefault>
+	TReturnType* CreateDefaultSubobject(FName SubobjectName, bool bTransient = false)
+	{
+		return static_cast<TReturnType*>(CreateDefaultSubobject(SubobjectName, TReturnType::StaticClass(), TClassToConstructByDefault::StaticClass(), /*bIsRequired =*/ true, /*bIsAbstract =*/ false, bTransient));
+	}
+	
+	/**
+	* Create optional component or subobject. Optional subobjects may not get created
+	* when a derived class specified DoNotCreateDefaultSubobject with the subobject's name.
+	* @param	TReturnType					class of return type, all overrides must be of this type
+	* @param	Outer						outer to construct the subobject in
+	* @param	SubobjectName				name of the new component
+	* @param bTransient		true if the component is being assigned to a transient property
+	*/
+	template<class TReturnType>
+	TReturnType* CreateOptionalDefaultSubobject(FName SubobjectName, bool bTransient = false)
+	{
+		UClass* ReturnType = TReturnType::StaticClass();
+		return static_cast<TReturnType*>(CreateDefaultSubobject(SubobjectName, ReturnType, ReturnType, /*bIsRequired =*/ false, /*bIsAbstract =*/ false, bTransient));
+	}
+	
+	/**
+	* Create optional component or subobject. Optional subobjects may not get created
+	* when a derived class specified DoNotCreateDefaultSubobject with the subobject's name.
+	* @param	TReturnType					class of return type, all overrides must be of this type
+	* @param	Outer						outer to construct the subobject in
+	* @param	SubobjectName				name of the new component
+	* @param bTransient		true if the component is being assigned to a transient property
+	*/
+	template<class TReturnType>
+	TReturnType* CreateAbstractDefaultSubobject(FName SubobjectName, bool bTransient = false)
+	{
+		UClass* ReturnType = TReturnType::StaticClass();
+		return static_cast<TReturnType*>(CreateDefaultSubobject(SubobjectName, ReturnType, ReturnType, /*bIsRequired =*/ true, /*bIsAbstract =*/ true, bTransient));
+	}
+
+	/**
+	* Gets all default subobjects associated with this object instance.
+	* @param	OutDefaultSubobjects	Array containing all default subobjects of this object.
+	*/
+	void GetDefaultSubobjects(TArray<UObject*>& OutDefaultSubobjects);
+
+	/**
+	* Finds a subobject associated with this object instance by its name
+	* @param	OutDefaultSubobjects	Array containing all default subobjects of this object.
+	*/
+	UObject* GetDefaultSubobjectByName(FName ToFind);
+
+	//==========================================
+	// UObject interface.
+	//==========================================	
+protected:
+    /** 
+     * This function actually does the work for the GetDetailInfo and is virtual.  
+     * It should only be called from GetDetailedInfo as GetDetailedInfo is safe to call on NULL object pointers
+     **/
+	virtual FString GetDetailedInfoInternal() const { return TEXT("No_Detailed_Info_Specified"); }
+public:
+	/**
+	 * Called after the C++ constructor and after the properties have been initialized, including those loaded from config.
+	 * mainly this is to emulate some behavior of when the constructor was called after the properties were initialized.
+	 */
+	virtual void PostInitProperties();
+
+	/**
+	 * Called from within SavePackage on the passed in base/ root. The return value of this function will be passed to
+	 * PostSaveRoot. This is used to allow objects used as base to perform required actions before saving and cleanup
+	 * afterwards.
+	 * @param Filename: Name of the file being saved to (includes path)
+	 * @param AdditionalPackagesToCook [out] Array of other packages the Root wants to make sure are cooked when this is cooked
+	 *
+	 * @return	Whether PostSaveRoot needs to perform internal cleanup
+	 */
+	virtual bool PreSaveRoot(const TCHAR* Filename, TArray<FString>& AdditionalPackagesToCook)
+	{
+		return false;
+	}
+	/**
+	 * Called from within SavePackage on the passed in base/ root. This function is being called after the package
+	 * has been saved and can perform cleanup.
+	 *
+	 * @param	bCleanupIsRequired	Whether PreSaveRoot dirtied state that needs to be cleaned up
+	 */
+	virtual void PostSaveRoot( bool bCleanupIsRequired ) {}
+
+	/**
+	 * Presave function. Gets called once before an object gets serialized for saving. This function is necessary
+	 * for save time computation as Serialize gets called three times per object from within SavePackage.
+	 *
+	 * @warning: Objects created from within PreSave will NOT have PreSave called on them!!!
+	 */
+	virtual void PreSave(const class ITargetPlatform* TargetPlatform);
+
+	/**
+	 * Note that the object will be modified.  If we are currently recording into the 
+	 * transaction buffer (undo/redo), save a copy of this object into the buffer and 
+	 * marks the package as needing to be saved.
+	 *
+	 * @param	bAlwaysMarkDirty	if true, marks the package dirty even if we aren't
+	 *								currently recording an active undo/redo transaction
+	 * @return true if the object was saved to the transaction buffer
+	 */
+	virtual bool Modify( bool bAlwaysMarkDirty=true );
+
+#if WITH_EDITOR
+	/** 
+	 * Called when the object was loaded from another class via active class redirects.
+	 */
+	virtual void LoadedFromAnotherClass(const FName& OldClassName) {}
+#endif
+
+	/** 
+	 * Do any object-specific cleanup required immediately after loading an object, 
+	 * and immediately after any undo/redo.
+	 */
+	virtual void PostLoad();
+
+	/**
+	 * Instances components for objects being loaded from disk, if necessary.  Ensures that component references
+	 * between nested components are fixed up correctly.
+	 *
+	 * @param	OuterInstanceGraph	when calling this method on subobjects, specifies the instancing graph which contains all instanced
+	 *								subobjects and components for a subobject root.
+	 */
+	virtual void PostLoadSubobjects(FObjectInstancingGraph* OuterInstanceGraph);
+	
+	/**
+	 * Called before destroying the object.  This is called immediately upon deciding to destroy the object, to allow the object to begin an
+	 * asynchronous cleanup process.
+	 */
+	virtual void BeginDestroy();
+
+	/**
+	 * Called to check if the object is ready for FinishDestroy.  This is called after BeginDestroy to check the completion of the
+	 * potentially asynchronous object cleanup.
+	 * @return True if the object's asynchronous cleanup has completed and it is ready for FinishDestroy to be called.
+	 */
+	virtual bool IsReadyForFinishDestroy() { return true; }
+
+#if WITH_EDITOR
+	/**
+	 * Called in response to the linker changing, this can only happen in the editor
+	 */
+	virtual void PostLinkerChange() {}
+#endif
+
+	/**
+	 * Called to finish destroying the object.  After UObject::FinishDestroy is called, the object's memory should no longer be accessed.
+	 *
+	 * note: because properties are destroyed here, Super::FinishDestroy() should always be called at the end of your child class's
+	 * FinishDestroy() method, rather than at the beginning.
+	 */
+	virtual void FinishDestroy();
+
+	/** UObject serializer. */
+	virtual void Serialize( FArchive& Ar );
+
+	virtual void ShutdownAfterError() {}
+
+	/** 
+	 * This is called when property is modified by InterpPropertyTracks
+	 *
+	 * @param PropertyThatChanged	Property that changed
+	 */
+	virtual void PostInterpChange(UProperty* PropertyThatChanged) {}
+
+#if WITH_EDITOR
+	/** 
+	 * This is called when property is about to be modified by InterpPropertyTracks
+	 *
+	 * @param PropertyThatWillChange	Property that will be changed
+	 */
+	virtual void PreEditChange(UProperty* PropertyAboutToChange);
+
+	/**
+	 * This alternate version of PreEditChange is called when properties inside structs are modified.  The property that was actually modified
+	 * is located at the tail of the list.  The head of the list of the UStructProperty member variable that contains the property that was modified.
+	 *
+	 * @param PropertyAboutToChange the property that is about to be modified
+	 */
+	virtual void PreEditChange( class FEditPropertyChain& PropertyAboutToChange );
+
+	/**
+	 * Called by the editor to query whether a property of this object is allowed to be modified.
+	 * The property editor uses this to disable controls for properties that should not be changed.
+	 * When overriding this function you should always call the parent implementation first.
+	 *
+	 * @param	InProperty	The property to query
+	 *
+	 * @return	true if the property can be modified in the editor, otherwise false
+	 */
+	virtual bool CanEditChange( const UProperty* InProperty ) const;
+
+	/** 
+	 * Intentionally non-virtual as it calls the FPropertyChangedEvent version
+	 */
+	void PostEditChange();
+
+	/**
+	 * Called when a property on this object has been modified externally
+	 *
+	 * @param PropertyThatChanged the property that was modified
+	 */
+	virtual void PostEditChangeProperty( struct FPropertyChangedEvent& PropertyChangedEvent);
+
+	/**
+	 * This alternate version of PostEditChange is called when properties inside structs are modified.  The property that was actually modified
+	 * is located at the tail of the list.  The head of the list of the UStructProperty member variable that contains the property that was modified.
+	 */
+	virtual void PostEditChangeChainProperty( struct FPropertyChangedChainEvent& PropertyChangedEvent );
+
+	/** Gathers external data required for applying an undo transaction */
+	virtual TSharedPtr<ITransactionObjectAnnotation> GetTransactionAnnotation() const { return NULL; }
+
+	/** Called before applying a transaction to the object.  Default implementation simply calls PreEditChange. */
+	virtual void PreEditUndo();
+
+	/** Called after applying a transaction to the object.  Default implementation simply calls PostEditChange. */
+	virtual void PostEditUndo();
+
+	/** Called after applying a transaction to the object in cases where transaction annotation was provided. Default implementation simply calls PostEditChange. */
+	virtual void PostEditUndo(TSharedPtr<ITransactionObjectAnnotation> TransactionAnnotation);
+
+
+#endif // WITH_EDITOR
+
+	// @todo document
+	virtual void PostRename(UObject* OldOuter, const FName OldName) {}
+	
+	/**
+	 * Called after duplication & serialization and before PostLoad. Used to e.g. make sure UStaticMesh's UModel gets copied as well.
+	 * Note: NOT called on components on actor duplication (alt-drag or copy-paste).  Use PostEditImport as well to cover that case.
+	 */
+	virtual void PostDuplicate(bool bDuplicateForPIE) {}
+
+	/**
+	 * Called during saving to determine the load flags to save with the object.
+	 * Upon reload, this object will be discarded on clients
+	 *
+	 * @return	true if this object should not be loaded on clients
+	 */
+	virtual bool NeedsLoadForClient() const 
+	{ 
+		return true; 
+	}
+
+	/**
+	 * Called during saving to determine the load flags to save with the object.
+	 * Upon reload, this object will be discarded on servers
+	 *
+	 * @return	true if this object should not be loaded on servers
+	 */
+	virtual bool NeedsLoadForServer() const;
+
+	/**
+	 * Called during saving to determine the load flags to save with the object.
+	 *
+	 * @return	true if this object should always be loaded for editor game
+	 */
+	virtual bool NeedsLoadForEditorGame() const
+	{
+		return false;
+	}
+
+	/**
+	* Called during async load to determine if PostLoad can be called on the loading thread.
+	*
+	* @return	true if this object's PostLoad is thread safe
+	*/
+	virtual bool IsPostLoadThreadSafe() const
+	{
+		return false;
+	}
+
+
+	/** 
+	 *	Determines if you can create an object from the supplied template in the current context (editor, client only, dedicated server, game/listen) 
+	 *	This calls NeedsLoadForClient & NeedsLoadForServer
+	 */
+	static bool CanCreateInCurrentContext(UObject* Template);
+
+	/**
+	* Exports the property values for the specified object as text to the output device. Required for Copy&Paste
+	* Most objects don't need this as unreal script can handle most cases.
+	*
+	* @param	Out				the output device to send the exported text to
+	* @param	Indent			number of spaces to prepend to each line of output
+	*
+	* see also: ImportCustomProperties()
+	*/
+	virtual void ExportCustomProperties(FOutputDevice& Out, uint32 Indent)	{}
+
+	/**
+	* Exports the property values for the specified object as text to the output device. Required for Copy&Paste
+	* Most objects don't need this as unreal script can handle most cases.
+	*
+	* @param	SourceText		the input data (zero terminated), must not be 0
+	* @param	Warn			for error reporting, must not be 0
+	*
+	* see also: ExportCustomProperties()
+	*/
+	virtual void ImportCustomProperties(const TCHAR* SourceText, FFeedbackContext* Warn) {}
+
+	/**
+	 * Called after importing property values for this object (paste, duplicate or .t3d import)
+	 * Allow the object to perform any cleanup for properties which shouldn't be duplicated or
+	 * are unsupported by the script serialization
+	 */
+	virtual void PostEditImport() {}
+
+	/**
+	 * Called from ReloadConfig after the object has reloaded its configuration data.
+	 */
+	virtual void PostReloadConfig( class UProperty* PropertyThatWasLoaded ) {}
+
+	/** Rename this object to a unique name.*/
+	virtual bool Rename( const TCHAR* NewName=NULL, UObject* NewOuter=NULL, ERenameFlags Flags=REN_None );
+
+
+	/** @return a one line description of an object for viewing in the thumbnail view of the generic browser */
+	virtual FString GetDesc() { return TEXT( "" ); }
+
+#if WITH_ENGINE
+	virtual class UWorld* GetWorld() const;
+	class UWorld* GetWorldChecked(bool& bSupported) const;
+	bool ImplementsGetWorld() const;
+#endif
+
+	/**
+	 * Callback for retrieving a textual representation of natively serialized properties.  Child classes should implement this method if they wish
+	 * to have natively serialized property values included in things like diffcommandlet output.
+	 *
+	 * @param	out_PropertyValues	receives the property names and values which should be reported for this object.  The map's key should be the name of
+	 *								the property and the map's value should be the textual representation of the property's value.  The property value should
+	 *								be formatted the same way that UProperty::ExportText formats property values (i.e. for arrays, wrap in quotes and use a comma
+	 *								as the delimiter between elements, etc.)
+	 * @param	ExportFlags			bitmask of EPropertyPortFlags used for modifying the format of the property values
+	 *
+	 * @return	true if property values were added to the map.
+	 */
+	virtual bool GetNativePropertyValues( TMap<FString,FString>& out_PropertyValues, uint32 ExportFlags=0 ) const
+	{
+		return false;
+	}
+
+	/**
+	 * Returns the size of the object/ resource for display to artists/ LDs in the Editor. The
+	 * default behavior is to return 0 which indicates that the resource shouldn't
+	 * display its size.
+	 *
+	 * @param	Type	Indicates which resource size should be returned
+	 * @return	Size of resource as to be displayed to artists/ LDs in the Editor.
+	 */
+	virtual SIZE_T GetResourceSize(EResourceSizeMode::Type Mode)
+	{
+		return 0;
+	}
+
+	/** 
+	 * Returns the name of the exporter factory used to export this object
+	 * Used when multiple factories have the same extension
+	 */
+	virtual FName GetExporterName( void )
+	{
+		return( FName( TEXT( "" ) ) );
+	}
+
+	/**
+	 * Callback used to allow object register its direct object references that are not already covered by
+	 * the token stream.
+	 *
+	 * @param InThis Object to collect references from.
+	 * @param Collector	FReferenceCollector objects to be used to collect references.
+	 */
+	static void AddReferencedObjects(UObject* InThis, FReferenceCollector& Collector);
+
+	/**
+	 * Helper function to call AddReferencedObjects for this object's class.
+	 *
+	 * @param Collector	FReferenceCollector objects to be used to collect references.
+	 */
+	void CallAddReferencedObjects(FReferenceCollector& Collector);
+
+	/**
+	 * Save information for StaticAllocateObject in the case of overwriting an existing object.
+	 * StaticAllocateObject will call delete on the result after calling Restore()
+	 *
+	 * @return An FRestoreForUObjectOverwrite that can restore the object or NULL if this is not necessary.
+	 */
+	virtual FRestoreForUObjectOverwrite* GetRestoreForUObjectOverwrite()
+	{
+		return NULL;
+	}
+
+	/**
+	 * Returns whether native properties are identical to the one of the passed in component.
+	 *
+	 * @param	Other	Other component to compare against
+	 *
+	 * @return true if native properties are identical, false otherwise
+	 */
+	virtual bool AreNativePropertiesIdenticalTo( UObject* Other ) const
+	{
+		return true;
+	}
+
+	/**
+	 * Gathers a list of asset registry searchable tags which are name/value pairs with some type information
+	 * This only needs to be implemented for asset objects
+	 *
+	 * @param	OutTags		A list of key-value pairs associated with this object and their types
+	 */
+	struct FAssetRegistryTag
+	{
+		enum ETagType
+		{
+			TT_Hidden,
+			TT_Alphabetical,
+			TT_Numerical,
+			TT_Dimensional
+		};
+
+		FName Name;
+		ETagType Type;
+		FString Value;
+
+		FAssetRegistryTag(FName InName, const FString& InValue, ETagType InType)
+			: Name(InName), Type(InType), Value(InValue) {}
+
+		/** Gathers a list of asset registry searchable tags from given objects properties */
+		COREUOBJECT_API static void GetAssetRegistryTagsFromSearchableProperties(const UObject* Object, TArray<FAssetRegistryTag>& OutTags);
+	};
+	virtual void GetAssetRegistryTags(TArray<FAssetRegistryTag>& OutTags) const;
+
+	/** Get the common tag name used for all asset source file import paths */
+	static const FName& SourceFileTagName();
+
+#if WITH_EDITOR
+	/**
+	 * Additional data pertaining to asset registry tags used by the editor
+	 */
+	struct FAssetRegistryTagMetadata
+	{
+		FText DisplayName;
+		FText TooltipText;
+		FText Suffix;
+		FString ImportantValue;
+
+		/** Set override display name */
+		FAssetRegistryTagMetadata& SetDisplayName(const FText& InDisplayName)
+		{
+			DisplayName = InDisplayName;
+			return *this;
+		}
+
+		/** Set tooltip text pertaining to the asset registry tag in the column view header */
+		FAssetRegistryTagMetadata& SetTooltip(const FText& InTooltipText)
+		{
+			TooltipText = InTooltipText;
+			return *this;
+		}
+
+		/** Set suffix appended to the tag value */
+		FAssetRegistryTagMetadata& SetSuffix(const FText& InSuffix)
+		{
+			Suffix = InSuffix;
+			return *this;
+		}
+
+		/** Set value deemed to be 'important' for this registry tag */
+		FAssetRegistryTagMetadata& SetImportantValue(const FString& InImportantValue)
+		{
+			ImportantValue = InImportantValue;
+			return *this;
+		}
+	};
+
+	/** Gathers a collection of asset registry tag metadata */
+	virtual void GetAssetRegistryTagMetadata(TMap<FName, FAssetRegistryTagMetadata>& OutMetadata) const;
+#endif
+
+	/** Returns true if this object is considered an asset. */
+	virtual bool IsAsset() const;
+
+	/** Returns true if this object is considered a localized resource. */
+	virtual bool IsLocalizedResource() const;
+
+	/** Returns true if this object is safe to add to the root set. */
+	virtual bool IsSafeForRootSet() const;
+
+	/** 
+	 * Tags objects that are part of the same asset with the specified object flag, used for GC checking
+	 *
+	 * @param	ObjectFlags	Object Flags to enable on the related objects
+	 */
+	virtual void TagSubobjects(EObjectFlags NewFlags);
+
+	/** Returns properties that are replicated for the lifetime of the actor channel */
+	virtual void GetLifetimeReplicatedProps( TArray< class FLifetimeProperty > & OutLifetimeProps ) const;
+
+	/** IsNameStableForNetworking means an object can be referred to its path name (relative to outer) over the network */
+	virtual bool IsNameStableForNetworking() const;
+
+	/** IsFullNameStableForNetworking means an object can be referred to its full path name over the network */
+	virtual bool IsFullNameStableForNetworking() const;
+
+	/** IsSupportedForNetworking means an object can be referenced over the network */
+	virtual bool IsSupportedForNetworking() const;
+
+	/** Returns a list of sub-objects that have stable names for networking */
+	virtual void GetSubobjectsWithStableNamesForNetworking(TArray<UObject*> &ObjList) {}
+
+	/** Called right before receiving a bunch */
+	virtual void PreNetReceive();
+
+	/** Called right after receiving a bunch */
+	virtual void PostNetReceive();
+
+	/** Called right after calling all OnRep notifies (called even when there are no notifies) */
+	virtual void PostRepNotifies() {}
+
+	/** Called right before being marked for destruction due to network replication */
+	virtual void PreDestroyFromReplication();
+
+	/**
+ 	 *******************************************************
+	 * Non virtual functions, not intended to be overridden
+	 *******************************************************
+	 */
+
+	/**
+	 * Test the selection state of a UObject
+	 *
+	 * @return		true if the object is selected, false otherwise.
+	 * @todo UE4 this doesn't belong here, but it doesn't belong anywhere else any better
+	 */
+	bool IsSelected() const;
+
+#if WITH_EDITOR
+	/**
+	 * Serializes all objects which have this object as their archetype into GMemoryArchive, then recursively calls this function
+	 * on each of those objects until the full list has been processed.
+	 * Called when a property value is about to be modified in an archetype object. 
+	 *
+	 * @param	AffectedObjects		the array of objects which have this object in their ObjectArchetype chain and will be affected by the change.
+	 *								Objects which have this object as their direct ObjectArchetype are removed from the list once they're processed.
+	 */
+	void PropagatePreEditChange( TArray<UObject*>& AffectedObjects, FEditPropertyChain& PropertyAboutToChange );
+
+	/**
+	 * De-serializes all objects which have this object as their archetype from the GMemoryArchive, then recursively calls this function
+	 * on each of those objects until the full list has been processed.
+	 *
+	 * @param	AffectedObjects		the array of objects which have this object in their ObjectArchetype chain and will be affected by the change.
+	 *								Objects which have this object as their direct ObjectArchetype are removed from the list once they're processed.
+	 */
+	void PropagatePostEditChange( TArray<UObject*>& AffectedObjects, FPropertyChangedChainEvent& PropertyChangedEvent );
+#endif // WITH_EDITOR
+
+	/**
+	 * Serializes the script property data located at Data.  When saving, only saves those properties which differ from the corresponding
+	 * value in the specified 'DiffObject' (usually the object's archetype).
+	 *
+	 * @param	Ar				the archive to use for serialization
+	 */
+	void SerializeScriptProperties( FArchive& Ar ) const;
+
+	/**
+	 * Wrapper function for InitProperties() which handles safely tearing down this object before re-initializing it
+	 * from the specified source object.
+	 *
+	 * @param	SourceObject	the object to use for initializing property values in this object.  If not specified, uses this object's archetype.
+	 * @param	InstanceGraph	contains the mappings of instanced objects and components to their templates
+	 */
+	void ReinitializeProperties( UObject* SourceObject=NULL, struct FObjectInstancingGraph* InstanceGraph=NULL );
+
+	/**
+	 * This will return detail info about this specific object. (e.g. AudioComponent will return the name of the cue,
+	 * ParticleSystemComponent will return the name of the ParticleSystem)  The idea here is that in many places
+	 * you have a component of interest but what you really want is some characteristic that you can use to track
+	 * down where it came from.  
+	 *
+	 * @note	safe to call on NULL object pointers!
+	 */
+	FString GetDetailedInfo() const;
+
+	/**
+	 * Called before destroying the object.  This is called immediately upon deciding to destroy the object, to allow the object to begin an
+	 * asynchronous cleanup process.
+	 */
+	bool ConditionalBeginDestroy();
+	bool ConditionalFinishDestroy();
+	
+	/** PostLoad if needed. */
+	void ConditionalPostLoad();
+
+	/**
+	 * Instances subobjects and components for objects being loaded from disk, if necessary.  Ensures that references
+	 * between nested components are fixed up correctly.
+	 *
+	 * @param	OuterInstanceGraph	when calling this method on subobjects, specifies the instancing graph which contains all instanced
+	 *								subobjects and components for a subobject root.
+	 */
+	void ConditionalPostLoadSubobjects( struct FObjectInstancingGraph* OuterInstanceGraph=NULL );
+#if WITH_EDITOR
+	/**
+	 * Starts caching of platform specific data for the target platform
+	 * Called when cooking before serialization so that object can prepare platform specific data
+	 * Not called during normal loading of objects
+	 * 
+	 * @param	TargetPlatform	target platform to cache platform specific data for
+	 */
+	virtual void BeginCacheForCookedPlatformData( const ITargetPlatform* TargetPlatform ) {  }
+	
+	/**
+	 * Have we finished loading all the cooked platform data for the target platforms requested in BeginCacheForCookedPlatformData
+	 * 
+	 * @param	TargetPlatform target platform to check for cooked platform data
+	 */
+	virtual bool IsCachedCookedPlatformDataLoaded( const ITargetPlatform* TargetPlatform ) { return true; }
+
+	/**
+	 * All caching has finished for this object (all IsCachedCookedPlatformDataLoaded functions have finished for all platforms)
+	 */
+	virtual void WillNeverCacheCookedPlatformDataAgain() { }
+
+	/**
+	 * Clears cached cooked platform data for specific platform
+	 * 
+	 * @param	TargetPlatform	target platform to cache platform specific data for
+	 */
+	virtual void ClearCachedCookedPlatformData( const ITargetPlatform* TargetPlatform ) {  }
+
+	/**
+	 * Clear all cached cooked platform data
+	 * 
+	 * @param	TargetPlatform	target platform to cache platform specific data for
+	 */
+	virtual void ClearAllCachedCookedPlatformData() { }
+#endif
+	/**
+	 * Determine if this object has SomeObject in its archetype chain.
+	 */
+	inline bool IsBasedOnArchetype( const UObject* const SomeObject ) const;
+
+	UFunction* FindFunction( FName InName ) const;
+	UFunction* FindFunctionChecked( FName InName ) const;
+
+	/**
+	 * Uses the TArchiveObjectReferenceCollector to build a list of all components referenced by this object which have this object as the outer
+	 *
+	 * @param	OutDefaultSubobjects	the array that should be populated with the default subobjects "owned" by this object
+	 * @param	bIncludeNestedSubobjects	controls whether subobjects which are contained by this object, but do not have this object
+	 *										as its direct Outer should be included
+	 */
+	void CollectDefaultSubobjects( TArray<UObject*>& OutDefaultSubobjects, bool bIncludeNestedSubobjects=false );
+
+	/**
+	 * Checks default sub-object assumptions.
+	 *
+	 * @param bForceCheck Force checks even if not enabled globally.
+	 * @return true if the assumptions are met, false otherwise.
+	 */
+	bool CheckDefaultSubobjects(bool bForceCheck = false);
+
+	/**
+	 * Save configuration.
+	 * @warning: Must be safe on class-default metaobjects.
+	 * !!may benefit from hierarchical propagation, deleting keys that match superclass...not sure what's best yet.
+	 */
+	void SaveConfig( uint64 Flags=CPF_Config, const TCHAR* Filename=NULL, FConfigCacheIni* Config=GConfig );
+
+	/**
+	 * Saves just the section(s) for this class into the default ini file for the class (with just the changes from base)
+	 */
+	void UpdateDefaultConfigFile(const FString& SpecificFileLocation = "");
+
+	/**
+	 * Saves just the section(s) for this class into the global user ini file for the class (with just the changes from base)
+	 */
+	void UpdateGlobalUserConfigFile();
+
+	/**
+	 * Saves just the property into the global user ini file for the class (with just the changes from base)
+	 */
+	void UpdateSinglePropertyInConfigFile(const UProperty* InProperty, const FString& InConfigIniName);
+
+private:
+	/**
+	 * Saves just the section(s) for this class into the given ini file for the class (with just the changes from base)
+	 */
+	void UpdateSingleSectionOfConfigFile(const FString& ConfigIniName);
+
+#if WITH_HOT_RELOAD_CTORS
+	/**
+	 * Ensures that current thread is NOT during vtable ptr retrieval process
+	 * of some UClass.
+	 */
+	void EnsureNotRetrievingVTablePtr() const;
+#endif // WITH_HOT_RELOAD_CTORS
+
+public:
+	
+	/**
+	 * Get the default config filename for the specified UObject
+	 */
+	FString GetDefaultConfigFilename() const;
+
+	/**
+	 * Get the global user override config filename for the specified UObject
+	 */
+	FString GetGlobalUserConfigFilename() const;
+
+	/**
+	 * Imports property values from an .ini file.
+	 *
+	 * @param	Class				the class to use for determining which section of the ini to retrieve text values from
+	 * @param	Filename			indicates the filename to load values from; if not specified, uses ConfigClass's ClassConfigName
+	 * @param	PropagationFlags	indicates how this call to LoadConfig should be propagated; expects a bitmask of UE4::ELoadConfigPropagationFlags values.
+	 * @param	PropertyToLoad		if specified, only the ini value for the specified property will be imported.
+	 */
+	void LoadConfig( UClass* ConfigClass=NULL, const TCHAR* Filename=NULL, uint32 PropagationFlags=UE4::LCPF_None, class UProperty* PropertyToLoad=NULL );
+
+	/**
+	 * Wrapper method for LoadConfig that is used when reloading the config data for objects at runtime which have already loaded their config data at least once.
+	 * Allows the objects the receive a callback that it's configuration data has been reloaded.
+	 *
+	 * @param	Class				the class to use for determining which section of the ini to retrieve text values from
+	 * @param	Filename			indicates the filename to load values from; if not specified, uses ConfigClass's ClassConfigName
+	 * @param	PropagationFlags	indicates how this call to LoadConfig should be propagated; expects a bitmask of UE4::ELoadConfigPropagationFlags values.
+	 * @param	PropertyToLoad		if specified, only the ini value for the specified property will be imported
+	 */
+	void ReloadConfig( UClass* ConfigClass=NULL, const TCHAR* Filename=NULL, uint32 PropagationFlags=UE4::LCPF_None, class UProperty* PropertyToLoad=NULL );
+
+	/** Import an object from a file. */
+	void ParseParms( const TCHAR* Parms );
+
+	/**
+	 * Outputs a string to an arbitrary output device, describing the list of objects which are holding references to this one.
+	 *
+	 * @param	Ar						the output device to send output to
+	 * @param	Referencers				optionally allows the caller to specify the list of references to output.
+	 */
+	void OutputReferencers( FOutputDevice& Ar, FReferencerInformationList* Referencers=NULL );
+	
+	// @todo document
+	void RetrieveReferencers( TArray<FReferencerInformation>* OutInternalReferencers, TArray<FReferencerInformation>* OutExternalReferencers);
+
+	/**
+	 * Changes the linker and linker index to the passed in one. A linker of NULL and linker index of INDEX_NONE
+	 * indicates that the object is without a linker.
+	 *
+	 * @param LinkerLoad				New LinkerLoad object to set
+	 * @param LinkerIndex				New LinkerIndex to set
+	 * @param bShouldDetachExisting		If true, detach existing linker and call PostLinkerChange
+	 */
+	void SetLinker( FLinkerLoad* LinkerLoad, int32 LinkerIndex, bool bShouldDetachExisting=true );
+
+	/**
+	 * Return the template that an object with this class, outer and name would be
+	 * 
+	 * @return the archetype for this object
+	 */
+	static UObject* GetArchetypeFromRequiredInfo(UClass* Class, UObject* Outer, FName Name, EObjectFlags ObjectFlags);
+
+	/**
+	 * Return the template this object is based on. 
+	 * 
+	 * @return the archetype for this object
+	 */
+	UObject* GetArchetype() const
+	{
+		return GetArchetypeFromRequiredInfo(GetClass(), GetOuter(), GetFName(), GetFlags());
+	}
+
+	/**
+	 * Builds a list of objects which have this object in their archetype chain.
+	 *
+	 * @param	Instances	receives the list of objects which have this one in their archetype chain
+	 */
+	void GetArchetypeInstances( TArray<UObject*>& Instances );
+
+	/**
+	 * Wrapper for calling UClass::InstanceSubobjectTemplates() for this object.
+	 */
+	void InstanceSubobjectTemplates( struct FObjectInstancingGraph* InstanceGraph = NULL );
+
+	/**
+	 * Returns true if this object implements the interface T, false otherwise.
+	 */
+	template<class T>
+	bool Implements() const;
+
+	//==========================================
+	// Kismet Virtual Machine
+	//==========================================
+	//@todo UE4 - do these really need to be virtual to support actors?
+
+	//
+	// Script processing functions.
+	//
+	virtual void ProcessEvent( UFunction* Function, void* Parms );
+
+	/**
+	 * Return the space this function should be called.   Checks to see if this function should
+	 * be called locally, remotely, or simply absorbed under the given conditions
+	 *
+	 * @param Function function to call
+	 * @param Parameters arguments to the function call
+	 * @param Stack stack frame for the function call
+	 * @return bitmask representing all callspaces that apply to this UFunction in the given context
+	 */
+	virtual int32 GetFunctionCallspace( UFunction* Function, void* Parameters, FFrame* Stack )
+	{
+		return FunctionCallspace::Local;
+	}
+
+	/**
+	 * Call the actor's function remotely
+	 *
+	 * @param Function function to call
+	 * @param Parameters arguments to the function call
+	 * @param Stack stack frame for the function call
+	 */
+	virtual bool CallRemoteFunction( UFunction* Function, void* Parms, struct FOutParmRec* OutParms, FFrame* Stack )
+	{
+		return false;
+	}
+
+	/** Command line. */
+	bool CallFunctionByNameWithArguments( const TCHAR* Cmd, FOutputDevice& Ar, UObject* Executor, bool bForceCallWithNonExec = false );
+
+	// Call a function
+	void CallFunction( FFrame& Stack, RESULT_DECL, UFunction* Function );
+
+	//
+	// Internal function call processing.
+	// @warning: might not write anything to Result if proper type isn't returned.
+	//
+	void ProcessInternal( FFrame& Stack, RESULT_DECL );
+
+	/**
+	 * This function handles a console exec sent to the object; it is virtual so 'nexus' objects like
+	 * a player controller can reroute the command to several different objects.
+	 */
+	virtual bool ProcessConsoleExec(const TCHAR* Cmd, FOutputDevice& Ar, UObject* Executor)
+	{
+		return CallFunctionByNameWithArguments(Cmd, Ar, Executor);
+	}
+
+	/** advances Stack's code past the parameters to the given Function and if the function has a return value, copies the zero value for that property to the memory for the return value
+	 * @param Stack the script stack frame
+	 * @param Result pointer to where the return value should be written
+	 * @param Function the function being called
+	 */
+	void SkipFunction(FFrame& Stack, RESULT_DECL, UFunction* Function);
+
+	/**
+	 * Called on the target when a class is loaded with ClassGeneratedBy is loaded.  Should regenerate the class if needed, and return the updated class
+	 * @return	Updated instance of the class, if needed, or NULL if no regeneration was necessary
+	 */
+	virtual UClass* RegenerateClass(UClass* ClassToRegenerate, UObject* PreviousCDO, TArray<UObject*>& ObjLoaded) { return NULL; };
+
+	/** 
+	 * Returns whether this object is contained in or part of a blueprint object
+	 */
+	bool IsInBlueprint() const;
+
+	/** 
+	 *  Destroy properties that won't be destroyed by the native destructor
+	 */
+	void DestroyNonNativeProperties();
+
+	virtual void MarkAsEditorOnlySubobject() { }
+
+	// UnrealScript intrinsics.
+
+	// Undefined native handler
+	DECLARE_FUNCTION(execUndefined);
+
+	
+	// @todo document below declared functions
+	// Variables
+	DECLARE_FUNCTION(execLocalVariable);
+	DECLARE_FUNCTION(execInstanceVariable);
+	DECLARE_FUNCTION(execDefaultVariable);
+	DECLARE_FUNCTION(execLocalOutVariable);
+	DECLARE_FUNCTION(execInterfaceVariable);
+	DECLARE_FUNCTION(execInterfaceContext);
+	DECLARE_FUNCTION(execArrayElement);
+	DECLARE_FUNCTION(execBoolVariable);
+	DECLARE_FUNCTION(execClassDefaultVariable);
+	DECLARE_FUNCTION(execEndFunctionParms);
+
+
+	// Do Nothing 
+	DECLARE_FUNCTION(execNothing);
+	DECLARE_FUNCTION(execNothingOp4a);
+
+	/** Breakpoint; only observed in the editor; executing it at any other time is a NOP */
+	DECLARE_FUNCTION(execBreakpoint);
+
+	/** Tracepoint; only observed in the editor; executing it at any other time is a NOP */
+	DECLARE_FUNCTION(execTracepoint);
+	DECLARE_FUNCTION(execWireTracepoint);
+
+	/** Instrumentation event for profiling; only observed in the builds with blueprint instrumentation */
+	DECLARE_FUNCTION(execInstrumentation);
+
+	DECLARE_FUNCTION(execEndOfScript);
+
+	/** failsafe for functions that return a value - returns the zero value for a property and logs that control reached the end of a non-void function */
+	DECLARE_FUNCTION(execReturnNothing);
+	DECLARE_FUNCTION(execEmptyParmValue);
+
+	// Commands
+	DECLARE_FUNCTION(execJump);
+	DECLARE_FUNCTION(execJumpIfNot);
+	DECLARE_FUNCTION(execAssert);
+
+	/** 
+	 * Push a code offset onto the execution flow stack for future execution.  
+	 * Current execution continues to the next instruction after the push one.
+	 */
+	DECLARE_FUNCTION(execPushExecutionFlow);
+
+	/**
+	 * Pops a code offset from the execution flow stack and starts execution there.
+	 * If there are no stack entries left, it is treated as an execution error.
+	 */
+	DECLARE_FUNCTION(execPopExecutionFlow);
+	DECLARE_FUNCTION(execComputedJump);
+
+	/** 
+	 * Pops a code offset from the execution flow stack and starts execution there, if a condition is not true.
+	 * If there are no stack entries left, it is treated as an execution error.
+	 */
+	DECLARE_FUNCTION(execPopExecutionFlowIfNot);
+
+
+	// Assignment
+	DECLARE_FUNCTION(execLet);
+	DECLARE_FUNCTION(execLetObj);
+	DECLARE_FUNCTION(execLetWeakObjPtr);
+	DECLARE_FUNCTION(execLetBool);
+	DECLARE_FUNCTION(execLetDelegate);
+	DECLARE_FUNCTION(execLetMulticastDelegate);
+
+	// Delegates 
+	DECLARE_FUNCTION(execAddMulticastDelegate);
+	DECLARE_FUNCTION(execClearMulticastDelegate);
+	DECLARE_FUNCTION(execEatReturnValue);
+	DECLARE_FUNCTION(execRemoveMulticastDelegate);
+
+	// Context expressions
+	DECLARE_FUNCTION(execSelf);
+	DECLARE_FUNCTION(execContext);
+	DECLARE_FUNCTION(execContext_FailSilent);
+	DECLARE_FUNCTION(execStructMemberContext);
+
+	// Function calls
+	DECLARE_FUNCTION(execVirtualFunction);
+	DECLARE_FUNCTION(execFinalFunction);
+
+	// Struct comparison
+	DECLARE_FUNCTION(execStructCmpEq);
+	DECLARE_FUNCTION(execStructCmpNe);
+	DECLARE_FUNCTION(execStructMember);
+
+	// @todo delegate: Delegate comparison is not supported for multi-cast delegates
+	DECLARE_FUNCTION(execEqualEqual_DelegateDelegate);
+	DECLARE_FUNCTION(execNotEqual_DelegateDelegate);
+	DECLARE_FUNCTION(execEqualEqual_DelegateFunction);
+	DECLARE_FUNCTION(execNotEqual_DelegateFunction);
+
+	// Constants
+	DECLARE_FUNCTION(execIntConst);
+	DECLARE_FUNCTION(execInt64Const);
+	DECLARE_FUNCTION(execUInt64Const);
+	DECLARE_FUNCTION(execSkipOffsetConst);
+	DECLARE_FUNCTION(execFloatConst);
+	DECLARE_FUNCTION(execStringConst);
+	DECLARE_FUNCTION(execUnicodeStringConst);
+	DECLARE_FUNCTION(execTextConst);
+	DECLARE_FUNCTION(execObjectConst);
+	DECLARE_FUNCTION(execAssetConst);
+
+	// @todo delegate: Multi-cast versions needed for script execution!		(Need Add, Remove, Clear/Empty)
+	DECLARE_FUNCTION(execInstanceDelegate);
+	DECLARE_FUNCTION(execNameConst);
+	DECLARE_FUNCTION(execByteConst);
+	DECLARE_FUNCTION(execIntZero);
+	DECLARE_FUNCTION(execIntOne);
+	DECLARE_FUNCTION(execTrue);
+	DECLARE_FUNCTION(execFalse);
+	DECLARE_FUNCTION(execNoObject);
+	DECLARE_FUNCTION(execNullInterface);
+	DECLARE_FUNCTION(execIntConstByte);
+	DECLARE_FUNCTION(execRotationConst);
+	DECLARE_FUNCTION(execVectorConst);
+	DECLARE_FUNCTION(execTransformConst);
+	DECLARE_FUNCTION(execStructConst);
+	DECLARE_FUNCTION(execSetArray);
+	DECLARE_FUNCTION(execArrayConst);
+
+	// Object construction
+	DECLARE_FUNCTION(execNew);
+	DECLARE_FUNCTION(execClassContext);
+	DECLARE_FUNCTION(execNativeParm);
+
+	// Conversions 
+	DECLARE_FUNCTION(execDynamicCast);
+	DECLARE_FUNCTION(execMetaCast);
+	DECLARE_FUNCTION(execPrimitiveCast);
+	DECLARE_FUNCTION(execInterfaceCast);
+
+	// Cast functions
+	DECLARE_FUNCTION(execObjectToBool);
+	DECLARE_FUNCTION(execInterfaceToBool);
+	DECLARE_FUNCTION(execObjectToInterface);
+	DECLARE_FUNCTION(execInterfaceToInterface);
+	DECLARE_FUNCTION(execInterfaceToObject);
+
+	// Dynamic array functions
+	// Array support
+	DECLARE_FUNCTION(execGetDynArrayElement);
+	DECLARE_FUNCTION(execSetDynArrayElement);
+	DECLARE_FUNCTION(execGetDynArrayLength);
+	DECLARE_FUNCTION(execSetDynArrayLength);
+	DECLARE_FUNCTION(execDynArrayInsert);
+	DECLARE_FUNCTION(execDynArrayRemove);
+	DECLARE_FUNCTION(execDynArrayFind);
+	DECLARE_FUNCTION(execDynArrayFindStruct);
+	DECLARE_FUNCTION(execDynArrayAdd);
+	DECLARE_FUNCTION(execDynArrayAddItem);
+	DECLARE_FUNCTION(execDynArrayInsertItem);
+	DECLARE_FUNCTION(execDynArrayRemoveItem);
+	DECLARE_FUNCTION(execDynArraySort);
+
+	DECLARE_FUNCTION(execBindDelegate);
+	DECLARE_FUNCTION(execCallMulticastDelegate);
+
+	DECLARE_FUNCTION(execLetValueOnPersistentFrame);
+
+	DECLARE_FUNCTION(execCallMathFunction);
+
+	DECLARE_FUNCTION(execSwitchValue);
+
+	DECLARE_FUNCTION(execArrayGetByRef);
+
+	// -- K2 support functions
+	struct Object_eventExecuteUbergraph_Parms
+	{
+		int32 EntryPoint;
+	};
+	void ExecuteUbergraph(int32 EntryPoint)
+	{
+		Object_eventExecuteUbergraph_Parms Parms;
+		Parms.EntryPoint=EntryPoint;
+		ProcessEvent(FindFunctionChecked(NAME_ExecuteUbergraph),&Parms);
+	}
+
+protected: 
+
+	/** Checks it's ok to perform subobjects check at this time. */
+	bool CanCheckDefaultSubObjects(bool bForceCheck, bool& bResult);
+
+	/**
+	* Checks default sub-object assumptions.
+	*
+	* @return true if the assumptions are met, false otherwise.
+	*/
+	virtual bool CheckDefaultSubobjectsInternal();
+
+private:
+	void ProcessContextOpcode(FFrame& Stack, RESULT_DECL, bool bCanFailSilent);
+
+	/**
+	* Create a component or subobject only to be used with the editor.
+	* @param	Outer						outer to construct the subobject in
+	* @param	ReturnType					type of the new component
+	* @param	SubobjectName				name of the new component
+	* @param	bTransient					true if the component is being assigned to a transient property
+	*/
+	UObject* CreateEditorOnlyDefaultSubobjectImpl(FName SubobjectName, UClass* ReturnType, bool bTransient = false);
+};
+
+/**
+* Test validity of object
+*
+* @param	Test			The object to test
+* @return	Return true if the object is usable: non-null and not pending kill
+*/
+FORCEINLINE bool IsValid(const UObject *Test)
+{
+	return Test && !Test->IsPendingKill();
+}
diff --git a/Engine/Source/Runtime/CoreUObject/Public/UObject/Package.h b/Engine/Source/Runtime/CoreUObject/Public/UObject/Package.h
new file mode 100644
--- /dev/null
+++ b/Engine/Source/Runtime/CoreUObject/Public/UObject/Package.h
@@ -0,0 +1,559 @@
+// Copyright 1998-2016 Epic Games, Inc. All Rights Reserved.
+
+#pragma once
+
+#include "ObjectBase.h"
+#include "WorldCompositionUtility.h"
+#include "GatherableTextData.h"
+#include "PropertyLocalizationDataGathering.h"
+
+/**
+ * Represents the result of saving a package
+ */
+enum class ESavePackageResult
+{
+	/** Package was saved successfully */
+	Success, 
+	/** Unknown error occured when saving package */
+	Error,
+	/** Canceled by user */
+	Canceled,
+	/** [When cooking] Package was not saved because it contained editor-only data */
+	ContainsEditorOnlyData,
+	/** [When cooking] Package was not saved because it was referenced by editor-only properties */
+	ReferencedOnlyByEditorOnlyData, 
+	/** [When cooking] Package was not saved because it contains assets that were converted into native code */
+	ReplaceCompletely,
+	/** [When cooking] Package was saved, but we should generate a stub so that other converted packages can interface with it*/
+	GenerateStub
+};
+
+/**
+ * A package.
+ */
+PRAGMA_DISABLE_DEPRECATION_WARNINGS // Required for auto-generated functions referencing PackageFlags
+class COREUOBJECT_API UPackage : public UObject
+{
+	// Have to unwind this macro to support the reference variable, can go back to commented declaration when removing the deprecated variable
+	// DECLARE_CASTED_CLASS_INTRINSIC(UPackage, UObject, 0, TEXT("/Script/CoreUObject"), CASTCLASS_UPackage)
+
+#if WITH_HOT_RELOAD_CTORS
+	DECLARE_CASTED_CLASS_INTRINSIC_NO_CTOR_NO_VTABLE_CTOR( UPackage, UObject, 0, TEXT("/Script/CoreUObject"), CASTCLASS_UPackage, NO_API )
+	/** DO NOT USE. This constructor is for internal usage only for hot-reload purposes. */
+	UPackage(FVTableHelper& Helper)
+		: Super(Helper)
+		, PackageFlagsPrivate(PackageFlags) 
+	{
+	};
+#else
+	DECLARE_CASTED_CLASS_INTRINSIC_NO_CTOR(UPackage, UObject, 0, TEXT("/Script/CoreUObject"), CASTCLASS_UPackage, NO_API)
+#endif
+
+
+public:
+
+	UPackage(const FObjectInitializer& ObjectInitializer = FObjectInitializer::Get())
+		: Super(ObjectInitializer)
+		, PackageFlagsPrivate(PackageFlags)
+	{
+	}
+
+	/** delegate type for package dirty state events.  ( Params: UPackage* ModifiedPackage ) */
+	DECLARE_MULTICAST_DELEGATE_OneParam(FOnPackageDirtyStateChanged, class UPackage*);
+	/** delegate type for package saved events ( Params: const FString& PackageFileName, UObject* Outer ) */
+	DECLARE_MULTICAST_DELEGATE_TwoParams(FOnPackageSaved, const FString&, UObject*);					
+	/** delegate type for when a package is marked as dirty via UObjectBaseUtilty::MarkPackageDirty ( Params: UPackage* ModifiedPackage, bool bWasDirty ) */
+	DECLARE_MULTICAST_DELEGATE_TwoParams(FOnPackageMarkedDirty, class UPackage*, bool);
+
+	/** Delegate to notify subscribers when a package has been saved. This is triggered when the package saving
+	 *  has completed and was successful. */
+	static FOnPackageSaved PackageSavedEvent;
+	/** Delegate to notify subscribers when the dirty state of a package is changed.
+	 *  Allows the editor to register the modified package as one that should be prompted for source control checkout. 
+	 *  Use Package->IsDirty() to get the updated dirty state of the package */
+	static FOnPackageDirtyStateChanged PackageDirtyStateChangedEvent;
+	/** 
+	 * Delegate to notify subscribers when a package is marked as dirty via UObjectBaseUtilty::MarkPackageDirty 
+	 * Note: Unlike FOnPackageDirtyStateChanged, this is always called, even when the package is already dirty
+	 * Use bWasDirty to check the previous dirty state of the package
+	 * Use Package->IsDirty() to get the updated dirty state of the package
+	 */
+	static FOnPackageMarkedDirty PackageMarkedDirtyEvent;
+
+private:
+	/** Used by the editor to determine if a package has been changed.																							*/
+	bool	bDirty;
+#if WITH_EDITORONLY_DATA
+	/** True if this package is only referenced by editor-only properties */
+	bool bLoadedByEditorPropertiesOnly;
+#endif
+public:
+	/** Whether this package has been fully loaded (aka had all it's exports created) at some point.															*/
+	bool	bHasBeenFullyLoaded;
+private:
+	/** Returns whether exports should be found in memory first before trying to load from disk from within CreateExport.										*/
+	bool	bShouldFindExportsInMemoryFirst;
+	/** Indicates which folder to display this package under in the Generic Browser's list of packages. If not specified, package is added to the root level.	*/
+	FName	FolderName;
+	/** Time in seconds it took to fully load this package. 0 if package is either in process of being loaded or has never been fully loaded.					*/
+	float LoadTime;
+
+	/** GUID of package if it was loaded from disk; used by netcode to make sure packages match between client and server */
+	FGuid Guid;
+	/** Chunk IDs for the streaming install chunks this package will be placed in.  Empty for no chunk */
+	TArray<int32> ChunkIDs;
+	/** for packages that were a forced export in another package (seekfree loading), the name of that base package, otherwise NAME_None */
+	FName ForcedExportBasePackageName;
+public:
+
+	virtual bool IsNameStableForNetworking() const override { return true; }		// For now, assume all packages have stable net names
+
+#if WITH_EDITORONLY_DATA
+	/** Sets the bLoadedByEditorPropertiesOnly flag */
+	void SetLoadedByEditorPropertiesOnly(bool bIsEditorOnly, bool bRecursive = false);
+	/** returns true when the package is only referenced by editor-only flag */
+	bool IsLoadedByEditorPropertiesOnly() const { return bLoadedByEditorPropertiesOnly; }
+#endif
+
+private:
+	/** New private Package Flags. When deprecated PackageFlags removed, make this a member instead of a reference */
+	uint32&	PackageFlagsPrivate;
+
+public:
+
+	/** Old deprecated public Package flags, serialized. When deprecation removed make PackageFlagsPrivate no longer a reference */
+	DEPRECATED(4.10, "Direct access of PackageFlags is deprecated, use Get/Set/Clear/HasAny/HasAll PackageFlags instead.")
+	uint32	PackageFlags;
+
+	/** The name of the file that this package was loaded from */
+	FName	FileName;
+
+	/** Linker load associated with this package */
+	class FLinkerLoad* LinkerLoad;
+
+	/** Linker package version this package has been serialized with. This is mostly used by PostLoad **/
+	int32 LinkerPackageVersion;
+
+	/** Linker licensee version this package has been serialized with. This is mostly used by PostLoad **/
+	int32 LinkerLicenseeVersion;
+
+	/** Linker custom version container this package has been serialized with. This is mostly used by PostLoad **/
+	FCustomVersionContainer LinkerCustomVersion;
+
+	/** size of the file for this package; if the package was not loaded from a file or was a forced export in another package, this will be zero */
+	int64 FileSize;
+	
+	/** Editor only: Thumbnails stored in this package */
+	TScopedPointer< FThumbnailMap > ThumbnailMap;
+
+	// MetaData for the editor, or NULL in the game
+	class UMetaData*	MetaData;
+	
+	// World browser information
+	TScopedPointer< FWorldTileInfo > WorldTileInfo;
+
+	TMap<FName, int32> ClassUniqueNameIndexMap;
+
+#if WITH_EDITOR
+	/** Editor only: PIE instance ID this package belongs to, INDEX_NONE otherwise */
+	int32 PIEInstanceID;
+#endif
+
+#if WITH_EDITORONLY_DATA
+	/** True if this packages has been cooked for the editor / opened cooked by the editor */
+	bool bIsCookedForEditor;
+#endif
+
+	/**
+	 * Called after the C++ constructor and after the properties have been initialized, but before the config has been loaded, etc.
+	 * mainly this is to emulate some behavior of when the constructor was called after the properties were initialized.
+	 */
+	virtual void PostInitProperties() override;
+
+	virtual void BeginDestroy() override;
+
+	/** Serializer */
+	virtual void Serialize( FArchive& Ar ) override;
+
+	static void AddReferencedObjects(UObject* InThis, FReferenceCollector& Collector);
+
+	/** Packages are never assets */
+	virtual bool IsAsset() const override { return false; }
+
+	// UPackage interface.
+
+	/**
+	 * Sets the time it took to load this package.
+	 */
+	void SetLoadTime( float InLoadTime )
+	{
+		LoadTime = InLoadTime;
+	}
+
+	/**
+	 * Returns the time it took the last time this package was fully loaded, 0 otherwise.
+	 *
+	 * @return Time it took to load.
+	 */
+	float GetLoadTime()
+	{
+		return LoadTime;
+	}
+
+	/**
+	 * Get the package's folder name
+	 * @return		Folder name
+	 */
+	FName GetFolderName() const
+	{
+		return FolderName;
+	}
+
+	/**
+	 * Set the package's folder name
+	 */
+	void SetFolderName (FName name)
+	{
+		FolderName = name;
+	}
+
+	/**
+	 * Marks/Unmarks the package's bDirty flag
+	 */
+	void SetDirtyFlag( bool bIsDirty );
+
+	/**
+	 * Returns whether the package needs to be saved.
+	 *
+	 * @return		true if the package is dirty and needs to be saved, false otherwise.
+	 */
+	bool IsDirty() const
+	{
+		return bDirty;
+	}
+
+	/**
+	 * Marks this package as being fully loaded.
+	 */
+	void MarkAsFullyLoaded()
+	{
+		bHasBeenFullyLoaded = true;
+	}
+
+	/**
+	 * Returns whether the package is fully loaded.
+	 *
+	 * @return true if fully loaded or no file associated on disk, false otherwise
+	 */
+	bool IsFullyLoaded();
+
+	/**
+	 * Fully loads this package. Safe to call multiple times and won't clobber already loaded assets.
+	 */
+	void FullyLoad();
+
+	/**
+	 * Tags the Package's metadata
+	 */
+	virtual void TagSubobjects(EObjectFlags NewFlags) override;
+
+	/**
+	 * Sets whether exports should be found in memory first or  not.
+	 *
+	 * @param bInShouldFindExportsInMemoryFirst	Whether to find in memory first or not
+	 */
+	void FindExportsInMemoryFirst( bool bInShouldFindExportsInMemoryFirst )
+	{
+		bShouldFindExportsInMemoryFirst = bInShouldFindExportsInMemoryFirst;
+	}
+
+	/**
+	 * Returns whether exports should be found in memory first before trying to load from disk
+	 * from within CreateExport.
+	 *
+	 * @return true if exports should be found via FindObject first, false otherwise.
+	 */
+	bool ShouldFindExportsInMemoryFirst()
+	{
+		return bShouldFindExportsInMemoryFirst;
+	}
+
+	/**
+	 * Called to indicate that this package contains a ULevel or UWorld object.
+	 */
+	void ThisContainsMap() 
+	{
+		SetPackageFlags(PKG_ContainsMap);
+	}
+
+	/**
+	 * Returns whether this package contains a ULevel or UWorld object.
+	 *
+	 * @return		true if package contains ULevel/ UWorld object, false otherwise.
+	 */
+	bool ContainsMap() const
+	{
+		return HasAnyPackageFlags(PKG_ContainsMap);
+	}
+
+	/**
+	 * Called to indicate that this package contains data required to be gathered for localization.
+	 */
+	void ThisRequiresLocalizationGather(bool Value)
+	{
+		if(Value)
+		{
+			SetPackageFlags(PKG_RequiresLocalizationGather);
+		}
+		else
+		{
+			ClearPackageFlags(PKG_RequiresLocalizationGather);
+		}
+	}
+
+	/**
+	 * Returns whether this package contains data required to be gathered for localization.
+	 *
+	 * @return		true if package contains contains data required to be gathered for localization, false otherwise.
+	 */
+	bool RequiresLocalizationGather() const
+	{
+		return HasAnyPackageFlags(PKG_RequiresLocalizationGather);
+	}
+
+// UE-21181 - trying to track when a flag gets set on a package due to PIE
+#if WITH_EDITOR
+	static UPackage* EditorPackage;
+	void SetPackageFlagsTo( uint32 NewFlags );
+#else
+	/**
+	 * Sets all package flags to the specified values.
+	 *
+	 * @param	NewFlags		New value for package flags
+	 */
+	FORCEINLINE void SetPackageFlagsTo( uint32 NewFlags )
+	{
+		PackageFlagsPrivate = NewFlags;
+	}
+#endif
+
+	/**
+	 * Set the specified flags to true. Does not affect any other flags.
+	 *
+	 * @param	NewFlags		Package flags to enable
+	 */
+	FORCEINLINE void SetPackageFlags( uint32 NewFlags )
+	{
+		SetPackageFlagsTo(PackageFlagsPrivate | NewFlags);
+	}
+
+	/**
+	 * Set the specified flags to false. Does not affect any other flags.
+	 *
+	 * @param	NewFlags		Package flags to disable
+	 */
+	FORCEINLINE void ClearPackageFlags( uint32 NewFlags )
+	{
+		SetPackageFlagsTo(PackageFlagsPrivate & ~NewFlags);
+	}
+
+	/**
+	 * Used to safely check whether the passed in flag is set.
+	 *
+	 * @param	FlagsToCheck		Package flags to check for
+	 *
+	 * @return	true if the passed in flag is set, false otherwise
+	 *			(including no flag passed in, unless the FlagsToCheck is CLASS_AllFlags)
+	 */
+	FORCEINLINE bool HasAnyPackageFlags( uint32 FlagsToCheck ) const
+	{
+		return (PackageFlagsPrivate & FlagsToCheck) != 0;
+	}
+
+	/**
+	 * Used to safely check whether all of the passed in flags are set.
+	 *
+	 * @param FlagsToCheck	Package flags to check for
+	 * @return true if all of the passed in flags are set (including no flags passed in), false otherwise
+	 */
+	FORCEINLINE bool HasAllPackagesFlags( uint32 FlagsToCheck ) const
+	{
+		return ((PackageFlagsPrivate & FlagsToCheck) == FlagsToCheck);
+	}
+
+	/**
+	 * Gets the package flags.
+	 *
+	 * @return	The package flags.
+	 */
+	FORCEINLINE uint32 GetPackageFlags() const
+	{
+		return PackageFlagsPrivate;
+	}
+
+	/** Returns true if this package has a thumbnail map */
+	bool HasThumbnailMap() const
+	{
+		return ThumbnailMap.IsValid();
+	}
+	
+	/** Returns the thumbnail map for this package (const).  Only call this if HasThumbnailMap returns true! */
+	const FThumbnailMap& GetThumbnailMap() const
+	{
+		check( HasThumbnailMap() );
+		return *ThumbnailMap;
+	}
+
+	/** Access the thumbnail map for this package.  Only call this if HasThumbnailMap returns true! */
+	FThumbnailMap& AccessThumbnailMap()
+	{
+		check( HasThumbnailMap() );
+		return *ThumbnailMap;
+	}
+	
+	/** returns our Guid */
+	FORCEINLINE FGuid GetGuid() const
+	{
+		return Guid;
+	}
+	/** makes our a new fresh Guid */
+	FORCEINLINE void MakeNewGuid()
+	{
+		Guid = FGuid::NewGuid();
+	}
+	/** sets a specific Guid */
+	FORCEINLINE void SetGuid(FGuid NewGuid)
+	{
+		Guid = NewGuid;
+	}
+
+	/** returns our FileSize */
+	FORCEINLINE int64 GetFileSize()
+	{
+		return FileSize;
+	}
+
+	/** returns our ChunkIDs */
+	FORCEINLINE const TArray<int32>& GetChunkIDs() const
+	{
+		return ChunkIDs;
+	}
+	/** sets our ChunkIDs */
+	FORCEINLINE void SetChunkIDs(const TArray<int32>& InChunkIDs)
+	{
+		ChunkIDs = InChunkIDs;
+	}
+
+	/** returns ForcedExportBasePackageName */
+	FORCEINLINE FName GetForcedExportBasePackageName()
+	{
+		return ForcedExportBasePackageName;
+	}
+
+	////////////////////////////////////////////////////////
+	// MetaData 
+
+	/**
+	 * Gets (after possibly creating) a metadata object for this package
+	 *
+	 * @param	bAllowLoadObject				Can load an object to find it's UMetaData if not currently loaded.
+	 *
+	 * @return A valid UMetaData pointer for all objects in this package
+	 */
+	class UMetaData* GetMetaData();
+
+	/**
+	 * Save one specific object (along with any objects it references contained within the same Outer) into an Unreal package.
+	 * 
+	 * @param	InOuter							the outer to use for the new package
+	 * @param	Base							the object that should be saved into the package
+	 * @param	TopLevelFlags					For all objects which are not referenced [either directly, or indirectly] through Base, only objects
+	 *											that contain any of these flags will be saved.  If 0 is specified, only objects which are referenced
+	 *											by Base will be saved into the package.
+	 * @param	Filename						the name to use for the new package file
+	 * @param	Error							error output
+	 * @param	Conform							if non-NULL, all index tables for this will be sorted to match the order of the corresponding index table
+	 *											in the conform package
+	 * @param	bForceByteSwapping				whether we should forcefully byte swap before writing to disk
+	 * @param	bWarnOfLongFilename				[opt] If true (the default), warn when saving to a long filename.
+	 * @param	SaveFlags						Flags to control saving
+	 * @param	TargetPlatform					The platform being saved for
+	 * @param	FinalTimeStamp					If not FDateTime::MinValue(), the timestamp the saved file should be set to. (Intended for cooking only...)
+	 *
+	 * @return	ESavePackageResult enum value with the result of saving a package.
+	 */
+	static ESavePackageResult Save(UPackage* InOuter, UObject* Base, EObjectFlags TopLevelFlags, const TCHAR* Filename,
+		FOutputDevice* Error=GError, FLinkerLoad* Conform=NULL, bool bForceByteSwapping=false, bool bWarnOfLongFilename=true, 
+		uint32 SaveFlags=SAVE_None, const class ITargetPlatform* TargetPlatform = NULL, const FDateTime& FinalTimeStamp = FDateTime::MinValue(), bool bSlowTask = true );
+
+	/**
+	* Save one specific object (along with any objects it references contained within the same Outer) into an Unreal package.
+	*
+	* @param	InOuter							the outer to use for the new package
+	* @param	Base							the object that should be saved into the package
+	* @param	TopLevelFlags					For all objects which are not referenced [either directly, or indirectly] through Base, only objects
+	*											that contain any of these flags will be saved.  If 0 is specified, only objects which are referenced
+	*											by Base will be saved into the package.
+	* @param	Filename						the name to use for the new package file
+	* @param	Error							error output
+	* @param	Conform							if non-NULL, all index tables for this will be sorted to match the order of the corresponding index table
+	*											in the conform package
+	* @param	bForceByteSwapping				whether we should forcefully byte swap before writing to disk
+	* @param	bWarnOfLongFilename				[opt] If true (the default), warn when saving to a long filename.
+	* @param	SaveFlags						Flags to control saving
+	* @param	TargetPlatform					The platform being saved for
+	* @param	FinalTimeStamp					If not FDateTime::MinValue(), the timestamp the saved file should be set to. (Intended for cooking only...)
+	*
+	* @return	true if the package was saved successfully.
+	*/
+	static bool SavePackage(UPackage* InOuter, UObject* Base, EObjectFlags TopLevelFlags, const TCHAR* Filename,
+		FOutputDevice* Error = GError, FLinkerLoad* Conform = NULL, bool bForceByteSwapping = false, bool bWarnOfLongFilename = true,
+		uint32 SaveFlags = SAVE_None, const class ITargetPlatform* TargetPlatform = NULL, const FDateTime& FinalTimeStamp = FDateTime::MinValue(), bool bSlowTask = true);
+
+	/** Wait for any SAVE_Async file writes to complete **/
+	static void WaitForAsyncFileWrites();
+
+	/**
+	 * Static: Saves thumbnail data for the specified package outer and linker
+	 *
+	 * @param	InOuter							the outer to use for the new package
+	 * @param	Linker							linker we're currently saving with
+	 */
+	static void SaveThumbnails( UPackage* InOuter, FLinkerSave* Linker );
+
+	/**
+	 * Static: Saves asset registry data for the specified package outer and linker
+	 *
+	 * @param	InOuter							the outer to use for the new package
+	 * @param	Linker							linker we're currently saving with
+	 */
+	static void SaveAssetRegistryData( UPackage* InOuter, FLinkerSave* Linker );
+
+	/**
+	 * Static: Saves the level information used by the World browser
+	 *
+	 * @param	InOuter							the outer to use for the new package
+	 * @param	Linker							linker we're currently saving with
+	 */
+	static void SaveWorldLevelInfo( UPackage* InOuter, FLinkerSave* Linker );
+
+	/**
+	  * Determines if a package contains no more assets.
+	  *
+	  * @param Package			the package to test
+	  * @param LastReferencer	the optional last UObject referencer to this package. This object will be excluded when determining if the package is empty
+	  * @return true if Package contains no more assets.
+	  */
+	static bool IsEmptyPackage(UPackage* Package, const UObject* LastReferencer = NULL);
+
+
+	/**
+	 * @param TargetPlatform - The platform being saved for
+	 * @param bIsCooking	 - Whether we are cooking or not
+	 *
+	 * @return Objects marks specific for the particular target platform, objects with these marks will be rejected from the cook
+	 */
+	static EObjectMark GetObjectMarksForTargetPlatform( const class ITargetPlatform* TargetPlatform, const bool bIsCooking );
+
+};
+PRAGMA_ENABLE_DEPRECATION_WARNINGS
