Copying //UE4/Dev-Build to //UE4/Dev-Main (Source: //UE4/Dev-Build @ 3092544)

diff --git a/Engine/Source/Runtime/CoreUObject/Public/UObject/UObject.h b/Engine/Source/Runtime/CoreUObject/Public/UObject/UObject.h
--- a/Engine/Source/Runtime/CoreUObject/Public/UObject/UObject.h
+++ b/Engine/Source/Runtime/CoreUObject/Public/UObject/UObject.h
@@ -2,1195 +2,14 @@
 
 #pragma once
 
-#include "UObjectBaseUtility.h"
+#include "Object.h"
 
-DECLARE_LOG_CATEGORY_EXTERN(LogObj, Log, All);
-
-
-namespace ECastCheckedType
-{
-	enum Type
-	{
-		NullAllowed,
-		NullChecked
-	};
-};
-
-/** Passed to GetResourceSize() to indicate which resource size should be returned.*/
-namespace EResourceSizeMode
-{
-	enum Type
-	{
-		/** Only exclusive resource size */
-		Exclusive,
-		/** Resource size of the object and all of its references */
-		Inclusive,
-	};
-};
-
-//
-// The base class of all objects.
-//
-class COREUOBJECT_API UObject : public UObjectBaseUtility
+DEPRECATED(4.14, "Including UObject.h has been deprecated. Please include Object.h instead.")
+inline void UObjectHeaderDeprecatedWarning()
 {
-	// Declarations.
-	DECLARE_CLASS(UObject,UObject,CLASS_Abstract|CLASS_NoExport,CASTCLASS_None,TEXT("/Script/CoreUObject"),NO_API)
-	DEFINE_DEFAULT_OBJECT_INITIALIZER_CONSTRUCTOR_CALL(UObject)
-#if WITH_HOT_RELOAD_CTORS
-	static UObject* __VTableCtorCaller(FVTableHelper& Helper)
-	{
-		return new (EC_InternalUseOnlyConstructor, (UObject*)GetTransientPackage(), NAME_None, RF_NeedLoad | RF_ClassDefaultObject | RF_TagGarbageTemp) UObject(Helper);
-	}
-#endif // WITH_HOT_RELOAD_CTORS
-
-	typedef UObject WithinClass;
-	static const TCHAR* StaticConfigName() {return TEXT("Engine");}
-
-	// Constructors and destructors.
-	UObject();
-	UObject(const FObjectInitializer& ObjectInitializer);
-	UObject( EStaticConstructor, EObjectFlags InFlags );
-#if WITH_HOT_RELOAD_CTORS
-	/** DO NOT USE. This constructor is for internal usage only for hot-reload purposes. */
-	UObject(FVTableHelper& Helper);
-#endif // WITH_HOT_RELOAD_CTORS
-
-	static void StaticRegisterNativesUObject() 
-	{
-	}
-
-	UObject* CreateDefaultSubobject(FName SubobjectFName, UClass* ReturnType, UClass* ClassToCreateByDefault, bool bIsRequired, bool bAbstract, bool bIsTransient);
-
-	/**
-	* Create a component or subobject only to be used with the editor.
-	* @param	TReturnType					class of return type, all overrides must be of this type
-	* @param	Outer						outer to construct the subobject in
-	* @param	SubobjectName				name of the new component
-	* @param bTransient		true if the component is being assigned to a transient property
-	*/
-	template<class TReturnType>
-	TReturnType* CreateEditorOnlyDefaultSubobject(FName SubobjectName, bool bTransient = false)
-	{
-		UClass* ReturnType = TReturnType::StaticClass();
-		return static_cast<TReturnType*>(CreateEditorOnlyDefaultSubobjectImpl(SubobjectName, ReturnType, bTransient));
-	}
-
-	/**
-	* Create a component or subobject
-	* @param	TReturnType					class of return type, all overrides must be of this type
-	* @param	Outer						outer to construct the subobject in
-	* @param	SubobjectName				name of the new component
-	* @param bTransient		true if the component is being assigned to a transient property
-	*/
-	template<class TReturnType>
-	TReturnType* CreateDefaultSubobject(FName SubobjectName, bool bTransient = false)
-	{
-		UClass* ReturnType = TReturnType::StaticClass();
-		return static_cast<TReturnType*>(CreateDefaultSubobject(SubobjectName, ReturnType, ReturnType, /*bIsRequired =*/ true, /*bIsAbstract =*/ false, bTransient));
-	}
-	
-	/**
-	* Create a component or subobject
-	* @param TReturnType class of return type, all overrides must be of this type
-	* @param TClassToConstructByDefault class to construct by default
-	* @param Outer outer to construct the subobject in
-	* @param SubobjectName name of the new component
-	* @param bTransient		true if the component is being assigned to a transient property
-	*/
-	template<class TReturnType, class TClassToConstructByDefault>
-	TReturnType* CreateDefaultSubobject(FName SubobjectName, bool bTransient = false)
-	{
-		return static_cast<TReturnType*>(CreateDefaultSubobject(SubobjectName, TReturnType::StaticClass(), TClassToConstructByDefault::StaticClass(), /*bIsRequired =*/ true, /*bIsAbstract =*/ false, bTransient));
-	}
-	
-	/**
-	* Create optional component or subobject. Optional subobjects may not get created
-	* when a derived class specified DoNotCreateDefaultSubobject with the subobject's name.
-	* @param	TReturnType					class of return type, all overrides must be of this type
-	* @param	Outer						outer to construct the subobject in
-	* @param	SubobjectName				name of the new component
-	* @param bTransient		true if the component is being assigned to a transient property
-	*/
-	template<class TReturnType>
-	TReturnType* CreateOptionalDefaultSubobject(FName SubobjectName, bool bTransient = false)
-	{
-		UClass* ReturnType = TReturnType::StaticClass();
-		return static_cast<TReturnType*>(CreateDefaultSubobject(SubobjectName, ReturnType, ReturnType, /*bIsRequired =*/ false, /*bIsAbstract =*/ false, bTransient));
-	}
-	
-	/**
-	* Create optional component or subobject. Optional subobjects may not get created
-	* when a derived class specified DoNotCreateDefaultSubobject with the subobject's name.
-	* @param	TReturnType					class of return type, all overrides must be of this type
-	* @param	Outer						outer to construct the subobject in
-	* @param	SubobjectName				name of the new component
-	* @param bTransient		true if the component is being assigned to a transient property
-	*/
-	template<class TReturnType>
-	TReturnType* CreateAbstractDefaultSubobject(FName SubobjectName, bool bTransient = false)
-	{
-		UClass* ReturnType = TReturnType::StaticClass();
-		return static_cast<TReturnType*>(CreateDefaultSubobject(SubobjectName, ReturnType, ReturnType, /*bIsRequired =*/ true, /*bIsAbstract =*/ true, bTransient));
-	}
-
-	/**
-	* Gets all default subobjects associated with this object instance.
-	* @param	OutDefaultSubobjects	Array containing all default subobjects of this object.
-	*/
-	void GetDefaultSubobjects(TArray<UObject*>& OutDefaultSubobjects);
-
-	/**
-	* Finds a subobject associated with this object instance by its name
-	* @param	OutDefaultSubobjects	Array containing all default subobjects of this object.
-	*/
-	UObject* GetDefaultSubobjectByName(FName ToFind);
-
-	//==========================================
-	// UObject interface.
-	//==========================================	
-protected:
-    /** 
-     * This function actually does the work for the GetDetailInfo and is virtual.  
-     * It should only be called from GetDetailedInfo as GetDetailedInfo is safe to call on NULL object pointers
-     **/
-	virtual FString GetDetailedInfoInternal() const { return TEXT("No_Detailed_Info_Specified"); }
-public:
-	/**
-	 * Called after the C++ constructor and after the properties have been initialized, including those loaded from config.
-	 * mainly this is to emulate some behavior of when the constructor was called after the properties were initialized.
-	 */
-	virtual void PostInitProperties();
-
-	/**
-	 * Called from within SavePackage on the passed in base/ root. The return value of this function will be passed to
-	 * PostSaveRoot. This is used to allow objects used as base to perform required actions before saving and cleanup
-	 * afterwards.
-	 * @param Filename: Name of the file being saved to (includes path)
-	 * @param AdditionalPackagesToCook [out] Array of other packages the Root wants to make sure are cooked when this is cooked
-	 *
-	 * @return	Whether PostSaveRoot needs to perform internal cleanup
-	 */
-	virtual bool PreSaveRoot(const TCHAR* Filename, TArray<FString>& AdditionalPackagesToCook)
-	{
-		return false;
-	}
-	/**
-	 * Called from within SavePackage on the passed in base/ root. This function is being called after the package
-	 * has been saved and can perform cleanup.
-	 *
-	 * @param	bCleanupIsRequired	Whether PreSaveRoot dirtied state that needs to be cleaned up
-	 */
-	virtual void PostSaveRoot( bool bCleanupIsRequired ) {}
-
-	/**
-	 * Presave function. Gets called once before an object gets serialized for saving. This function is necessary
-	 * for save time computation as Serialize gets called three times per object from within SavePackage.
-	 *
-	 * @warning: Objects created from within PreSave will NOT have PreSave called on them!!!
-	 */
-	virtual void PreSave(const class ITargetPlatform* TargetPlatform);
-
-	/**
-	 * Note that the object will be modified.  If we are currently recording into the 
-	 * transaction buffer (undo/redo), save a copy of this object into the buffer and 
-	 * marks the package as needing to be saved.
-	 *
-	 * @param	bAlwaysMarkDirty	if true, marks the package dirty even if we aren't
-	 *								currently recording an active undo/redo transaction
-	 * @return true if the object was saved to the transaction buffer
-	 */
-	virtual bool Modify( bool bAlwaysMarkDirty=true );
-
-#if WITH_EDITOR
-	/** 
-	 * Called when the object was loaded from another class via active class redirects.
-	 */
-	virtual void LoadedFromAnotherClass(const FName& OldClassName) {}
-#endif
-
-	/** 
-	 * Do any object-specific cleanup required immediately after loading an object, 
-	 * and immediately after any undo/redo.
-	 */
-	virtual void PostLoad();
-
-	/**
-	 * Instances components for objects being loaded from disk, if necessary.  Ensures that component references
-	 * between nested components are fixed up correctly.
-	 *
-	 * @param	OuterInstanceGraph	when calling this method on subobjects, specifies the instancing graph which contains all instanced
-	 *								subobjects and components for a subobject root.
-	 */
-	virtual void PostLoadSubobjects(FObjectInstancingGraph* OuterInstanceGraph);
-	
-	/**
-	 * Called before destroying the object.  This is called immediately upon deciding to destroy the object, to allow the object to begin an
-	 * asynchronous cleanup process.
-	 */
-	virtual void BeginDestroy();
-
-	/**
-	 * Called to check if the object is ready for FinishDestroy.  This is called after BeginDestroy to check the completion of the
-	 * potentially asynchronous object cleanup.
-	 * @return True if the object's asynchronous cleanup has completed and it is ready for FinishDestroy to be called.
-	 */
-	virtual bool IsReadyForFinishDestroy() { return true; }
-
-#if WITH_EDITOR
-	/**
-	 * Called in response to the linker changing, this can only happen in the editor
-	 */
-	virtual void PostLinkerChange() {}
-#endif
-
-	/**
-	 * Called to finish destroying the object.  After UObject::FinishDestroy is called, the object's memory should no longer be accessed.
-	 *
-	 * note: because properties are destroyed here, Super::FinishDestroy() should always be called at the end of your child class's
-	 * FinishDestroy() method, rather than at the beginning.
-	 */
-	virtual void FinishDestroy();
-
-	/** UObject serializer. */
-	virtual void Serialize( FArchive& Ar );
-
-	virtual void ShutdownAfterError() {}
-
-	/** 
-	 * This is called when property is modified by InterpPropertyTracks
-	 *
-	 * @param PropertyThatChanged	Property that changed
-	 */
-	virtual void PostInterpChange(UProperty* PropertyThatChanged) {}
-
-#if WITH_EDITOR
-	/** 
-	 * This is called when property is about to be modified by InterpPropertyTracks
-	 *
-	 * @param PropertyThatWillChange	Property that will be changed
-	 */
-	virtual void PreEditChange(UProperty* PropertyAboutToChange);
-
-	/**
-	 * This alternate version of PreEditChange is called when properties inside structs are modified.  The property that was actually modified
-	 * is located at the tail of the list.  The head of the list of the UStructProperty member variable that contains the property that was modified.
-	 *
-	 * @param PropertyAboutToChange the property that is about to be modified
-	 */
-	virtual void PreEditChange( class FEditPropertyChain& PropertyAboutToChange );
-
-	/**
-	 * Called by the editor to query whether a property of this object is allowed to be modified.
-	 * The property editor uses this to disable controls for properties that should not be changed.
-	 * When overriding this function you should always call the parent implementation first.
-	 *
-	 * @param	InProperty	The property to query
-	 *
-	 * @return	true if the property can be modified in the editor, otherwise false
-	 */
-	virtual bool CanEditChange( const UProperty* InProperty ) const;
-
-	/** 
-	 * Intentionally non-virtual as it calls the FPropertyChangedEvent version
-	 */
-	void PostEditChange();
-
-	/**
-	 * Called when a property on this object has been modified externally
-	 *
-	 * @param PropertyThatChanged the property that was modified
-	 */
-	virtual void PostEditChangeProperty( struct FPropertyChangedEvent& PropertyChangedEvent);
-
-	/**
-	 * This alternate version of PostEditChange is called when properties inside structs are modified.  The property that was actually modified
-	 * is located at the tail of the list.  The head of the list of the UStructProperty member variable that contains the property that was modified.
-	 */
-	virtual void PostEditChangeChainProperty( struct FPropertyChangedChainEvent& PropertyChangedEvent );
-
-	/** Gathers external data required for applying an undo transaction */
-	virtual TSharedPtr<ITransactionObjectAnnotation> GetTransactionAnnotation() const { return NULL; }
-
-	/** Called before applying a transaction to the object.  Default implementation simply calls PreEditChange. */
-	virtual void PreEditUndo();
-
-	/** Called after applying a transaction to the object.  Default implementation simply calls PostEditChange. */
-	virtual void PostEditUndo();
-
-	/** Called after applying a transaction to the object in cases where transaction annotation was provided. Default implementation simply calls PostEditChange. */
-	virtual void PostEditUndo(TSharedPtr<ITransactionObjectAnnotation> TransactionAnnotation);
-
-
-#endif // WITH_EDITOR
-
-	// @todo document
-	virtual void PostRename(UObject* OldOuter, const FName OldName) {}
-	
-	/**
-	 * Called after duplication & serialization and before PostLoad. Used to e.g. make sure UStaticMesh's UModel gets copied as well.
-	 * Note: NOT called on components on actor duplication (alt-drag or copy-paste).  Use PostEditImport as well to cover that case.
-	 */
-	virtual void PostDuplicate(bool bDuplicateForPIE) {}
-
-	/**
-	 * Called during saving to determine the load flags to save with the object.
-	 * Upon reload, this object will be discarded on clients
-	 *
-	 * @return	true if this object should not be loaded on clients
-	 */
-	virtual bool NeedsLoadForClient() const 
-	{ 
-		return true; 
-	}
-
-	/**
-	 * Called during saving to determine the load flags to save with the object.
-	 * Upon reload, this object will be discarded on servers
-	 *
-	 * @return	true if this object should not be loaded on servers
-	 */
-	virtual bool NeedsLoadForServer() const;
-
-	/**
-	 * Called during saving to determine the load flags to save with the object.
-	 *
-	 * @return	true if this object should always be loaded for editor game
-	 */
-	virtual bool NeedsLoadForEditorGame() const
-	{
-		return false;
-	}
-
-	/**
-	* Called during async load to determine if PostLoad can be called on the loading thread.
-	*
-	* @return	true if this object's PostLoad is thread safe
-	*/
-	virtual bool IsPostLoadThreadSafe() const
-	{
-		return false;
-	}
-
-
-	/** 
-	 *	Determines if you can create an object from the supplied template in the current context (editor, client only, dedicated server, game/listen) 
-	 *	This calls NeedsLoadForClient & NeedsLoadForServer
-	 */
-	static bool CanCreateInCurrentContext(UObject* Template);
-
-	/**
-	* Exports the property values for the specified object as text to the output device. Required for Copy&Paste
-	* Most objects don't need this as unreal script can handle most cases.
-	*
-	* @param	Out				the output device to send the exported text to
-	* @param	Indent			number of spaces to prepend to each line of output
-	*
-	* see also: ImportCustomProperties()
-	*/
-	virtual void ExportCustomProperties(FOutputDevice& Out, uint32 Indent)	{}
-
-	/**
-	* Exports the property values for the specified object as text to the output device. Required for Copy&Paste
-	* Most objects don't need this as unreal script can handle most cases.
-	*
-	* @param	SourceText		the input data (zero terminated), must not be 0
-	* @param	Warn			for error reporting, must not be 0
-	*
-	* see also: ExportCustomProperties()
-	*/
-	virtual void ImportCustomProperties(const TCHAR* SourceText, FFeedbackContext* Warn) {}
-
-	/**
-	 * Called after importing property values for this object (paste, duplicate or .t3d import)
-	 * Allow the object to perform any cleanup for properties which shouldn't be duplicated or
-	 * are unsupported by the script serialization
-	 */
-	virtual void PostEditImport() {}
-
-	/**
-	 * Called from ReloadConfig after the object has reloaded its configuration data.
-	 */
-	virtual void PostReloadConfig( class UProperty* PropertyThatWasLoaded ) {}
-
-	/** Rename this object to a unique name.*/
-	virtual bool Rename( const TCHAR* NewName=NULL, UObject* NewOuter=NULL, ERenameFlags Flags=REN_None );
-
-
-	/** @return a one line description of an object for viewing in the thumbnail view of the generic browser */
-	virtual FString GetDesc() { return TEXT( "" ); }
-
-#if WITH_ENGINE
-	virtual class UWorld* GetWorld() const;
-	class UWorld* GetWorldChecked(bool& bSupported) const;
-	bool ImplementsGetWorld() const;
-#endif
-
-	/**
-	 * Callback for retrieving a textual representation of natively serialized properties.  Child classes should implement this method if they wish
-	 * to have natively serialized property values included in things like diffcommandlet output.
-	 *
-	 * @param	out_PropertyValues	receives the property names and values which should be reported for this object.  The map's key should be the name of
-	 *								the property and the map's value should be the textual representation of the property's value.  The property value should
-	 *								be formatted the same way that UProperty::ExportText formats property values (i.e. for arrays, wrap in quotes and use a comma
-	 *								as the delimiter between elements, etc.)
-	 * @param	ExportFlags			bitmask of EPropertyPortFlags used for modifying the format of the property values
-	 *
-	 * @return	true if property values were added to the map.
-	 */
-	virtual bool GetNativePropertyValues( TMap<FString,FString>& out_PropertyValues, uint32 ExportFlags=0 ) const
-	{
-		return false;
-	}
-
-	/**
-	 * Returns the size of the object/ resource for display to artists/ LDs in the Editor. The
-	 * default behavior is to return 0 which indicates that the resource shouldn't
-	 * display its size.
-	 *
-	 * @param	Type	Indicates which resource size should be returned
-	 * @return	Size of resource as to be displayed to artists/ LDs in the Editor.
-	 */
-	virtual SIZE_T GetResourceSize(EResourceSizeMode::Type Mode)
-	{
-		return 0;
-	}
-
-	/** 
-	 * Returns the name of the exporter factory used to export this object
-	 * Used when multiple factories have the same extension
-	 */
-	virtual FName GetExporterName( void )
-	{
-		return( FName( TEXT( "" ) ) );
-	}
-
-	/**
-	 * Callback used to allow object register its direct object references that are not already covered by
-	 * the token stream.
-	 *
-	 * @param InThis Object to collect references from.
-	 * @param Collector	FReferenceCollector objects to be used to collect references.
-	 */
-	static void AddReferencedObjects(UObject* InThis, FReferenceCollector& Collector);
-
-	/**
-	 * Helper function to call AddReferencedObjects for this object's class.
-	 *
-	 * @param Collector	FReferenceCollector objects to be used to collect references.
-	 */
-	void CallAddReferencedObjects(FReferenceCollector& Collector);
-
-	/**
-	 * Save information for StaticAllocateObject in the case of overwriting an existing object.
-	 * StaticAllocateObject will call delete on the result after calling Restore()
-	 *
-	 * @return An FRestoreForUObjectOverwrite that can restore the object or NULL if this is not necessary.
-	 */
-	virtual FRestoreForUObjectOverwrite* GetRestoreForUObjectOverwrite()
-	{
-		return NULL;
-	}
-
-	/**
-	 * Returns whether native properties are identical to the one of the passed in component.
-	 *
-	 * @param	Other	Other component to compare against
-	 *
-	 * @return true if native properties are identical, false otherwise
-	 */
-	virtual bool AreNativePropertiesIdenticalTo( UObject* Other ) const
-	{
-		return true;
-	}
-
-	/**
-	 * Gathers a list of asset registry searchable tags which are name/value pairs with some type information
-	 * This only needs to be implemented for asset objects
-	 *
-	 * @param	OutTags		A list of key-value pairs associated with this object and their types
-	 */
-	struct FAssetRegistryTag
-	{
-		enum ETagType
-		{
-			TT_Hidden,
-			TT_Alphabetical,
-			TT_Numerical,
-			TT_Dimensional
-		};
-
-		FName Name;
-		ETagType Type;
-		FString Value;
-
-		FAssetRegistryTag(FName InName, const FString& InValue, ETagType InType)
-			: Name(InName), Type(InType), Value(InValue) {}
-
-		/** Gathers a list of asset registry searchable tags from given objects properties */
-		COREUOBJECT_API static void GetAssetRegistryTagsFromSearchableProperties(const UObject* Object, TArray<FAssetRegistryTag>& OutTags);
-	};
-	virtual void GetAssetRegistryTags(TArray<FAssetRegistryTag>& OutTags) const;
-
-	/** Get the common tag name used for all asset source file import paths */
-	static const FName& SourceFileTagName();
-
-#if WITH_EDITOR
-	/**
-	 * Additional data pertaining to asset registry tags used by the editor
-	 */
-	struct FAssetRegistryTagMetadata
-	{
-		FText DisplayName;
-		FText TooltipText;
-		FText Suffix;
-		FString ImportantValue;
-
-		/** Set override display name */
-		FAssetRegistryTagMetadata& SetDisplayName(const FText& InDisplayName)
-		{
-			DisplayName = InDisplayName;
-			return *this;
-		}
-
-		/** Set tooltip text pertaining to the asset registry tag in the column view header */
-		FAssetRegistryTagMetadata& SetTooltip(const FText& InTooltipText)
-		{
-			TooltipText = InTooltipText;
-			return *this;
-		}
-
-		/** Set suffix appended to the tag value */
-		FAssetRegistryTagMetadata& SetSuffix(const FText& InSuffix)
-		{
-			Suffix = InSuffix;
-			return *this;
-		}
-
-		/** Set value deemed to be 'important' for this registry tag */
-		FAssetRegistryTagMetadata& SetImportantValue(const FString& InImportantValue)
-		{
-			ImportantValue = InImportantValue;
-			return *this;
-		}
-	};
-
-	/** Gathers a collection of asset registry tag metadata */
-	virtual void GetAssetRegistryTagMetadata(TMap<FName, FAssetRegistryTagMetadata>& OutMetadata) const;
-#endif
-
-	/** Returns true if this object is considered an asset. */
-	virtual bool IsAsset() const;
-
-	/** Returns true if this object is considered a localized resource. */
-	virtual bool IsLocalizedResource() const;
-
-	/** Returns true if this object is safe to add to the root set. */
-	virtual bool IsSafeForRootSet() const;
-
-	/** 
-	 * Tags objects that are part of the same asset with the specified object flag, used for GC checking
-	 *
-	 * @param	ObjectFlags	Object Flags to enable on the related objects
-	 */
-	virtual void TagSubobjects(EObjectFlags NewFlags);
-
-	/** Returns properties that are replicated for the lifetime of the actor channel */
-	virtual void GetLifetimeReplicatedProps( TArray< class FLifetimeProperty > & OutLifetimeProps ) const;
-
-	/** IsNameStableForNetworking means an object can be referred to its path name (relative to outer) over the network */
-	virtual bool IsNameStableForNetworking() const;
-
-	/** IsFullNameStableForNetworking means an object can be referred to its full path name over the network */
-	virtual bool IsFullNameStableForNetworking() const;
-
-	/** IsSupportedForNetworking means an object can be referenced over the network */
-	virtual bool IsSupportedForNetworking() const;
-
-	/** Returns a list of sub-objects that have stable names for networking */
-	virtual void GetSubobjectsWithStableNamesForNetworking(TArray<UObject*> &ObjList) {}
-
-	/** Called right before receiving a bunch */
-	virtual void PreNetReceive();
-
-	/** Called right after receiving a bunch */
-	virtual void PostNetReceive();
-
-	/** Called right after calling all OnRep notifies (called even when there are no notifies) */
-	virtual void PostRepNotifies() {}
-
-	/** Called right before being marked for destruction due to network replication */
-	virtual void PreDestroyFromReplication();
-
-	/**
- 	 *******************************************************
-	 * Non virtual functions, not intended to be overridden
-	 *******************************************************
-	 */
-
-	/**
-	 * Test the selection state of a UObject
-	 *
-	 * @return		true if the object is selected, false otherwise.
-	 * @todo UE4 this doesn't belong here, but it doesn't belong anywhere else any better
-	 */
-	bool IsSelected() const;
-
-#if WITH_EDITOR
-	/**
-	 * Serializes all objects which have this object as their archetype into GMemoryArchive, then recursively calls this function
-	 * on each of those objects until the full list has been processed.
-	 * Called when a property value is about to be modified in an archetype object. 
-	 *
-	 * @param	AffectedObjects		the array of objects which have this object in their ObjectArchetype chain and will be affected by the change.
-	 *								Objects which have this object as their direct ObjectArchetype are removed from the list once they're processed.
-	 */
-	void PropagatePreEditChange( TArray<UObject*>& AffectedObjects, FEditPropertyChain& PropertyAboutToChange );
-
-	/**
-	 * De-serializes all objects which have this object as their archetype from the GMemoryArchive, then recursively calls this function
-	 * on each of those objects until the full list has been processed.
-	 *
-	 * @param	AffectedObjects		the array of objects which have this object in their ObjectArchetype chain and will be affected by the change.
-	 *								Objects which have this object as their direct ObjectArchetype are removed from the list once they're processed.
-	 */
-	void PropagatePostEditChange( TArray<UObject*>& AffectedObjects, FPropertyChangedChainEvent& PropertyChangedEvent );
-#endif // WITH_EDITOR
-
-	/**
-	 * Serializes the script property data located at Data.  When saving, only saves those properties which differ from the corresponding
-	 * value in the specified 'DiffObject' (usually the object's archetype).
-	 *
-	 * @param	Ar				the archive to use for serialization
-	 */
-	void SerializeScriptProperties( FArchive& Ar ) const;
-
-	/**
-	 * Wrapper function for InitProperties() which handles safely tearing down this object before re-initializing it
-	 * from the specified source object.
-	 *
-	 * @param	SourceObject	the object to use for initializing property values in this object.  If not specified, uses this object's archetype.
-	 * @param	InstanceGraph	contains the mappings of instanced objects and components to their templates
-	 */
-	void ReinitializeProperties( UObject* SourceObject=NULL, struct FObjectInstancingGraph* InstanceGraph=NULL );
-
-	/**
-	 * This will return detail info about this specific object. (e.g. AudioComponent will return the name of the cue,
-	 * ParticleSystemComponent will return the name of the ParticleSystem)  The idea here is that in many places
-	 * you have a component of interest but what you really want is some characteristic that you can use to track
-	 * down where it came from.  
-	 *
-	 * @note	safe to call on NULL object pointers!
-	 */
-	FString GetDetailedInfo() const;
-
-	/**
-	 * Called before destroying the object.  This is called immediately upon deciding to destroy the object, to allow the object to begin an
-	 * asynchronous cleanup process.
-	 */
-	bool ConditionalBeginDestroy();
-	bool ConditionalFinishDestroy();
-	
-	/** PostLoad if needed. */
-	void ConditionalPostLoad();
-
-	/**
-	 * Instances subobjects and components for objects being loaded from disk, if necessary.  Ensures that references
-	 * between nested components are fixed up correctly.
-	 *
-	 * @param	OuterInstanceGraph	when calling this method on subobjects, specifies the instancing graph which contains all instanced
-	 *								subobjects and components for a subobject root.
-	 */
-	void ConditionalPostLoadSubobjects( struct FObjectInstancingGraph* OuterInstanceGraph=NULL );
-#if WITH_EDITOR
-	/**
-	 * Starts caching of platform specific data for the target platform
-	 * Called when cooking before serialization so that object can prepare platform specific data
-	 * Not called during normal loading of objects
-	 * 
-	 * @param	TargetPlatform	target platform to cache platform specific data for
-	 */
-	virtual void BeginCacheForCookedPlatformData( const ITargetPlatform* TargetPlatform ) {  }
-	
-	/**
-	 * Have we finished loading all the cooked platform data for the target platforms requested in BeginCacheForCookedPlatformData
-	 * 
-	 * @param	TargetPlatform target platform to check for cooked platform data
-	 */
-	virtual bool IsCachedCookedPlatformDataLoaded( const ITargetPlatform* TargetPlatform ) { return true; }
-
-	/**
-	 * All caching has finished for this object (all IsCachedCookedPlatformDataLoaded functions have finished for all platforms)
-	 */
-	virtual void WillNeverCacheCookedPlatformDataAgain() { }
-
-	/**
-	 * Clears cached cooked platform data for specific platform
-	 * 
-	 * @param	TargetPlatform	target platform to cache platform specific data for
-	 */
-	virtual void ClearCachedCookedPlatformData( const ITargetPlatform* TargetPlatform ) {  }
-
-	/**
-	 * Clear all cached cooked platform data
-	 * 
-	 * @param	TargetPlatform	target platform to cache platform specific data for
-	 */
-	virtual void ClearAllCachedCookedPlatformData() { }
-#endif
-	/**
-	 * Determine if this object has SomeObject in its archetype chain.
-	 */
-	inline bool IsBasedOnArchetype( const UObject* const SomeObject ) const;
-
-	UFunction* FindFunction( FName InName ) const;
-	UFunction* FindFunctionChecked( FName InName ) const;
-
-	/**
-	 * Uses the TArchiveObjectReferenceCollector to build a list of all components referenced by this object which have this object as the outer
-	 *
-	 * @param	OutDefaultSubobjects	the array that should be populated with the default subobjects "owned" by this object
-	 * @param	bIncludeNestedSubobjects	controls whether subobjects which are contained by this object, but do not have this object
-	 *										as its direct Outer should be included
-	 */
-	void CollectDefaultSubobjects( TArray<UObject*>& OutDefaultSubobjects, bool bIncludeNestedSubobjects=false );
-
-	/**
-	 * Checks default sub-object assumptions.
-	 *
-	 * @param bForceCheck Force checks even if not enabled globally.
-	 * @return true if the assumptions are met, false otherwise.
-	 */
-	bool CheckDefaultSubobjects(bool bForceCheck = false);
-
-	/**
-	 * Save configuration.
-	 * @warning: Must be safe on class-default metaobjects.
-	 * !!may benefit from hierarchical propagation, deleting keys that match superclass...not sure what's best yet.
-	 */
-	void SaveConfig( uint64 Flags=CPF_Config, const TCHAR* Filename=NULL, FConfigCacheIni* Config=GConfig );
-
-	/**
-	 * Saves just the section(s) for this class into the default ini file for the class (with just the changes from base)
-	 */
-	void UpdateDefaultConfigFile(const FString& SpecificFileLocation = "");
-
-	/**
-	 * Saves just the section(s) for this class into the global user ini file for the class (with just the changes from base)
-	 */
-	void UpdateGlobalUserConfigFile();
-
-	/**
-	 * Saves just the property into the global user ini file for the class (with just the changes from base)
-	 */
-	void UpdateSinglePropertyInConfigFile(const UProperty* InProperty, const FString& InConfigIniName);
-
-private:
-	/**
-	 * Saves just the section(s) for this class into the given ini file for the class (with just the changes from base)
-	 */
-	void UpdateSingleSectionOfConfigFile(const FString& ConfigIniName);
-
-#if WITH_HOT_RELOAD_CTORS
-	/**
-	 * Ensures that current thread is NOT during vtable ptr retrieval process
-	 * of some UClass.
-	 */
-	void EnsureNotRetrievingVTablePtr() const;
-#endif // WITH_HOT_RELOAD_CTORS
-
-public:
-	
-	/**
-	 * Get the default config filename for the specified UObject
-	 */
-	FString GetDefaultConfigFilename() const;
-
-	/**
-	 * Get the global user override config filename for the specified UObject
-	 */
-	FString GetGlobalUserConfigFilename() const;
-
-	/**
-	 * Imports property values from an .ini file.
-	 *
-	 * @param	Class				the class to use for determining which section of the ini to retrieve text values from
-	 * @param	Filename			indicates the filename to load values from; if not specified, uses ConfigClass's ClassConfigName
-	 * @param	PropagationFlags	indicates how this call to LoadConfig should be propagated; expects a bitmask of UE4::ELoadConfigPropagationFlags values.
-	 * @param	PropertyToLoad		if specified, only the ini value for the specified property will be imported.
-	 */
-	void LoadConfig( UClass* ConfigClass=NULL, const TCHAR* Filename=NULL, uint32 PropagationFlags=UE4::LCPF_None, class UProperty* PropertyToLoad=NULL );
-
-	/**
-	 * Wrapper method for LoadConfig that is used when reloading the config data for objects at runtime which have already loaded their config data at least once.
-	 * Allows the objects the receive a callback that it's configuration data has been reloaded.
-	 *
-	 * @param	Class				the class to use for determining which section of the ini to retrieve text values from
-	 * @param	Filename			indicates the filename to load values from; if not specified, uses ConfigClass's ClassConfigName
-	 * @param	PropagationFlags	indicates how this call to LoadConfig should be propagated; expects a bitmask of UE4::ELoadConfigPropagationFlags values.
-	 * @param	PropertyToLoad		if specified, only the ini value for the specified property will be imported
-	 */
-	void ReloadConfig( UClass* ConfigClass=NULL, const TCHAR* Filename=NULL, uint32 PropagationFlags=UE4::LCPF_None, class UProperty* PropertyToLoad=NULL );
-
-	/** Import an object from a file. */
-	void ParseParms( const TCHAR* Parms );
-
-	/**
-	 * Outputs a string to an arbitrary output device, describing the list of objects which are holding references to this one.
-	 *
-	 * @param	Ar						the output device to send output to
-	 * @param	Referencers				optionally allows the caller to specify the list of references to output.
-	 */
-	void OutputReferencers( FOutputDevice& Ar, FReferencerInformationList* Referencers=NULL );
-	
-	// @todo document
-	void RetrieveReferencers( TArray<FReferencerInformation>* OutInternalReferencers, TArray<FReferencerInformation>* OutExternalReferencers);
-
-	/**
-	 * Changes the linker and linker index to the passed in one. A linker of NULL and linker index of INDEX_NONE
-	 * indicates that the object is without a linker.
-	 *
-	 * @param LinkerLoad				New LinkerLoad object to set
-	 * @param LinkerIndex				New LinkerIndex to set
-	 * @param bShouldDetachExisting		If true, detach existing linker and call PostLinkerChange
-	 */
-	void SetLinker( FLinkerLoad* LinkerLoad, int32 LinkerIndex, bool bShouldDetachExisting=true );
-
-	/**
-	 * Return the template that an object with this class, outer and name would be
-	 * 
-	 * @return the archetype for this object
-	 */
-	static UObject* GetArchetypeFromRequiredInfo(UClass* Class, UObject* Outer, FName Name, EObjectFlags ObjectFlags);
-
-	/**
-	 * Return the template this object is based on. 
-	 * 
-	 * @return the archetype for this object
-	 */
-	UObject* GetArchetype() const
-	{
-		return GetArchetypeFromRequiredInfo(GetClass(), GetOuter(), GetFName(), GetFlags());
-	}
-
-	/**
-	 * Builds a list of objects which have this object in their archetype chain.
-	 *
-	 * @param	Instances	receives the list of objects which have this one in their archetype chain
-	 */
-	void GetArchetypeInstances( TArray<UObject*>& Instances );
-
-	/**
-	 * Wrapper for calling UClass::InstanceSubobjectTemplates() for this object.
-	 */
-	void InstanceSubobjectTemplates( struct FObjectInstancingGraph* InstanceGraph = NULL );
-
-	/**
-	 * Returns true if this object implements the interface T, false otherwise.
-	 */
-	template<class T>
-	bool Implements() const;
-
-	//==========================================
-	// Kismet Virtual Machine
-	//==========================================
-	//@todo UE4 - do these really need to be virtual to support actors?
-
-	//
-	// Script processing functions.
-	//
-	virtual void ProcessEvent( UFunction* Function, void* Parms );
-
-	/**
-	 * Return the space this function should be called.   Checks to see if this function should
-	 * be called locally, remotely, or simply absorbed under the given conditions
-	 *
-	 * @param Function function to call
-	 * @param Parameters arguments to the function call
-	 * @param Stack stack frame for the function call
-	 * @return bitmask representing all callspaces that apply to this UFunction in the given context
-	 */
-	virtual int32 GetFunctionCallspace( UFunction* Function, void* Parameters, FFrame* Stack )
-	{
-		return FunctionCallspace::Local;
-	}
-
-	/**
-	 * Call the actor's function remotely
-	 *
-	 * @param Function function to call
-	 * @param Parameters arguments to the function call
-	 * @param Stack stack frame for the function call
-	 */
-	virtual bool CallRemoteFunction( UFunction* Function, void* Parms, struct FOutParmRec* OutParms, FFrame* Stack )
-	{
-		return false;
-	}
-
-	/** Command line. */
-	bool CallFunctionByNameWithArguments( const TCHAR* Cmd, FOutputDevice& Ar, UObject* Executor, bool bForceCallWithNonExec = false );
-
-	// Call a function
-	void CallFunction( FFrame& Stack, RESULT_DECL, UFunction* Function );
-
-	//
-	// Internal function call processing.
-	// @warning: might not write anything to Result if proper type isn't returned.
-	//
-	void ProcessInternal( FFrame& Stack, RESULT_DECL );
-
-	/**
-	 * This function handles a console exec sent to the object; it is virtual so 'nexus' objects like
-	 * a player controller can reroute the command to several different objects.
-	 */
-	virtual bool ProcessConsoleExec(const TCHAR* Cmd, FOutputDevice& Ar, UObject* Executor)
-	{
-		return CallFunctionByNameWithArguments(Cmd, Ar, Executor);
-	}
-
-	/** advances Stack's code past the parameters to the given Function and if the function has a return value, copies the zero value for that property to the memory for the return value
-	 * @param Stack the script stack frame
-	 * @param Result pointer to where the return value should be written
-	 * @param Function the function being called
-	 */
-	void SkipFunction(FFrame& Stack, RESULT_DECL, UFunction* Function);
-
-	/**
-	 * Called on the target when a class is loaded with ClassGeneratedBy is loaded.  Should regenerate the class if needed, and return the updated class
-	 * @return	Updated instance of the class, if needed, or NULL if no regeneration was necessary
-	 */
-	virtual UClass* RegenerateClass(UClass* ClassToRegenerate, UObject* PreviousCDO, TArray<UObject*>& ObjLoaded) { return NULL; };
-
-	/** 
-	 * Returns whether this object is contained in or part of a blueprint object
-	 */
-	bool IsInBlueprint() const;
-
-	/** 
-	 *  Destroy properties that won't be destroyed by the native destructor
-	 */
-	void DestroyNonNativeProperties();
-
-	virtual void MarkAsEditorOnlySubobject() { }
-
-	// UnrealScript intrinsics.
-
-	// Undefined native handler
-	DECLARE_FUNCTION(execUndefined);
-
-	
-	// @todo document below declared functions
-	// Variables
-	DECLARE_FUNCTION(execLocalVariable);
-	DECLARE_FUNCTION(execInstanceVariable);
-	DECLARE_FUNCTION(execDefaultVariable);
-	DECLARE_FUNCTION(execLocalOutVariable);
-	DECLARE_FUNCTION(execInterfaceVariable);
-	DECLARE_FUNCTION(execInterfaceContext);
-	DECLARE_FUNCTION(execArrayElement);
-	DECLARE_FUNCTION(execBoolVariable);
-	DECLARE_FUNCTION(execClassDefaultVariable);
-	DECLARE_FUNCTION(execEndFunctionParms);
-
-
-	// Do Nothing 
-	DECLARE_FUNCTION(execNothing);
-	DECLARE_FUNCTION(execNothingOp4a);
-
-	/** Breakpoint; only observed in the editor; executing it at any other time is a NOP */
-	DECLARE_FUNCTION(execBreakpoint);
-
-	/** Tracepoint; only observed in the editor; executing it at any other time is a NOP */
-	DECLARE_FUNCTION(execTracepoint);
-	DECLARE_FUNCTION(execWireTracepoint);
-
-	/** Instrumentation event for profiling; only observed in the builds with blueprint instrumentation */
-	DECLARE_FUNCTION(execInstrumentation);
-
-	DECLARE_FUNCTION(execEndOfScript);
-
-	/** failsafe for functions that return a value - returns the zero value for a property and logs that control reached the end of a non-void function */
-	DECLARE_FUNCTION(execReturnNothing);
-	DECLARE_FUNCTION(execEmptyParmValue);
-
-	// Commands
-	DECLARE_FUNCTION(execJump);
-	DECLARE_FUNCTION(execJumpIfNot);
-	DECLARE_FUNCTION(execAssert);
-
-	/** 
-	 * Push a code offset onto the execution flow stack for future execution.  
-	 * Current execution continues to the next instruction after the push one.
-	 */
-	DECLARE_FUNCTION(execPushExecutionFlow);
-
-	/**
-	 * Pops a code offset from the execution flow stack and starts execution there.
-	 * If there are no stack entries left, it is treated as an execution error.
-	 */
-	DECLARE_FUNCTION(execPopExecutionFlow);
-	DECLARE_FUNCTION(execComputedJump);
-
-	/** 
-	 * Pops a code offset from the execution flow stack and starts execution there, if a condition is not true.
-	 * If there are no stack entries left, it is treated as an execution error.
-	 */
-	DECLARE_FUNCTION(execPopExecutionFlowIfNot);
-
-
-	// Assignment
-	DECLARE_FUNCTION(execLet);
-	DECLARE_FUNCTION(execLetObj);
-	DECLARE_FUNCTION(execLetWeakObjPtr);
-	DECLARE_FUNCTION(execLetBool);
-	DECLARE_FUNCTION(execLetDelegate);
-	DECLARE_FUNCTION(execLetMulticastDelegate);
-
-	// Delegates 
-	DECLARE_FUNCTION(execAddMulticastDelegate);
-	DECLARE_FUNCTION(execClearMulticastDelegate);
-	DECLARE_FUNCTION(execEatReturnValue);
-	DECLARE_FUNCTION(execRemoveMulticastDelegate);
-
-	// Context expressions
-	DECLARE_FUNCTION(execSelf);
-	DECLARE_FUNCTION(execContext);
-	DECLARE_FUNCTION(execContext_FailSilent);
-	DECLARE_FUNCTION(execStructMemberContext);
-
-	// Function calls
-	DECLARE_FUNCTION(execVirtualFunction);
-	DECLARE_FUNCTION(execFinalFunction);
-
-	// Struct comparison
-	DECLARE_FUNCTION(execStructCmpEq);
-	DECLARE_FUNCTION(execStructCmpNe);
-	DECLARE_FUNCTION(execStructMember);
-
-	// @todo delegate: Delegate comparison is not supported for multi-cast delegates
-	DECLARE_FUNCTION(execEqualEqual_DelegateDelegate);
-	DECLARE_FUNCTION(execNotEqual_DelegateDelegate);
-	DECLARE_FUNCTION(execEqualEqual_DelegateFunction);
-	DECLARE_FUNCTION(execNotEqual_DelegateFunction);
-
-	// Constants
-	DECLARE_FUNCTION(execIntConst);
-	DECLARE_FUNCTION(execInt64Const);
-	DECLARE_FUNCTION(execUInt64Const);
-	DECLARE_FUNCTION(execSkipOffsetConst);
-	DECLARE_FUNCTION(execFloatConst);
-	DECLARE_FUNCTION(execStringConst);
-	DECLARE_FUNCTION(execUnicodeStringConst);
-	DECLARE_FUNCTION(execTextConst);
-	DECLARE_FUNCTION(execObjectConst);
-	DECLARE_FUNCTION(execAssetConst);
-
-	// @todo delegate: Multi-cast versions needed for script execution!		(Need Add, Remove, Clear/Empty)
-	DECLARE_FUNCTION(execInstanceDelegate);
-	DECLARE_FUNCTION(execNameConst);
-	DECLARE_FUNCTION(execByteConst);
-	DECLARE_FUNCTION(execIntZero);
-	DECLARE_FUNCTION(execIntOne);
-	DECLARE_FUNCTION(execTrue);
-	DECLARE_FUNCTION(execFalse);
-	DECLARE_FUNCTION(execNoObject);
-	DECLARE_FUNCTION(execNullInterface);
-	DECLARE_FUNCTION(execIntConstByte);
-	DECLARE_FUNCTION(execRotationConst);
-	DECLARE_FUNCTION(execVectorConst);
-	DECLARE_FUNCTION(execTransformConst);
-	DECLARE_FUNCTION(execStructConst);
-	DECLARE_FUNCTION(execSetArray);
-	DECLARE_FUNCTION(execArrayConst);
-
-	// Object construction
-	DECLARE_FUNCTION(execNew);
-	DECLARE_FUNCTION(execClassContext);
-	DECLARE_FUNCTION(execNativeParm);
-
-	// Conversions 
-	DECLARE_FUNCTION(execDynamicCast);
-	DECLARE_FUNCTION(execMetaCast);
-	DECLARE_FUNCTION(execPrimitiveCast);
-	DECLARE_FUNCTION(execInterfaceCast);
-
-	// Cast functions
-	DECLARE_FUNCTION(execObjectToBool);
-	DECLARE_FUNCTION(execInterfaceToBool);
-	DECLARE_FUNCTION(execObjectToInterface);
-	DECLARE_FUNCTION(execInterfaceToInterface);
-	DECLARE_FUNCTION(execInterfaceToObject);
-
-	// Dynamic array functions
-	// Array support
-	DECLARE_FUNCTION(execGetDynArrayElement);
-	DECLARE_FUNCTION(execSetDynArrayElement);
-	DECLARE_FUNCTION(execGetDynArrayLength);
-	DECLARE_FUNCTION(execSetDynArrayLength);
-	DECLARE_FUNCTION(execDynArrayInsert);
-	DECLARE_FUNCTION(execDynArrayRemove);
-	DECLARE_FUNCTION(execDynArrayFind);
-	DECLARE_FUNCTION(execDynArrayFindStruct);
-	DECLARE_FUNCTION(execDynArrayAdd);
-	DECLARE_FUNCTION(execDynArrayAddItem);
-	DECLARE_FUNCTION(execDynArrayInsertItem);
-	DECLARE_FUNCTION(execDynArrayRemoveItem);
-	DECLARE_FUNCTION(execDynArraySort);
-
-	DECLARE_FUNCTION(execBindDelegate);
-	DECLARE_FUNCTION(execCallMulticastDelegate);
-
-	DECLARE_FUNCTION(execLetValueOnPersistentFrame);
-
-	DECLARE_FUNCTION(execCallMathFunction);
-
-	DECLARE_FUNCTION(execSwitchValue);
-
-	DECLARE_FUNCTION(execArrayGetByRef);
-
-	// -- K2 support functions
-	struct Object_eventExecuteUbergraph_Parms
-	{
-		int32 EntryPoint;
-	};
-	void ExecuteUbergraph(int32 EntryPoint)
-	{
-		Object_eventExecuteUbergraph_Parms Parms;
-		Parms.EntryPoint=EntryPoint;
-		ProcessEvent(FindFunctionChecked(NAME_ExecuteUbergraph),&Parms);
-	}
-
-protected: 
-
-	/** Checks it's ok to perform subobjects check at this time. */
-	bool CanCheckDefaultSubObjects(bool bForceCheck, bool& bResult);
-
-	/**
-	* Checks default sub-object assumptions.
-	*
-	* @return true if the assumptions are met, false otherwise.
-	*/
-	virtual bool CheckDefaultSubobjectsInternal();
-
-private:
-	void ProcessContextOpcode(FFrame& Stack, RESULT_DECL, bool bCanFailSilent);
-
-	/**
-	* Create a component or subobject only to be used with the editor.
-	* @param	Outer						outer to construct the subobject in
-	* @param	ReturnType					type of the new component
-	* @param	SubobjectName				name of the new component
-	* @param	bTransient					true if the component is being assigned to a transient property
-	*/
-	UObject* CreateEditorOnlyDefaultSubobjectImpl(FName SubobjectName, UClass* ReturnType, bool bTransient = false);
-};
+}
 
-/**
-* Test validity of object
-*
-* @param	Test			The object to test
-* @return	Return true if the object is usable: non-null and not pending kill
-*/
-FORCEINLINE bool IsValid(const UObject *Test)
+inline void TriggerUObjectHeaderDeprecatedWarning()
 {
-	return Test && !Test->IsPendingKill();
+	UObjectHeaderDeprecatedWarning();
 }
diff --git a/Engine/Source/Runtime/DatabaseSupport/Private/DatabaseSupportPrivatePCH.h b/Engine/Source/Runtime/DatabaseSupport/Private/DatabaseSupportPrivatePCH.h
--- a/Engine/Source/Runtime/DatabaseSupport/Private/DatabaseSupportPrivatePCH.h
+++ b/Engine/Source/Runtime/DatabaseSupport/Private/DatabaseSupportPrivatePCH.h
@@ -1,6 +1,5 @@
 // Copyright 1998-2016 Epic Games, Inc. All Rights Reserved.
 
-#include "Database.h"
+#pragma once
 
-// You should place include statements to your module's private header files here.  You only need to
-// add includes for headers that are used in most of your module's source files though.
+#include "Database.h"
diff --git a/Engine/Source/Runtime/Engine/Classes/Components/ActorComponent.h b/Engine/Source/Runtime/Engine/Classes/Components/ActorComponent.h
--- a/Engine/Source/Runtime/Engine/Classes/Components/ActorComponent.h
+++ b/Engine/Source/Runtime/Engine/Classes/Components/ActorComponent.h
@@ -789,25 +789,6 @@ private:
 //////////////////////////////////////////////////////////////////////////
 // UActorComponent inlines
 
-FORCEINLINE_DEBUGGABLE class AActor* UActorComponent::GetOwner() const
-{
-#if WITH_EDITOR
-	// During undo/redo the cached owner is unreliable so just used GetTypedOuter
-	if (bCanUseCachedOwner)
-	{
-		checkSlow(OwnerPrivate == GetTypedOuter<AActor>()); // verify cached value is correct
-		return OwnerPrivate;
-	}
-	else
-	{
-		return GetTypedOuter<AActor>();
-	}
-#else
-	checkSlow(OwnerPrivate == GetTypedOuter<AActor>()); // verify cached value is correct
-	return OwnerPrivate;
-#endif
-}
-
 FORCEINLINE bool UActorComponent::CanEverAffectNavigation() const
 {
 	return bCanEverAffectNavigation;
diff --git a/Engine/Source/Runtime/Engine/Classes/Engine/Level.h b/Engine/Source/Runtime/Engine/Classes/Engine/Level.h
--- a/Engine/Source/Runtime/Engine/Classes/Engine/Level.h
+++ b/Engine/Source/Runtime/Engine/Classes/Engine/Level.h
@@ -72,6 +72,12 @@ struct ENGINE_API FDynamicTextureInstance : public FStreamableTextureInstance
 	friend FArchive& operator<<( FArchive& Ar, FDynamicTextureInstance& TextureInstance );
 };
 
+/** Manually implement TPointerIsConvertibleFromTo for AActor so that TWeakObjectPtr can use it below when AActor is forward declared. */
+template<> struct TPointerIsConvertibleFromTo<AActor, const volatile UObject>
+{
+	enum { Value = 1 };
+};
+
 /** Struct that holds on to information about Actors that wish to be auto enabled for input before the player controller has been created */
 struct FPendingAutoReceiveInputActor
 {
diff --git a/Engine/Source/Runtime/Engine/Classes/Engine/MemberReference.h b/Engine/Source/Runtime/Engine/Classes/Engine/MemberReference.h
--- a/Engine/Source/Runtime/Engine/Classes/Engine/MemberReference.h
+++ b/Engine/Source/Runtime/Engine/Classes/Engine/MemberReference.h
@@ -447,13 +447,13 @@ protected:
 	/** 
 	 * A mapping from old pin name to new pin name for each K2 node.  Get primed from INI files, and should contain entries for node class, and old param name and new param name
 	 */
-	static TMultiMap<UClass*, FParamRemapInfo> ParamRedirectMap;
+	ENGINE_API static TMultiMap<UClass*, FParamRemapInfo> ParamRedirectMap;
 
 	/** Has the field map been initialized this run */
 	static bool bFieldRedirectMapInitialized;
 
 	/** Init the field redirect map (if not already done) from .ini file entries */
-	static void InitFieldRedirectMap();
+	ENGINE_API static void InitFieldRedirectMap();
 
 	/** 
 	 * Searches the field replacement map for an appropriately named field in the specified scope, and returns an updated name if this property has been listed in the remapping table
diff --git a/Engine/Source/Runtime/Engine/Classes/GameFramework/Actor.h b/Engine/Source/Runtime/Engine/Classes/GameFramework/Actor.h
--- a/Engine/Source/Runtime/Engine/Classes/GameFramework/Actor.h
+++ b/Engine/Source/Runtime/Engine/Classes/GameFramework/Actor.h
@@ -11,14 +11,14 @@
 #include "TimerManager.h"
 #include "Engine/Level.h"
 
+#include "Actor.generated.h"
+
 struct FHitResult;
 class AActor;
 class FTimerManager; 
 class UNetDriver;
 struct FNetViewer;
 
-#include "Actor.generated.h"
-
 ENGINE_API DECLARE_LOG_CATEGORY_EXTERN(LogActor, Log, Warning);
  
 
@@ -3027,6 +3027,28 @@ FORCEINLINE_DEBUGGABLE void AActor::SetNetUpdateTime(float NewUpdateTime)
 	NetUpdateTime = NewUpdateTime;
 }
 
+//////////////////////////////////////////////////////////////////////////
+// UActorComponent inlines
+
+FORCEINLINE_DEBUGGABLE class AActor* UActorComponent::GetOwner() const
+{
+#if WITH_EDITOR
+	// During undo/redo the cached owner is unreliable so just used GetTypedOuter
+	if (bCanUseCachedOwner)
+	{
+		checkSlow(OwnerPrivate == GetTypedOuter<AActor>()); // verify cached value is correct
+		return OwnerPrivate;
+	}
+	else
+	{
+		return GetTypedOuter<AActor>();
+	}
+#else
+	checkSlow(OwnerPrivate == GetTypedOuter<AActor>()); // verify cached value is correct
+	return OwnerPrivate;
+#endif
+}
+
 //////////////////////////////////////////////////////////////////////////
 // Macro to hide common Transform functions in native code for classes where they don't make sense.
 // Note that this doesn't prevent access through function calls from parent classes (ie an AActor*), but
diff --git a/Engine/Source/Runtime/Engine/Classes/Sound/SoundConcurrency.h b/Engine/Source/Runtime/Engine/Classes/Sound/SoundConcurrency.h
--- a/Engine/Source/Runtime/Engine/Classes/Sound/SoundConcurrency.h
+++ b/Engine/Source/Runtime/Engine/Classes/Sound/SoundConcurrency.h
@@ -164,7 +164,7 @@ class FSoundConcurrencyManager
 {
 public:
 	FSoundConcurrencyManager(class FAudioDevice* InAudioDevice);
-	~FSoundConcurrencyManager();
+	ENGINE_API ~FSoundConcurrencyManager();
 
 	/** Returns a newly allocated active sound given the input active sound struct. Will return nullptr if the active sound concurrency evaluation doesn't allow for it. */
 	FActiveSound* CreateNewActiveSound(const FActiveSound& NewActiveSound);
diff --git a/Engine/Source/Runtime/Engine/Private/Animation/AnimNode_AssetPlayerBase.cpp b/Engine/Source/Runtime/Engine/Private/Animation/AnimNode_AssetPlayerBase.cpp
new file mode 100644
--- /dev/null
+++ b/Engine/Source/Runtime/Engine/Private/Animation/AnimNode_AssetPlayerBase.cpp
@@ -0,0 +1,76 @@
+// Copyright 1998-2016 Epic Games, Inc. All Rights Reserved.
+
+#include "EnginePrivate.h"
+#include "Animation/AnimNode_AssetPlayerBase.h"
+#include "Animation/AnimInstanceProxy.h"
+
+FAnimNode_AssetPlayerBase::FAnimNode_AssetPlayerBase()
+	: bIgnoreForRelevancyTest(false)
+	, GroupIndex(INDEX_NONE)
+	, BlendWeight(0.0f)
+	, InternalTimeAccumulator(0.0f)
+	, bHasBeenFullWeight(false)
+{
+
+}
+
+void FAnimNode_AssetPlayerBase::Initialize(const FAnimationInitializeContext& Context)
+{
+	FAnimNode_Base::Initialize(Context);
+
+	MarkerTickRecord.Reset();
+	bHasBeenFullWeight = false;
+}
+
+void FAnimNode_AssetPlayerBase::Update(const FAnimationUpdateContext& Context)
+{
+	// Cache the current weight and update the node
+	BlendWeight = Context.GetFinalBlendWeight();
+	bHasBeenFullWeight = bHasBeenFullWeight || (BlendWeight >= (1.0f - ZERO_ANIMWEIGHT_THRESH));
+
+	UpdateAssetPlayer(Context);
+}
+
+void FAnimNode_AssetPlayerBase::CreateTickRecordForNode(const FAnimationUpdateContext& Context, UAnimSequenceBase* Sequence, bool bLooping, float PlayRate)
+{
+	// Create a tick record and fill it out
+	const float FinalBlendWeight = Context.GetFinalBlendWeight();
+	FAnimGroupInstance* SyncGroup;
+	const int32 GroupIndexToUse = ((GroupRole != EAnimGroupRole::TransitionLeader) || bHasBeenFullWeight) ? GroupIndex : INDEX_NONE;
+
+	FAnimTickRecord& TickRecord = Context.AnimInstanceProxy->CreateUninitializedTickRecord(GroupIndexToUse, /*out*/ SyncGroup);
+
+	Context.AnimInstanceProxy->MakeSequenceTickRecord(TickRecord, Sequence, bLooping, PlayRate, FinalBlendWeight, /*inout*/ InternalTimeAccumulator, MarkerTickRecord);
+
+	// Update the sync group if it exists
+	if (SyncGroup != NULL)
+	{
+		SyncGroup->TestTickRecordForLeadership(GroupRole);
+	}
+}
+
+float FAnimNode_AssetPlayerBase::GetCachedBlendWeight()
+{
+	return BlendWeight;
+}
+
+float FAnimNode_AssetPlayerBase::GetAccumulatedTime()
+{
+	return InternalTimeAccumulator;
+}
+
+void FAnimNode_AssetPlayerBase::SetAccumulatedTime(const float& NewTime)
+{
+	InternalTimeAccumulator = NewTime;
+}
+
+UAnimationAsset* FAnimNode_AssetPlayerBase::GetAnimAsset()
+{
+	return nullptr;
+}
+
+void FAnimNode_AssetPlayerBase::ClearCachedBlendWeight()
+{
+	BlendWeight = 0.0f;
+}
+
diff --git a/Engine/Source/Runtime/Engine/Private/CapturePin.h b/Engine/Source/Runtime/Engine/Private/CapturePin.h
--- a/Engine/Source/Runtime/Engine/Private/CapturePin.h
+++ b/Engine/Source/Runtime/Engine/Private/CapturePin.h
@@ -3,8 +3,8 @@
 /*=============================================================================
 	CapturePin.h: CapturePin definition
 =============================================================================*/
-#ifndef _CAPTUREPIN_HEADER_
-#define _CAPTUREPIN_HEADER_
+
+#pragma once
 
 #if PLATFORM_WINDOWS && !UE_BUILD_MINIMAL
 
@@ -84,5 +84,3 @@ public:
 #endif //#if PLATFORM_WINDOWS
 
 
-#endif	//#ifndef _CAPTUREPIN_HEADER_
-
diff --git a/Engine/Source/Runtime/Engine/Private/CaptureSource.h b/Engine/Source/Runtime/Engine/Private/CaptureSource.h
--- a/Engine/Source/Runtime/Engine/Private/CaptureSource.h
+++ b/Engine/Source/Runtime/Engine/Private/CaptureSource.h
@@ -4,8 +4,7 @@
 	CaptureSource.h: CaptureSource definition
 =============================================================================*/
 
-#ifndef _CAPTURESOURCE_HEADER_
-#define _CAPTURESOURCE_HEADER_
+#pragma once
 
 #if PLATFORM_WINDOWS && !UE_BUILD_MINIMAL
 #pragma warning(push)
@@ -39,4 +38,3 @@ private:
 };
 #endif //#if PLATFORM_WINDOWS
 
-#endif	//#ifndef _CAPTURESOURCE_HEADER_
\ No newline at end of file
diff --git a/Engine/Source/Runtime/Engine/Private/Collision/PhysXCollision.cpp b/Engine/Source/Runtime/Engine/Private/Collision/PhysXCollision.cpp
--- a/Engine/Source/Runtime/Engine/Private/Collision/PhysXCollision.cpp
+++ b/Engine/Source/Runtime/Engine/Private/Collision/PhysXCollision.cpp
@@ -7,6 +7,7 @@
 #include "PhysicsEngine/PhysicsSettings.h"
 #include "PhysicsEngine/BodySetup.h"
 #include "PhysicalMaterials/PhysicalMaterial.h"
+#include "PhysicsEngine/PhysXSupport.h"
 #include "PhysicsFiltering.h"
 
 #if WITH_BOX2D
diff --git a/Engine/Source/Runtime/Engine/Private/Engine/PoseWatch.cpp b/Engine/Source/Runtime/Engine/Private/Engine/PoseWatch.cpp
new file mode 100644
--- /dev/null
+++ b/Engine/Source/Runtime/Engine/Private/Engine/PoseWatch.cpp
@@ -0,0 +1,9 @@
+// Copyright 1998-2016 Epic Games, Inc. All Rights Reserved.
+
+#include "EnginePrivate.h"
+#include "Engine/PoseWatch.h"
+
+UPoseWatch::UPoseWatch(const FObjectInitializer& ObjectInitializer)
+	: Super(ObjectInitializer)
+{
+}
\ No newline at end of file
diff --git a/Engine/Source/Runtime/Engine/Private/PhysicsEngine2D/Box2DIntegration.cpp b/Engine/Source/Runtime/Engine/Private/PhysicsEngine2D/Box2DIntegration.cpp
--- a/Engine/Source/Runtime/Engine/Private/PhysicsEngine2D/Box2DIntegration.cpp
+++ b/Engine/Source/Runtime/Engine/Private/PhysicsEngine2D/Box2DIntegration.cpp
@@ -2,6 +2,7 @@
 
 #include "EnginePrivate.h"
 #include "Box2DIntegration.h"
+#include "PhysicsEngine/PhysXSupport.h"
 #include "Collision/PhysicsFiltering.h"
 
 #if WITH_BOX2D
diff --git a/Engine/Source/Runtime/Engine/Private/PointLightSceneProxy.h b/Engine/Source/Runtime/Engine/Private/PointLightSceneProxy.h
--- a/Engine/Source/Runtime/Engine/Private/PointLightSceneProxy.h
+++ b/Engine/Source/Runtime/Engine/Private/PointLightSceneProxy.h
@@ -4,8 +4,8 @@
 	PointLightSceneProxy.h: Point light scene info definition.
 =============================================================================*/
 
-#ifndef __PointLightSceneProxy_H__
-#define __PointLightSceneProxy_H__
+#pragma once
+
 #include "Components/PointLightComponent.h"
 
 
@@ -178,4 +178,3 @@ public:
 	}
 };
 
-#endif
diff --git a/Engine/Source/Runtime/Engine/Private/Rendering/SubDSurface.cpp b/Engine/Source/Runtime/Engine/Private/Rendering/SubDSurface.cpp
--- a/Engine/Source/Runtime/Engine/Private/Rendering/SubDSurface.cpp
+++ b/Engine/Source/Runtime/Engine/Private/Rendering/SubDSurface.cpp
@@ -2,7 +2,6 @@
 
 #include "EnginePrivate.h"
 #include "SeparableSSS.h"
-#include "EnginePrivate.h"
 #include "RendererInterface.h"
 #include "Engine/SubDSurface.h"
 
diff --git a/Engine/Source/Runtime/Engine/Private/SceneInterface.cpp b/Engine/Source/Runtime/Engine/Private/SceneInterface.cpp
new file mode 100644
--- /dev/null
+++ b/Engine/Source/Runtime/Engine/Private/SceneInterface.cpp
@@ -0,0 +1,9 @@
+// Copyright 1998-2016 Epic Games, Inc. All Rights Reserved.
+
+#include "EnginePrivate.h"
+#include "SceneInterface.h"
+
+TArray<FPrimitiveComponentId> FSceneInterface::GetScenePrimitiveComponentIds() const
+{ 
+	return TArray<FPrimitiveComponentId>(); 
+}
diff --git a/Engine/Source/Runtime/Engine/Public/AudioThread.h b/Engine/Source/Runtime/Engine/Public/AudioThread.h
--- a/Engine/Source/Runtime/Engine/Public/AudioThread.h
+++ b/Engine/Source/Runtime/Engine/Public/AudioThread.h
@@ -4,8 +4,7 @@
 	AudioThread.h: Rendering thread definitions.
 =============================================================================*/
 
-#ifndef __AudioThread_H__
-#define __AudioThread_H__
+#pragma once
 
 #include "TaskGraphInterfaces.h"
 
@@ -140,4 +139,3 @@ private:
 	mutable FGraphEventRef CompletionEvent;
 };
 
-#endif
diff --git a/Engine/Source/Runtime/Engine/Public/BatchedElements.h b/Engine/Source/Runtime/Engine/Public/BatchedElements.h
--- a/Engine/Source/Runtime/Engine/Public/BatchedElements.h
+++ b/Engine/Source/Runtime/Engine/Public/BatchedElements.h
@@ -4,8 +4,7 @@
 	BatchedElements.h: Batched element rendering.
 =============================================================================*/
 
-#ifndef _INC_BATCHEDELEMENTS
-#define _INC_BATCHEDELEMENTS
+#pragma once
 
 #include "HitProxies.h"
 #include "StaticBoundShaderState.h"
@@ -360,4 +359,3 @@ private:
 		) const;
 };
 
-#endif
diff --git a/Engine/Source/Runtime/Engine/Public/ChartCreation.h b/Engine/Source/Runtime/Engine/Public/ChartCreation.h
--- a/Engine/Source/Runtime/Engine/Public/ChartCreation.h
+++ b/Engine/Source/Runtime/Engine/Public/ChartCreation.h
@@ -5,8 +5,7 @@
  *
  */
 
-#ifndef _CHART_CREATION_INC_
-#define _CHART_CREATION_INC_
+#pragma once
 
 #define DO_CHARTING 1
 #if DO_CHARTING
@@ -111,5 +110,3 @@ struct FHitchChartEntry
 extern ENGINE_API FHitchChartEntry GHitchChart[ STAT_FPSChart_LastHitchBucketStat - STAT_FPSChart_FirstHitchStat ];
 
 #endif // DO_CHARTING
-
-#endif // _CHART_CREATION_INC_
\ No newline at end of file
diff --git a/Engine/Source/Runtime/Engine/Public/DebugRenderSceneProxy.h b/Engine/Source/Runtime/Engine/Public/DebugRenderSceneProxy.h
--- a/Engine/Source/Runtime/Engine/Public/DebugRenderSceneProxy.h
+++ b/Engine/Source/Runtime/Engine/Public/DebugRenderSceneProxy.h
@@ -2,12 +2,9 @@
 
 /*=============================================================================
 	DebugRenderSceneProxy.h: Useful scene proxy for rendering non performance-critical information.
-
-
 =============================================================================*/
 
-#ifndef _INC_DEBUGRENDERSCENEPROXY
-#define _INC_DEBUGRENDERSCENEPROXY
+#pragma once
 
 #include "PrimitiveSceneProxy.h"
 #include "DynamicMeshBuilder.h"
@@ -248,4 +245,3 @@ public:
 	TWeakObjectPtr<class UMaterial> SolidMeshMaterial;
 };
 
-#endif
diff --git a/Engine/Source/Runtime/Engine/Public/EngineAnalytics.h b/Engine/Source/Runtime/Engine/Public/EngineAnalytics.h
--- a/Engine/Source/Runtime/Engine/Public/EngineAnalytics.h
+++ b/Engine/Source/Runtime/Engine/Public/EngineAnalytics.h
@@ -41,7 +41,7 @@ private:
 	static bool bIsInitialized;
 	static bool bIsEditorRun;
 	static bool bIsGameRun;
-	static TSharedPtr<IAnalyticsProviderET> Analytics;
+	static ENGINE_API TSharedPtr<IAnalyticsProviderET> Analytics;
 	static TSharedPtr<FEngineSessionManager> SessionManager;
 };
 
diff --git a/Engine/Source/Runtime/Engine/Public/EngineDefines.h b/Engine/Source/Runtime/Engine/Public/EngineDefines.h
--- a/Engine/Source/Runtime/Engine/Public/EngineDefines.h
+++ b/Engine/Source/Runtime/Engine/Public/EngineDefines.h
@@ -54,10 +54,6 @@
 	#define ENABLE_VISUAL_LOG (PLATFORM_DESKTOP && !NO_LOGGING && !(UE_BUILD_SHIPPING || UE_BUILD_TEST))
 #endif
 
-#ifndef WITH_FIXED_AREA_ENTERING_COST
-	#define WITH_FIXED_AREA_ENTERING_COST 1
-#endif // WITH_FIXED_AREA_ENTERING_COST
-
 // If set, recast will use async workers for rebuilding tiles in runtime
 // All access to tile data must be guarded with critical sections
 #ifndef RECAST_ASYNC_REBUILDING
diff --git a/Engine/Source/Runtime/Engine/Public/EngineUtils.h b/Engine/Source/Runtime/Engine/Public/EngineUtils.h
--- a/Engine/Source/Runtime/Engine/Public/EngineUtils.h
+++ b/Engine/Source/Runtime/Engine/Public/EngineUtils.h
@@ -4,8 +4,7 @@
 	Engine.h: Unreal engine public header file.
 =============================================================================*/
 
-#ifndef _INC_ENGINE_UTILS
-#define _INC_ENGINE_UTILS
+#pragma once
 
 #include "Engine/Brush.h"
 #include "Engine/Level.h"
@@ -803,4 +802,3 @@ public:
 	}
 };
 
-#endif // _INC_ENGINE_UTILS
diff --git a/Engine/Source/Runtime/Engine/Public/GenericOctree.inl b/Engine/Source/Runtime/Engine/Public/GenericOctree.inl
--- a/Engine/Source/Runtime/Engine/Public/GenericOctree.inl
+++ b/Engine/Source/Runtime/Engine/Public/GenericOctree.inl
@@ -4,8 +4,7 @@
 	GenericOctree.inl: Generic octree implementation.
 =============================================================================*/
 
-#ifndef __GENERICOCTREE_INL__
-#define __GENERICOCTREE_INL__
+#pragma once
 
 ENGINE_API DECLARE_LOG_CATEGORY_EXTERN(LogGenericOctree, Log, All);
 
@@ -400,5 +399,3 @@ void TOctree<ElementType,OctreeSemantics>::ApplyOffset(const FVector& InOffset)
 	// Saved nodes will be deleted on scope exit
 }
 
-
-#endif
diff --git a/Engine/Source/Runtime/Engine/Public/HModel.h b/Engine/Source/Runtime/Engine/Public/HModel.h
--- a/Engine/Source/Runtime/Engine/Public/HModel.h
+++ b/Engine/Source/Runtime/Engine/Public/HModel.h
@@ -4,8 +4,7 @@
 	HModel.h: HModel definition.
 =============================================================================*/
 
-#ifndef _INC_HMODEL
-#define _INC_HMODEL
+#pragma once
 
 class UModelComponent;
 
@@ -51,4 +50,3 @@ private:
 	UModel* Model;
 };
 
-#endif
diff --git a/Engine/Source/Runtime/Engine/Public/InterpolationHitProxy.h b/Engine/Source/Runtime/Engine/Public/InterpolationHitProxy.h
--- a/Engine/Source/Runtime/Engine/Public/InterpolationHitProxy.h
+++ b/Engine/Source/Runtime/Engine/Public/InterpolationHitProxy.h
@@ -4,8 +4,7 @@
 	InterpolationHitProxy.h
 =============================================================================*/
 
-#ifndef _INC_INTERPOLATIONHITPROXY
-#define _INC_INTERPOLATIONHITPROXY
+#pragma once
 
 #include "Interpolation.h"
 
@@ -98,5 +97,3 @@ struct HInterpTrackKeyHandleProxy : public HHitProxy
 };
 
 
-
-#endif
diff --git a/Engine/Source/Runtime/Engine/Public/MallocProfilerEx.h b/Engine/Source/Runtime/Engine/Public/MallocProfilerEx.h
--- a/Engine/Source/Runtime/Engine/Public/MallocProfilerEx.h
+++ b/Engine/Source/Runtime/Engine/Public/MallocProfilerEx.h
@@ -4,8 +4,7 @@
 	MallocProfilerEx.h: Extended memory profiling support.
 =============================================================================*/
 
-#ifndef INC_MALLOCPROFILEREX_H
-#define INC_MALLOCPROFILEREX_H
+#pragma once
 
 #if USE_MALLOC_PROFILER
 
@@ -40,4 +39,3 @@ public:
 
 #endif // USE_MALLOC_PROFILER
 
-#endif // INC_MALLOCPROFILEREX_H
diff --git a/Engine/Source/Runtime/Engine/Public/ParticleEmitterInstances.h b/Engine/Source/Runtime/Engine/Public/ParticleEmitterInstances.h
--- a/Engine/Source/Runtime/Engine/Public/ParticleEmitterInstances.h
+++ b/Engine/Source/Runtime/Engine/Public/ParticleEmitterInstances.h
@@ -1253,6 +1253,7 @@ struct FParticleTrailsEmitterInstance_Base : public FParticleEmitterInstance
 	{
 #if !(UE_BUILD_SHIPPING || UE_BUILD_TEST)
 		for (uint32 TrailIdx = 0; TrailIdx < 128; TrailIdx++)
+		{
 			if (CurrentEndIndices[TrailIdx] != INDEX_NONE)
 			{
 				DECLARE_PARTICLE_PTR(EndParticle, ParticleData + ParticleStride * CurrentEndIndices[TrailIdx]);
diff --git a/Engine/Source/Runtime/Engine/Public/SceneInterface.h b/Engine/Source/Runtime/Engine/Public/SceneInterface.h
--- a/Engine/Source/Runtime/Engine/Public/SceneInterface.h
+++ b/Engine/Source/Runtime/Engine/Public/SceneInterface.h
@@ -383,7 +383,7 @@ public:
 	/**
 	 * Returns the FPrimitiveComponentId for all primitives in the scene
 	 */
-	virtual TArray<FPrimitiveComponentId> GetScenePrimitiveComponentIds() const { return TArray<FPrimitiveComponentId>(); }
+	virtual ENGINE_API TArray<FPrimitiveComponentId> GetScenePrimitiveComponentIds() const;
 
 protected:
 	virtual ~FSceneInterface() {}
diff --git a/Engine/Source/Runtime/Engine/Public/TileRendering.h b/Engine/Source/Runtime/Engine/Public/TileRendering.h
--- a/Engine/Source/Runtime/Engine/Public/TileRendering.h
+++ b/Engine/Source/Runtime/Engine/Public/TileRendering.h
@@ -4,8 +4,7 @@
 	TileRendering.h: Simple tile rendering implementation.
 =============================================================================*/
 
-#ifndef _INC_TILERENDERING
-#define _INC_TILERENDERING
+#pragma once
 
 class FTileRenderer
 {
@@ -29,4 +28,3 @@ private:
 	FTileRenderer() {}
 };
 
-#endif
diff --git a/Engine/Source/Runtime/Engine/Public/TriangleRendering.h b/Engine/Source/Runtime/Engine/Public/TriangleRendering.h
--- a/Engine/Source/Runtime/Engine/Public/TriangleRendering.h
+++ b/Engine/Source/Runtime/Engine/Public/TriangleRendering.h
@@ -4,8 +4,7 @@
 	TriangleRendering.h: Simple triangle rendering implementation.
 =============================================================================*/
 
-#ifndef _INC_TRIANGLERENDERING
-#define _INC_TRIANGLERENDERING
+#pragma once
 
 class FTriangleRenderer
 {
@@ -27,4 +26,3 @@ private:
 	FTriangleRenderer() {}
 };
 
-#endif
diff --git a/Engine/Source/Runtime/HeadMountedDisplay/Public/HeadMountedDisplay.h b/Engine/Source/Runtime/HeadMountedDisplay/Public/HeadMountedDisplay.h
--- a/Engine/Source/Runtime/HeadMountedDisplay/Public/HeadMountedDisplay.h
+++ b/Engine/Source/Runtime/HeadMountedDisplay/Public/HeadMountedDisplay.h
@@ -1,9 +1,7 @@
 // Copyright 1998-2016 Epic Games, Inc. All Rights Reserved.
 
-#ifndef __HEADMOUNTEDDISPLAY_H__
-#define __HEADMOUNTEDDISPLAY_H__
+#pragma once
 
 #include "IHeadMountedDisplayModule.h"
 #include "IHeadMountedDisplay.h"
 
-#endif
diff --git a/Engine/Source/Runtime/ImageCore/Private/ImageCorePCH.h b/Engine/Source/Runtime/ImageCore/Private/ImageCorePCH.h
--- a/Engine/Source/Runtime/ImageCore/Private/ImageCorePCH.h
+++ b/Engine/Source/Runtime/ImageCore/Private/ImageCorePCH.h
@@ -1,10 +1,10 @@
 // Copyright 1998-2016 Epic Games, Inc. All Rights Reserved.
 
-#include "ImageCore.h"
-
+#pragma once
 
 /* Dependencies
  *****************************************************************************/
 
 #include "Core.h"
 #include "ModuleManager.h"
+#include "ImageCore.h"
diff --git a/Engine/Source/Runtime/InputDevice/Public/InputDevice.h b/Engine/Source/Runtime/InputDevice/Public/InputDevice.h
--- a/Engine/Source/Runtime/InputDevice/Public/InputDevice.h
+++ b/Engine/Source/Runtime/InputDevice/Public/InputDevice.h
@@ -1,11 +1,9 @@
 // Copyright 1998-2016 Epic Games, Inc. All Rights Reserved.
 
-#ifndef __INPUTDEVICE_H__
-#define __INPUTDEVICE_H__
+#pragma once
 
 #include "Core.h"
 #include "Engine.h"
 #include "IInputDevice.h"
 #include "IInputDeviceModule.h"
 
-#endif
diff --git a/Engine/Source/Runtime/Json/Public/Json.h b/Engine/Source/Runtime/Json/Public/Json.h
--- a/Engine/Source/Runtime/Json/Public/Json.h
+++ b/Engine/Source/Runtime/Json/Public/Json.h
@@ -12,8 +12,7 @@
 /* Public includes
  *****************************************************************************/
 
-JSON_API DECLARE_LOG_CATEGORY_EXTERN(LogJson, Log, All);
-
+#include "JsonGlobals.h"
 
 #include "JsonPrintPolicy.h"
 #include "PrettyJsonPrintPolicy.h"
diff --git a/Engine/Source/Runtime/Json/Public/JsonGlobals.h b/Engine/Source/Runtime/Json/Public/JsonGlobals.h
new file mode 100644
--- /dev/null
+++ b/Engine/Source/Runtime/Json/Public/JsonGlobals.h
@@ -0,0 +1,5 @@
+// Copyright 1998-2016 Epic Games, Inc. All Rights Reserved.
+
+#pragma once
+
+JSON_API DECLARE_LOG_CATEGORY_EXTERN(LogJson, Log, All);
diff --git a/Engine/Source/Runtime/Landscape/Landscape.Build.cs b/Engine/Source/Runtime/Landscape/Landscape.Build.cs
--- a/Engine/Source/Runtime/Landscape/Landscape.Build.cs
+++ b/Engine/Source/Runtime/Landscape/Landscape.Build.cs
@@ -10,6 +10,7 @@ public class Landscape : ModuleRules
 		PrivateIncludePaths.AddRange(
 			new string[] {
 				"Runtime/Engine/Private", // for Engine/Private/Collision/PhysXCollision.h
+				"Runtime/Landscape/Private"
 			}
 		);
 
diff --git a/Engine/Source/Runtime/Landscape/Private/Landscape.cpp b/Engine/Source/Runtime/Landscape/Private/Landscape.cpp
--- a/Engine/Source/Runtime/Landscape/Private/Landscape.cpp
+++ b/Engine/Source/Runtime/Landscape/Private/Landscape.cpp
@@ -4,6 +4,7 @@
 Landscape.cpp: Terrain rendering
 =============================================================================*/
 
+#include "LandscapePrivatePCH.h"
 #include "Landscape.h"
 #include "LandscapeStreamingProxy.h"
 #include "EditorSupportDelegates.h"
diff --git a/Engine/Source/Runtime/Landscape/Private/LandscapeBlueprintSupport.cpp b/Engine/Source/Runtime/Landscape/Private/LandscapeBlueprintSupport.cpp
--- a/Engine/Source/Runtime/Landscape/Private/LandscapeBlueprintSupport.cpp
+++ b/Engine/Source/Runtime/Landscape/Private/LandscapeBlueprintSupport.cpp
@@ -4,6 +4,7 @@
 	LandscapeBlueprintSupport.cpp: Landscape blueprint functions
   =============================================================================*/
 
+#include "LandscapePrivatePCH.h"
 #include "Landscape.h"
 #include "LandscapeSplineRaster.h"
 #include "LandscapeProxy.h"
diff --git a/Engine/Source/Runtime/Landscape/Private/LandscapeGizmoActor.cpp b/Engine/Source/Runtime/Landscape/Private/LandscapeGizmoActor.cpp
--- a/Engine/Source/Runtime/Landscape/Private/LandscapeGizmoActor.cpp
+++ b/Engine/Source/Runtime/Landscape/Private/LandscapeGizmoActor.cpp
@@ -1,6 +1,6 @@
 // Copyright 1998-2016 Epic Games, Inc. All Rights Reserved.
 
-
+#include "LandscapePrivatePCH.h"
 #include "Landscape.h"
 #include "LandscapeDataAccess.h"
 #include "LandscapeRender.h"
diff --git a/Engine/Source/Runtime/Landscape/Private/LandscapePrivatePCH.h b/Engine/Source/Runtime/Landscape/Private/LandscapePrivatePCH.h
new file mode 100644
--- /dev/null
+++ b/Engine/Source/Runtime/Landscape/Private/LandscapePrivatePCH.h
@@ -0,0 +1,5 @@
+// Copyright 1998-2016 Epic Games, Inc. All Rights Reserved.
+
+#pragma once
+
+#include "Landscape.h"
diff --git a/Engine/Source/Runtime/Landscape/Public/LandscapeEdit.h b/Engine/Source/Runtime/Landscape/Public/LandscapeEdit.h
--- a/Engine/Source/Runtime/Landscape/Public/LandscapeEdit.h
+++ b/Engine/Source/Runtime/Landscape/Public/LandscapeEdit.h
@@ -4,8 +4,7 @@
 LandscapeEdit.h: Classes for the editor to access to Landscape data
 =============================================================================*/
 
-#ifndef _LANDSCAPEEDIT_H
-#define _LANDSCAPEEDIT_H
+#pragma once
 
 #define MAX_LANDSCAPE_LOD_DISTANCE_FACTOR 10.f
 
@@ -237,4 +236,3 @@ void FLandscapeEditDataInterface::ShrinkData(TArray<T>& Data, int32 OldMinX, int
 
 #endif
 
-#endif // _LANDSCAPEEDIT_H
diff --git a/Engine/Source/Runtime/Landscape/Public/LandscapeRenderMobile.h b/Engine/Source/Runtime/Landscape/Public/LandscapeRenderMobile.h
--- a/Engine/Source/Runtime/Landscape/Public/LandscapeRenderMobile.h
+++ b/Engine/Source/Runtime/Landscape/Public/LandscapeRenderMobile.h
@@ -4,8 +4,7 @@
 LandscapeRenderMobile.h: Mobile landscape rendering
 =============================================================================*/
 
-#ifndef _LANDSCAPERENDERMOBILE_H
-#define _LANDSCAPERENDERMOBILE_H
+#pragma once
 
 #include "Landscape.h"
 
@@ -132,5 +131,3 @@ public:
 
 	friend class FLandscapeVertexBufferMobile;
 };
-
-#endif // _LANDSCAPERENDERMOBILE_H
\ No newline at end of file
diff --git a/Engine/Source/Runtime/Launch/Private/LaunchEngineLoop.cpp b/Engine/Source/Runtime/Launch/Private/LaunchEngineLoop.cpp
--- a/Engine/Source/Runtime/Launch/Private/LaunchEngineLoop.cpp
+++ b/Engine/Source/Runtime/Launch/Private/LaunchEngineLoop.cpp
@@ -3322,6 +3322,8 @@ bool FEngineLoop::AppInit( )
 	// Print compiler version info
 #if defined(__clang__)
 	UE_LOG(LogInit, Log, TEXT("Compiled with Clang: %s"), ANSI_TO_TCHAR( __clang_version__ ) );
+#elif defined(__INTEL_COMPILER)
+	UE_LOG(LogInit, Log, TEXT("Compiled with ICL: %d"), __INTEL_COMPILER);
 #elif defined( _MSC_VER )
 	#ifndef __INTELLISENSE__	// Intellisense compiler doesn't support _MSC_FULL_VER
 	{
diff --git a/Engine/Source/Runtime/MoviePlayer/Private/DefaultGameMoviePlayer.cpp b/Engine/Source/Runtime/MoviePlayer/Private/DefaultGameMoviePlayer.cpp
--- a/Engine/Source/Runtime/MoviePlayer/Private/DefaultGameMoviePlayer.cpp
+++ b/Engine/Source/Runtime/MoviePlayer/Private/DefaultGameMoviePlayer.cpp
@@ -1,5 +1,6 @@
 // Copyright 1998-2016 Epic Games, Inc. All Rights Reserved.
 
+#include "MoviePlayerPrivatePCH.h"
 #include "MoviePlayer.h"
 
 #include "Engine.h"
diff --git a/Engine/Source/Runtime/MoviePlayer/Private/MoviePlayer.cpp b/Engine/Source/Runtime/MoviePlayer/Private/MoviePlayer.cpp
--- a/Engine/Source/Runtime/MoviePlayer/Private/MoviePlayer.cpp
+++ b/Engine/Source/Runtime/MoviePlayer/Private/MoviePlayer.cpp
@@ -1,6 +1,6 @@
-
 // Copyright 1998-2016 Epic Games, Inc. All Rights Reserved.
 
+#include "MoviePlayerPrivatePCH.h"
 #include "MoviePlayer.h"
 
 #include "Engine.h"
diff --git a/Engine/Source/Runtime/MoviePlayer/Private/MoviePlayerPrivatePCH.h b/Engine/Source/Runtime/MoviePlayer/Private/MoviePlayerPrivatePCH.h
new file mode 100644
--- /dev/null
+++ b/Engine/Source/Runtime/MoviePlayer/Private/MoviePlayerPrivatePCH.h
@@ -0,0 +1,7 @@
+// Copyright 1998-2016 Epic Games, Inc. All Rights Reserved.
+
+#pragma once
+
+#include "Engine.h"
+#include "ModuleInterface.h"
+#include "SlateBasics.h"
diff --git a/Engine/Source/Runtime/MoviePlayer/Private/MoviePlayerSettings.cpp b/Engine/Source/Runtime/MoviePlayer/Private/MoviePlayerSettings.cpp
--- a/Engine/Source/Runtime/MoviePlayer/Private/MoviePlayerSettings.cpp
+++ b/Engine/Source/Runtime/MoviePlayer/Private/MoviePlayerSettings.cpp
@@ -1,5 +1,6 @@
 // Copyright 1998-2016 Epic Games, Inc. All Rights Reserved.
 
+#include "MoviePlayerPrivatePCH.h"
 #include "MoviePlayer.h"
 #include "MoviePlayerSettings.h"
 
diff --git a/Engine/Source/Runtime/MovieScene/Private/MovieScenePrivatePCH.h b/Engine/Source/Runtime/MovieScene/Private/MovieScenePrivatePCH.h
--- a/Engine/Source/Runtime/MovieScene/Private/MovieScenePrivatePCH.h
+++ b/Engine/Source/Runtime/MovieScene/Private/MovieScenePrivatePCH.h
@@ -1,7 +1,6 @@
 // Copyright 1998-2016 Epic Games, Inc. All Rights Reserved.
 
-#ifndef MOVIE_SCENE_CORE_PRIVATE
-#define MOVIE_SCENE_CORE_PRIVATE
+#pragma once
 
 #include "Core.h"
 #include "Engine.h"
@@ -16,4 +15,3 @@
 #include "MovieSceneSequence.h"
 #include "MovieSceneCommonHelpers.h"
 
-#endif
\ No newline at end of file
diff --git a/Engine/Source/Runtime/Navmesh/Public/Detour/DetourNavMeshQuery.h b/Engine/Source/Runtime/Navmesh/Public/Detour/DetourNavMeshQuery.h
--- a/Engine/Source/Runtime/Navmesh/Public/Detour/DetourNavMeshQuery.h
+++ b/Engine/Source/Runtime/Navmesh/Public/Detour/DetourNavMeshQuery.h
@@ -27,9 +27,7 @@
 #include "DetourCommon.h"
 
 //@UE4 BEGIN
-#ifndef WITH_FIXED_AREA_ENTERING_COST
 #define WITH_FIXED_AREA_ENTERING_COST 1
-#endif
 
 #ifndef FLT_MAX
 #include <float.h>
@@ -59,7 +57,7 @@ struct NAVMESH_API dtQuerySpecialLinkFilter
 };
 
 // [UE4: moved all filter variables to struct, DO NOT mess with virtual functions here!]
-struct dtQueryFilterData
+struct NAVMESH_API dtQueryFilterData
 {
 	float m_areaCost[DT_MAX_AREAS];		///< Cost per area type. (Used by default implementation.)
 #if WITH_FIXED_AREA_ENTERING_COST
diff --git a/Engine/Source/Runtime/OpenGLDrv/Private/OpenGLDevice.cpp b/Engine/Source/Runtime/OpenGLDrv/Private/OpenGLDevice.cpp
--- a/Engine/Source/Runtime/OpenGLDrv/Private/OpenGLDevice.cpp
+++ b/Engine/Source/Runtime/OpenGLDrv/Private/OpenGLDevice.cpp
@@ -554,7 +554,7 @@ static void InitRHICapabilitiesForGL()
 #endif
 
 	// GL vendor and version information.
-#if !defined(__GNUC__) && !defined(__clang__) && !(defined(_MSC_VER) && _MSC_VER >= 1900)
+#if !defined(__GNUC__) && !defined(__clang__) && !defined(__INTEL_COMPILER) && !(defined(_MSC_VER) && _MSC_VER >= 1900)
 	#define LOG_GL_STRING(StringEnum) UE_LOG(LogRHI, Log, TEXT("  ") ## TEXT(#StringEnum) ## TEXT(": %s"), ANSI_TO_TCHAR((const ANSICHAR*)glGetString(StringEnum)))
 #else
 	#define LOG_GL_STRING(StringEnum) UE_LOG(LogRHI, Log, TEXT("  " #StringEnum ": %s"), ANSI_TO_TCHAR((const ANSICHAR*)glGetString(StringEnum)))
@@ -625,7 +625,7 @@ static void InitRHICapabilitiesForGL()
 	FOpenGL::InitDebugContext();
 
 	// Log and get various limits.
-#if !defined(__GNUC__) && !defined(__clang__) && !(defined(_MSC_VER) && _MSC_VER >= 1900)
+#if !defined(__GNUC__) && !defined(__clang__) && !defined(__INTEL_COMPILER) && !(defined(_MSC_VER) && _MSC_VER >= 1900)
 #define LOG_AND_GET_GL_INT_TEMP(IntEnum,Default) GLint Value_##IntEnum = Default; if (IntEnum) {glGetIntegerv(IntEnum, &Value_##IntEnum); glGetError();} else {Value_##IntEnum = Default;} UE_LOG(LogRHI, Log, TEXT("  ") ## TEXT(#IntEnum) ## TEXT(": %d"), Value_##IntEnum)
 #else
 #define LOG_AND_GET_GL_INT_TEMP(IntEnum,Default) GLint Value_##IntEnum = Default; if (IntEnum) {glGetIntegerv(IntEnum, &Value_##IntEnum); glGetError();} else {Value_##IntEnum = Default;} UE_LOG(LogRHI, Log, TEXT("  " #IntEnum ": %d"), Value_##IntEnum)
diff --git a/Engine/Source/Runtime/OpenGLDrv/Private/OpenGLDrvPrivate.h b/Engine/Source/Runtime/OpenGLDrv/Private/OpenGLDrvPrivate.h
--- a/Engine/Source/Runtime/OpenGLDrv/Private/OpenGLDrvPrivate.h
+++ b/Engine/Source/Runtime/OpenGLDrv/Private/OpenGLDrvPrivate.h
@@ -4,8 +4,7 @@
 	OpenGLDrvPrivate.h: Private OpenGL RHI definitions.
 =============================================================================*/
 
-#ifndef __OPENGLDRVPRIVATE_H__
-#define __OPENGLDRVPRIVATE_H__
+#pragma once
 
 #include "OpenGLDrv.h"
 
@@ -547,4 +546,3 @@ inline bool OpenGLShaderPlatformSeparable(const EShaderPlatform InShaderPlatform
 	}
 }
 
-#endif // __OPENGLDRVPRIVATE_H__
diff --git a/Engine/Source/Runtime/OpenGLDrv/Public/OpenGLResources.h b/Engine/Source/Runtime/OpenGLDrv/Public/OpenGLResources.h
--- a/Engine/Source/Runtime/OpenGLDrv/Public/OpenGLResources.h
+++ b/Engine/Source/Runtime/OpenGLDrv/Public/OpenGLResources.h
@@ -757,7 +757,7 @@ inline GLenum GetOpenGLTargetFromRHITexture(FRHITexture* Texture)
 }
 
 
-class FOpenGLTextureBase
+class OPENGLDRV_API FOpenGLTextureBase
 {
 protected:
 	class FOpenGLDynamicRHI* OpenGLRHI;
@@ -1372,7 +1372,7 @@ protected:
 void OPENGLDRV_API OpenGLTextureDeleted(FRHITexture* Texture);
 void OPENGLDRV_API OpenGLTextureAllocated( FRHITexture* Texture , uint32 Flags);
 
-extern void ReleaseOpenGLFramebuffers(FOpenGLDynamicRHI* Device, FTextureRHIParamRef TextureRHI);
+void OPENGLDRV_API ReleaseOpenGLFramebuffers(FOpenGLDynamicRHI* Device, FTextureRHIParamRef TextureRHI);
 
 /** A OpenGL event query resource. */
 class FOpenGLEventQuery : public FRenderResource
diff --git a/Engine/Source/Runtime/Projects/Private/PluginManager.cpp b/Engine/Source/Runtime/Projects/Private/PluginManager.cpp
--- a/Engine/Source/Runtime/Projects/Private/PluginManager.cpp
+++ b/Engine/Source/Runtime/Projects/Private/PluginManager.cpp
@@ -151,30 +151,31 @@ FPluginManager::~FPluginManager()
 void FPluginManager::RefreshPluginsList()
 {
 	// Read a new list of all plugins
-	TArray<TSharedRef<FPlugin>> NewPlugins;
+	TMap<FString, TSharedRef<FPlugin>> NewPlugins;
 	ReadAllPlugins(NewPlugins, PluginDiscoveryPaths);
 
 	// Build a list of filenames for plugins which are enabled, and remove the rest
 	TArray<FString> EnabledPluginFileNames;
-	for(int32 Idx = 0; Idx < AllPlugins.Num(); Idx++)
+	for(TMap<FString, TSharedRef<FPlugin>>::TIterator Iter(AllPlugins); Iter; ++Iter)
 	{
-		const TSharedRef<FPlugin>& Plugin = AllPlugins[Idx];
+		const TSharedRef<FPlugin>& Plugin = Iter.Value();
 		if(Plugin->bEnabled)
 		{
 			EnabledPluginFileNames.Add(Plugin->FileName);
 		}
 		else
 		{
-			AllPlugins.RemoveAt(Idx--);
+			Iter.RemoveCurrent();
 		}
 	}
 
 	// Add all the plugins which aren't already enabled
-	for(TSharedRef<FPlugin>& NewPlugin: NewPlugins)
+	for(const TPair<FString, TSharedRef<FPlugin>>& NewPluginPair: NewPlugins)
 	{
+		const TSharedRef<FPlugin>& NewPlugin = NewPluginPair.Value;
 		if(!EnabledPluginFileNames.Contains(NewPlugin->FileName))
 		{
-			AllPlugins.Add(NewPlugin);
+			AllPlugins.Add(NewPlugin->GetName(), NewPlugin);
 		}
 	}
 }
@@ -187,13 +188,14 @@ void FPluginManager::DiscoverAllPlugins()
 	ReadAllPlugins(AllPlugins, PluginDiscoveryPaths);
 }
 
-void FPluginManager::ReadAllPlugins(TArray<TSharedRef<FPlugin>>& Plugins, const TSet<FString>& ExtraSearchPaths)
+void FPluginManager::ReadAllPlugins(TMap<FString, TSharedRef<FPlugin>>& Plugins, const TSet<FString>& ExtraSearchPaths)
 {
 #if (WITH_ENGINE && !IS_PROGRAM) || WITH_PLUGIN_SUPPORT
 	// Find "built-in" plugins.  That is, plugins situated right within the Engine directory.
 	ReadPluginsInDirectory(FPaths::EnginePluginsDir(), EPluginLoadedFrom::Engine, Plugins);
 
-	// Find plugins in the game project directory (<MyGameProject>/Plugins)
+	// Find plugins in the game project directory (<MyGameProject>/Plugins). If there are any engine plugins matching the name of a game plugin,
+	// assume that the game plugin version is preferred.
 	if( FApp::HasGameName() )
 	{
 		ReadPluginsInDirectory(FPaths::GamePluginsDir(), EPluginLoadedFrom::GameProject, Plugins);
@@ -206,7 +208,7 @@ void FPluginManager::ReadAllPlugins(TArray<TSharedRef<FPlugin>>& Plugins, const
 #endif
 }
 
-void FPluginManager::ReadPluginsInDirectory(const FString& PluginsDirectory, const EPluginLoadedFrom LoadedFrom, TArray<TSharedRef<FPlugin>>& Plugins)
+void FPluginManager::ReadPluginsInDirectory(const FString& PluginsDirectory, const EPluginLoadedFrom LoadedFrom, TMap<FString, TSharedRef<FPlugin>>& Plugins)
 {
 	// Make sure the directory even exists
 	if(FPlatformFileManager::Get().GetPlatformFile().DirectoryExists(*PluginsDirectory))
@@ -225,7 +227,20 @@ void FPluginManager::ReadPluginsInDirectory(const FString& PluginsDirectory, con
 				FString FullPath = FPaths::ConvertRelativePathToFull(FileName);
 				UE_LOG(LogPluginManager, Log, TEXT("Loaded Plugin %s, From %s"), *Plugin->GetName(), *FullPath);
 
-				Plugins.Add(Plugin);
+				const TSharedRef<FPlugin>* ExistingPlugin = Plugins.Find(Plugin->GetName());
+				if (ExistingPlugin == nullptr)
+				{
+					Plugins.Add(Plugin->GetName(), Plugin);
+				}
+				else if ((*ExistingPlugin)->LoadedFrom == EPluginLoadedFrom::Engine && LoadedFrom == EPluginLoadedFrom::GameProject)
+				{
+					Plugins[Plugin->GetName()] = Plugin;
+					UE_LOG(LogPluginManager, Log, TEXT("Replacing engine version of '%s' plugin with game version"), *Plugin->GetName());
+				}
+				else if( (*ExistingPlugin)->LoadedFrom != EPluginLoadedFrom::GameProject || LoadedFrom != EPluginLoadedFrom::Engine)
+				{
+					UE_LOG(LogPluginManager, Warning, TEXT("Plugin '%s' exists at '%s' and '%s' - second location will be ignored"), *Plugin->GetName(), *(*ExistingPlugin)->FileName, *Plugin->FileName);
+				}
 			}
 			else
 			{
@@ -364,8 +379,9 @@ bool FPluginManager::ConfigureEnabledPlugins()
 		AllEnabledPlugins.Append(MoveTemp(EnabledPluginNames));
 
 		// Enable all the plugins by name
-		for (const TSharedRef< FPlugin > Plugin : AllPlugins)
+		for (const TPair<FString, TSharedRef< FPlugin >> PluginPair : AllPlugins)
 		{
+			const TSharedRef<FPlugin>& Plugin = PluginPair.Value;
 			if (AllEnabledPlugins.Contains(Plugin->Name))
 			{
 #if IS_PROGRAM
@@ -427,8 +443,9 @@ bool FPluginManager::ConfigureEnabledPlugins()
 		// If we made it here, we have all the required plugins
 		bHaveAllRequiredPlugins = true;
 
-		for(const TSharedRef<FPlugin>& Plugin: AllPlugins)
+		for(const TPair<FString, TSharedRef<FPlugin>>& PluginPair: AllPlugins)
 		{
+			const TSharedRef<FPlugin>& Plugin = PluginPair.Value;
 			if (Plugin->bEnabled)
 			{
 				// Add the plugin binaries directory
@@ -518,16 +535,15 @@ bool FPluginManager::ConfigureEnabledPlugins()
 
 TSharedPtr<FPlugin> FPluginManager::FindPluginInstance(const FString& Name)
 {
-	TSharedPtr<FPlugin> Result;
-	for(const TSharedRef<FPlugin>& Instance : AllPlugins)
+	const TSharedRef<FPlugin>* Instance = AllPlugins.Find(Name);
+	if (Instance == nullptr)
 	{
-		if(Instance->Name == Name)
-		{
-			Result = Instance;
-			break;
-		}
+		return TSharedPtr<FPlugin>();
+	}
+	else
+	{
+		return TSharedPtr<FPlugin>(*Instance);
 	}
-	return Result;
 }
 
 bool FPluginManager::LoadModulesForEnabledPlugins( const ELoadingPhase::Type LoadingPhase )
@@ -541,8 +557,9 @@ bool FPluginManager::LoadModulesForEnabledPlugins( const ELoadingPhase::Type Loa
 	FScopedSlowTask SlowTask(AllPlugins.Num());
 
 	// Load plugins!
-	for( const TSharedRef< FPlugin > Plugin : AllPlugins )
+	for( const TPair<FString, TSharedRef< FPlugin >> PluginPair : AllPlugins )
 	{
+		const TSharedRef<FPlugin> &Plugin = PluginPair.Value;
 		SlowTask.EnterProgressFrame(1);
 
 		if ( Plugin->bEnabled )
@@ -607,8 +624,9 @@ void FPluginManager::GetLocalizationPathsForEnabledPlugins( TArray<FString>& Out
 	}
 
 	// Gather the paths from all plugins that have localization targets that are loaded based on the current runtime environment
-	for (const TSharedRef<FPlugin>& Plugin : AllPlugins)
+	for (const TPair<FString, TSharedRef<FPlugin>>& PluginPair : AllPlugins)
 	{
+		const TSharedRef<FPlugin>& Plugin = PluginPair.Value;
 		if (!Plugin->bEnabled || Plugin->GetDescriptor().LocalizationTargets.Num() == 0)
 		{
 			continue;
@@ -643,9 +661,9 @@ bool FPluginManager::CheckModuleCompatibility(TArray<FString>& OutIncompatibleMo
 	}
 
 	bool bResult = true;
-	for (TArray< TSharedRef< FPlugin > >::TConstIterator Iter(AllPlugins); Iter; ++Iter)
+	for(const TPair<FString, TSharedRef<FPlugin>>& PluginPair : AllPlugins)
 	{
-		const TSharedRef< FPlugin > &Plugin = *Iter;
+		const TSharedRef< FPlugin > &Plugin = PluginPair.Value;
 		if (Plugin->bEnabled && !FModuleDescriptor::CheckModuleCompatibility(Plugin->Descriptor.Modules, Plugin->LoadedFrom == EPluginLoadedFrom::GameProject, OutIncompatibleModules))
 		{
 			bResult = false;
@@ -667,23 +685,23 @@ IPluginManager& IPluginManager::Get()
 
 TSharedPtr<IPlugin> FPluginManager::FindPlugin(const FString& Name)
 {
-	TSharedPtr<IPlugin> Plugin;
-	for(TSharedRef<FPlugin>& PossiblePlugin : AllPlugins)
+	const TSharedRef<FPlugin>* Instance = AllPlugins.Find(Name);
+	if (Instance == nullptr)
 	{
-		if(PossiblePlugin->Name == Name)
-		{
-			Plugin = PossiblePlugin;
-			break;
-		}
+		return TSharedPtr<IPlugin>();
+	}
+	else
+	{
+		return TSharedPtr<IPlugin>(*Instance);
 	}
-	return Plugin;
 }
 
 TArray<TSharedRef<IPlugin>> FPluginManager::GetEnabledPlugins()
 {
 	TArray<TSharedRef<IPlugin>> Plugins;
-	for(TSharedRef<FPlugin>& PossiblePlugin : AllPlugins)
+	for(TPair<FString, TSharedRef<FPlugin>>& PluginPair : AllPlugins)
 	{
+		TSharedRef<FPlugin>& PossiblePlugin = PluginPair.Value;
 		if(PossiblePlugin->bEnabled)
 		{
 			Plugins.Add(PossiblePlugin);
@@ -695,9 +713,9 @@ TArray<TSharedRef<IPlugin>> FPluginManager::GetEnabledPlugins()
 TArray<TSharedRef<IPlugin>> FPluginManager::GetDiscoveredPlugins()
 {
 	TArray<TSharedRef<IPlugin>> Plugins;
-	for(TSharedRef<FPlugin>& Plugin : AllPlugins)
+	for (TPair<FString, TSharedRef<FPlugin>>& PluginPair : AllPlugins)
 	{
-		Plugins.Add(Plugin);
+		Plugins.Add(PluginPair.Value);
 	}
 	return Plugins;
 }
@@ -706,9 +724,9 @@ TArray< FPluginStatus > FPluginManager::QueryStatusForAllPlugins() const
 {
 	TArray< FPluginStatus > PluginStatuses;
 
-	for( auto PluginIt( AllPlugins.CreateConstIterator() ); PluginIt; ++PluginIt )
+	for( const TPair<FString, TSharedRef<FPlugin>>& PluginPair : AllPlugins )
 	{
-		const TSharedRef< FPlugin >& Plugin = *PluginIt;
+		const TSharedRef< FPlugin >& Plugin = PluginPair.Value;
 		
 		FPluginStatus PluginStatus;
 		PluginStatus.Name = Plugin->Name;
diff --git a/Engine/Source/Runtime/Projects/Private/PluginManager.h b/Engine/Source/Runtime/Projects/Private/PluginManager.h
--- a/Engine/Source/Runtime/Projects/Private/PluginManager.h
+++ b/Engine/Source/Runtime/Projects/Private/PluginManager.h
@@ -80,10 +80,10 @@ private:
 	void DiscoverAllPlugins();
 
 	/** Reads all the plugin descriptors */
-	static void ReadAllPlugins(TArray<TSharedRef<FPlugin>>& Plugins, const TSet<FString>& ExtraSearchPaths);
+	static void ReadAllPlugins(TMap<FString, TSharedRef<FPlugin>>& Plugins, const TSet<FString>& ExtraSearchPaths);
 
 	/** Reads all the plugin descriptors from disk */
-	static void ReadPluginsInDirectory(const FString& PluginsDirectory, const EPluginLoadedFrom LoadedFrom, TArray<TSharedRef<FPlugin>>& Plugins);
+	static void ReadPluginsInDirectory(const FString& PluginsDirectory, const EPluginLoadedFrom LoadedFrom, TMap<FString, TSharedRef<FPlugin>>& Plugins);
 
 	/** Finds all the plugin descriptors underneath a given directory */
 	static void FindPluginsInDirectory(const FString& PluginsDirectory, TArray<FString>& FileNames);
@@ -99,7 +99,7 @@ private:
 
 private:
 	/** All of the plugins that we know about */
-	TArray< TSharedRef< FPlugin > > AllPlugins;
+	TMap< FString, TSharedRef< FPlugin > > AllPlugins;
 
 	/** Delegate for mounting content paths.  Bound by FPackageName code in CoreUObject, so that we can access
 	    content path mounting functionality from Core. */
diff --git a/Engine/Source/Runtime/RHI/Private/Android/AndroidDynamicRHI.cpp b/Engine/Source/Runtime/RHI/Private/Android/AndroidDynamicRHI.cpp
--- a/Engine/Source/Runtime/RHI/Private/Android/AndroidDynamicRHI.cpp
+++ b/Engine/Source/Runtime/RHI/Private/Android/AndroidDynamicRHI.cpp
@@ -1,5 +1,6 @@
 // Copyright 1998-2016 Epic Games, Inc. All Rights Reserved.
 
+#include "RHIPrivatePCH.h"
 #include "RHI.h"
 #include "ModuleManager.h"
 
diff --git a/Engine/Source/Runtime/RHI/Private/Apple/AppleDynamicRHI.cpp b/Engine/Source/Runtime/RHI/Private/Apple/AppleDynamicRHI.cpp
--- a/Engine/Source/Runtime/RHI/Private/Apple/AppleDynamicRHI.cpp
+++ b/Engine/Source/Runtime/RHI/Private/Apple/AppleDynamicRHI.cpp
@@ -1,5 +1,6 @@
 // Copyright 1998-2016 Epic Games, Inc. All Rights Reserved.
 
+#include "RHIPrivatePCH.h"
 #include "RHI.h"
 #include "ModuleManager.h"
 
diff --git a/Engine/Source/Runtime/RHI/Private/RHICommandList.cpp b/Engine/Source/Runtime/RHI/Private/RHICommandList.cpp
--- a/Engine/Source/Runtime/RHI/Private/RHICommandList.cpp
+++ b/Engine/Source/Runtime/RHI/Private/RHICommandList.cpp
@@ -1,7 +1,7 @@
 // Copyright 1998-2016 Epic Games, Inc. All Rights Reserved.
 
 
-#include "RHI.h"
+#include "RHIPrivatePCH.h"
 #include "RHICommandList.h"
 
 DECLARE_CYCLE_STAT(TEXT("Nonimmed. Command List Execute"), STAT_NonImmedCmdListExecuteTime, STATGROUP_RHICMDLIST);
diff --git a/Engine/Source/Runtime/RHI/Private/RHIPrivatePCH.h b/Engine/Source/Runtime/RHI/Private/RHIPrivatePCH.h
new file mode 100644
--- /dev/null
+++ b/Engine/Source/Runtime/RHI/Private/RHIPrivatePCH.h
@@ -0,0 +1,5 @@
+// Copyright 1998-2016 Epic Games, Inc. All Rights Reserved.
+
+#pragma once
+
+#include "RHI.h"
diff --git a/Engine/Source/Runtime/RHI/Private/RHIUtilities.cpp b/Engine/Source/Runtime/RHI/Private/RHIUtilities.cpp
--- a/Engine/Source/Runtime/RHI/Private/RHIUtilities.cpp
+++ b/Engine/Source/Runtime/RHI/Private/RHIUtilities.cpp
@@ -3,7 +3,7 @@
 /*=============================================================================
 RHIUtilities.cpp:
 =============================================================================*/
-#include "RHI.h"
+#include "RHIPrivatePCH.h"
 #include "RHIUtilities.h"
 #include "Core.h"
 
diff --git a/Engine/Source/Runtime/RHI/Public/RHI.h b/Engine/Source/Runtime/RHI/Public/RHI.h
--- a/Engine/Source/Runtime/RHI/Public/RHI.h
+++ b/Engine/Source/Runtime/RHI/Public/RHI.h
@@ -4,8 +4,7 @@
 	RHI.h: Render Hardware Interface definitions.
 =============================================================================*/
 
-#ifndef __RHI_h__
-#define __RHI_h__
+#pragma once
 
 #include "Core.h"
 #include "RHIDefinitions.h"
@@ -1253,5 +1252,3 @@ extern RHI_API void RHIExit();
 
 // RHI utility functions that depend on the RHI definitions.
 #include "RHIUtilities.h"
-
-#endif // __RHI_h__
diff --git a/Engine/Source/Runtime/RHI/Public/RHICommandList.inl b/Engine/Source/Runtime/RHI/Public/RHICommandList.inl
--- a/Engine/Source/Runtime/RHI/Public/RHICommandList.inl
+++ b/Engine/Source/Runtime/RHI/Public/RHICommandList.inl
@@ -4,8 +4,7 @@
 	RHICommandList.inl: RHI Command List inline definitions.
 =============================================================================*/
 
-#ifndef __RHICOMMANDLIST_INL__
-#define __RHICOMMANDLIST_INL__
+#pragma once
 
 FORCEINLINE_DEBUGGABLE void FRHICommandListBase::Flush()
 {
@@ -146,4 +145,3 @@ private:
 	uint32 CmdListNumCommands;
 };
 
-#endif	// __RHICOMMANDLIST_INL__
diff --git a/Engine/Source/Runtime/RHI/RHI.Build.cs b/Engine/Source/Runtime/RHI/RHI.Build.cs
--- a/Engine/Source/Runtime/RHI/RHI.Build.cs
+++ b/Engine/Source/Runtime/RHI/RHI.Build.cs
@@ -42,5 +42,7 @@ public class RHI : ModuleRules
 		{
 			PrivateIncludePathModuleNames.AddRange(new string[] { "TaskGraph" });
 		}
+
+		PrivateIncludePaths.Add("Runtime/RHI/Private");
 	}
 }
diff --git a/Engine/Source/Runtime/RenderCore/Private/RenderCore.cpp b/Engine/Source/Runtime/RenderCore/Private/RenderCore.cpp
--- a/Engine/Source/Runtime/RenderCore/Private/RenderCore.cpp
+++ b/Engine/Source/Runtime/RenderCore/Private/RenderCore.cpp
@@ -4,6 +4,7 @@
 	RenderCore.h: Render core module implementation.
 =============================================================================*/
 
+#include "RenderCorePrivatePCH.h"
 #include "RenderCore.h"
 #include "UniformBuffer.h"
 #include "ModuleManager.h"
diff --git a/Engine/Source/Runtime/RenderCore/Private/RenderCorePrivatePCH.h b/Engine/Source/Runtime/RenderCore/Private/RenderCorePrivatePCH.h
new file mode 100644
--- /dev/null
+++ b/Engine/Source/Runtime/RenderCore/Private/RenderCorePrivatePCH.h
@@ -0,0 +1,5 @@
+// Copyright 1998-2016 Epic Games, Inc. All Rights Reserved.
+
+#pragma once
+
+#include "RenderCore.h"
diff --git a/Engine/Source/Runtime/RenderCore/Private/RenderResource.cpp b/Engine/Source/Runtime/RenderCore/Private/RenderResource.cpp
--- a/Engine/Source/Runtime/RenderCore/Private/RenderResource.cpp
+++ b/Engine/Source/Runtime/RenderCore/Private/RenderResource.cpp
@@ -4,6 +4,7 @@
 	RenderResource.cpp: Render resource implementation.
 =============================================================================*/
 
+#include "RenderCorePrivatePCH.h"
 #include "RenderCore.h"
 #include "RenderResource.h"
 
diff --git a/Engine/Source/Runtime/RenderCore/Private/RenderUtils.cpp b/Engine/Source/Runtime/RenderCore/Private/RenderUtils.cpp
--- a/Engine/Source/Runtime/RenderCore/Private/RenderUtils.cpp
+++ b/Engine/Source/Runtime/RenderCore/Private/RenderUtils.cpp
@@ -1,6 +1,6 @@
 // Copyright 1998-2016 Epic Games, Inc. All Rights Reserved.
 
-
+#include "RenderCorePrivatePCH.h"
 #include "RenderCore.h"
 #include "RenderUtils.h"
 #include "RHI.h"
diff --git a/Engine/Source/Runtime/RenderCore/Public/RenderingThread.h b/Engine/Source/Runtime/RenderCore/Public/RenderingThread.h
--- a/Engine/Source/Runtime/RenderCore/Public/RenderingThread.h
+++ b/Engine/Source/Runtime/RenderCore/Public/RenderingThread.h
@@ -4,8 +4,7 @@
 	RenderingThread.h: Rendering thread definitions.
 =============================================================================*/
 
-#ifndef __RenderingThread_H__
-#define __RenderingThread_H__
+#pragma once
 
 #include "RenderCore.h"
 #include "TaskGraphInterfaces.h"
@@ -543,5 +542,3 @@ extern RENDERCORE_API void BeginCleanup(FDeferredCleanupInterface* CleanupObject
  * @return A pointer to the set of pending cleanup objects.  The called is responsible for deletion.
  */
 extern RENDERCORE_API FPendingCleanupObjects* GetPendingCleanupObjects();
-
-#endif
