Copying //UE4/Dev-Build to //UE4/Dev-Main (Source: //UE4/Dev-Build @ 3104200)

diff --git a/Engine/Source/Runtime/CoreUObject/Public/Serialization/ArchiveAsync.h b/Engine/Source/Runtime/CoreUObject/Public/Serialization/ArchiveAsync.h
new file mode 100644
--- /dev/null
+++ b/Engine/Source/Runtime/CoreUObject/Public/Serialization/ArchiveAsync.h
@@ -0,0 +1,186 @@
+// Copyright 1998-2016 Epic Games, Inc. All Rights Reserved.
+
+#pragma once
+
+#include "AsyncFileHandle.h"
+#include "TextPackageNamespaceUtil.h"
+
+
+#if !USE_NEW_ASYNC_IO
+/*----------------------------------------------------------------------------
+	FArchiveAsync.
+----------------------------------------------------------------------------*/
+
+/**
+ * Rough and basic version of async archive. The code relies on Serialize only ever to be called on the last
+ * precached region.
+ */
+class COREUOBJECT_API FArchiveAsync final: public FArchive
+{
+public:
+	/**
+	 * Constructor, initializing member variables.
+	 */
+	FArchiveAsync( const TCHAR* InFileName );
+
+	/**
+	 * Virtual destructor cleaning up internal file reader.
+	 */
+	virtual ~FArchiveAsync();
+
+	/**
+	 * Close archive and return whether there has been an error.
+	 *
+	 * @return	true if there were NO errors, false otherwise
+	 */
+	virtual bool Close();
+
+	/**
+	 * Sets mapping from offsets/ sizes that are going to be used for seeking and serialization to what
+	 * is actually stored on disk. If the archive supports dealing with compression in this way it is 
+	 * going to return true.
+	 *
+	 * @param	CompressedChunks	Pointer to array containing information about [un]compressed chunks
+	 * @param	CompressionFlags	Flags determining compression format associated with mapping
+	 *
+	 * @return true if archive supports translating offsets & uncompressing on read, false otherwise
+	 */
+	virtual bool SetCompressionMap( TArray<FCompressedChunk>* CompressedChunks, ECompressionFlags CompressionFlags );
+
+	/**
+	 * Hint the archive that the region starting at passed in offset and spanning the passed in size
+	 * is going to be read soon and should be precached.
+	 *
+	 * The function returns whether the precache operation has completed or not which is an important
+	 * hint for code knowing that it deals with potential async I/O. The archive is free to either not 
+	 * implement this function or only partially precache so it is required that given sufficient time
+	 * the function will return true. Archives not based on async I/O should always return true.
+	 *
+	 * This function will not change the current archive position.
+	 *
+	 * @param	PrecacheOffset	Offset at which to begin precaching.
+	 * @param	PrecacheSize	Number of bytes to precache
+	 * @return	false if precache operation is still pending, true otherwise
+	 */
+	virtual bool Precache( int64 PrecacheOffset, int64 PrecacheSize );
+
+	/**
+	 * Serializes data from archive.
+	 *
+	 * @param	Data	Pointer to serialize to
+	 * @param	Num		Number of bytes to read
+	 */
+	virtual void Serialize( void* Data, int64 Num );
+
+	/**
+	 * Returns the current position in the archive as offset in bytes from the beginning.
+	 *
+	 * @return	Current position in the archive (offset in bytes from the beginning)
+	 */
+	virtual int64 Tell();
+	/**
+	 * Returns the total size of the archive in bytes.
+	 *
+	 * @return total size of the archive in bytes
+	 */
+	virtual int64 TotalSize();
+
+	/**
+	 * Sets the current position.
+	 *
+	 * @param InPos	New position (as offset from beginning in bytes)
+	 */
+	virtual void Seek( int64 InPos );
+
+	/**
+	 * Flushes cache and frees internal data.
+	 */
+	virtual void FlushCache();
+private:
+
+	/**
+	 * Swaps current and next buffer. Relies on calling code to ensure that there are no outstanding
+	 * async read operations into the buffers.
+	 */
+	void BufferSwitcheroo();
+
+	/**
+	 * Whether the current precache buffer contains the passed in request.
+	 *
+	 * @param	RequestOffset	Offset in bytes from start of file
+	 * @param	RequestSize		Size in bytes requested
+	 *
+	 * @return true if buffer contains request, false othwerise
+	 */
+	FORCEINLINE bool PrecacheBufferContainsRequest( int64 RequestOffset, int64 RequestSize );
+
+	/**
+	 * Finds and returns the compressed chunk index associated with the passed in offset.
+	 *
+	 * @param	RequestOffset	Offset in file to find associated chunk index for
+	 *
+	 * @return Index into CompressedChunks array matching this offset
+	 */
+	int32 FindCompressedChunkIndex( int64 RequestOffset );
+
+	/**
+	 * Precaches compressed chunk of passed in index using buffer at passed in index.
+	 *
+	 * @param	ChunkIndex	Index of compressed chunk
+	 * @param	BufferIndex	Index of buffer to precache into	
+	 */
+	void PrecacheCompressedChunk( int64 ChunkIndex, int64 BufferIndex );
+
+	/** Anon enum used to index precache data. */
+	enum
+	{
+		CURRENT = 0,
+		NEXT	= 1,
+	};
+
+	/** Cached filename for debugging.												*/
+	FString							FileName;
+	/** Cached file size															*/
+	int64							FileSize;
+	/** Cached uncompressed file size (!= FileSize for compressed packages.			*/
+	int64							UncompressedFileSize;
+	/** BulkData area size ( size of the bulkdata are at the end of the file. That is
+		not part of the CompressedChunks, as BulkData is compressed separately		*/
+	int64							BulkDataAreaSize;
+	/** Current position of archive.												*/
+	int64							CurrentPos;
+
+	/** Start position of current precache request.									*/
+	int64							PrecacheStartPos[2];
+	/** End position (exclusive) of current precache request.						*/
+	int64							PrecacheEndPos[2];
+	/** Buffer containing precached data.											*/
+	uint8*							PrecacheBuffer[2];
+	/** Precahce buffer size */
+	SIZE_T PrecacheBufferSize[2];
+	/** True if the precache buffer is protected */
+	bool PrecacheBufferProtected[2];
+
+	/** Status of pending read, a value of 0 means no outstanding reads.			*/
+	FThreadSafeCounter				PrecacheReadStatus[2];
+	
+	/** Mapping of compressed <-> uncompresses sizes and offsets, NULL if not used.	*/
+	TArray<FCompressedChunk>*		CompressedChunks;
+	/** Current index into compressed chunks array.									*/
+	int64							CurrentChunkIndex;
+	/** Compression flags determining compression of CompressedChunks.				*/
+	ECompressionFlags				CompressionFlags;
+	/** Caches the return value of FPlatformMisc::SupportsMultithreading (comes up in profiles often) */
+	bool PlatformIsSinglethreaded;
+};
+
+#else
+
+COREUOBJECT_API FArchive* NewFArchiveAsync2(const TCHAR* InFileName);
+
+COREUOBJECT_API void HintFutureReadDone(const TCHAR * FileName);
+
+COREUOBJECT_API void HintFutureRead(const TCHAR * FileName);
+
+#endif // USE_NEW_ASYNC_IO
+
diff --git a/Engine/Source/Runtime/CoreUObject/Public/Serialization/ArchiveCountMem.h b/Engine/Source/Runtime/CoreUObject/Public/Serialization/ArchiveCountMem.h
new file mode 100644
--- /dev/null
+++ b/Engine/Source/Runtime/CoreUObject/Public/Serialization/ArchiveCountMem.h
@@ -0,0 +1,48 @@
+// Copyright 1998-2016 Epic Games, Inc. All Rights Reserved.
+
+#pragma once
+
+#include "AsyncFileHandle.h"
+#include "TextPackageNamespaceUtil.h"
+struct FObjectInstancingGraph;
+
+/**
+ * Archive for counting memory usage.
+ */
+class FArchiveCountMem : public FArchiveUObject
+{
+public:
+	FArchiveCountMem( UObject* Src )
+	:	Num(0)
+	,	Max(0)
+	{
+		ArIsCountingMemory = true;
+		if( Src )
+		{
+			Src->Serialize( *this );
+		}
+	}
+	SIZE_T GetNum()
+	{
+		return Num;
+	}
+	SIZE_T GetMax()
+	{
+		return Max;
+	}
+	void CountBytes( SIZE_T InNum, SIZE_T InMax )
+	{
+		Num += InNum;
+		Max += InMax;
+	}
+	/**
+	 * Returns the name of the Archive.  Useful for getting the name of the package a struct or object
+	 * is in when a loading error occurs.
+	 *
+	 * This is overridden for the specific Archive Types
+	 **/
+	virtual FString GetArchiveName() const { return TEXT("FArchiveCountMem"); }
+
+protected:
+	SIZE_T Num, Max;
+};
diff --git a/Engine/Source/Runtime/CoreUObject/Public/Serialization/ArchiveObjectCrc32.h b/Engine/Source/Runtime/CoreUObject/Public/Serialization/ArchiveObjectCrc32.h
new file mode 100644
--- /dev/null
+++ b/Engine/Source/Runtime/CoreUObject/Public/Serialization/ArchiveObjectCrc32.h
@@ -0,0 +1,49 @@
+// Copyright 1998-2016 Epic Games, Inc. All Rights Reserved.
+
+#pragma once
+
+#include "AsyncFileHandle.h"
+#include "TextPackageNamespaceUtil.h"
+
+/*----------------------------------------------------------------------------
+FArchiveObjectCrc32
+----------------------------------------------------------------------------*/
+
+/**
+* Calculates a checksum on an object's serialized data stream.
+*/
+class COREUOBJECT_API FArchiveObjectCrc32 : public FArchiveUObject
+{
+public:
+	/**
+	* Default constructor.
+	*/
+	FArchiveObjectCrc32();
+
+	//~ Begin FArchive Interface
+	virtual void Serialize(void* Data, int64 Length);
+	virtual FArchive& operator<<(class FName& Name);
+	virtual FArchive& operator<<(class UObject*& Object);
+	virtual FString GetArchiveName() const { return TEXT("FArchiveObjectCrc32"); }
+	//~ End FArchive Interface
+
+	/**
+	* Serialize the given object, calculate and return its checksum.
+	*/
+	uint32 Crc32(UObject* Object, uint32 CRC = 0);
+
+protected:
+	/** Return if object was already serialized */
+	virtual bool CustomSerialize(class UObject* Object) { return false; }
+
+	/** Internal byte array used for serialization */
+	TArray<uint8> SerializedObjectData;
+	/** Internal archive used for serialization */
+	FMemoryWriter MemoryWriter;
+	/** Internal queue of object references awaiting serialization */
+	TQueue<UObject*> ObjectsToSerialize;
+	/** Internal currently serialized object */
+	const UObject* ObjectBeingSerialized;
+	/** Internal root object */
+	const UObject* RootObject;
+};
diff --git a/Engine/Source/Runtime/CoreUObject/Public/Serialization/ArchiveObjectGraph.h b/Engine/Source/Runtime/CoreUObject/Public/Serialization/ArchiveObjectGraph.h
new file mode 100644
--- /dev/null
+++ b/Engine/Source/Runtime/CoreUObject/Public/Serialization/ArchiveObjectGraph.h
@@ -0,0 +1,104 @@
+// Copyright 1998-2016 Epic Games, Inc. All Rights Reserved.
+
+#pragma once
+
+#include "AsyncFileHandle.h"
+#include "TextPackageNamespaceUtil.h"
+
+struct FTraceRouteRecord
+{
+	struct FObjectGraphNode*	GraphNode;
+	TArray<UProperty*>			ReferencerProperties;
+
+	FTraceRouteRecord( struct FObjectGraphNode* InGraphNode, UProperty* InReferencerProperty)
+	: GraphNode(InGraphNode)
+	{
+		ReferencerProperties.Add(InReferencerProperty);
+	}
+
+	FTraceRouteRecord( struct FObjectGraphNode* InGraphNode, const TArray<UProperty*>&	InReferencerProperties )
+		: GraphNode(InGraphNode)
+	{
+		ReferencerProperties = InReferencerProperties;
+	}
+
+	void Add(UProperty* InReferencerProperty)
+	{
+		ReferencerProperties.Add(InReferencerProperty);
+	}
+};
+
+struct FObjectGraphNode
+{
+	/** the object this node represents */
+	UObject*	NodeObject;
+
+	/** Records for objects which reference this object */
+	TMap<UObject*, FTraceRouteRecord>	ReferencerRecords;
+
+	/** Records for objects which are referenced by this object */
+	TMap<UObject*, FTraceRouteRecord>	ReferencedObjects;
+
+	/** the number of links between NodeObject and the target object */
+	int32									ReferenceDepth;
+
+	/** Used during search - Visited or not */
+	bool							Visited;
+
+	/**
+	 * The property that references NodeObject; only set on nodes which are part
+	 * of the calculated shortest route
+	 */
+	TArray<UProperty*>							ReferencerProperties;
+
+	/** Default constructor */
+	FObjectGraphNode( UObject* InNodeObject=NULL )
+	:	NodeObject(InNodeObject)
+	,	ReferenceDepth(MAX_int32)
+	,	Visited(false)
+	{}
+};
+
+/*----------------------------------------------------------------------------
+	FArchiveObjectGraph.
+----------------------------------------------------------------------------*/
+// This is from FArchiveTraceRoute -This only creates object graph of all objects 
+// This can be used by other classes such as FTraceReferences - trace references of one object
+class FArchiveObjectGraph : public FArchiveUObject
+{
+	/** Handles serialization of UObject references */
+	FArchive& operator<<( class UObject*& Obj );
+
+	/**
+	* The object currently being serialized; used by the overloaded serialization operator to determine the referencing object.
+	*/
+	UObject* CurrentReferencer;
+	/** The set of objects encountered while serializing CurrentReferencer */
+	TArray<UObject*> ObjectsToSerialize;
+
+	/** true if we should serialize objects marked RF_Transient */
+	bool bIncludeTransients;
+
+	/**
+	* A bitmask of object flags which indicates which objects should be included in the group of initial objects to be serialized;
+	* RF_RootSet will be automatically added to this bitmask, and OBJECTMARK_TagExp will automatically be removed.
+	*/
+	EObjectFlags	RequiredFlags;
+
+public:
+	FArchiveObjectGraph(bool IncludeTransients, EObjectFlags KeepFlags);
+	~FArchiveObjectGraph();
+
+	/**
+	* Serializes the objects in the specified set; any objects encountered during serialization
+	* of an object are added to the object set and processed until no new objects are added.
+	* DO NOT MAKE THIS VIRTUAL - this is called by constructor. If you wish to do so, please change where be called
+	* @param	Objects		the original set of objects to serialize; the original set will be preserved.
+	*/
+	void GenerateObjectGraph( TArray<UObject*>& Objects );
+
+	void ClearSearchFlags();
+
+	/** A complete graph of all references between all objects in memory */
+	TMap<UObject*, FObjectGraphNode*> ObjectGraph;
+};
diff --git a/Engine/Source/Runtime/CoreUObject/Public/Serialization/ArchiveReplaceArchetype.h b/Engine/Source/Runtime/CoreUObject/Public/Serialization/ArchiveReplaceArchetype.h
new file mode 100644
--- /dev/null
+++ b/Engine/Source/Runtime/CoreUObject/Public/Serialization/ArchiveReplaceArchetype.h
@@ -0,0 +1,30 @@
+// Copyright 1998-2016 Epic Games, Inc. All Rights Reserved.
+
+#pragma once
+
+#include "AsyncFileHandle.h"
+#include "TextPackageNamespaceUtil.h"
+#include "ReloadObjectArc.h"
+
+/*----------------------------------------------------------------------------
+	FArchiveReplaceArchetype.
+----------------------------------------------------------------------------*/
+/**
+ * This specialized version of the FReloadObjectArc is used when changing the archetype for a fully initialized object.  It handles saving and restoring
+ * the values which have been changed in the instance, as well as remapping archetypes for subobjects to the corresponding subobject in the new archetype.
+ * If a corresponding subobject cannot be found, the subobject's archetype is reset to the CDO for that subobject.
+ */
+class FArchiveReplaceArchetype : public FReloadObjectArc
+{
+public:
+	FArchiveReplaceArchetype();
+
+	/**
+	 * Returns the name of the Archive.  Useful for getting the name of the package a struct or object
+	 * is in when a loading error occurs.
+	 *
+	 * This is overridden for the specific Archive Types
+	 **/
+	virtual FString GetArchiveName() const { return TEXT("FArchiveReplaceArchetype"); }
+};
+
diff --git a/Engine/Source/Runtime/CoreUObject/Public/Serialization/ArchiveReplaceObjectRef.h b/Engine/Source/Runtime/CoreUObject/Public/Serialization/ArchiveReplaceObjectRef.h
new file mode 100644
--- /dev/null
+++ b/Engine/Source/Runtime/CoreUObject/Public/Serialization/ArchiveReplaceObjectRef.h
@@ -0,0 +1,179 @@
+// Copyright 1998-2016 Epic Games, Inc. All Rights Reserved.
+
+#pragma once
+
+#include "AsyncFileHandle.h"
+#include "TextPackageNamespaceUtil.h"
+
+/** Base class for object replacement archives */ 
+class COREUOBJECT_API FArchiveReplaceObjectRefBase : public FArchiveUObject
+{
+protected:
+	/**
+	* Serializes a single object
+	*/
+	void SerializeObject(UObject* ObjectToSerialize);
+};
+
+/*----------------------------------------------------------------------------
+	FArchiveReplaceObjectRef.
+----------------------------------------------------------------------------*/
+/**
+ * Archive for replacing a reference to an object. This classes uses
+ * serialization to replace all references to one object with another.
+ * Note that this archive will only traverse objects with an Outer
+ * that matches InSearchObject.
+ *
+ * NOTE: The template type must be a child of UObject or this class will not compile.
+ */
+template< class T >
+class FArchiveReplaceObjectRef : public FArchiveReplaceObjectRefBase
+{
+public:
+	/**
+	 * Initializes variables and starts the serialization search
+	 *
+	 * @param InSearchObject		The object to start the search on
+	 * @param ReplacementMap		Map of objects to find -> objects to replace them with (null zeros them)
+	 * @param bNullPrivateRefs		Whether references to non-public objects not contained within the SearchObject
+	 *								should be set to null
+	 * @param bIgnoreOuterRef		Whether we should replace Outer pointers on Objects.
+	 * @param bIgnoreArchetypeRef	Whether we should replace the ObjectArchetype reference on Objects.
+	 * @param bDelayStart			Specify true to prevent the constructor from starting the process.  Allows child classes' to do initialization stuff in their ctor
+	 */
+	FArchiveReplaceObjectRef
+	(
+		UObject* InSearchObject,
+		const TMap<T*,T*>& inReplacementMap,
+		bool bNullPrivateRefs,
+		bool bIgnoreOuterRef,
+		bool bIgnoreArchetypeRef,
+		bool bDelayStart = false,
+		bool bIgnoreClassGeneratedByRef = true
+	)
+	: SearchObject(InSearchObject), ReplacementMap(inReplacementMap)
+	, Count(0), bNullPrivateReferences(bNullPrivateRefs)
+	{
+		ArIsObjectReferenceCollector = true;
+		ArIsModifyingWeakAndStrongReferences = true;		// Also replace weak references too!
+		ArIgnoreArchetypeRef = bIgnoreArchetypeRef;
+		ArIgnoreOuterRef = bIgnoreOuterRef;
+		ArIgnoreClassGeneratedByRef = bIgnoreClassGeneratedByRef;
+
+		if ( !bDelayStart )
+		{
+			SerializeSearchObject();
+		}
+	}
+
+	/**
+	 * Starts the serialization of the root object
+	 */
+	void SerializeSearchObject()
+	{
+		ReplacedReferences.Empty();
+
+		if (SearchObject != NULL && !SerializedObjects.Find(SearchObject)
+		&&	(ReplacementMap.Num() > 0 || bNullPrivateReferences))
+		{
+			// start the initial serialization
+			SerializedObjects.Add(SearchObject);
+			SerializeObject(SearchObject);
+			for (int32 Iter = 0; Iter < PendingSerializationObjects.Num(); Iter++)
+			{
+				SerializeObject(PendingSerializationObjects[Iter]);
+			}
+			PendingSerializationObjects.Reset();
+		}
+	}
+
+	/**
+	 * Returns the number of times the object was referenced
+	 */
+	int64 GetCount() const
+	{
+		return Count;
+	}
+
+	/**
+	 * Returns a reference to the object this archive is operating on
+	 */
+	const UObject* GetSearchObject() const { return SearchObject; }
+
+	/**
+	 * Returns a reference to the replaced references map
+	 */
+	const TMap<UObject*, TArray<UProperty*>>& GetReplacedReferences() const { return ReplacedReferences; }
+
+	/**
+	 * Serializes the reference to the object
+	 */
+	FArchive& operator<<( UObject*& Obj )
+	{
+		if (Obj != NULL)
+		{
+			// If these match, replace the reference
+			T* const* ReplaceWith = (T*const*)((const TMap<UObject*,UObject*>*)&ReplacementMap)->Find(Obj);
+			if ( ReplaceWith != NULL )
+			{
+				Obj = *ReplaceWith;
+				ReplacedReferences.FindOrAdd(Obj).AddUnique(GetSerializedProperty());
+				Count++;
+			}
+			// A->IsIn(A) returns false, but we don't want to NULL that reference out, so extra check here.
+			else if ( Obj == SearchObject || Obj->IsIn(SearchObject) )
+			{
+#if 0
+				// DEBUG: Log when we are using the A->IsIn(A) path here.
+				if(Obj == SearchObject)
+				{
+					FString ObjName = Obj->GetPathName();
+					UE_LOG(LogSerialization, Log,  TEXT("FArchiveReplaceObjectRef: Obj == SearchObject : '%s'"), *ObjName );
+				}
+#endif
+				bool bAlreadyAdded = false;
+				SerializedObjects.Add(Obj, &bAlreadyAdded);
+				if (!bAlreadyAdded)
+				{
+					// No recursion
+					PendingSerializationObjects.Add(Obj);
+				}
+			}
+			else if ( bNullPrivateReferences && !Obj->HasAnyFlags(RF_Public) )
+			{
+				Obj = NULL;
+			}
+		}
+		return *this;
+	}
+
+	/**
+	 * Returns the name of this archive.
+	 **/
+	virtual FString GetArchiveName() const { return TEXT("ReplaceObjectRef"); }
+
+protected:
+	/** Initial object to start the reference search from */
+	UObject* SearchObject;
+
+	/** Map of objects to find references to -> object to replace references with */
+	const TMap<T*,T*>& ReplacementMap;
+	
+	/** The number of times encountered */
+	int32 Count;
+
+	/** List of objects that have already been serialized */
+	TSet<UObject*> SerializedObjects;
+
+	/** Object that will be serialized */
+	TArray<UObject*> PendingSerializationObjects;
+
+	/** Map of referencing objects to referencing properties */
+	TMap<UObject*, TArray<UProperty*>> ReplacedReferences;
+
+	/**
+	 * Whether references to non-public objects not contained within the SearchObject
+	 * should be set to null
+	 */
+	bool bNullPrivateReferences;
+};
diff --git a/Engine/Source/Runtime/CoreUObject/Public/Serialization/ArchiveTraceRoute.h b/Engine/Source/Runtime/CoreUObject/Public/Serialization/ArchiveTraceRoute.h
new file mode 100644
--- /dev/null
+++ b/Engine/Source/Runtime/CoreUObject/Public/Serialization/ArchiveTraceRoute.h
@@ -0,0 +1,114 @@
+// Copyright 1998-2016 Epic Games, Inc. All Rights Reserved.
+
+#pragma once
+
+#include "AsyncFileHandle.h"
+#include "TextPackageNamespaceUtil.h"
+
+/**
+ * Archive for finding shortest path from root to a particular object.
+ * Depth-first search.
+ */
+class FArchiveTraceRoute : public FArchiveUObject
+{
+	/**
+	 * Simple container struct for grouping two FObjectGraphNodes together.
+	 */
+	struct FRouteLink
+	{
+		/**
+		 * Constructor
+		 */
+		FRouteLink( FObjectGraphNode* InParent=NULL, FObjectGraphNode* InChild=NULL )
+		: LinkParent(InParent), LinkChild(InChild)
+		{}
+
+		/**
+		 * The node corresponding to the "referencing" object.
+		 */
+		FObjectGraphNode* LinkParent;
+
+		/**
+		 * The node corresponding to the "referenced" object.
+		 */
+		FObjectGraphNode* LinkChild;
+	};
+
+public:
+	static COREUOBJECT_API TMap<UObject*,UProperty*> FindShortestRootPath( UObject* Object, bool bIncludeTransients, EObjectFlags KeepFlags );
+
+	/**
+	 * Retuns path to root created by e.g. FindShortestRootPath via a string.
+	 *
+	 * @param TargetObject	object marking one end of the route
+	 * @param Route			route to print to log.
+	 * @param String of root path
+	 */
+	static COREUOBJECT_API FString PrintRootPath( const TMap<UObject*,UProperty*>& Route, const UObject* TargetObject );
+
+	/**
+	 * Returns the name of the Archive.  Useful for getting the name of the package a struct or object
+	 * is in when a loading error occurs.
+	 *
+	 * This is overridden for the specific Archive Types
+	 **/
+	virtual FString GetArchiveName() const { return TEXT("FArchiveTraceRoute"); }
+
+	/**
+	 * Serializes the objects in the specified set; any objects encountered during serialization
+	 * of an object are added to the object set and processed until no new objects are added.
+	 *
+	 * @param	Objects		the original set of objects to serialize; the original set will be preserved.
+	 */
+	void GenerateObjectGraph( TSparseArray<UObject*>& Objects );
+
+	/**
+	 * Recursively iterates over the referencing objects for the specified node, marking each with
+	 * the current Depth value.  Stops once it reaches a route root.
+	 *
+	 * @param	ObjectNode	the node to evaluate.
+	 */
+	void CalculateReferenceDepthsForNode( FObjectGraphNode* ObjectNode );
+
+	/**
+	 * Searches through the objects referenced by CurrentNode for a record with a Depth lower than LowestDepth.
+	 *
+	 * @param	CurrentNode		the node containing the list of referenced objects that will be searched.
+	 * @param	LowestDepth		the current number of links we are from the target object.
+	 * @param	ClosestLink		if a trace route record is found with a lower depth value than LowestDepth, the link is saved to this value.
+	 *
+	 * @return	true if a closer link was discovered; false if no links were closer than lowest depth, or if we've reached the target object.
+	 */
+	bool FindClosestLink( FObjectGraphNode* CurrentNode, int32& LowestDepth, FRouteLink& ClosestLink );
+
+private:
+	FArchiveTraceRoute( UObject* TargetObject, TMap<UObject*,FTraceRouteRecord>& InRoutes, bool bShouldIncludeTransients, EObjectFlags KeepFlags );
+	~FArchiveTraceRoute();
+
+	/** Handles serialization of UObject references */
+	FArchive& operator<<( class UObject*& Obj );
+
+	/** A complete graph of all references between all objects in memory */
+	TMap<UObject*, FObjectGraphNode*> ObjectGraph;
+
+	/**
+	 * The object currently being serialized; used by the overloaded serialization operator to determine the referencing object.
+	 */
+	UObject* CurrentReferencer;
+
+	/** The set of objects encountered while serializing CurrentReferencer */
+	TArray<UObject*> ObjectsToSerialize;
+
+	/** the current number of object reference links away from the target object */
+	int32 Depth;
+
+	/** true if we should serialize objects marked RF_Transient */
+	bool bIncludeTransients;
+
+	/**
+	 * A bitmask of object flags which indicates which objects should be included in the group of initial objects to be serialized;
+	 * RF_RootSet will be automatically added to this bitmask, and OBJECTMARK_TagExp will automatically be removed.
+	 */
+	EObjectFlags	RequiredFlags;
+
+};
diff --git a/Engine/Source/Runtime/CoreUObject/Public/Serialization/ArchiveUObject.h b/Engine/Source/Runtime/CoreUObject/Public/Serialization/ArchiveUObject.h
--- a/Engine/Source/Runtime/CoreUObject/Public/Serialization/ArchiveUObject.h
+++ b/Engine/Source/Runtime/CoreUObject/Public/Serialization/ArchiveUObject.h
@@ -2,1662 +2,18 @@
 
 #pragma once
 
-#include "AsyncFileHandle.h"
-#include "TextPackageNamespaceUtil.h"
-struct FObjectInstancingGraph;
-
-
 /**
- * Archive for counting memory usage.
+ * Base FArchive for serializing UObjects. Supports FLazyObjectPtr and FAssetPtr serialization.
  */
-class FArchiveCountMem : public FArchiveUObject
+class COREUOBJECT_API FArchiveUObject : public FArchive
 {
 public:
-	FArchiveCountMem( UObject* Src )
-	:	Num(0)
-	,	Max(0)
-	{
-		ArIsCountingMemory = true;
-		if( Src )
-		{
-			Src->Serialize( *this );
-		}
-	}
-	SIZE_T GetNum()
-	{
-		return Num;
-	}
-	SIZE_T GetMax()
-	{
-		return Max;
-	}
-	void CountBytes( SIZE_T InNum, SIZE_T InMax )
-	{
-		Num += InNum;
-		Max += InMax;
-	}
-	/**
-	 * Returns the name of the Archive.  Useful for getting the name of the package a struct or object
-	 * is in when a loading error occurs.
-	 *
-	 * This is overridden for the specific Archive Types
-	 **/
-	virtual FString GetArchiveName() const { return TEXT("FArchiveCountMem"); }
-
-protected:
-	SIZE_T Num, Max;
-};
-
 
-/**
- * Implements a proxy archive that serializes UObjects and FNames as string data.
- *
- * Expected use is:
- *    FArchive* SomeAr = CreateAnAr();
- *    FObjectAndNameAsStringProxyArchive Ar(*SomeAr);
- *    SomeObject->Serialize(Ar);
- *    FinalizeAr(SomeAr);
- * 
- * @param InInnerArchive The actual FArchive object to serialize normal data types (FStrings, INTs, etc)
- */
-struct FObjectAndNameAsStringProxyArchive : public FNameAsStringProxyArchive
-{
-	/**
-	 * Creates and initializes a new instance.
-	 *
-	 * @param InInnerArchive - The inner archive to proxy.
-	 * @param bInLoadIfFindFails - Indicates whether to try and load a ref'd object if we don't find it
-	 */
-	FObjectAndNameAsStringProxyArchive(FArchive& InInnerArchive, bool bInLoadIfFindFails)
-		:	FNameAsStringProxyArchive(InInnerArchive)
-		,	bLoadIfFindFails(bInLoadIfFindFails)
-	{ }
-
-	/** If we fail to find an object during loading, try and load it. */
-	bool bLoadIfFindFails;
-
-	COREUOBJECT_API virtual FArchive& operator<<(class UObject*& Obj);
-};
-
-/**
- * UObject Memory Writer Archive.
- */
-class FObjectWriter : public FMemoryWriter
-{
-
-public:
-	FObjectWriter(UObject* Obj, TArray<uint8>& InBytes, bool bIgnoreClassRef = false, bool bIgnoreArchetypeRef = false, bool bDoDelta = true, uint32 AdditionalPortFlags = 0)
-		: FMemoryWriter(InBytes)
-	{
-		ArIgnoreClassRef = bIgnoreClassRef;
-		ArIgnoreArchetypeRef = bIgnoreArchetypeRef;
-		ArNoDelta = !bDoDelta;
-		ArPortFlags |= AdditionalPortFlags;
-
-#if USE_STABLE_LOCALIZATION_KEYS
-		if (GIsEditor && !(ArPortFlags & PPF_DuplicateForPIE))
-		{
-			SetLocalizationNamespace(TextNamespaceUtil::EnsurePackageNamespace(Obj));
-		}
-#endif // USE_STABLE_LOCALIZATION_KEYS
-
-		Obj->Serialize(*this);
-	}
+	using FArchive::operator<<; // For visibility of the overloads we don't override
 
 	//~ Begin FArchive Interface
-	COREUOBJECT_API virtual FArchive& operator<<( class FName& N ) override;
-	COREUOBJECT_API virtual FArchive& operator<<( class UObject*& Res ) override;
-	COREUOBJECT_API virtual FArchive& operator<<( FLazyObjectPtr& LazyObjectPtr ) override;
-	COREUOBJECT_API virtual FArchive& operator<<( FAssetPtr& AssetPtr ) override;
-	COREUOBJECT_API virtual FArchive& operator<<(FStringAssetReference& AssetPtr) override;
-	COREUOBJECT_API virtual FString GetArchiveName() const override;
+	virtual FArchive& operator<< (class FLazyObjectPtr& Value) override;
+	virtual FArchive& operator<< (class FAssetPtr& Value) override;
+	virtual FArchive& operator<< (struct FStringAssetReference& Value) override;
 	//~ End FArchive Interface
-
-protected:
-	FObjectWriter(TArray<uint8>& InBytes)
-		: FMemoryWriter(InBytes)
-	{
-		ArIgnoreClassRef = false;
-		ArIgnoreArchetypeRef = false;
-	}
-};
-
-/**
- * UObject Memory Reader Archive.
- */
-class FObjectReader : public FMemoryArchive
-{
-public:
-	FObjectReader(UObject* Obj, TArray<uint8>& InBytes, bool bIgnoreClassRef = false, bool bIgnoreArchetypeRef = false)
-		: Bytes(InBytes)
-	{
-		ArIsLoading = true;
-		ArIsPersistent = false;
-		ArIgnoreClassRef = bIgnoreClassRef;
-		ArIgnoreArchetypeRef = bIgnoreArchetypeRef;
-
-#if USE_STABLE_LOCALIZATION_KEYS
-		if (GIsEditor && !(ArPortFlags & PPF_DuplicateForPIE))
-		{
-			SetLocalizationNamespace(TextNamespaceUtil::EnsurePackageNamespace(Obj));
-		}
-#endif // USE_STABLE_LOCALIZATION_KEYS
-
-		Obj->Serialize(*this);
-	}
-
-	//~ Begin FArchive Interface
-
-	int64 TotalSize()
-	{
-		return (int64)Bytes.Num();
-	}
-
-	void Serialize(void* Data, int64 Num)
-	{
-		if (Num && !ArIsError)
-		{
-			// Only serialize if we have the requested amount of data
-			if (Offset + Num <= TotalSize())
-			{
-				FMemory::Memcpy(Data, &Bytes[Offset], Num);
-				Offset += Num;
-			}
-			else
-			{
-				ArIsError = true;
-			}
-		}
-	}
-
-	COREUOBJECT_API virtual FArchive& operator<<( class FName& N ) override;
-	COREUOBJECT_API virtual FArchive& operator<<( class UObject*& Res ) override;
-	COREUOBJECT_API virtual FArchive& operator<<( FLazyObjectPtr& LazyObjectPtr ) override;
-	COREUOBJECT_API virtual FArchive& operator<<( FAssetPtr& AssetPtr ) override;
-	COREUOBJECT_API virtual FArchive& operator<<(FStringAssetReference& AssetPtr) override;
-	COREUOBJECT_API virtual FString GetArchiveName() const override;
-	//~ End FArchive Interface
-
-
-
-protected:
-	FObjectReader(TArray<uint8>& InBytes)
-		: Bytes(InBytes)
-	{
-		ArIsLoading = true;
-		ArIsPersistent = false;
-		ArIgnoreClassRef = false;
-		ArIgnoreArchetypeRef = false;
-	}
-
-	const TArray<uint8>& Bytes;
-};
-
-/**
- * Archive for reloading UObjects without requiring the UObject to be completely unloaded.
- * Used for easily repropagating defaults to UObjects after class data has been changed.
- * Aggregates FMemoryReader and FMemoryWriter to encapsulate management and coordination 
- * of the UObject data being saved/loaded.
- * <p>
- * UObject references are not serialized directly into the memory archive.  Instead, we use
- * a system similar to the Export/ImportMap of FLinker - the pointer to the UObject is added
- * to a persistent (from the standpoint of the FReloadObjectArc) array.  The location into
- * this array is what is actually stored in the archive's buffer.
- * <p>
- * This is primarily necessary for loading UObject references from the memory archive.  The
- * UObject pointer ref. passed into the UObject overloaded serialization operator is not
- * guaranteed to be valid when reading data from the memory archive.  Since we need to determine
- * whether we *should* serialize the object before we actually do, we must introduce a level of
- * indirection.  Using the index retrieved from the archive's buffer, we can look at the UObject
- * before we attempt to serialize it.
- */
-class COREUOBJECT_API FReloadObjectArc : public FArchiveUObject
-{
-public:
-
-	/**
-	 * Changes this memory archive to "read" mode, for reading UObject
-	 * data from the temporary location back into the UObjects.
-	 */
-	void ActivateReader()
-	{
-		ArIsSaving = false;
-		ArIsLoading = true;
-	}
-
-	/**
-	 * Changes this memory archive to "write" mode, for storing UObject
-	 * data to the temporary location.
-	 *
-	 * @note: called from ctors in child classes - should never be made virtual
-	 */
-	void ActivateWriter()
-	{
-		ArIsSaving = true;
-		ArIsLoading = false;
-	}
-
-	/**
-	 * Begin serializing a UObject into the memory archive.
-	 *
-	 * @param	Obj		the object to serialize
-	 */
-	void SerializeObject( UObject* Obj );
-
-	/**
-	 * Resets the archive so that it can be loaded from again from scratch
-	 * as if it was never serialized as a Reader
-	 */
-	void Reset();
-
-	/** FArchive Interface */
-	int64 TotalSize()
-	{
-		return Bytes.Num();
-	}
-	void Seek( int64 InPos )
-	{
-		if ( IsLoading() )
-			Reader.Seek(InPos);
-		else if ( IsSaving() )
-			Writer.Seek(InPos);
-	}
-	int64 Tell()
-	{
-		return IsLoading() ? Reader.Tell() : Writer.Tell();
-	}
-	FArchive& operator<<( class FName& Name );
-	FArchive& operator<<(class UObject*& Obj);
-
-	/** Constructor */
-	FReloadObjectArc();
-
-	/** Destructor */
-	virtual ~FReloadObjectArc();
-
-	/**
-	 * Returns the name of the Archive.  Useful for getting the name of the package a struct or object
-	 * is in when a loading error occurs.
-	 *
-	 * This is overridden for the specific Archive Types
-	 **/
-	virtual FString GetArchiveName() const { return TEXT("FReloadObjectArc"); }
-
-protected:
-
-	/**
-	 * Raw I/O function.  Routes the call to the appropriate archive, depending on whether
-	 * we're loading/saving.
-	 */
-	void Serialize( void* Data, int64 Num )
-	{
-		if ( IsLoading() )
-		{
-			Reader.Serialize(Data, Num);
-		}
-		else if ( IsSaving() )
-		{
-			Writer.Serialize(Data,Num);
-		}
-	}
-
-	/**
-	 * Sets the root object for this memory archive.
-	 * 
-	 * @param	NewRoot		the UObject that should be the new root
-	 */
-	void SetRootObject( UObject* NewRoot );
-
-	/** the raw UObject data contained by this archive */
-	TArray<uint8>		Bytes;
-
-	/** moves UObject data from storage into UObject address space */
-	FMemoryReader		Reader;
-
-	/** stores UObject data in a temporary location for later retrieval */
-	FMemoryWriter		Writer;
-
-	/** UObjects for which all data is stored in the memory archive */
-	TArray<UObject*>	CompleteObjects;
-
-	/** UObjects for which only a reference to the object is stored in the memory archive */
-	TArray<UObject*>	ReferencedObjects;
-
-	/**
-	 * List of top-level objects that have saved into the memory archive.  Used to prevent objects
-	 * from being serialized into storage multiple times.
-	 */
-	TSet<UObject*>	SavedObjects;
-
-	/**
-	 * List of top-level objects that have been loaded using the memory archive.  Used to prevent
-	 * objects from being serialized multiple times from the same memory archive.
-	 */
-	TSet<UObject*>	LoadedObjects;
-
-	/** A mapping of "UObject" => "the offset for that UObject's data in the Bytes array" for the objects stored in this archive */
-	TMap<UObject*,int32>	ObjectMap;
-
-	/**
-	 * This is the current top-level object.  For any UObjects contained
-	 * within this object, the complete UObject data will be stored in the
-	 * archive's buffer
-	 */
-	UObject*			RootObject;
-
-	/**
-	 * Used for tracking the subobjects and components that are instanced during this object reload.
-	 */
-	struct FObjectInstancingGraph*	InstanceGraph;
-
-	/**
-	 * Indicates whether this archive will serialize references to objects with the RF_Transient flag. (defaults to false)
-	 */
-	bool				bAllowTransientObjects;
-
-	/**
-	 * Indicates whether this archive should call InstanceSubobjectTemplates/InstanceSubobjects on objects that it re-initializes.
-	 * Specify false if the object needs special handling before calling InstanceSubobjectTemplates
-	 */
-	bool				bInstanceSubobjectsOnLoad;
-};
-
-/*----------------------------------------------------------------------------
-	FArchiveReplaceArchetype.
-----------------------------------------------------------------------------*/
-/**
- * This specialized version of the FReloadObjectArc is used when changing the archetype for a fully initialized object.  It handles saving and restoring
- * the values which have been changed in the instance, as well as remapping archetypes for subobjects to the corresponding subobject in the new archetype.
- * If a corresponding subobject cannot be found, the subobject's archetype is reset to the CDO for that subobject.
- */
-class FArchiveReplaceArchetype : public FReloadObjectArc
-{
-public:
-	FArchiveReplaceArchetype();
-
-	/**
-	 * Returns the name of the Archive.  Useful for getting the name of the package a struct or object
-	 * is in when a loading error occurs.
-	 *
-	 * This is overridden for the specific Archive Types
-	 **/
-	virtual FString GetArchiveName() const { return TEXT("FArchiveReplaceArchetype"); }
-};
-
-/*----------------------------------------------------------------------------
-	FArchiveShowReferences.
-----------------------------------------------------------------------------*/
-/**
- * Archive for displaying all objects referenced by a particular object.
- */
-class FArchiveShowReferences : public FArchiveUObject
-{
-	/**
-	 * I/O function.  Called when an object reference is encountered.
-	 *
-	 * @param	Obj		a pointer to the object that was encountered
-	 */
-	FArchive& operator<<( UObject*& Obj );
-
-	/** the object to display references to */
-	UObject* SourceObject;
-
-	/** ignore references to objects have the same Outer as our Target */
-	UObject* SourceOuter;
-
-	/** output device for logging results */
-	FOutputDevice& OutputAr;
-
-	/**
-	 * list of Outers to ignore;  any objects encountered that have one of
-	 * these objects as an Outer will also be ignored
-	 */
-	class TArray<UObject*>& Exclude;
-
-	/** list of objects that have been found */
-	class TArray<UObject*> Found;
-
-	bool DidRef;
-
-public:
-
-	/**
-	 * Constructor
-	 * 
-	 * @param	inOutputAr		archive to use for logging results
-	 * @param	inOuter			only consider objects that do not have this object as its Outer
-	 * @param	inSource		object to show references for
-	 * @param	inExclude		list of objects that should be ignored if encountered while serializing SourceObject
-	 */
-	FArchiveShowReferences( FOutputDevice& inOutputAr, UObject* inOuter, UObject* inSource, TArray<UObject*>& InExclude );
-
-	/**
-	 * Returns the name of the Archive.  Useful for getting the name of the package a struct or object
-	 * is in when a loading error occurs.
-	 *
-	 * This is overridden for the specific Archive Types
-	 **/
-	virtual FString GetArchiveName() const { return TEXT("FArchiveShowReferences"); }
-};
-
-/*----------------------------------------------------------------------------
-	FFindReferencersArchive.
-----------------------------------------------------------------------------*/
-/**
- * Archive for mapping out the referencers of a collection of objects.
- */
-class FFindReferencersArchive : public FArchiveUObject
-{
-public:
-	/**
-	 * Constructor
-	 *
-	 * @param	PotentialReferencer		the object to serialize which may contain references to our target objects
-	 * @param	InTargetObjects			array of objects to search for references to
-	 * @param	bFindAlsoWeakReferences should we also look into weak references?
-	 */
-	COREUOBJECT_API FFindReferencersArchive(class UObject* PotentialReferencer, const TArray<class UObject*>& InTargetObjects, bool bFindAlsoWeakReferences = false);
-
-	/**
-	 * Retrieves the number of references from PotentialReferencer to the object specified.
-	 *
-	 * @param	TargetObject	the object to might be referenced
-	 * @param	out_ReferencingProperties
-	 *							receives the list of properties which were holding references to TargetObject
-	 *
-	 * @return	the number of references to TargetObject which were encountered when PotentialReferencer
-	 *			was serialized.
-	 */
-	COREUOBJECT_API int32 GetReferenceCount( class UObject* TargetObject, TArray<class UProperty*>* out_ReferencingProperties=NULL ) const;
-
-	/**
-	 * Retrieves the number of references from PotentialReferencer list of TargetObjects
-	 *
-	 * @param	out_ReferenceCounts		receives the number of references to each of the TargetObjects
-	 *
-	 * @return	the number of objects which were referenced by PotentialReferencer.
-	 */
-	COREUOBJECT_API int32 GetReferenceCounts( TMap<class UObject*, int32>& out_ReferenceCounts ) const;
-
-	/**
-	 * Retrieves the number of references from PotentialReferencer list of TargetObjects
-	 *
-	 * @param	out_ReferenceCounts			receives the number of references to each of the TargetObjects
-	 * @param	out_ReferencingProperties	receives the map of properties holding references to each referenced object.
-	 *
-	 * @return	the number of objects which were referenced by PotentialReferencer.
-	 */
-	COREUOBJECT_API int32 GetReferenceCounts( TMap<class UObject*, int32>& out_ReferenceCounts, TMultiMap<class UObject*,class UProperty*>& out_ReferencingProperties ) const;
-
-	/**
-	 * Returns the name of the Archive.  Useful for getting the name of the package a struct or object
-	 * is in when a loading error occurs.
-	 *
-	 * This is overridden for the specific Archive Types
-	 **/
-	COREUOBJECT_API virtual FString GetArchiveName() const { return TEXT("FFindReferencersArchive"); }
-
-	/**
-	 * Resets the reference counts.  Keeps the same target objects but sets up everything to test a new potential referencer.
-	 * @param	PotentialReferencer		the object to serialize which may contain references to our target objects
-	 **/
-	COREUOBJECT_API void ResetPotentialReferencer(UObject* InPotentialReferencer);
-
-protected:
-	TMap<class UObject*, int32>	TargetObjects;
-
-	/** a mapping of target object => the properties in PotentialReferencer that hold the reference to target object */
-	TMultiMap<class UObject*,class UProperty*> ReferenceMap;
-
-	/** The potential referencer we ignore */
-	class UObject* PotentialReferencer;
-
-private:
-
-	/**
-	 * Serializer - if Obj is one of the objects we're looking for, increments the reference count for that object
-	 */
-	COREUOBJECT_API FArchive& operator<<( class UObject*& Obj );
-};
-
-/**
- * This class is used to find which objects reference any element from a list of "TargetObjects".  When invoked,
- * it will generate a mapping of each target object to an array of objects referencing that target object.
- *
- * Each key corresponds to an element of the input TargetObjects array which was referenced
- * by some other object.  The values for these keys are the objects which are referencing them.
- */
-template< class T >
-class TFindObjectReferencers : public TMultiMap<T*, UObject*>
-{
-public:
-
-	/**
-	 * Default constructor
-	 *
-	 * @param	TargetObjects	the list of objects to find references to
-	 * @param	PackageToCheck	if specified, only objects contained in this package will be searched
-	 *							for references to 
-	 * @param	bIgnoreTemplates If true, do not record template objects
-	 * @param	bFindAlsoWeakReferences If true, also look into weak references
-	 */
-	TFindObjectReferencers( TArray< T* > TargetObjects, UPackage* PackageToCheck=NULL, bool bIgnoreTemplates = true, bool bFindAlsoWeakReferences = false)
-	: TMultiMap< T*, UObject* >()
-	{
-		TArray<UObject*> ReferencedObjects;
-		TMap<UObject*, int32> ReferenceCounts;
-
-		FFindReferencersArchive FindReferencerAr(nullptr, (TArray<UObject*>&)TargetObjects, bFindAlsoWeakReferences);
-
-		// Loop over every object to find any reference that may exist for the target objects
-		for (FObjectIterator It; It; ++It)
-		{
-			UObject* PotentialReferencer = *It;
-			if ( !TargetObjects.Contains(PotentialReferencer)
-			&&	(PackageToCheck == NULL || PotentialReferencer->IsIn(PackageToCheck))
-			&&	(!bIgnoreTemplates || !PotentialReferencer->IsTemplate()) )
-			{
-				FindReferencerAr.ResetPotentialReferencer(PotentialReferencer);
-
-				ReferenceCounts.Reset();
-				if ( FindReferencerAr.GetReferenceCounts(ReferenceCounts) > 0 )
-				{
-					// here we don't really care about the number of references from PotentialReferencer to the target object...just that he's a referencer
-					ReferencedObjects.Reset();
-					ReferenceCounts.GenerateKeyArray(ReferencedObjects);
-					for ( int32 RefIndex = 0; RefIndex < ReferencedObjects.Num(); RefIndex++ )
-					{
-						this->Add(static_cast<T*>(ReferencedObjects[RefIndex]), PotentialReferencer);
-					}
-				}
-			}
-		}
-	}
-
-private:
-	/**
-	 * This is a mapping of TargetObjects to the list of objects which references each one combined with
-	 * the list of properties which are holding the reference to the TargetObject in that referencer.
-	 *
-	 * @todo - not yet implemented
-	 */
-//	TMap< T*, TMultiMap<UObject*, UProperty*> >		ReferenceProperties;
-};
-
-/*----------------------------------------------------------------------------
-	FArchiveFindCulprit.
-----------------------------------------------------------------------------*/
-/**
- * Archive for finding who references an object.
- */
-class FArchiveFindCulprit : public FArchiveUObject
-{
-public:
-	/**
-	 * Constructor
-	 *
-	 * @param	InFind	the object that we'll be searching for references to
-	 * @param	Src		the object to serialize which may contain a reference to InFind
-	 * @param	InPretendSaving		if true, marks the archive as saving and persistent, so that a different serialization codepath is followed
-	 */
-	COREUOBJECT_API FArchiveFindCulprit( UObject* InFind, UObject* Src, bool InPretendSaving );
-
-	int32 GetCount() const
-	{
-		return Count;
-	}
-	int32 GetCount( TArray<const UProperty*>& Properties )
-	{
-		Properties = Referencers;
-		return Count;
-	}
-
-	/**
-	 * Returns the name of the Archive.  Useful for getting the name of the package a struct or object
-	 * is in when a loading error occurs.
-	 *
-	 * This is overridden for the specific Archive Types
-	 **/
-	virtual FString GetArchiveName() const { return TEXT("FArchiveFindCulprit"); }
-
-protected:
-	UObject*			Find;
-	int32					Count;
-	bool				PretendSaving;
-	class TArray<const UProperty*>	Referencers;
-
-private:
-	COREUOBJECT_API FArchive& operator<<( class UObject*& Obj );
-};
-
-struct FTraceRouteRecord
-{
-	struct FObjectGraphNode*	GraphNode;
-	TArray<UProperty*>			ReferencerProperties;
-
-	FTraceRouteRecord( struct FObjectGraphNode* InGraphNode, UProperty* InReferencerProperty)
-	: GraphNode(InGraphNode)
-	{
-		ReferencerProperties.Add(InReferencerProperty);
-	}
-
-	FTraceRouteRecord( struct FObjectGraphNode* InGraphNode, const TArray<UProperty*>&	InReferencerProperties )
-		: GraphNode(InGraphNode)
-	{
-		ReferencerProperties = InReferencerProperties;
-	}
-
-	void Add(UProperty* InReferencerProperty)
-	{
-		ReferencerProperties.Add(InReferencerProperty);
-	}
-};
-
-struct FObjectGraphNode
-{
-	/** the object this node represents */
-	UObject*	NodeObject;
-
-	/** Records for objects which reference this object */
-	TMap<UObject*, FTraceRouteRecord>	ReferencerRecords;
-
-	/** Records for objects which are referenced by this object */
-	TMap<UObject*, FTraceRouteRecord>	ReferencedObjects;
-
-	/** the number of links between NodeObject and the target object */
-	int32									ReferenceDepth;
-
-	/** Used during search - Visited or not */
-	bool							Visited;
-
-	/**
-	 * The property that references NodeObject; only set on nodes which are part
-	 * of the calculated shortest route
-	 */
-	TArray<UProperty*>							ReferencerProperties;
-
-	/** Default constructor */
-	FObjectGraphNode( UObject* InNodeObject=NULL )
-	:	NodeObject(InNodeObject)
-	,	ReferenceDepth(MAX_int32)
-	,	Visited(false)
-	{}
-};
-
-/*----------------------------------------------------------------------------
-	FArchiveFindCulprit.
-----------------------------------------------------------------------------*/
-// This is from FArchiveTraceRoute -This only creates object graph of all objects 
-// This can be used by other classes such as FTraceReferences - trace references of one object
-class FArchiveObjectGraph : public FArchiveUObject
-{
-	/** Handles serialization of UObject references */
-	FArchive& operator<<( class UObject*& Obj );
-
-	/**
-	* The object currently being serialized; used by the overloaded serialization operator to determine the referencing object.
-	*/
-	UObject* CurrentReferencer;
-	/** The set of objects encountered while serializing CurrentReferencer */
-	TArray<UObject*> ObjectsToSerialize;
-
-	/** true if we should serialize objects marked RF_Transient */
-	bool bIncludeTransients;
-
-	/**
-	* A bitmask of object flags which indicates which objects should be included in the group of initial objects to be serialized;
-	* RF_RootSet will be automatically added to this bitmask, and OBJECTMARK_TagExp will automatically be removed.
-	*/
-	EObjectFlags	RequiredFlags;
-
-public:
-	FArchiveObjectGraph(bool IncludeTransients, EObjectFlags KeepFlags);
-	~FArchiveObjectGraph();
-
-	/**
-	* Serializes the objects in the specified set; any objects encountered during serialization
-	* of an object are added to the object set and processed until no new objects are added.
-	* DO NOT MAKE THIS VIRTUAL - this is called by constructor. If you wish to do so, please change where be called
-	* @param	Objects		the original set of objects to serialize; the original set will be preserved.
-	*/
-	void GenerateObjectGraph( TArray<UObject*>& Objects );
-
-	void ClearSearchFlags();
-
-	/** A complete graph of all references between all objects in memory */
-	TMap<UObject*, FObjectGraphNode*> ObjectGraph;
-};
-
-// This traces referenced/referencer of an object using FArchiveObjectGraph 
-class FTraceReferences
-{
-	FArchiveObjectGraph ArchiveObjectGraph;
-
-	// internal recursive function for referencers/referenced
-	void GetReferencerInternal( UObject* CurrentObject, TArray<FObjectGraphNode*> &OutReferencer, int32 CurrentDepth, int32 TargetDepth );
-	void GetReferencedInternal( UObject* CurrentObject, TArray<FObjectGraphNode*> &OutReferenced, int32 CurrentDepth, int32 TargetDepth );
-
-public:
-	FTraceReferences( bool bIncludeTransients = false, EObjectFlags KeepFlags = RF_AllFlags );
-
-	// returns referencer string of an object
-	FString GetReferencerString( UObject* Object, int32 Depth = 100 ); 
-	// returns referenced string of an object	
-	FString GetReferencedString( UObject* Object, int32 Depth = 100 );
-
-	// returns referencer object list of an object	
-	int32 GetReferencer( UObject* Object, TArray<FObjectGraphNode*> &Referencer, bool bExcludeSelf=true, int32 Depth = 100 );
-	// returns referenced object list of an object		
-	int32 GetReferenced( UObject* Object, TArray<FObjectGraphNode*> &Referenced, bool bExcludeSelf=true, int32 Depth = 100 );
-};
-/**
- * Archive for finding shortest path from root to a particular object.
- * Depth-first search.
- */
-class FArchiveTraceRoute : public FArchiveUObject
-{
-	/**
-	 * Simple container struct for grouping two FObjectGraphNodes together.
-	 */
-	struct FRouteLink
-	{
-		/**
-		 * Constructor
-		 */
-		FRouteLink( FObjectGraphNode* InParent=NULL, FObjectGraphNode* InChild=NULL )
-		: LinkParent(InParent), LinkChild(InChild)
-		{}
-
-		/**
-		 * The node corresponding to the "referencing" object.
-		 */
-		FObjectGraphNode* LinkParent;
-
-		/**
-		 * The node corresponding to the "referenced" object.
-		 */
-		FObjectGraphNode* LinkChild;
-	};
-
-public:
-	static COREUOBJECT_API TMap<UObject*,UProperty*> FindShortestRootPath( UObject* Object, bool bIncludeTransients, EObjectFlags KeepFlags );
-
-	/**
-	 * Retuns path to root created by e.g. FindShortestRootPath via a string.
-	 *
-	 * @param TargetObject	object marking one end of the route
-	 * @param Route			route to print to log.
-	 * @param String of root path
-	 */
-	static COREUOBJECT_API FString PrintRootPath( const TMap<UObject*,UProperty*>& Route, const UObject* TargetObject );
-
-	/**
-	 * Returns the name of the Archive.  Useful for getting the name of the package a struct or object
-	 * is in when a loading error occurs.
-	 *
-	 * This is overridden for the specific Archive Types
-	 **/
-	virtual FString GetArchiveName() const { return TEXT("FArchiveTraceRoute"); }
-
-	/**
-	 * Serializes the objects in the specified set; any objects encountered during serialization
-	 * of an object are added to the object set and processed until no new objects are added.
-	 *
-	 * @param	Objects		the original set of objects to serialize; the original set will be preserved.
-	 */
-	void GenerateObjectGraph( TSparseArray<UObject*>& Objects );
-
-	/**
-	 * Recursively iterates over the referencing objects for the specified node, marking each with
-	 * the current Depth value.  Stops once it reaches a route root.
-	 *
-	 * @param	ObjectNode	the node to evaluate.
-	 */
-	void CalculateReferenceDepthsForNode( FObjectGraphNode* ObjectNode );
-
-	/**
-	 * Searches through the objects referenced by CurrentNode for a record with a Depth lower than LowestDepth.
-	 *
-	 * @param	CurrentNode		the node containing the list of referenced objects that will be searched.
-	 * @param	LowestDepth		the current number of links we are from the target object.
-	 * @param	ClosestLink		if a trace route record is found with a lower depth value than LowestDepth, the link is saved to this value.
-	 *
-	 * @return	true if a closer link was discovered; false if no links were closer than lowest depth, or if we've reached the target object.
-	 */
-	bool FindClosestLink( FObjectGraphNode* CurrentNode, int32& LowestDepth, FRouteLink& ClosestLink );
-
-private:
-	FArchiveTraceRoute( UObject* TargetObject, TMap<UObject*,FTraceRouteRecord>& InRoutes, bool bShouldIncludeTransients, EObjectFlags KeepFlags );
-	~FArchiveTraceRoute();
-
-	/** Handles serialization of UObject references */
-	FArchive& operator<<( class UObject*& Obj );
-
-	/** A complete graph of all references between all objects in memory */
-	TMap<UObject*, FObjectGraphNode*> ObjectGraph;
-
-	/**
-	 * The object currently being serialized; used by the overloaded serialization operator to determine the referencing object.
-	 */
-	UObject* CurrentReferencer;
-
-	/** The set of objects encountered while serializing CurrentReferencer */
-	TArray<UObject*> ObjectsToSerialize;
-
-	/** the current number of object reference links away from the target object */
-	int32 Depth;
-
-	/** true if we should serialize objects marked RF_Transient */
-	bool bIncludeTransients;
-
-	/**
-	 * A bitmask of object flags which indicates which objects should be included in the group of initial objects to be serialized;
-	 * RF_RootSet will be automatically added to this bitmask, and OBJECTMARK_TagExp will automatically be removed.
-	 */
-	EObjectFlags	RequiredFlags;
-
-};
-
-/*----------------------------------------------------------------------------
-	FDuplicateDataReader.
-----------------------------------------------------------------------------*/
-
-/**
- * Information about a duplicated object
- * For use with a dense object annotation
- */
-struct FDuplicatedObject
-{
-	/** The duplicated object */
-	UObject* DuplicatedObject;
-
-	FDuplicatedObject()
-		: DuplicatedObject(NULL)
-	{
-
-	}
-
-	FDuplicatedObject( UObject* InDuplicatedObject )
-		: DuplicatedObject( InDuplicatedObject )
-	{
-
-	}
-
-	/**
-	 * @return true if this is the default annotation and holds no information about a duplicated object
-	 */
-	FORCEINLINE bool IsDefault()
-	{
-		return DuplicatedObject == NULL;
-	}
-};
-
-template <> struct TIsPODType<FDuplicatedObject> { enum { Value = true }; };
-
-/**
- * Reads duplicated objects from a memory buffer, replacing object references to duplicated objects.
- */
-class FDuplicateDataReader : public FArchiveUObject
-{
-private:
-
-	class FUObjectAnnotationSparse<FDuplicatedObject,false>&	DuplicatedObjectAnnotation;
-	const TArray<uint8>&					ObjectData;
-	int32									Offset;
-
-	//~ Begin FArchive Interface.
-
-	virtual FArchive& operator<<(FName& N);
-	virtual FArchive& operator<<(UObject*& Object);
-	virtual FArchive& operator<<(FLazyObjectPtr& LazyObjectPtr);
-	virtual FArchive& operator<<(FAssetPtr& AssetPtr);
-	virtual FArchive& operator<<(FStringAssetReference& StringAssetReference);
-	
-	void SerializeFail();
-
-	virtual void Serialize(void* Data,int64 Num)
-	{
-		if(Num)
-		{
-			if (!(Offset + Num <= ObjectData.Num()))
-			{
-				SerializeFail();
-			}
-			FMemory::Memcpy(Data,&ObjectData[Offset],Num);
-			Offset += Num;
-		}
-	}
-
-	virtual void Seek(int64 InPos)
-	{
-		Offset = InPos;
-	}
-
-public:
-	/**
-	 * Returns the name of the Archive.  Useful for getting the name of the package a struct or object
-	 * is in when a loading error occurs.
-	 *
-	 * This is overridden for the specific Archive Types
-	 **/
-	virtual FString GetArchiveName() const { return TEXT("FDuplicateDataReader"); }
-
-	virtual int64 Tell()
-	{
-		return Offset;
-	}
-	virtual int64 TotalSize()
-	{
-		return ObjectData.Num();
-	}
-
-	/**
-	 * Constructor
-	 * 
-	 * @param	InDuplicatedObjectAnnotation		Annotation for storing a mapping from source to duplicated object
-	 * @param	InObjectData					Object data to read from
-	 */
-	FDuplicateDataReader( FUObjectAnnotationSparse<FDuplicatedObject,false>& InDuplicatedObjectAnnotation, const TArray<uint8>& InObjectData, uint32 InPortFlags, UObject* InDestOuter );
-};
-
-/*----------------------------------------------------------------------------
-	FDuplicateDataWriter.
-----------------------------------------------------------------------------*/
-/**
- * Writes duplicated objects to a memory buffer, duplicating referenced inner objects and adding the duplicates to the DuplicatedObjects map.
- */
-class FDuplicateDataWriter : public FArchiveUObject
-{
-private:
-
-	class FUObjectAnnotationSparse<FDuplicatedObject,false>&	DuplicatedObjectAnnotation;
-	TArray<uint8>&							ObjectData;
-	int64										Offset;
-	EObjectFlags							FlagMask;
-	EObjectFlags							ApplyFlags;
-	EInternalObjectFlags InternalFlagMask;
-	EInternalObjectFlags ApplyInternalFlags;
-
-	/**
-	 * This is used to prevent object & component instancing resulting from the calls to StaticConstructObject(); instancing subobjects and components is pointless,
-	 * since we do that manually and replace the current value with our manually created object anyway.
-	 */
-	struct FObjectInstancingGraph*			InstanceGraph;
-
-	//~ Begin FArchive Interface.
-
-	virtual FArchive& operator<<(FName& N);
-	virtual FArchive& operator<<(UObject*& Object);
-	virtual FArchive& operator<<(FLazyObjectPtr& LazyObjectPtr);
-	virtual FArchive& operator<<(FAssetPtr& AssetPtr);
-
-	virtual void Serialize(void* Data,int64 Num)
-	{
-		// Don't try to add/memcpy zero sized items
-		if (Data != NULL && Num > 0)
-		{
-			if(Offset == ObjectData.Num())
-			{
-				ObjectData.AddUninitialized(Num);
-			}
-			FMemory::Memcpy(&ObjectData[Offset],Data,Num);
-			Offset += Num;
-		}
-	}
-
-	virtual void Seek(int64 InPos)
-	{
-		Offset = InPos;
-	}
-
-	/**
-	 * Places a new duplicate in the DuplicatedObjects map as well as the UnserializedObjects list
-	 * 
-	 * @param	SourceObject	the original version of the object
-	 * @param	DuplicateObject	the copy of the object
-	 *
-	 * @return	a pointer to the copy of the object
-	 */
-	UObject* AddDuplicate(UObject* SourceObject,UObject* DuplicateObject);
-
-public:
-	/**
-	 * Returns the name of the Archive.  Useful for getting the name of the package a struct or object
-	 * is in when a loading error occurs.
-	 *
-	 * This is overridden for the specific Archive Types
-	 **/
-	virtual FString GetArchiveName() const { return TEXT("FDuplicateDataWriter"); }
-
-	virtual int64 Tell()
-	{
-		return Offset;
-	}
-
-	virtual int64 TotalSize()
-	{
-		return ObjectData.Num();
-	}
-	TArray<UObject*>	UnserializedObjects;
-
-	/**
-	 * Returns a pointer to the duplicate of a given object, creating the duplicate object if necessary.
-	 * 
-	 * @param	Object	the object to find a duplicate for
-	 * @param	bCreateIfMissing Create the duplicated object if it's missing.
-	 *
-	 * @return	a pointer to the duplicate of the specified object
-	 */
-	UObject* GetDuplicatedObject(UObject* Object, bool bCreateIfMissing = true);
-
-	/**
-	 * Constructor
-	 * 
-	 * @param	InDuplicatedObjects		Annotation for storing a mapping from source to duplicated object
-	 * @param	InObjectData			will store the serialized data
-	 * @param	SourceObject			the object to copy
-	 * @param	DestObject				the object to copy to
-	 * @param	InFlagMask				the flags that should be copied when the object is duplicated
-	 * @param	InApplyFlags			the flags that should always be set on the duplicated objects (regardless of whether they're set on the source)
-	 * @param	InInstanceGraph			the instancing graph to use when creating the duplicate objects.
-	 */
-	FDuplicateDataWriter(
-		FUObjectAnnotationSparse<FDuplicatedObject, false>& InDuplicatedObjects,
-		TArray<uint8>& InObjectData,
-		UObject* SourceObject,
-		UObject* DestObject,
-		EObjectFlags InFlagMask,
-		EObjectFlags InApplyMask,
-		EInternalObjectFlags InInternalFlagMask,
-		EInternalObjectFlags InApplyInternalFlags,
-		FObjectInstancingGraph* InInstanceGraph,
-		uint32 InPortFlags);
-};
-
-/** Base class for object replacement archives */ 
-class COREUOBJECT_API FArchiveReplaceObjectRefBase : public FArchiveUObject
-{
-protected:
-	/**
-	* Serializes a single object
-	*/
-	void SerializeObject(UObject* ObjectToSerialize);
-};
-
-/*----------------------------------------------------------------------------
-	FArchiveReplaceObjectRef.
-----------------------------------------------------------------------------*/
-/**
- * Archive for replacing a reference to an object. This classes uses
- * serialization to replace all references to one object with another.
- * Note that this archive will only traverse objects with an Outer
- * that matches InSearchObject.
- *
- * NOTE: The template type must be a child of UObject or this class will not compile.
- */
-template< class T >
-class FArchiveReplaceObjectRef : public FArchiveReplaceObjectRefBase
-{
-public:
-	/**
-	 * Initializes variables and starts the serialization search
-	 *
-	 * @param InSearchObject		The object to start the search on
-	 * @param ReplacementMap		Map of objects to find -> objects to replace them with (null zeros them)
-	 * @param bNullPrivateRefs		Whether references to non-public objects not contained within the SearchObject
-	 *								should be set to null
-	 * @param bIgnoreOuterRef		Whether we should replace Outer pointers on Objects.
-	 * @param bIgnoreArchetypeRef	Whether we should replace the ObjectArchetype reference on Objects.
-	 * @param bDelayStart			Specify true to prevent the constructor from starting the process.  Allows child classes' to do initialization stuff in their ctor
-	 */
-	FArchiveReplaceObjectRef
-	(
-		UObject* InSearchObject,
-		const TMap<T*,T*>& inReplacementMap,
-		bool bNullPrivateRefs,
-		bool bIgnoreOuterRef,
-		bool bIgnoreArchetypeRef,
-		bool bDelayStart = false,
-		bool bIgnoreClassGeneratedByRef = true
-	)
-	: SearchObject(InSearchObject), ReplacementMap(inReplacementMap)
-	, Count(0), bNullPrivateReferences(bNullPrivateRefs)
-	{
-		ArIsObjectReferenceCollector = true;
-		ArIsModifyingWeakAndStrongReferences = true;		// Also replace weak references too!
-		ArIgnoreArchetypeRef = bIgnoreArchetypeRef;
-		ArIgnoreOuterRef = bIgnoreOuterRef;
-		ArIgnoreClassGeneratedByRef = bIgnoreClassGeneratedByRef;
-
-		if ( !bDelayStart )
-		{
-			SerializeSearchObject();
-		}
-	}
-
-	/**
-	 * Starts the serialization of the root object
-	 */
-	void SerializeSearchObject()
-	{
-		ReplacedReferences.Empty();
-
-		if (SearchObject != NULL && !SerializedObjects.Find(SearchObject)
-		&&	(ReplacementMap.Num() > 0 || bNullPrivateReferences))
-		{
-			// start the initial serialization
-			SerializedObjects.Add(SearchObject);
-			SerializeObject(SearchObject);
-			for (int32 Iter = 0; Iter < PendingSerializationObjects.Num(); Iter++)
-			{
-				SerializeObject(PendingSerializationObjects[Iter]);
-			}
-			PendingSerializationObjects.Reset();
-		}
-	}
-
-	/**
-	 * Returns the number of times the object was referenced
-	 */
-	int64 GetCount() const
-	{
-		return Count;
-	}
-
-	/**
-	 * Returns a reference to the object this archive is operating on
-	 */
-	const UObject* GetSearchObject() const { return SearchObject; }
-
-	/**
-	 * Returns a reference to the replaced references map
-	 */
-	const TMap<UObject*, TArray<UProperty*>>& GetReplacedReferences() const { return ReplacedReferences; }
-
-	/**
-	 * Serializes the reference to the object
-	 */
-	FArchive& operator<<( UObject*& Obj )
-	{
-		if (Obj != NULL)
-		{
-			// If these match, replace the reference
-			T* const* ReplaceWith = (T*const*)((const TMap<UObject*,UObject*>*)&ReplacementMap)->Find(Obj);
-			if ( ReplaceWith != NULL )
-			{
-				Obj = *ReplaceWith;
-				ReplacedReferences.FindOrAdd(Obj).AddUnique(GetSerializedProperty());
-				Count++;
-			}
-			// A->IsIn(A) returns false, but we don't want to NULL that reference out, so extra check here.
-			else if ( Obj == SearchObject || Obj->IsIn(SearchObject) )
-			{
-#if 0
-				// DEBUG: Log when we are using the A->IsIn(A) path here.
-				if(Obj == SearchObject)
-				{
-					FString ObjName = Obj->GetPathName();
-					UE_LOG(LogSerialization, Log,  TEXT("FArchiveReplaceObjectRef: Obj == SearchObject : '%s'"), *ObjName );
-				}
-#endif
-				bool bAlreadyAdded = false;
-				SerializedObjects.Add(Obj, &bAlreadyAdded);
-				if (!bAlreadyAdded)
-				{
-					// No recursion
-					PendingSerializationObjects.Add(Obj);
-				}
-			}
-			else if ( bNullPrivateReferences && !Obj->HasAnyFlags(RF_Public) )
-			{
-				Obj = NULL;
-			}
-		}
-		return *this;
-	}
-
-	/**
-	 * Returns the name of this archive.
-	 **/
-	virtual FString GetArchiveName() const { return TEXT("ReplaceObjectRef"); }
-
-protected:
-	/** Initial object to start the reference search from */
-	UObject* SearchObject;
-
-	/** Map of objects to find references to -> object to replace references with */
-	const TMap<T*,T*>& ReplacementMap;
-	
-	/** The number of times encountered */
-	int32 Count;
-
-	/** List of objects that have already been serialized */
-	TSet<UObject*> SerializedObjects;
-
-	/** Object that will be serialized */
-	TArray<UObject*> PendingSerializationObjects;
-
-	/** Map of referencing objects to referencing properties */
-	TMap<UObject*, TArray<UProperty*>> ReplacedReferences;
-
-	/**
-	 * Whether references to non-public objects not contained within the SearchObject
-	 * should be set to null
-	 */
-	bool bNullPrivateReferences;
-};
-
-/*----------------------------------------------------------------------------
-FArchiveReplaceOrClearExternalReferences.
-----------------------------------------------------------------------------*/
-/**
-* Identical to FArchiveReplaceObjectRef, but for references to private objects
-* in other packages we clear the reference instead of preserving it (unless it
-* makes it into the replacement map)
-*/
-template< class T >
-class FArchiveReplaceOrClearExternalReferences : public FArchiveReplaceObjectRef<T>
-{
-	typedef FArchiveReplaceObjectRef<T> TSuper;
-public:
-	FArchiveReplaceOrClearExternalReferences
-		( UObject* InSearchObject
-		, const TMap<T*, T*>& InReplacementMap
-		, UPackage* InDestPackage
-		, bool bDelayStart = false )
-		: TSuper(InSearchObject, InReplacementMap, false, false, false, true)
-		, DestPackage(InDestPackage)
-	{
-		if (!bDelayStart)
-		{
-			this->SerializeSearchObject();
-		}
-	}
-
-	FArchive& operator<<(UObject*& Obj)
-	{
-		UObject* Resolved = Obj;
-		TSuper::operator<<(Resolved);
-		// if Resolved is a private object in another package just clear the reference:
-		if (Resolved)
-		{
-			UObject* Outermost = Resolved->GetOutermost();
-			if (Outermost)
-			{
-				UPackage* ObjPackage = dynamic_cast<UPackage*>(Outermost);
-				if (ObjPackage)
-				{
-					CA_SUPPRESS(6011);
-					if (ObjPackage != Obj && 
-						DestPackage != ObjPackage && 
-						!Obj->HasAnyFlags(RF_Public))
-					{
-						Resolved = nullptr;
-					}
-				}
-			}
-		}
-		Obj = Resolved;
-		return *this;
-	}
-
-protected:
-	/** Package that we are loading into, references to private objects in other packages will be cleared */
-	UPackage* DestPackage;
-};
-
-/*----------------------------------------------------------------------------
-	FArchiveObjectPropertyMapper.
-----------------------------------------------------------------------------*/
-
-/**
- * Class for collecting references to objects, along with the properties that
- * references that object.
- */
-class FArchiveObjectPropertyMapper : public FArchiveUObject
-{
-public:
-	/**
-	 * Constructor
-	 *
-	 * @param	InObjectArray			Array to add object references to
-	 * @param	InBase					only objects with this outer will be considered, or NULL to disregard outers
-	 * @param	InLimitClass			only objects of this class (and children) will be considered, or null to disregard object class
-	 * @param	bInRequireDirectOuter	determines whether objects contained within 'InOuter', but that do not have an Outer
-	 *									of 'InOuter' are included.  i.e. for objects that have GetOuter()->GetOuter() == InOuter.
-	 *									If InOuter is NULL, this parameter is ignored.
-	 * @param	bInSerializeRecursively	only applicable when LimitOuter != NULL && bRequireDirectOuter==true;
-	 *									serializes each object encountered looking for subobjects of referenced
-	 *									objects that have LimitOuter for their Outer (i.e. nested subobjects/components)
-	 */
-	FArchiveObjectPropertyMapper( TMap<UProperty*,UObject*>* InObjectGraph, UObject* InOuter=NULL, UClass* InLimitClass=NULL, bool bInRequireDirectOuter=true, bool bInSerializeRecursively=true )
-	:	ObjectGraph( InObjectGraph ), LimitOuter(InOuter), LimitClass(InLimitClass), bRequireDirectOuter(bInRequireDirectOuter), bSerializeRecursively(bInSerializeRecursively)
-	{
-		ArIsObjectReferenceCollector = true;
-		bSerializeRecursively = bInSerializeRecursively && LimitOuter != NULL;
-	}
-private:
-	/** 
-	 * UObject serialize operator implementation
-	 *
-	 * @param Object	reference to Object reference
-	 * @return reference to instance of this class
-	 */
-	FArchive& operator<<( UObject*& Object )
-	{
-		// Avoid duplicate entries.
-		if ( Object != NULL )
-		{
-			if ((LimitClass == NULL || Object->IsA(LimitClass)) &&
-				(LimitOuter == NULL || (Object->GetOuter() == LimitOuter || (!bRequireDirectOuter && Object->IsIn(LimitOuter)))) )
-			{
-				ObjectGraph->Add(GetSerializedProperty(), Object);
-				if ( bSerializeRecursively && !ObjectArray.Contains(Object) )
-				{
-					ObjectArray.Add( Object );
-
-					// check this object for any potential object references
-					Object->Serialize(*this);
-				}
-			}
-		}
-
-		return *this;
-	}
-
-	/** Tracks the objects which have been serialized by this archive, to prevent recursion */
-	TArray<UObject*>			ObjectArray;
-
-	/** Stored pointer to array of objects we add object references to */
-	TMap<UProperty*,UObject*>*	ObjectGraph;
-
-	/** only objects with this outer will be considered, NULL value indicates that outers are disregarded */
-	UObject*			LimitOuter;
-
-	/** only objects of this type will be considered, NULL value indicates that all classes are considered */
-	UClass*				LimitClass;
-
-	/** determines whether nested objects contained within LimitOuter are considered */
-	bool				bRequireDirectOuter;
-
-	/** determines whether we serialize objects that are encounterd by this archive */
-	bool				bSerializeRecursively;
-};
-
-
-/**
- * This archive marks all objects referenced by the specified "root set" of objects.
- */
-class FArchiveReferenceMarker : public FArchiveUObject
-{
-public:
-	FArchiveReferenceMarker( TArray<UObject*>& SourceObjects )
-	{
-		ArIsObjectReferenceCollector = true;
-		ArIgnoreOuterRef = true;
-
-		for ( int32 ObjectIndex = 0; ObjectIndex < SourceObjects.Num(); ObjectIndex++ )
-		{
-			UObject* Object = SourceObjects[ObjectIndex];
-			Object->Mark(OBJECTMARK_TagImp);
-
-			// OBJECTMARK_TagImp is used to allow serialization of objects which we would otherwise ignore.
-			Object->Serialize(*this);
-		}
-
-		for ( int32 ObjectIndex = 0; ObjectIndex < SourceObjects.Num(); ObjectIndex++ )
-		{
-			UObject* Object = SourceObjects[ObjectIndex];
-			Object->UnMark(OBJECTMARK_TagImp);
-		}
-	}
-
-	/** 
-	* UObject serialize operator implementation
-	*
-	* @param Object	reference to Object reference
-	* @return reference to instance of this class
-	*/
-	FArchive& operator<<( UObject*& Object )
-	{
-		if (Object != NULL && !(Object->HasAnyMarks(OBJECTMARK_TagExp) || Object->IsPendingKillOrUnreachable()) )
-		{
-			Object->Mark(OBJECTMARK_TagExp);
-
-			const bool bIgnoreObject = 
-				// No need to call Serialize from here for any objects that were part of our root set.
-				// By preventing re-entrant serialization using the OBJECTMARK_TagImp flag (instead of just marking each object in the root set with
-				// OBJECTMARK_TagExp prior to calling Serialize) we can determine which objects from our root set are being referenced
-				// by other objects in our root set.
-				Object->HasAnyMarks(OBJECTMARK_TagImp);
-
-			if ( bIgnoreObject == false )
-			{
-				Object->Serialize( *this );
-			}
-		}
-
-		return *this;
-	}
-};
-
-#if !USE_NEW_ASYNC_IO
-/*----------------------------------------------------------------------------
-	FArchiveAsync.
-----------------------------------------------------------------------------*/
-
-/**
- * Rough and basic version of async archive. The code relies on Serialize only ever to be called on the last
- * precached region.
- */
-class COREUOBJECT_API FArchiveAsync final: public FArchive
-{
-public:
-	/**
-	 * Constructor, initializing member variables.
-	 */
-	FArchiveAsync( const TCHAR* InFileName );
-
-	/**
-	 * Virtual destructor cleaning up internal file reader.
-	 */
-	virtual ~FArchiveAsync();
-
-	/**
-	 * Close archive and return whether there has been an error.
-	 *
-	 * @return	true if there were NO errors, false otherwise
-	 */
-	virtual bool Close();
-
-	/**
-	 * Sets mapping from offsets/ sizes that are going to be used for seeking and serialization to what
-	 * is actually stored on disk. If the archive supports dealing with compression in this way it is 
-	 * going to return true.
-	 *
-	 * @param	CompressedChunks	Pointer to array containing information about [un]compressed chunks
-	 * @param	CompressionFlags	Flags determining compression format associated with mapping
-	 *
-	 * @return true if archive supports translating offsets & uncompressing on read, false otherwise
-	 */
-	virtual bool SetCompressionMap( TArray<FCompressedChunk>* CompressedChunks, ECompressionFlags CompressionFlags );
-
-	/**
-	 * Hint the archive that the region starting at passed in offset and spanning the passed in size
-	 * is going to be read soon and should be precached.
-	 *
-	 * The function returns whether the precache operation has completed or not which is an important
-	 * hint for code knowing that it deals with potential async I/O. The archive is free to either not 
-	 * implement this function or only partially precache so it is required that given sufficient time
-	 * the function will return true. Archives not based on async I/O should always return true.
-	 *
-	 * This function will not change the current archive position.
-	 *
-	 * @param	PrecacheOffset	Offset at which to begin precaching.
-	 * @param	PrecacheSize	Number of bytes to precache
-	 * @return	false if precache operation is still pending, true otherwise
-	 */
-	virtual bool Precache( int64 PrecacheOffset, int64 PrecacheSize );
-
-	/**
-	 * Serializes data from archive.
-	 *
-	 * @param	Data	Pointer to serialize to
-	 * @param	Num		Number of bytes to read
-	 */
-	virtual void Serialize( void* Data, int64 Num );
-
-	/**
-	 * Returns the current position in the archive as offset in bytes from the beginning.
-	 *
-	 * @return	Current position in the archive (offset in bytes from the beginning)
-	 */
-	virtual int64 Tell();
-	/**
-	 * Returns the total size of the archive in bytes.
-	 *
-	 * @return total size of the archive in bytes
-	 */
-	virtual int64 TotalSize();
-
-	/**
-	 * Sets the current position.
-	 *
-	 * @param InPos	New position (as offset from beginning in bytes)
-	 */
-	virtual void Seek( int64 InPos );
-
-	/**
-	 * Flushes cache and frees internal data.
-	 */
-	virtual void FlushCache();
-private:
-
-	/**
-	 * Swaps current and next buffer. Relies on calling code to ensure that there are no outstanding
-	 * async read operations into the buffers.
-	 */
-	void BufferSwitcheroo();
-
-	/**
-	 * Whether the current precache buffer contains the passed in request.
-	 *
-	 * @param	RequestOffset	Offset in bytes from start of file
-	 * @param	RequestSize		Size in bytes requested
-	 *
-	 * @return true if buffer contains request, false othwerise
-	 */
-	FORCEINLINE bool PrecacheBufferContainsRequest( int64 RequestOffset, int64 RequestSize );
-
-	/**
-	 * Finds and returns the compressed chunk index associated with the passed in offset.
-	 *
-	 * @param	RequestOffset	Offset in file to find associated chunk index for
-	 *
-	 * @return Index into CompressedChunks array matching this offset
-	 */
-	int32 FindCompressedChunkIndex( int64 RequestOffset );
-
-	/**
-	 * Precaches compressed chunk of passed in index using buffer at passed in index.
-	 *
-	 * @param	ChunkIndex	Index of compressed chunk
-	 * @param	BufferIndex	Index of buffer to precache into	
-	 */
-	void PrecacheCompressedChunk( int64 ChunkIndex, int64 BufferIndex );
-
-	/** Anon enum used to index precache data. */
-	enum
-	{
-		CURRENT = 0,
-		NEXT	= 1,
-	};
-
-	/** Cached filename for debugging.												*/
-	FString							FileName;
-	/** Cached file size															*/
-	int64							FileSize;
-	/** Cached uncompressed file size (!= FileSize for compressed packages.			*/
-	int64							UncompressedFileSize;
-	/** BulkData area size ( size of the bulkdata are at the end of the file. That is
-		not part of the CompressedChunks, as BulkData is compressed separately		*/
-	int64							BulkDataAreaSize;
-	/** Current position of archive.												*/
-	int64							CurrentPos;
-
-	/** Start position of current precache request.									*/
-	int64							PrecacheStartPos[2];
-	/** End position (exclusive) of current precache request.						*/
-	int64							PrecacheEndPos[2];
-	/** Buffer containing precached data.											*/
-	uint8*							PrecacheBuffer[2];
-	/** Precahce buffer size */
-	SIZE_T PrecacheBufferSize[2];
-	/** True if the precache buffer is protected */
-	bool PrecacheBufferProtected[2];
-
-	/** Status of pending read, a value of 0 means no outstanding reads.			*/
-	FThreadSafeCounter				PrecacheReadStatus[2];
-	
-	/** Mapping of compressed <-> uncompresses sizes and offsets, NULL if not used.	*/
-	TArray<FCompressedChunk>*		CompressedChunks;
-	/** Current index into compressed chunks array.									*/
-	int64							CurrentChunkIndex;
-	/** Compression flags determining compression of CompressedChunks.				*/
-	ECompressionFlags				CompressionFlags;
-	/** Caches the return value of FPlatformMisc::SupportsMultithreading (comes up in profiles often) */
-	bool PlatformIsSinglethreaded;
-};
-
-#else
-
-COREUOBJECT_API FArchive* NewFArchiveAsync2(const TCHAR* InFileName);
-
-COREUOBJECT_API void HintFutureReadDone(const TCHAR * FileName);
-
-COREUOBJECT_API void HintFutureRead(const TCHAR * FileName);
-
-#endif // USE_NEW_ASYNC_IO
-
-/*----------------------------------------------------------------------------
-FArchiveObjectCrc32
-----------------------------------------------------------------------------*/
-
-/**
-* Calculates a checksum on an object's serialized data stream.
-*/
-class COREUOBJECT_API FArchiveObjectCrc32 : public FArchiveUObject
-{
-public:
-	/**
-	* Default constructor.
-	*/
-	FArchiveObjectCrc32();
-
-	//~ Begin FArchive Interface
-	virtual void Serialize(void* Data, int64 Length);
-	virtual FArchive& operator<<(class FName& Name);
-	virtual FArchive& operator<<(class UObject*& Object);
-	virtual FString GetArchiveName() const { return TEXT("FArchiveObjectCrc32"); }
-	//~ End FArchive Interface
-
-	/**
-	* Serialize the given object, calculate and return its checksum.
-	*/
-	uint32 Crc32(UObject* Object, uint32 CRC = 0);
-
-protected:
-	/** Return if object was already serialized */
-	virtual bool CustomSerialize(class UObject* Object) { return false; }
-
-	/** Internal byte array used for serialization */
-	TArray<uint8> SerializedObjectData;
-	/** Internal archive used for serialization */
-	FMemoryWriter MemoryWriter;
-	/** Internal queue of object references awaiting serialization */
-	TQueue<UObject*> ObjectsToSerialize;
-	/** Internal currently serialized object */
-	const UObject* ObjectBeingSerialized;
-	/** Internal root object */
-	const UObject* RootObject;
 };
diff --git a/Engine/Source/Runtime/CoreUObject/Public/Serialization/ArchiveUObjectBase.h b/Engine/Source/Runtime/CoreUObject/Public/Serialization/ArchiveUObjectBase.h
--- a/Engine/Source/Runtime/CoreUObject/Public/Serialization/ArchiveUObjectBase.h
+++ b/Engine/Source/Runtime/CoreUObject/Public/Serialization/ArchiveUObjectBase.h
@@ -1,55 +1,5 @@
 // Copyright 1998-2016 Epic Games, Inc. All Rights Reserved.
 
-/*====================================================================================
-	ArchiveUObjectBase.h: Implements the FArchiveUObject class for serializing UObjects
-======================================================================================*/
-
 #pragma once
 
-/**
- * Base FArchive for serializing UObjects. Supports FLazyObjectPtr and FAssetPtr serialization.
- */
-class COREUOBJECT_API FArchiveUObject : public FArchive
-{
-public:
-
-	using FArchive::operator<<; // For visibility of the overloads we don't override
-
-	//~ Begin FArchive Interface
-	virtual FArchive& operator<< (class FLazyObjectPtr& Value) override;
-	virtual FArchive& operator<< (class FAssetPtr& Value) override;
-	virtual FArchive& operator<< (struct FStringAssetReference& Value) override;
-	//~ End FArchive Interface
-};
-
-class UProperty;
-
-/** Helper class to set and restore serialized property on an archive */
-class FSerializedPropertyScope
-{
-	FArchive& Ar;
-	UProperty* Property;
-	UProperty* OldProperty;
-#if WITH_EDITORONLY_DATA
-	void PushEditorOnlyProperty();
-	void PopEditorOnlyProperty();
-#endif
-public:
-	FSerializedPropertyScope(FArchive& InAr, UProperty* InProperty)
-		: Ar(InAr)
-		, Property(InProperty)
-	{
-		OldProperty = Ar.GetSerializedProperty();
-		Ar.SetSerializedProperty(Property);
-#if WITH_EDITORONLY_DATA
-		PushEditorOnlyProperty();
-#endif
-	}
-	~FSerializedPropertyScope()
-	{
-#if WITH_EDITORONLY_DATA
-		PopEditorOnlyProperty();
-#endif
-		Ar.SetSerializedProperty(OldProperty);
-	}
-};
+#include "ArchiveUObject.h"
diff --git a/Engine/Source/Runtime/CoreUObject/Public/Serialization/DuplicatedDataReader.h b/Engine/Source/Runtime/CoreUObject/Public/Serialization/DuplicatedDataReader.h
new file mode 100644
--- /dev/null
+++ b/Engine/Source/Runtime/CoreUObject/Public/Serialization/DuplicatedDataReader.h
@@ -0,0 +1,76 @@
+// Copyright 1998-2016 Epic Games, Inc. All Rights Reserved.
+
+#pragma once
+
+#include "AsyncFileHandle.h"
+#include "TextPackageNamespaceUtil.h"
+
+/*----------------------------------------------------------------------------
+	FDuplicateDataReader.
+----------------------------------------------------------------------------*/
+
+/**
+ * Reads duplicated objects from a memory buffer, replacing object references to duplicated objects.
+ */
+class FDuplicateDataReader : public FArchiveUObject
+{
+private:
+
+	class FUObjectAnnotationSparse<FDuplicatedObject,false>&	DuplicatedObjectAnnotation;
+	const TArray<uint8>&					ObjectData;
+	int32									Offset;
+
+	//~ Begin FArchive Interface.
+
+	virtual FArchive& operator<<(FName& N);
+	virtual FArchive& operator<<(UObject*& Object);
+	virtual FArchive& operator<<(FLazyObjectPtr& LazyObjectPtr);
+	virtual FArchive& operator<<(FAssetPtr& AssetPtr);
+	virtual FArchive& operator<<(FStringAssetReference& StringAssetReference);
+	
+	void SerializeFail();
+
+	virtual void Serialize(void* Data,int64 Num)
+	{
+		if(Num)
+		{
+			if (!(Offset + Num <= ObjectData.Num()))
+			{
+				SerializeFail();
+			}
+			FMemory::Memcpy(Data,&ObjectData[Offset],Num);
+			Offset += Num;
+		}
+	}
+
+	virtual void Seek(int64 InPos)
+	{
+		Offset = InPos;
+	}
+
+public:
+	/**
+	 * Returns the name of the Archive.  Useful for getting the name of the package a struct or object
+	 * is in when a loading error occurs.
+	 *
+	 * This is overridden for the specific Archive Types
+	 **/
+	virtual FString GetArchiveName() const { return TEXT("FDuplicateDataReader"); }
+
+	virtual int64 Tell()
+	{
+		return Offset;
+	}
+	virtual int64 TotalSize()
+	{
+		return ObjectData.Num();
+	}
+
+	/**
+	 * Constructor
+	 * 
+	 * @param	InDuplicatedObjectAnnotation		Annotation for storing a mapping from source to duplicated object
+	 * @param	InObjectData					Object data to read from
+	 */
+	FDuplicateDataReader( FUObjectAnnotationSparse<FDuplicatedObject,false>& InDuplicatedObjectAnnotation, const TArray<uint8>& InObjectData, uint32 InPortFlags, UObject* InDestOuter );
+};
diff --git a/Engine/Source/Runtime/CoreUObject/Public/Serialization/DuplicatedObject.h b/Engine/Source/Runtime/CoreUObject/Public/Serialization/DuplicatedObject.h
new file mode 100644
--- /dev/null
+++ b/Engine/Source/Runtime/CoreUObject/Public/Serialization/DuplicatedObject.h
@@ -0,0 +1,39 @@
+// Copyright 1998-2016 Epic Games, Inc. All Rights Reserved.
+
+#pragma once
+
+#include "AsyncFileHandle.h"
+#include "TextPackageNamespaceUtil.h"
+
+/**
+ * Information about a duplicated object
+ * For use with a dense object annotation
+ */
+struct FDuplicatedObject
+{
+	/** The duplicated object */
+	UObject* DuplicatedObject;
+
+	FDuplicatedObject()
+		: DuplicatedObject(NULL)
+	{
+
+	}
+
+	FDuplicatedObject( UObject* InDuplicatedObject )
+		: DuplicatedObject( InDuplicatedObject )
+	{
+
+	}
+
+	/**
+	 * @return true if this is the default annotation and holds no information about a duplicated object
+	 */
+	FORCEINLINE bool IsDefault()
+	{
+		return DuplicatedObject == NULL;
+	}
+};
+
+template <> struct TIsPODType<FDuplicatedObject> { enum { Value = true }; };
+
diff --git a/Engine/Source/Runtime/CoreUObject/Public/Serialization/ObjectAndNameAsStringProxyArchive.h b/Engine/Source/Runtime/CoreUObject/Public/Serialization/ObjectAndNameAsStringProxyArchive.h
new file mode 100644
--- /dev/null
+++ b/Engine/Source/Runtime/CoreUObject/Public/Serialization/ObjectAndNameAsStringProxyArchive.h
@@ -0,0 +1,38 @@
+// Copyright 1998-2016 Epic Games, Inc. All Rights Reserved.
+
+#pragma once
+
+#include "AsyncFileHandle.h"
+#include "TextPackageNamespaceUtil.h"
+struct FObjectInstancingGraph;
+
+/**
+ * Implements a proxy archive that serializes UObjects and FNames as string data.
+ *
+ * Expected use is:
+ *    FArchive* SomeAr = CreateAnAr();
+ *    FObjectAndNameAsStringProxyArchive Ar(*SomeAr);
+ *    SomeObject->Serialize(Ar);
+ *    FinalizeAr(SomeAr);
+ * 
+ * @param InInnerArchive The actual FArchive object to serialize normal data types (FStrings, INTs, etc)
+ */
+struct FObjectAndNameAsStringProxyArchive : public FNameAsStringProxyArchive
+{
+	/**
+	 * Creates and initializes a new instance.
+	 *
+	 * @param InInnerArchive - The inner archive to proxy.
+	 * @param bInLoadIfFindFails - Indicates whether to try and load a ref'd object if we don't find it
+	 */
+	FObjectAndNameAsStringProxyArchive(FArchive& InInnerArchive, bool bInLoadIfFindFails)
+		:	FNameAsStringProxyArchive(InInnerArchive)
+		,	bLoadIfFindFails(bInLoadIfFindFails)
+	{ }
+
+	/** If we fail to find an object during loading, try and load it. */
+	bool bLoadIfFindFails;
+
+	COREUOBJECT_API virtual FArchive& operator<<(class UObject*& Obj);
+};
+
diff --git a/Engine/Source/Runtime/CoreUObject/Public/Serialization/ObjectReader.h b/Engine/Source/Runtime/CoreUObject/Public/Serialization/ObjectReader.h
new file mode 100644
--- /dev/null
+++ b/Engine/Source/Runtime/CoreUObject/Public/Serialization/ObjectReader.h
@@ -0,0 +1,78 @@
+// Copyright 1998-2016 Epic Games, Inc. All Rights Reserved.
+
+#pragma once
+
+#include "AsyncFileHandle.h"
+#include "TextPackageNamespaceUtil.h"
+struct FObjectInstancingGraph;
+
+/**
+ * UObject Memory Reader Archive.
+ */
+class FObjectReader : public FMemoryArchive
+{
+public:
+	FObjectReader(UObject* Obj, TArray<uint8>& InBytes, bool bIgnoreClassRef = false, bool bIgnoreArchetypeRef = false)
+		: Bytes(InBytes)
+	{
+		ArIsLoading = true;
+		ArIsPersistent = false;
+		ArIgnoreClassRef = bIgnoreClassRef;
+		ArIgnoreArchetypeRef = bIgnoreArchetypeRef;
+
+#if USE_STABLE_LOCALIZATION_KEYS
+		if (GIsEditor && !(ArPortFlags & PPF_DuplicateForPIE))
+		{
+			SetLocalizationNamespace(TextNamespaceUtil::EnsurePackageNamespace(Obj));
+		}
+#endif // USE_STABLE_LOCALIZATION_KEYS
+
+		Obj->Serialize(*this);
+	}
+
+	//~ Begin FArchive Interface
+
+	int64 TotalSize()
+	{
+		return (int64)Bytes.Num();
+	}
+
+	void Serialize(void* Data, int64 Num)
+	{
+		if (Num && !ArIsError)
+		{
+			// Only serialize if we have the requested amount of data
+			if (Offset + Num <= TotalSize())
+			{
+				FMemory::Memcpy(Data, &Bytes[Offset], Num);
+				Offset += Num;
+			}
+			else
+			{
+				ArIsError = true;
+			}
+		}
+	}
+
+	COREUOBJECT_API virtual FArchive& operator<<( class FName& N ) override;
+	COREUOBJECT_API virtual FArchive& operator<<( class UObject*& Res ) override;
+	COREUOBJECT_API virtual FArchive& operator<<( FLazyObjectPtr& LazyObjectPtr ) override;
+	COREUOBJECT_API virtual FArchive& operator<<( FAssetPtr& AssetPtr ) override;
+	COREUOBJECT_API virtual FArchive& operator<<(FStringAssetReference& AssetPtr) override;
+	COREUOBJECT_API virtual FString GetArchiveName() const override;
+	//~ End FArchive Interface
+
+
+
+protected:
+	FObjectReader(TArray<uint8>& InBytes)
+		: Bytes(InBytes)
+	{
+		ArIsLoading = true;
+		ArIsPersistent = false;
+		ArIgnoreClassRef = false;
+		ArIgnoreArchetypeRef = false;
+	}
+
+	const TArray<uint8>& Bytes;
+};
diff --git a/Engine/Source/Runtime/CoreUObject/Public/Serialization/ObjectWriter.h b/Engine/Source/Runtime/CoreUObject/Public/Serialization/ObjectWriter.h
new file mode 100644
--- /dev/null
+++ b/Engine/Source/Runtime/CoreUObject/Public/Serialization/ObjectWriter.h
@@ -0,0 +1,50 @@
+// Copyright 1998-2016 Epic Games, Inc. All Rights Reserved.
+
+#pragma once
+
+#include "AsyncFileHandle.h"
+#include "TextPackageNamespaceUtil.h"
+struct FObjectInstancingGraph;
+
+/**
+ * UObject Memory Writer Archive.
+ */
+class FObjectWriter : public FMemoryWriter
+{
+
+public:
+	FObjectWriter(UObject* Obj, TArray<uint8>& InBytes, bool bIgnoreClassRef = false, bool bIgnoreArchetypeRef = false, bool bDoDelta = true, uint32 AdditionalPortFlags = 0)
+		: FMemoryWriter(InBytes)
+	{
+		ArIgnoreClassRef = bIgnoreClassRef;
+		ArIgnoreArchetypeRef = bIgnoreArchetypeRef;
+		ArNoDelta = !bDoDelta;
+		ArPortFlags |= AdditionalPortFlags;
+
+#if USE_STABLE_LOCALIZATION_KEYS
+		if (GIsEditor && !(ArPortFlags & PPF_DuplicateForPIE))
+		{
+			SetLocalizationNamespace(TextNamespaceUtil::EnsurePackageNamespace(Obj));
+		}
+#endif // USE_STABLE_LOCALIZATION_KEYS
+
+		Obj->Serialize(*this);
+	}
+
+	//~ Begin FArchive Interface
+	COREUOBJECT_API virtual FArchive& operator<<( class FName& N ) override;
+	COREUOBJECT_API virtual FArchive& operator<<( class UObject*& Res ) override;
+	COREUOBJECT_API virtual FArchive& operator<<( FLazyObjectPtr& LazyObjectPtr ) override;
+	COREUOBJECT_API virtual FArchive& operator<<( FAssetPtr& AssetPtr ) override;
+	COREUOBJECT_API virtual FArchive& operator<<(FStringAssetReference& AssetPtr) override;
+	COREUOBJECT_API virtual FString GetArchiveName() const override;
+	//~ End FArchive Interface
+
+protected:
+	FObjectWriter(TArray<uint8>& InBytes)
+		: FMemoryWriter(InBytes)
+	{
+		ArIgnoreClassRef = false;
+		ArIgnoreArchetypeRef = false;
+	}
+};
diff --git a/Engine/Source/Runtime/CoreUObject/Public/Serialization/ReloadObjectArc.h b/Engine/Source/Runtime/CoreUObject/Public/Serialization/ReloadObjectArc.h
new file mode 100644
--- /dev/null
+++ b/Engine/Source/Runtime/CoreUObject/Public/Serialization/ReloadObjectArc.h
@@ -0,0 +1,176 @@
+// Copyright 1998-2016 Epic Games, Inc. All Rights Reserved.
+
+#pragma once
+
+#include "AsyncFileHandle.h"
+#include "TextPackageNamespaceUtil.h"
+struct FObjectInstancingGraph;
+
+/**
+ * Archive for reloading UObjects without requiring the UObject to be completely unloaded.
+ * Used for easily repropagating defaults to UObjects after class data has been changed.
+ * Aggregates FMemoryReader and FMemoryWriter to encapsulate management and coordination 
+ * of the UObject data being saved/loaded.
+ * <p>
+ * UObject references are not serialized directly into the memory archive.  Instead, we use
+ * a system similar to the Export/ImportMap of FLinker - the pointer to the UObject is added
+ * to a persistent (from the standpoint of the FReloadObjectArc) array.  The location into
+ * this array is what is actually stored in the archive's buffer.
+ * <p>
+ * This is primarily necessary for loading UObject references from the memory archive.  The
+ * UObject pointer ref. passed into the UObject overloaded serialization operator is not
+ * guaranteed to be valid when reading data from the memory archive.  Since we need to determine
+ * whether we *should* serialize the object before we actually do, we must introduce a level of
+ * indirection.  Using the index retrieved from the archive's buffer, we can look at the UObject
+ * before we attempt to serialize it.
+ */
+class COREUOBJECT_API FReloadObjectArc : public FArchiveUObject
+{
+public:
+
+	/**
+	 * Changes this memory archive to "read" mode, for reading UObject
+	 * data from the temporary location back into the UObjects.
+	 */
+	void ActivateReader()
+	{
+		ArIsSaving = false;
+		ArIsLoading = true;
+	}
+
+	/**
+	 * Changes this memory archive to "write" mode, for storing UObject
+	 * data to the temporary location.
+	 *
+	 * @note: called from ctors in child classes - should never be made virtual
+	 */
+	void ActivateWriter()
+	{
+		ArIsSaving = true;
+		ArIsLoading = false;
+	}
+
+	/**
+	 * Begin serializing a UObject into the memory archive.
+	 *
+	 * @param	Obj		the object to serialize
+	 */
+	void SerializeObject( UObject* Obj );
+
+	/**
+	 * Resets the archive so that it can be loaded from again from scratch
+	 * as if it was never serialized as a Reader
+	 */
+	void Reset();
+
+	/** FArchive Interface */
+	int64 TotalSize()
+	{
+		return Bytes.Num();
+	}
+	void Seek( int64 InPos )
+	{
+		if ( IsLoading() )
+			Reader.Seek(InPos);
+		else if ( IsSaving() )
+			Writer.Seek(InPos);
+	}
+	int64 Tell()
+	{
+		return IsLoading() ? Reader.Tell() : Writer.Tell();
+	}
+	FArchive& operator<<( class FName& Name );
+	FArchive& operator<<(class UObject*& Obj);
+
+	/** Constructor */
+	FReloadObjectArc();
+
+	/** Destructor */
+	virtual ~FReloadObjectArc();
+
+	/**
+	 * Returns the name of the Archive.  Useful for getting the name of the package a struct or object
+	 * is in when a loading error occurs.
+	 *
+	 * This is overridden for the specific Archive Types
+	 **/
+	virtual FString GetArchiveName() const { return TEXT("FReloadObjectArc"); }
+
+protected:
+
+	/**
+	 * Raw I/O function.  Routes the call to the appropriate archive, depending on whether
+	 * we're loading/saving.
+	 */
+	void Serialize( void* Data, int64 Num )
+	{
+		if ( IsLoading() )
+		{
+			Reader.Serialize(Data, Num);
+		}
+		else if ( IsSaving() )
+		{
+			Writer.Serialize(Data,Num);
+		}
+	}
+
+	/**
+	 * Sets the root object for this memory archive.
+	 * 
+	 * @param	NewRoot		the UObject that should be the new root
+	 */
+	void SetRootObject( UObject* NewRoot );
+
+	/** the raw UObject data contained by this archive */
+	TArray<uint8>		Bytes;
+
+	/** moves UObject data from storage into UObject address space */
+	FMemoryReader		Reader;
+
+	/** stores UObject data in a temporary location for later retrieval */
+	FMemoryWriter		Writer;
+
+	/** UObjects for which all data is stored in the memory archive */
+	TArray<UObject*>	CompleteObjects;
+
+	/** UObjects for which only a reference to the object is stored in the memory archive */
+	TArray<UObject*>	ReferencedObjects;
+
+	/**
+	 * List of top-level objects that have saved into the memory archive.  Used to prevent objects
+	 * from being serialized into storage multiple times.
+	 */
+	TSet<UObject*>	SavedObjects;
+
+	/**
+	 * List of top-level objects that have been loaded using the memory archive.  Used to prevent
+	 * objects from being serialized multiple times from the same memory archive.
+	 */
+	TSet<UObject*>	LoadedObjects;
+
+	/** A mapping of "UObject" => "the offset for that UObject's data in the Bytes array" for the objects stored in this archive */
+	TMap<UObject*,int32>	ObjectMap;
+
+	/**
+	 * This is the current top-level object.  For any UObjects contained
+	 * within this object, the complete UObject data will be stored in the
+	 * archive's buffer
+	 */
+	UObject*			RootObject;
+
+	/**
+	 * Used for tracking the subobjects and components that are instanced during this object reload.
+	 */
+	struct FObjectInstancingGraph*	InstanceGraph;
+
+	/**
+	 * Indicates whether this archive will serialize references to objects with the RF_Transient flag. (defaults to false)
+	 */
+	bool				bAllowTransientObjects;
+
+	/**
+	 * Indicates whether this archive should call InstanceSubobjectTemplates/InstanceSubobjects on objects that it re-initializes.
+	 * Specify false if the object needs special handling before calling InstanceSubobjectTemplates
+	 */
+	bool				bInstanceSubobjectsOnLoad;
+};
diff --git a/Engine/Source/Runtime/CoreUObject/Public/Serialization/SerializedPropertyScope.h b/Engine/Source/Runtime/CoreUObject/Public/Serialization/SerializedPropertyScope.h
new file mode 100644
--- /dev/null
+++ b/Engine/Source/Runtime/CoreUObject/Public/Serialization/SerializedPropertyScope.h
@@ -0,0 +1,37 @@
+// Copyright 1998-2016 Epic Games, Inc. All Rights Reserved.
+
+/*=============================================================================
+	AsyncPackage.h: Unreal async loading definitions.
+=============================================================================*/
+
+#pragma once
+
+/** Helper class to set and restore serialized property on an archive */
+class FSerializedPropertyScope
+{
+	FArchive& Ar;
+	UProperty* Property;
+	UProperty* OldProperty;
+#if WITH_EDITORONLY_DATA
+	void PushEditorOnlyProperty();
+	void PopEditorOnlyProperty();
+#endif
+public:
+	FSerializedPropertyScope(FArchive& InAr, UProperty* InProperty)
+		: Ar(InAr)
+		, Property(InProperty)
+	{
+		OldProperty = Ar.GetSerializedProperty();
+		Ar.SetSerializedProperty(Property);
+#if WITH_EDITORONLY_DATA
+		PushEditorOnlyProperty();
+#endif
+	}
+	~FSerializedPropertyScope()
+	{
+#if WITH_EDITORONLY_DATA
+		PopEditorOnlyProperty();
+#endif
+		Ar.SetSerializedProperty(OldProperty);
+	}
+};
