Copying //UE4/Dev-Build to //UE4/Dev-Main (Source: //UE4/Dev-Build @ 3104200)

diff --git a/Engine/Source/Runtime/CoreUObject/Public/Serialization/TraceReferences.h b/Engine/Source/Runtime/CoreUObject/Public/Serialization/TraceReferences.h
new file mode 100644
--- /dev/null
+++ b/Engine/Source/Runtime/CoreUObject/Public/Serialization/TraceReferences.h
@@ -0,0 +1,29 @@
+// Copyright 1998-2016 Epic Games, Inc. All Rights Reserved.
+
+#pragma once
+
+#include "AsyncFileHandle.h"
+#include "TextPackageNamespaceUtil.h"
+
+// This traces referenced/referencer of an object using FArchiveObjectGraph 
+class FTraceReferences
+{
+	FArchiveObjectGraph ArchiveObjectGraph;
+
+	// internal recursive function for referencers/referenced
+	void GetReferencerInternal( UObject* CurrentObject, TArray<FObjectGraphNode*> &OutReferencer, int32 CurrentDepth, int32 TargetDepth );
+	void GetReferencedInternal( UObject* CurrentObject, TArray<FObjectGraphNode*> &OutReferenced, int32 CurrentDepth, int32 TargetDepth );
+
+public:
+	FTraceReferences( bool bIncludeTransients = false, EObjectFlags KeepFlags = RF_AllFlags );
+
+	// returns referencer string of an object
+	FString GetReferencerString( UObject* Object, int32 Depth = 100 ); 
+	// returns referenced string of an object	
+	FString GetReferencedString( UObject* Object, int32 Depth = 100 );
+
+	// returns referencer object list of an object	
+	int32 GetReferencer( UObject* Object, TArray<FObjectGraphNode*> &Referencer, bool bExcludeSelf=true, int32 Depth = 100 );
+	// returns referenced object list of an object		
+	int32 GetReferenced( UObject* Object, TArray<FObjectGraphNode*> &Referenced, bool bExcludeSelf=true, int32 Depth = 100 );
+};
diff --git a/Engine/Source/Runtime/CoreUObject/Public/UObject/ErrorException.h b/Engine/Source/Runtime/CoreUObject/Public/UObject/ErrorException.h
new file mode 100644
--- /dev/null
+++ b/Engine/Source/Runtime/CoreUObject/Public/UObject/ErrorException.h
@@ -0,0 +1,35 @@
+// Copyright 1998-2016 Epic Games, Inc. All Rights Reserved.
+
+#pragma once
+
+#if HACK_HEADER_GENERATOR
+	#include "Templates/IsValidVariadicFunctionArg.h"
+	#include "Templates/AndOr.h"
+
+	/** 
+	 * FError
+	 * Set of functions for error reporting 
+	 **/
+	struct COREUOBJECT_API FError
+	{
+		/**
+		 * Throws a printf-formatted exception as a const TCHAR*.
+		 */
+		template <typename... Types>
+		FUNCTION_NO_RETURN_START
+		static void VARARGS Throwf(const TCHAR* Fmt, Types... Args)
+		FUNCTION_NO_RETURN_END
+		{
+			static_assert(TAnd<TIsValidVariadicFunctionArg<Types>...>::Value, "Invalid argument(s) passed to FError::Throwf");
+
+			ThrowfImpl(Fmt, Args...);
+		}
+
+	private:
+		FUNCTION_NO_RETURN_START
+		static void VARARGS ThrowfImpl(const TCHAR* Fmt, ...)
+		FUNCTION_NO_RETURN_END;
+	};
+#endif
+
+
diff --git a/Engine/Source/Runtime/CoreUObject/Public/UObject/GarbageCollection.h b/Engine/Source/Runtime/CoreUObject/Public/UObject/GarbageCollection.h
--- a/Engine/Source/Runtime/CoreUObject/Public/UObject/GarbageCollection.h
+++ b/Engine/Source/Runtime/CoreUObject/Public/UObject/GarbageCollection.h
@@ -6,7 +6,7 @@
 
 #pragma once
 
-#include "ArchiveUObjectBase.h"
+#include "ArchiveUObject.h"
 
 /** Context sensitive keep flags for garbage collection */
 #define GARBAGE_COLLECTION_KEEPFLAGS	(GIsEditor ? RF_Standalone : RF_NoFlags)
diff --git a/Engine/Source/Runtime/CoreUObject/Public/UObject/Interface.h b/Engine/Source/Runtime/CoreUObject/Public/UObject/Interface.h
new file mode 100644
--- /dev/null
+++ b/Engine/Source/Runtime/CoreUObject/Public/UObject/Interface.h
@@ -0,0 +1,27 @@
+// Copyright 1998-2016 Epic Games, Inc. All Rights Reserved.
+
+#pragma once
+
+#include "ObjectMacros.h"
+
+/**
+ * Base class for all interfaces
+ *
+ */
+
+class COREUOBJECT_API UInterface : public UObject
+{
+	DECLARE_CLASS_INTRINSIC(UInterface, UObject, CLASS_Interface | CLASS_Abstract, TEXT("/Script/CoreUObject"))
+};
+
+class COREUOBJECT_API IInterface
+{
+protected:
+
+	virtual ~IInterface() {}
+
+public:
+
+	typedef UInterface UClassType;
+};
+
diff --git a/Engine/Source/Runtime/CoreUObject/Public/UObject/LevelGuids.h b/Engine/Source/Runtime/CoreUObject/Public/UObject/LevelGuids.h
new file mode 100644
--- /dev/null
+++ b/Engine/Source/Runtime/CoreUObject/Public/UObject/LevelGuids.h
@@ -0,0 +1,22 @@
+// Copyright 1998-2016 Epic Games, Inc. All Rights Reserved.
+
+#pragma once
+
+#include "ObjectMacros.h"
+#include "WorldCompositionUtility.h"
+#include "GatherableTextData.h"
+#include "PropertyLocalizationDataGathering.h"
+#include "Package.h"
+#include "MetaData.h"
+
+/**
+ * Structure to hold information about an external packages objects used in cross-level references
+ */
+struct FLevelGuids
+{
+	/** Name of the external level */
+	FName LevelName;
+
+	/** Array of Guids possible in the other level (can be emptied out if all references are resolved after level load) */
+	TArray<FGuid> Guids;
+};
diff --git a/Engine/Source/Runtime/CoreUObject/Public/UObject/Linker.h b/Engine/Source/Runtime/CoreUObject/Public/UObject/Linker.h
--- a/Engine/Source/Runtime/CoreUObject/Public/UObject/Linker.h
+++ b/Engine/Source/Runtime/CoreUObject/Public/UObject/Linker.h
@@ -2,2317 +2,507 @@
 
 #pragma once
 
-#include "ObjectBase.h"
+#include "ObjectMacros.h"
 #include "EngineVersion.h"
 #include "GatherableTextData.h"
+#include "PackageFileSummary.h"
+#include "ObjectResource.h"
 
 
 DECLARE_LOG_CATEGORY_EXTERN(LogLinker, Log, All);
 
 /**
- * Wrapper for index into a ULnker's ImportMap or ExportMap.
- * Values greater than zero indicate that this is an index into the ExportMap.  The
- * actual array index will be (FPackageIndex - 1).
- *
- * Values less than zero indicate that this is an index into the ImportMap. The actual
- * array index will be (-FPackageIndex - 1)
+ * Information about a compressed chunk in a file.
  */
-class FPackageIndex
+struct FCompressedChunk
 {
-	/**
-	 * Values greater than zero indicate that this is an index into the ExportMap.  The
-	 * actual array index will be (FPackageIndex - 1).
-	 *
-	 * Values less than zero indicate that this is an index into the ImportMap. The actual
-	 * array index will be (-FPackageIndex - 1)
-	 */
-	int32 Index;
+	/** Default constructor, zero initializing all members. */
+	FCompressedChunk();
 
-	/** Internal constructor, sets the index directly **/
-	FORCEINLINE explicit FPackageIndex(int32 InIndex)
-		: Index(InIndex)
-	{
+	/** Original offset in uncompressed file.	*/
+	int32		UncompressedOffset;
+	/** Uncompressed size in bytes.				*/
+	int32		UncompressedSize;
+	/** Offset in compressed file.				*/
+	int32		CompressedOffset;
+	/** Compressed size in bytes.				*/
+	int32		CompressedSize;
 
-	}
+	/** I/O function */
+	friend COREUOBJECT_API FArchive& operator<<(FArchive& Ar,FCompressedChunk& Chunk);
+};
+
+
+class FLinkerTables
+{
 public:
-	/** Constructor, sets the value to null **/
-	FORCEINLINE FPackageIndex()
-		: Index(0)
-	{
+	/** The list of FObjectImports found in the package */
+	TArray<FObjectImport> ImportMap;
+	/** The list of FObjectExports found in the package */
+	TArray<FObjectExport> ExportMap;
+	/** List of dependency lists for each export */
+	TArray<TArray<FPackageIndex> > DependsMap;
+	/** Map that holds info about string asset references from the package. */
+	TArray<FString> StringAssetReferencesMap;
 
-	}
-	/** return true if this is an index into the import map **/
-	FORCEINLINE bool IsImport() const
+	/**
+	 * Check that this Index is non-null and return an import or export
+	 * @param	Index	Package index to get
+	 * @return	the resource corresponding to this index
+	 */
+	FORCEINLINE FObjectResource& ImpExp(FPackageIndex Index)
 	{
-		return Index < 0;
+		check(!Index.IsNull());
+		if (Index.IsImport())
+		{
+			return Imp(Index);
+		}
+		else
+		{
+			return Exp(Index);
+		}
 	}
-	/** return true if this is an index into the export map **/
-	FORCEINLINE bool IsExport() const
+	/**
+	 * Check that this Index is non-null and return an import or export
+	 * @param	Index	Package index to get
+	 * @return	the resource corresponding to this index
+	 */
+	FORCEINLINE FObjectResource const& ImpExp(FPackageIndex Index) const
 	{
-		return Index > 0;
-	}
-	/** return true if this null (i.e. neither an import nor an export) **/
-	FORCEINLINE bool IsNull() const
+		check(!Index.IsNull());
+		if (Index.IsImport())
+		{
+			return Imp(Index);
+		}
+		else
+		{
+			return Exp(Index);
+		}
+	}	
+	/**
+	 * Return an import or export for this index
+	 * @param	Index	Package index to get
+	 * @return	the resource corresponding to this index, or NULL if the package index is null
+	 */
+	FORCEINLINE FObjectResource* ImpExpPtr(FPackageIndex Index)
 	{
-		return Index == 0;
+		if (Index.IsImport())
+		{
+			return ImpPtr(Index);
+		}
+		else
+		{
+			return ExpPtr(Index);
+		}
 	}
-	/** Check that this is an import and return the index into the import map **/
-	FORCEINLINE int32 ToImport() const
+
+	/**
+	 * Check that this Index is non-null and an import and return an import
+	 * @param	Index	Package index to get, must be an import
+	 * @return	the import corresponding to this index
+	 */
+	FORCEINLINE FObjectImport& Imp(FPackageIndex Index)
 	{
-		check(IsImport());
-		return -Index - 1;
+		check(Index.IsImport() && ImportMap.IsValidIndex(Index.ToImport()));
+		return ImportMap[Index.ToImport()];
 	}
-	/** Check that this is an export and return the index into the export map **/
-	FORCEINLINE int32 ToExport() const
+	FORCEINLINE FObjectImport const& Imp(FPackageIndex Index) const
 	{
-		check(IsExport());
-		return Index - 1;
+		check(Index.IsImport() && ImportMap.IsValidIndex(Index.ToImport()));
+		return ImportMap[Index.ToImport()];
 	}
-	/** Return the raw value, for debugging purposes**/
-	FORCEINLINE int32 ForDebugging() const
+	/**
+	 * Return an import for this index
+	 * @param	Index	Package index to get
+	 * @return	the import corresponding to this index, or NULL if the package index is null or an export
+	 */
+	FORCEINLINE FObjectImport* ImpPtr(FPackageIndex Index)
 	{
-		return Index;
+		if (Index.IsImport())
+		{
+			check(ImportMap.IsValidIndex(Index.ToImport()));
+			return &ImportMap[Index.ToImport()];
+		}
+		return NULL;
 	}
 
-	/** Create a FPackageIndex from an import index **/
-	FORCEINLINE static FPackageIndex FromImport(int32 ImportIndex)
-	{
-		check(ImportIndex >= 0);
-		return FPackageIndex(-ImportIndex - 1);
-	}
-	/** Create a FPackageIndex from an export index **/
-	FORCEINLINE static FPackageIndex FromExport(int32 ExportIndex)
+	/**
+	 * Check that this Index is non-null and an export and return an import
+	 * @param	Index	Package index to get, must be an export
+	 * @return	the export corresponding to this index
+	 */
+	FORCEINLINE FObjectExport& Exp(FPackageIndex Index)
 	{
-		check(ExportIndex >= 0);
-		return FPackageIndex(ExportIndex + 1);
+		check(Index.IsExport() && ExportMap.IsValidIndex(Index.ToExport()));
+		return ExportMap[Index.ToExport()];
 	}
-
-	/** Compare package indecies for equality **/
-	FORCEINLINE bool operator==(const FPackageIndex& Other) const
+	FORCEINLINE FObjectExport const& Exp(FPackageIndex Index) const
 	{
-		return Index == Other.Index;
+		check(Index.IsExport() && ExportMap.IsValidIndex(Index.ToExport()));
+		return ExportMap[Index.ToExport()];
 	}
-	/** Compare package indecies for inequality **/
-	FORCEINLINE bool operator!=(const FPackageIndex& Other) const
+	/**
+	 * Return an export for this index
+	 * @param	Index	Package index to get
+	 * @return	the export corresponding to this index, or NULL if the package index is null or an import
+	 */
+	FORCEINLINE FObjectExport* ExpPtr(FPackageIndex Index)
 	{
-		return Index != Other.Index;
+		if (Index.IsExport())
+		{
+			check(ExportMap.IsValidIndex(Index.ToExport()));
+			return &ExportMap[Index.ToExport()];
+		}
+		return NULL;
 	}
+};
 
-	/** Compare package indecies **/
-	FORCEINLINE bool operator<(const FPackageIndex& Other) const
-	{
-		return Index < Other.Index;
-	}
-	FORCEINLINE bool operator>(const FPackageIndex& Other) const
+
+struct FLinkerNamePairKeyFuncs : DefaultKeyFuncs<FName, false>
+{
+	static FORCEINLINE bool Matches(FName A, FName B)
 	{
-		return Index > Other.Index;
+		// The linker requires that FNames preserve case, but the numeric suffix can be ignored since
+		// that is stored separately for each FName instance saved
+		return A.IsEqual(B, ENameCase::CaseSensitive, false/*bCompareNumber*/);
 	}
-	FORCEINLINE bool operator<=(const FPackageIndex& Other) const
+
+	static FORCEINLINE uint32 GetKeyHash(FName Key)
 	{
-		return Index <= Other.Index;
+		return Key.GetComparisonIndex();
 	}
-	FORCEINLINE bool operator>=(const FPackageIndex& Other) const
+};
+
+
+template<typename ValueType>
+struct TLinkerNameMapKeyFuncs : TDefaultMapKeyFuncs<FName, ValueType, false>
+{
+	static FORCEINLINE bool Matches(FName A, FName B)
 	{
-		return Index >= Other.Index;
+		// The linker requires that FNames preserve case, but the numeric suffix can be ignored since
+		// that is stored separately for each FName instance saved
+		return A.IsEqual(B, ENameCase::CaseSensitive, false/*bCompareNumber*/);
 	}
-	/**
-	 * Serializes a package index value from or into an archive.
-	 *
-	 * @param Ar - The archive to serialize from or to.
-	 * @param Value - The value to serialize.
-	 */
-	FORCEINLINE friend FArchive& operator<<(FArchive& Ar, FPackageIndex& Value)
+
+	static FORCEINLINE uint32 GetKeyHash(FName Key)
 	{
-		Ar << Value.Index;
-		return Ar;
+		return Key.GetComparisonIndex();
 	}
 };
 
 
-/**
- * Base class for UObject resource types.  FObjectResources are used to store UObjects on disk
- * via FLinker's ImportMap (for resources contained in other packages) and ExportMap (for resources
- * contained within the same package)
- */
-struct FObjectResource
+/*----------------------------------------------------------------------------
+	FLinker.
+----------------------------------------------------------------------------*/
+namespace ELinkerType
 {
-	/**
-	 * The name of the UObject represented by this resource.
-	 * Serialized
-	 */
-	FName			ObjectName;
-
-	/**
-	 * Location of the resource for this resource's Outer.  Values of 0 indicate that this resource
-	 * represents a top-level UPackage object (the linker's LinkerRoot).
-	 * Serialized
-	 */
-	FPackageIndex	OuterIndex;
-
-#if WITH_EDITOR
-	/**
-	 * Name of the class this object was serialized with (in case active class redirects have changed it)
-	 */
-	FName			OldClassName;
-#endif
-
-	FObjectResource();
-	FObjectResource( UObject* InObject );
-};
-
-/*-----------------------------------------------------------------------------
-	FObjectExport.
------------------------------------------------------------------------------*/
+	enum Type
+	{
+		None,
+		Load,
+		Save
+	};
+}
 
 /**
- * UObject resource type for objects that are contained within this package and can
- * be referenced by other packages.
+ * Manages the data associated with an Unreal package.  Acts as the bridge between
+ * the file on disk and the UPackage object in memory for all Unreal package types.
  */
-struct FObjectExport : public FObjectResource
+class FLinker : public FLinkerTables
 {
-	/**
-	 * Location of the resource for this export's class (if non-zero).  A value of zero
-	 * indicates that this export represents a UClass object; there is no resource for
-	 * this export's class object
-	 * Serialized
-	 */
-	FPackageIndex  	ClassIndex;
-
-	/**
-	* Location of this resource in export map. Used for export fixups while loading packages.
-	* Value of zero indicates resource is invalid and shouldn't be loaded.
-	* Not serialized.
-	*/
-	FPackageIndex ThisIndex;
-
-	/**
-	 * Location of the resource for this export's SuperField (parent).  Only valid if
-	 * this export represents a UStruct object. A value of zero indicates that the object
-	 * represented by this export isn't a UStruct-derived object.
-	 * Serialized
-	 */
-	FPackageIndex 	SuperIndex;
-
-	/**
-	 * The object flags for the UObject represented by this resource.  Only flags that
-	 * match the RF_Load combination mask will be loaded from disk and applied to the UObject.
-	 * Serialized
-	 */
-	EObjectFlags	ObjectFlags;
-
-	/**
-	 * The number of bytes to serialize when saving/loading this export's UObject.
-	 * Serialized
-	 */
-	int32         	SerialSize;
-
-	/**
-	 * The location (into the FLinker's underlying file reader archive) of the beginning of the
-	 * data for this export's UObject.  Used for verification only.
-	 * Serialized
-	 */
-	int32         	SerialOffset;
-
-	/**
-	 * The location (into the FLinker's underlying file reader archive) of the beginning of the
-	 * portion of this export's data that is serialized using script serialization.
-	 * Transient
-	 */
-	int32				ScriptSerializationStartOffset;
-
-	/**
-	 * The location (into the FLinker's underlying file reader archive) of the end of the
-	 * portion of this export's data that is serialized using script serialization.
-	 * Transient
-	 */
-	int32				ScriptSerializationEndOffset;
-
-	/**
-	 * The UObject represented by this export.  Assigned the first time CreateExport is called for this export.
-	 * Transient
-	 */
-	UObject*		Object;
+private:	
 
-	/**
-	 * The index into the FLinker's ExportMap for the next export in the linker's export hash table.
-	 * Transient
-	 */
-	int32				HashNext;
+	ELinkerType::Type LinkerType;
 
-	/**
-	 * whether the export was forced into the export table via OBJECTMARK_ForceTagExp.
-	 * Serialized
-	 */
-	bool			bForcedExport;   
+public:
 
-	/**
-	 * whether the export should be loaded on clients
-	 * Serialized
-	 */
-	bool			bNotForClient;   
+	/** The top-level UPackage object for the package associated with this linker */
+	UPackage*				LinkerRoot;
 
-	/**
-	* whether the export should be loaded on servers
-	 * Serialized
-	 */
-	bool			bNotForServer;
+	/** Table of contents for this package's file */
+	FPackageFileSummary		Summary;
 
-	/**
-	* whether the export should be always loaded in editor game
-	* False means that the object is necessary for editor game,
-	* True doesn't means, that the object won't be loaded.
-	* Serialized
-	*/
-	bool			bNotForEditorGame;
+	/** Names used by objects contained within this package */
+	TArray<FName>			NameMap;
 
-	/**
-	* True if this export is an asset object.
-	*/
-	bool			bIsAsset;
+	/** Gatherable text data contained within this package */
+	TArray<FGatherableTextData> GatherableTextDataMap;
 
-	/**
-	* Force this export to not load, it failed because the outer didn't exist.
-	*/
-	bool			bExportLoadFailed;
+	/** The name of the file for this package */
+	FString					Filename;
 
-	/**
-	* Export is a dynamic class.
-	*/
-	bool			bDynamicClass;
+	/** If true, filter out exports that are for clients but not servers */
+	bool					FilterClientButNotServer;
 
-	/**
-	* Export was filtered out on load
-	*/
-	bool			bWasFiltered;
+	/** If true, filter out exports that are for servers but not clients */
+	bool					FilterServerButNotClient;
 
-	/** If this object is a top level package (which must have been forced into the export table via OBJECTMARK_ForceTagExp)
-	 * this is the GUID for the original package file
-	 * Serialized
-	 */
-	FGuid			PackageGuid;
+	/** The SHA1 key generator for this package, if active */
+	class FSHA1*			ScriptSHA;
 
-	/** If this object is a top level package (which must have been forced into the export table via OBJECTMARK_ForceTagExp)
-	 * this is the package flags for the original package file
-	 * Serialized
-	 */
-	uint32			PackageFlags;
+	/** Constructor. */
+	FLinker(ELinkerType::Type InType, UPackage* InRoot, const TCHAR* InFilename);
 
-	/**
-	 * Constructors
-	 */
-	COREUOBJECT_API FObjectExport();
-	FObjectExport( UObject* InObject );
-	
-	/** I/O function */
-	friend COREUOBJECT_API FArchive& operator<<( FArchive& Ar, FObjectExport& E );
+	virtual ~FLinker();
 
-};
-
-/*-----------------------------------------------------------------------------
-	FObjectImport.
------------------------------------------------------------------------------*/
-
-/**
- * UObject resource type for objects that are referenced by this package, but contained
- * within another package.
- */
-struct FObjectImport : public FObjectResource
-{
-	/**
-	 * The name of the package that contains the class of the UObject represented by this resource.
-	 * Serialized
-	 */
-	FName			ClassPackage;
-
-	/**
-	 * The name of the class for the UObject represented by this resource.
-	 * Serialized
-	 */
-	FName			ClassName;
-
-	/**
-	 * The UObject represented by this resource.  Assigned the first time CreateImport is called for this import.
-	 * Transient
-	 */
-	UObject*		XObject;
-
-	/**
-	 * The linker that contains the original FObjectExport resource associated with this import.
-	 * Transient
-	 */
-	FLinkerLoad*	SourceLinker;
-
-	/**
-	 * Index into SourceLinker's ExportMap for the export associated with this import's UObject.
-	 * Transient
-	 */
-	int32             SourceIndex;
-
-	/**
-	 * Constructors
-	 */
-	COREUOBJECT_API FObjectImport();
-	FObjectImport( UObject* InObject );
-	FObjectImport( UObject* InObject, UClass* InClass );
-	
-	/** I/O function */
-	friend COREUOBJECT_API FArchive& operator<<( FArchive& Ar, FObjectImport& I );
-};
-
-/**
- * Information about a compressed chunk in a file.
- */
-struct FCompressedChunk
-{
-	/** Default constructor, zero initializing all members. */
-	FCompressedChunk();
-
-	/** Original offset in uncompressed file.	*/
-	int32		UncompressedOffset;
-	/** Uncompressed size in bytes.				*/
-	int32		UncompressedSize;
-	/** Offset in compressed file.				*/
-	int32		CompressedOffset;
-	/** Compressed size in bytes.				*/
-	int32		CompressedSize;
-
-	/** I/O function */
-	friend COREUOBJECT_API FArchive& operator<<(FArchive& Ar,FCompressedChunk& Chunk);
-};
-
-/*----------------------------------------------------------------------------
-	Items stored in Unrealfiles.
-----------------------------------------------------------------------------*/
-
-/**
- * Revision data for an Unreal package file.
- */
-//@todo: shouldn't need ExportCount/NameCount with the linker free package map; if so, clean this stuff up
-struct FGenerationInfo
-{
-	/**
-	 * Number of exports in the linker's ExportMap for this generation.
-	 */
-	int32 ExportCount;
-
-	/**
-	 * Number of names in the linker's NameMap for this generation.
-	 */
-	int32 NameCount;
-
-	/** Constructor */
-	FGenerationInfo( int32 InExportCount, int32 InNameCount );
-
-	/** I/O function
-	 * we use a function instead of operator<< so we can pass in the package file summary for version tests, since archive version hasn't been set yet
-	 */
-	void Serialize(FArchive& Ar, const struct FPackageFileSummary& Summary);
-};
-
-#if WITH_ENGINE
-/**
- * Information about the textures stored in the package.
- */
-struct FTextureAllocations
-{
-	/**
-	 * Stores an export index for each texture of a certain type (size, format, etc).
-	 */
-	struct FTextureType
-	{
-		FTextureType();
-		COREUOBJECT_API FTextureType( int32 SizeX, int32 SizeY, int32 NumMips, uint32 Format, uint32 TexCreateFlags );
-
-		/** Width of the largest mip-level stored in the package. */
-		int32				SizeX;
-		/** Height of the largest mip-level stored in the package. */
-		int32				SizeY;
-		/** Number of mips */
-		int32				NumMips;
-		/** Texture format */
-		uint32			Format;
-		/** ETextureCreateFlags bit flags */
-		uint32			TexCreateFlags;
-		/** Index into the package ExportMap. */
-		TArray<int32>		ExportIndices;
-
-		/** Not serialized. ResourceMems are constructed on load. */
-		TArray<FTexture2DResourceMem*>	Allocations;
-		/** Note serialized. Number of ExportIndices processed during load. */
-		int32				NumExportIndicesProcessed;
-
-		/**
-		 * Checks whether all potential allocations for this TextureType have been considered yet (during load).
-		 *
-		 * @return true if all allocations have been started
-		 */
-		bool	HaveAllAllocationsBeenConsidered() const
-		{
-			return NumExportIndicesProcessed == ExportIndices.Num();
-		}
-
-		/**
-		 * Serializes an FTextureType
-		 */
-		friend FArchive& operator<<( FArchive& Ar, FTextureAllocations::FTextureType& TextureType );
-	};
-
-	FTextureAllocations()
-	:	PendingAllocationSize(0)
-	,	NumTextureTypesConsidered(0)
-	{
-	}
-
-	/** Array of texture types in the package. */
-	TArray<FTextureType>	TextureTypes;
-	/** Number of allocations that haven't been completed yet. */
-	FThreadSafeCounter		PendingAllocationCount;
-	/** Number of allocated bytes that has yet to be claimed by UTexture2D serialization. */
-	int32						PendingAllocationSize;
-	/** Number of TextureTypes that have been considered for potential allocations so far, during load. */
-	int32						NumTextureTypesConsidered;
-
-	/**
-	 * Finds a suitable ResourceMem allocation, removes it from this container and returns it to the user.
-	 *
-	 * @param SizeX				Width of texture
-	 * @param SizeY				Height of texture
-	 * @param NumMips			Number of mips
-	 * @param Format			Texture format (EPixelFormat)
-	 * @param TexCreateFlags	ETextureCreateFlags bit flags
-	 **/
-	FTexture2DResourceMem*	FindAndRemove( int32 SizeX, int32 SizeY, int32 NumMips, uint32 Format, uint32 TexCreateFlags );
-
-	/**
-	 * Cancels any pending ResourceMem allocation that hasn't been claimed by a texture yet,
-	 * just in case there are any mismatches at run-time.
-	 *
-	 * @param bCancelEverything		If true, cancels all allocations. If false, only cancels allocations that haven't been completed yet.
-	 */
-	void	CancelRemainingAllocations( bool bCancelEverything );
-
-	/**
-	 * Checks if all allocations that should be started have been started (during load).
-	 *
-	 * @return true if all allocations have been started
-	 */
-	bool	HaveAllAllocationsBeenConsidered() const
-	{
-		return NumTextureTypesConsidered == TextureTypes.Num();
-	}
-
-	/**
-	 * Checks if all ResourceMem allocations has completed.
-	 *
-	 * @return true if all ResourceMem allocations has completed
-	 */
-	bool	HasCompleted() const
-	{
-		return PendingAllocationCount.GetValue() == 0;
-	}
-
-	/**
-	 * Checks if all ResourceMem allocations have been claimed by a texture.
-	 *
-	 * @return true if there are no more pending ResourceMem allocations waiting for a texture to claim it
-	 */
-	bool	HasBeenFullyClaimed() const
-	{
-		return PendingAllocationSize == 0;
-	}
-
-	/**
-	 * Serializes an FTextureType
-	 */
-	friend FArchive& operator<<( FArchive& Ar, FTextureAllocations::FTextureType& TextureAllocationType );
-
-	/**
-	 * Serializes an FTextureAllocations struct
-	 */
-	friend FArchive& operator<<( FArchive& Ar, FTextureAllocations& TextureAllocations );
-
-	FTextureAllocations( const FTextureAllocations& Other );
-	void operator=(const FTextureAllocations& Other);
-
-private:
-	/**
-	 * Finds a texture type that matches the given specifications.
-	 *
-	 * @param SizeX				Width of the largest mip-level stored in the package
-	 * @param SizeY				Height of the largest mip-level stored in the package
-	 * @param NumMips			Number of mips
-	 * @param Format			Texture format (EPixelFormat)
-	 * @param TexCreateFlags	ETextureCreateFlags bit flags
-	 * @return					Matching texture type, or NULL if none was found
-	 */
-	FTextureType*	FindTextureType( int32 SizeX, int32 SizeY, int32 NumMips, uint32 Format, uint32 TexCreateFlags );
-
-	/**
-	 * Adds a dummy export index (-1) for a specified texture type.
-	 * Creates the texture type entry if needed.
-	 *
-	 * @param SizeX				Width of the largest mip-level stored in the package
-	 * @param SizeY				Height of the largest mip-level stored in the package
-	 * @param NumMips			Number of mips
-	 * @param Format			Texture format (EPixelFormat)
-	 * @param TexCreateFlags	ETextureCreateFlags bit flags
-	 */
-	void			AddResourceMemInfo( int32 SizeX, int32 SizeY, int32 NumMips, uint32 Format, uint32 TexCreateFlags );
-};
-#endif		// WITH_ENGINE
-
-/**
- * A "table of contents" for an Unreal package file.  Stored at the top of the file.
- */
-struct FPackageFileSummary
-{
-	/**
-	 * Magic tag compared against PACKAGE_FILE_TAG to ensure that package is an Unreal package.
-	 */
-	int32		Tag;
-
-private:
-	/* UE4 file version */
-	int32		FileVersionUE4;
-	/* Licensee file version */
-	int32		FileVersionLicenseeUE4;
-	/* Custom version numbers. Keyed off a unique tag for each custom component. */
-	FCustomVersionContainer CustomVersionContainer;
-
-public:
-	/**
-	 * Total size of all information that needs to be read in to create a FLinkerLoad. This includes
-	 * the package file summary, name table and import & export maps.
-	 */
-	int32		TotalHeaderSize;
-
-	/**
-	 * The flags for the package
-	 */
-	uint32	PackageFlags;
-
-	/**
-	 * The Generic Browser folder name that this package lives in
-	 */
-	FString	FolderName;
-
-	/**
-	 * Number of names used in this package
-	 */
-	int32		NameCount;
-
-	/**
-	 * Location into the file on disk for the name data
-	 */
-	int32 	NameOffset;
-
-	/**
-	 * Number of gatherable text data items in this package
-	 */
-	int32	GatherableTextDataCount;
-
-	/**
-	 * Location into the file on disk for the gatherable text data items
-	 */
-	int32 	GatherableTextDataOffset;
-	
-	/**
-	 * Number of exports contained in this package
-	 */
-	int32		ExportCount;
-
-	/**
-	 * Location into the file on disk for the ExportMap data
-	 */
-	int32		ExportOffset;
-
-	/**
-	 * Number of imports contained in this package
-	 */
-	int32		ImportCount;
-
-	/**
-	 * Location into the file on disk for the ImportMap data
-	 */
-	int32		ImportOffset;
-
-	/**
-	* Location into the file on disk for the DependsMap data
-	*/
-	int32		DependsOffset;
-
-	/**
-	 * Number of references contained in this package
-	 */
-	int32		StringAssetReferencesCount;
-
-	/**
-	 * Location into the file on disk for the string asset references map data
-	 */
-	int32		StringAssetReferencesOffset;
-
-	/**
-	 * Thumbnail table offset
-	 */
-	int32		ThumbnailTableOffset;
-
-	/**
-	 * Current id for this package
-	 */
-	FGuid	Guid;
-
-	/**
-	 * Data about previous versions of this package
-	 */
-	TArray<FGenerationInfo> Generations;
-
-	/**
-	 * Engine version this package was saved with. For hotfix releases and engine versions which maintain strict binary compatibility with another version, this may differ from CompatibleWithEngineVersion.
-	 */
-	FEngineVersion SavedByEngineVersion;
-
-	/**
-	 * Engine version this package is compatible with. See SavedByEngineVersion.
-	 */
-	FEngineVersion CompatibleWithEngineVersion;
-
-	/**
-	 * Flags used to compress the file on save and uncompress on load.
-	 */
-	uint32	CompressionFlags;
-
-	/**
-	 * Value that is used to determine if the package was saved by Epic (or licensee) or by a modder, etc
-	 */
-	uint32	PackageSource;
-
-	/**
-	 * Array of compressed chunks in case this package was stored compressed.
-	 */
-	TArray<FCompressedChunk> CompressedChunks;
-
-	/**
-	 * List of additional packages that are needed to be cooked for this package (ie streaming levels)
-	 */
-	TArray<FString>	AdditionalPackagesToCook;
-
-	/** 
-	 * If true, this file will not be saved with version numbers or was saved without version numbers. In this case they are assumed to be the current version. 
-	 * This is only used for full cooks for distribution because it is hard to guarantee correctness 
-	 **/
-	bool bUnversioned;
-
-#if WITH_ENGINE
-	/**
-	 * Information about the textures stored in the package.
-	 */
-	FTextureAllocations	TextureAllocations;
-#endif		// WITH_ENGINE
-
-	/**
-	 * Location into the file on disk for the asset registry tag data
-	 */
-	int32 	AssetRegistryDataOffset;
-
-	/** Offset to the location in the file where the bulkdata starts */
-	int64	BulkDataStartOffset;
-	/**
-	 * Offset to the location in the file where the FWorldTileInfo data starts
-	 */
-	int32 	WorldTileInfoDataOffset;
-
-	/**
-	 * Streaming install ChunkIDs
-	 */
-	TArray<int32>	ChunkIDs;
-
-
-	/** Constructor */
-	COREUOBJECT_API FPackageFileSummary();
-
-	int32 GetFileVersionUE4() const
-	{
-		return FileVersionUE4;
-	}
-
-	int32 GetFileVersionLicenseeUE4() const
-	{
-		return FileVersionLicenseeUE4;
-	}
-
-	const FCustomVersionContainer& GetCustomVersionContainer() const
-	{
-		return CustomVersionContainer;
-	}
-
-	void SetCustomVersionContainer(const FCustomVersionContainer& InContainer)
-	{
-		CustomVersionContainer = InContainer;
-	}
-
-	void SetFileVersions(const int32 EpicUE4, const int32 LicenseeUE4, const bool bInSaveUnversioned = false)
-	{
-		FileVersionUE4 = EpicUE4;
-		FileVersionLicenseeUE4 = LicenseeUE4;
-		bUnversioned = bInSaveUnversioned;
-	}
-
-	/** I/O function */
-	friend COREUOBJECT_API FArchive& operator<<( FArchive& Ar, FPackageFileSummary& Sum );
-};
-
-class FLinkerTables
-{
-public:
-	/** The list of FObjectImports found in the package */
-	TArray<FObjectImport> ImportMap;
-	/** The list of FObjectExports found in the package */
-	TArray<FObjectExport> ExportMap;
-	/** List of dependency lists for each export */
-	TArray<TArray<FPackageIndex> > DependsMap;
-	/** Map that holds info about string asset references from the package. */
-	TArray<FString> StringAssetReferencesMap;
-
-	/**
-	 * Check that this Index is non-null and return an import or export
-	 * @param	Index	Package index to get
-	 * @return	the resource corresponding to this index
-	 */
-	FORCEINLINE FObjectResource& ImpExp(FPackageIndex Index)
-	{
-		check(!Index.IsNull());
-		if (Index.IsImport())
-		{
-			return Imp(Index);
-		}
-		else
-		{
-			return Exp(Index);
-		}
-	}
-	/**
-	 * Check that this Index is non-null and return an import or export
-	 * @param	Index	Package index to get
-	 * @return	the resource corresponding to this index
-	 */
-	FORCEINLINE FObjectResource const& ImpExp(FPackageIndex Index) const
-	{
-		check(!Index.IsNull());
-		if (Index.IsImport())
-		{
-			return Imp(Index);
-		}
-		else
-		{
-			return Exp(Index);
-		}
-	}	
-	/**
-	 * Return an import or export for this index
-	 * @param	Index	Package index to get
-	 * @return	the resource corresponding to this index, or NULL if the package index is null
-	 */
-	FORCEINLINE FObjectResource* ImpExpPtr(FPackageIndex Index)
-	{
-		if (Index.IsImport())
-		{
-			return ImpPtr(Index);
-		}
-		else
-		{
-			return ExpPtr(Index);
-		}
-	}
-
-	/**
-	 * Check that this Index is non-null and an import and return an import
-	 * @param	Index	Package index to get, must be an import
-	 * @return	the import corresponding to this index
-	 */
-	FORCEINLINE FObjectImport& Imp(FPackageIndex Index)
-	{
-		check(Index.IsImport() && ImportMap.IsValidIndex(Index.ToImport()));
-		return ImportMap[Index.ToImport()];
-	}
-	FORCEINLINE FObjectImport const& Imp(FPackageIndex Index) const
-	{
-		check(Index.IsImport() && ImportMap.IsValidIndex(Index.ToImport()));
-		return ImportMap[Index.ToImport()];
-	}
-	/**
-	 * Return an import for this index
-	 * @param	Index	Package index to get
-	 * @return	the import corresponding to this index, or NULL if the package index is null or an export
-	 */
-	FORCEINLINE FObjectImport* ImpPtr(FPackageIndex Index)
-	{
-		if (Index.IsImport())
-		{
-			check(ImportMap.IsValidIndex(Index.ToImport()));
-			return &ImportMap[Index.ToImport()];
-		}
-		return NULL;
-	}
-
-	/**
-	 * Check that this Index is non-null and an export and return an import
-	 * @param	Index	Package index to get, must be an export
-	 * @return	the export corresponding to this index
-	 */
-	FORCEINLINE FObjectExport& Exp(FPackageIndex Index)
-	{
-		check(Index.IsExport() && ExportMap.IsValidIndex(Index.ToExport()));
-		return ExportMap[Index.ToExport()];
-	}
-	FORCEINLINE FObjectExport const& Exp(FPackageIndex Index) const
-	{
-		check(Index.IsExport() && ExportMap.IsValidIndex(Index.ToExport()));
-		return ExportMap[Index.ToExport()];
-	}
-	/**
-	 * Return an export for this index
-	 * @param	Index	Package index to get
-	 * @return	the export corresponding to this index, or NULL if the package index is null or an import
-	 */
-	FORCEINLINE FObjectExport* ExpPtr(FPackageIndex Index)
-	{
-		if (Index.IsExport())
-		{
-			check(ExportMap.IsValidIndex(Index.ToExport()));
-			return &ExportMap[Index.ToExport()];
-		}
-		return NULL;
-	}
-};
-
-
-struct FLinkerNamePairKeyFuncs : DefaultKeyFuncs<FName, false>
-{
-	static FORCEINLINE bool Matches(FName A, FName B)
-	{
-		// The linker requires that FNames preserve case, but the numeric suffix can be ignored since
-		// that is stored separately for each FName instance saved
-		return A.IsEqual(B, ENameCase::CaseSensitive, false/*bCompareNumber*/);
-	}
-
-	static FORCEINLINE uint32 GetKeyHash(FName Key)
-	{
-		return Key.GetComparisonIndex();
-	}
-};
-
-
-template<typename ValueType>
-struct TLinkerNameMapKeyFuncs : TDefaultMapKeyFuncs<FName, ValueType, false>
-{
-	static FORCEINLINE bool Matches(FName A, FName B)
-	{
-		// The linker requires that FNames preserve case, but the numeric suffix can be ignored since
-		// that is stored separately for each FName instance saved
-		return A.IsEqual(B, ENameCase::CaseSensitive, false/*bCompareNumber*/);
-	}
-
-	static FORCEINLINE uint32 GetKeyHash(FName Key)
-	{
-		return Key.GetComparisonIndex();
-	}
-};
-
-
-/*----------------------------------------------------------------------------
-	FLinker.
-----------------------------------------------------------------------------*/
-namespace ELinkerType
-{
-	enum Type
-	{
-		None,
-		Load,
-		Save
-	};
-}
-
-/**
- * Manages the data associated with an Unreal package.  Acts as the bridge between
- * the file on disk and the UPackage object in memory for all Unreal package types.
- */
-class FLinker : public FLinkerTables
-{
-private:	
-
-	ELinkerType::Type LinkerType;
-
-public:
-
-	/** The top-level UPackage object for the package associated with this linker */
-	UPackage*				LinkerRoot;
-
-	/** Table of contents for this package's file */
-	FPackageFileSummary		Summary;
-
-	/** Names used by objects contained within this package */
-	TArray<FName>			NameMap;
-
-	/** Gatherable text data contained within this package */
-	TArray<FGatherableTextData> GatherableTextDataMap;
-
-	/** The name of the file for this package */
-	FString					Filename;
-
-	/** If true, filter out exports that are for clients but not servers */
-	bool					FilterClientButNotServer;
-
-	/** If true, filter out exports that are for servers but not clients */
-	bool					FilterServerButNotClient;
-
-	/** The SHA1 key generator for this package, if active */
-	class FSHA1*			ScriptSHA;
-
-	/** Constructor. */
-	FLinker(ELinkerType::Type InType, UPackage* InRoot, const TCHAR* InFilename);
-
-	virtual ~FLinker();
-
-	/** Gets the class name for the specified index in the export map. */
-	COREUOBJECT_API FName GetExportClassName(int32 ExportIdx);
-	/** Gets the class name for the specified index in the import map. */
-	FName GetExportClassName(FPackageIndex PackageIndex)
-	{
-		if (PackageIndex.IsExport())
-		{
-			return GetExportClassName(PackageIndex.ToExport());
-		}
-		return NAME_None;
-	}
-	/** Gets the class name for the specified index in the import map. */
-	FName GetImportClassName(int32 ImportIdx)
-	{
-		return ImportMap[ImportIdx].ClassName;
-	}
-	/** Gets the class name for the specified index in the import map. */
-	FName GetImportClassName(FPackageIndex PackageIndex)
-	{
-		if (PackageIndex.IsImport())
-		{
-			return GetImportClassName(PackageIndex.ToImport());
-		}
-		return NAME_None;
-	}
-	/** Gets the class name for the specified package index. */
-	FName GetClassName(FPackageIndex PackageIndex)
-	{
-		if (PackageIndex.IsImport())
-		{
-			return GetImportClassName(PackageIndex);
-		}
-		else if (PackageIndex.IsExport())
-		{
-			return GetExportClassName(PackageIndex);
-		}
-		return NAME_None;
-	}
-
-
-	FORCEINLINE ELinkerType::Type GetType() const
-	{
-		return LinkerType;
-	}
-
-	/**
-	 * I/O function
-	*
-	 * @param	Ar	the archive to read/write into
-	 */
-	void Serialize(FArchive& Ar);
-	void AddReferencedObjects(FReferenceCollector& Collector);
-	/**
-	 * Return the path name of the UObject represented by the specified import. 
-	 * (can be used with StaticFindObject)
-	 * 
-	 * @param	ImportIndex	index into the ImportMap for the resource to get the name for
-	 *
-	 * @return	the path name of the UObject represented by the resource at ImportIndex
-	 */
-	COREUOBJECT_API FString GetImportPathName(int32 ImportIndex);
-	/**
-	 * Return the path name of the UObject represented by the specified import. 
-	 * (can be used with StaticFindObject)
-	 * 
-	 * @param	PackageIndex	package index for the resource to get the name for
-	 *
-	 * @return	the path name of the UObject represented by the resource at PackageIndex, or the empty string if this isn't an import
-	 */
-	FString GetImportPathName(FPackageIndex PackageIndex)
-	{
-		if (PackageIndex.IsImport())
-		{
-			return GetImportPathName(PackageIndex.ToImport());
-		}
-		return FString();
-	}
-	/**
-	 * Return the path name of the UObject represented by the specified export.
-	 * (can be used with StaticFindObject)
-	 * 
-	 * @param	ExportIndex				index into the ExportMap for the resource to get the name for
-	 * @param	FakeRoot				Optional name to replace use as the root package of this object instead of the linker
-	 * @param	bResolveForcedExports	if true, the package name part of the return value will be the export's original package,
-	 *									not the name of the package it's currently contained within.
-	 *
-	 * @return	the path name of the UObject represented by the resource at ExportIndex
-	 */
-	COREUOBJECT_API FString GetExportPathName(int32 ExportIndex, const TCHAR* FakeRoot=NULL,bool bResolveForcedExports=false);
-	/**
-	 * Return the path name of the UObject represented by the specified export.
-	 * (can be used with StaticFindObject)
-	 * 
-	 * @param	PackageIndex			package index for the resource to get the name for
-	 * @param	FakeRoot				Optional name to replace use as the root package of this object instead of the linker
-	 * @param	bResolveForcedExports	if true, the package name part of the return value will be the export's original package,
-	 *									not the name of the package it's currently contained within.
-	 *
-	 * @return	the path name of the UObject represented by the resource at PackageIndex, or the empty string if this isn't an export
-	 */
-	FString GetExportPathName(FPackageIndex PackageIndex, const TCHAR* FakeRoot=NULL,bool bResolveForcedExports=false)
-	{
-		if (PackageIndex.IsExport())
-		{
-			return GetExportPathName(PackageIndex.ToExport(), FakeRoot, bResolveForcedExports);
-		}
-		return FString();
-	}
-
-	/**
-	 * Return the path name of the UObject represented by the specified import. 
-	 * (can be used with StaticFindObject)
-	 * 
-	 * @param	PackageIndex	package index
-	 *
-	 * @return	the path name of the UObject represented by the resource at PackageIndex, or the empty string if this is null
-	 */
-	FString GetPathName(FPackageIndex PackageIndex)
-	{
-		if (PackageIndex.IsImport())
-		{
-			return GetImportPathName(PackageIndex);
-		}
-		else if (PackageIndex.IsExport())
-		{
-			return GetExportPathName(PackageIndex);
-		}
-		return FString();
-	}
-	/**
-	 * Return the full name of the UObject represented by the specified import.
-	 * 
-	 * @param	ImportIndex	index into the ImportMap for the resource to get the name for
-	 *
-	 * @return	the full name of the UObject represented by the resource at ImportIndex
-	 */
-	COREUOBJECT_API FString GetImportFullName(int32 ImportIndex);
-	/**
-	 * Return the full name of the UObject represented by the specified package index
-	 * 
-	 * @param	PackageIndex	package index for the resource to get the name for
-	 *
-	 * @return	the full name of the UObject represented by the resource at PackageIndex
-	 */
-	FString GetImportFullName(FPackageIndex PackageIndex)
-	{
-		if (PackageIndex.IsImport())
-		{
-			return GetImportFullName(PackageIndex.ToImport());
-		}
-		return FString();
-	}
-
-	/**
-	 * Return the full name of the UObject represented by the specified export.
-	 * 
-	 * @param	ExportIndex				index into the ExportMap for the resource to get the name for
-	 * @param	FakeRoot				Optional name to replace use as the root package of this object instead of the linker
-	 * @param	bResolveForcedExports	if true, the package name part of the return value will be the export's original package,
-	 *									not the name of the package it's currently contained within.
-	 *
-	 * @return	the full name of the UObject represented by the resource at ExportIndex
-	 */
-	COREUOBJECT_API FString GetExportFullName(int32 ExportIndex, const TCHAR* FakeRoot=NULL,bool bResolveForcedExports=false);
-	/**
-	 * Return the full name of the UObject represented by the specified package index
-	 * 
-	 * @param	PackageIndex			package index for the resource to get the name for
-	 * @param	FakeRoot				Optional name to replace use as the root package of this object instead of the linker
-	 * @param	bResolveForcedExports	if true, the package name part of the return value will be the export's original package,
-	 *									not the name of the package it's currently contained within.
-	 *
-	 * @return	the full name of the UObject represented by the resource at PackageIndex
-	 */
-	FString GetExportFullName(FPackageIndex PackageIndex, const TCHAR* FakeRoot=NULL,bool bResolveForcedExports=false)
-	{
-		if (PackageIndex.IsExport())
-		{
-			return GetExportFullName(PackageIndex.ToExport(), FakeRoot, bResolveForcedExports);
-		}
-		return FString();
-	}
-
-	/**
-	 * Return the full name of the UObject represented by the specified export.
-	 * 
-	 * @param	PackageIndex	package index
-	 *
-	 * @return	the path name of the UObject represented by the resource at PackageIndex, or the empty string if this is null
-	 */
-	FString GetFullImpExpName(FPackageIndex PackageIndex)
-	{
-		if (PackageIndex.IsImport())
-		{
-			return GetImportFullName(PackageIndex);
-		}
-		else if (PackageIndex.IsExport())
-		{
-			return GetExportFullName(PackageIndex);
-		}
-		return FString();
-	}
-
-	/**
-	 * Tell this linker to start SHA calculations
-	 */
-	void StartScriptSHAGeneration();
-
-	/**
-	 * If generating a script SHA key, update the key with this script code
-	 *
-	 * @param ScriptCode Code to SHAify
-	 */
-	void UpdateScriptSHAKey(const TArray<uint8>& ScriptCode);
-
-	/**
-	 * After generating the SHA key for all of the 
-	 *
-	 * @param OutKey Storage for the key bytes (20 bytes)
-	 */
-	void GetScriptSHAKey(uint8* OutKey);
-
-	/**
-	 * Test and object against the load flag filters
-	 *
-	 * @return	true if the object should be filtered and not loaded
-	 */
-	bool FilterExport(const FObjectExport& Export)
-	{
-		if (Export.bExportLoadFailed || Export.bWasFiltered)
-		{
-			return true;
-		}
-#if WITH_EDITOR
-		if (!Export.bNotForEditorGame)
-		{
-			return false;
-		}
-#endif
-		if (FilterClientButNotServer && Export.bNotForServer) // "we are a dedicated server"
-		{
-			return true;
-		}
-		if (FilterServerButNotClient && Export.bNotForClient) // "we are a client only"
-		{
-			return true;
-		}
-		if (Export.ThisIndex.IsNull()) // Export is invalid and shouldn't be processed.
-		{
-			return true;
-		}
-		return false;
-	}
-
-	FORCEINLINE bool IsCompressed() const { return !!(Summary.PackageFlags & PKG_StoreCompressed); }
-};
-
-/*----------------------------------------------------------------------------
-	FLinkerLoad.
-----------------------------------------------------------------------------*/
-
-/**
- * Helper struct to keep track of all objects needed by an export (recursive dependency caching)
- */
-struct FDependencyRef
-{
-	/** The Linker the export lives in */
-	FLinkerLoad* Linker;
-
-	/** Index into Linker's ExportMap for this object */
-	int32 ExportIndex;
-
-	/**
-	 * Comparison operator
-	 */
-	bool operator==(const FDependencyRef& Other) const
-	{
-		return Linker == Other.Linker && ExportIndex == Other.ExportIndex;
-	}
-
-	/**
-	 * Type hash implementation. Export indices are usually less than 100k, so are linker indices.
-	 *
-	 * @param	Ref		Reference to hash
-	 * @return	hash value
-	 */
-	friend uint32 GetTypeHash( const FDependencyRef& Ref  );
-};
-
-/**
- * Handles loading Unreal package files, including reading UObject data from disk.
- */
-#if USE_NEW_ASYNC_IO
-	class FArchiveAsync2;
-#endif
-class FLinkerLoad : public FLinker, public FArchiveUObject
-{
-	// Friends.
-	friend class UObject;
-	friend class UPackageMap;
-	friend struct FAsyncPackage;
-protected:
-	/** Linker loading status. */
-	enum ELinkerStatus
-	{
-		/** Error occured when loading. */
-		LINKER_Failed = 0,
-		/** Operation completed successfully. */
-		LINKER_Loaded = 1,
-		/** Operation took more time than allowed. */
-		LINKER_TimedOut = 2
-	};
-
-	/** Verify result. */
-	enum EVerifyResult
-	{
-		/** Error occured when verifying import (can be fatal). */
-		VERIFY_Failed = 0,
-		/** Verify completed successfully. */
-		VERIFY_Success = 1,
-		/** Verify completed successfully and followed a redirector. */
-		VERIFY_Redirected = 2
-	};
-
-	// Variables.
-public:
-
-	FORCEINLINE static ELinkerType::Type StaticType()
-	{
-		return ELinkerType::Load;
-	}
-
-	virtual ~FLinkerLoad();
-
-	/** Flags determining loading behavior.																					*/
-	uint32					LoadFlags;
-	/** Indicates whether the imports for this loader have been verified													*/
-	bool					bHaveImportsBeenVerified;
-	/** Indicates that this linker was created for a dynamic class package and will not use Loader */
-	bool					bDynamicClassLinker;
-#if USE_NEW_ASYNC_IO
-	/** True if Loader is FArchiveAsync2  */
-	bool					bLoaderIsFArchiveAsync2;
-	FORCEINLINE FArchiveAsync2* GetFArchiveAsync2Loader()
-	{
-		return bLoaderIsFArchiveAsync2 ? (FArchiveAsync2*)Loader : nullptr;
-	}
-#endif
-	/** The archive that actually reads the raw data from disk.																*/
-	FArchive*				Loader;
-	/** The async package associated with this linker */
-	struct FAsyncPackage* AsyncRoot;
-#if WITH_EDITOR
-	/** Bulk data that does not need to be loaded when the linker is loaded.												*/
-	TArray<FUntypedBulkData*> BulkDataLoaders;
-#endif // WITH_EDITOR
-
-	/** Hash table for exports.																								*/
-	int32						ExportHash[256];
-
-	/** OldClassName to NewClassName for ImportMap */
-	static TMap<FName, FName> ObjectNameRedirects;
-	/** Additional info for some ObjectName redirects to also redirect the class and class package  */
-	static TMap<FName, TPair<FName, FName>> ObjectNameClassRedirects;
-	/** OldClassName to NewClassName for ExportMap */
-	static TMap<FName, FName> ObjectNameRedirectsInstanceOnly;
-	/** Object name to NewClassName for export map */
-	static TMap<FName, FName> ObjectNameRedirectsObjectOnly;	
-	/** Old game name to new game name for ImportMap */
-	static TMap<FName, FName> GameNameRedirects;
-	/** Add a new redirect from old game name to new game name for ImportMap */
-	COREUOBJECT_API static void AddGameNameRedirect(const FName OldName, const FName NewName);
-	/** Old struct name to new struct name mapping */
-	static TMap<FName, FName> StructNameRedirects;
-	/** Old plugin name to new plugin name mapping */
-	static TMap<FString, FString> PluginNameRedirects;
-	/** Packages that are known to be missing when verifying imports that we don't want a message about */
-	COREUOBJECT_API static TSet<FName> KnownMissingPackages;
-
-	/** Object name to required class and new name for load-time remapping */
-	struct FSubobjectRedirect
-	{
-		FName MatchClass;
-		FName NewName;
-		FSubobjectRedirect(FName InMatchClass, FName InNewName)
-			: MatchClass(InMatchClass)
-			, NewName(InNewName)
-		{
-		}
-	};
-	static TMap<FName, FSubobjectRedirect> SubobjectNameRedirects;	
-
-	/* Makes sure active redirects map has been initialized */
-	static bool bActiveRedirectsMapInitialized;
-
-	/** 
-	 * Utility functions to query the object name redirects list for previous names for a class
-	 * @param CurrentClassPath The current name of the class, with a full path
-	 * @param bIsInstance If true, we're an instance, so check instance only maps as well
-	 * @return Names without path of all classes that were redirected to this name. Empty if none found.
-	 */
-	COREUOBJECT_API static TArray<FName> FindPreviousNamesForClass(FString CurrentClassPath, bool bIsInstance);
-
-	/** 
-	 * Utility functions to query the object name redirects list for the current name for a class
-	 * @param OldClassName An old class name, without path
-	 * @param bIsInstance If true, we're an instance, so check instance only maps as well
-	 * @return Current full path of this class. It will be None if no redirect found
-	 */
-	COREUOBJECT_API static FName FindNewNameForClass(FName OldClassName, bool bIsInstance);
-
-private:
-	// Variables used during async linker creation.
-
-	/** Current index into name map, used by async linker creation for spreading out serializing name entries.					*/
-	int32						NameMapIndex;
-	/** Current index into gatherable text data map, used by async linker creation for spreading out serializing text entries.	*/
-	int32						GatherableTextDataMapIndex;
-	/** Current index into import map, used by async linker creation for spreading out serializing importmap entries.			*/
-	int32						ImportMapIndex;
-	/** Current index into export map, used by async linker creation for spreading out serializing exportmap entries.			*/
-	int32						ExportMapIndex;
-	/** Current index into export map, used by async linker creation for spreading out serializing exportmap entries.			*/
-	int32						FirstNotLoadedExportMapIndex;
-	/** Current index into depends map, used by async linker creation for spreading out serializing dependsmap entries.			*/
-	int32						DependsMapIndex;
-	/** Current index into export hash map, used by async linker creation for spreading out hashing exports.					*/
-	int32						ExportHashIndex;
-
-
-	/** Whether we already serialized the package file summary.																*/
-	bool					bHasSerializedPackageFileSummary;
-	/** Whether we already fixed up import map.																				*/
-	bool					bHasFixedUpImportMap;
-	/** Whether we already matched up existing exports.																		*/
-	bool					bHasFoundExistingExports;
-	/** Whether we are already fully initialized.																			*/
-	bool					bHasFinishedInitialization;
-	/** Whether we are gathering dependencies, can be used to streamline VerifyImports, etc									*/
-	bool					bIsGatheringDependencies;
-	/** Whether time limit is/ has been exceeded in current/ last tick.														*/
-	bool					bTimeLimitExceeded;
-	/** Whether to use a time limit for async linker creation.																*/
-	bool					bUseTimeLimit;
-	/** Whether to use the full time limit, even if we're blocked on I/O													*/
-	bool					bUseFullTimeLimit;
-	/** Call count of IsTimeLimitExceeded.																					*/
-	int32					IsTimeLimitExceededCallCount;
-	/** Current time limit to use if bUseTimeLimit is true.																	*/
-	float					TimeLimit;
-	/** Time at begin of Tick function. Used for time limit determination.													*/
-	double					TickStartTime;
-
-	/** Used for ActiveClassRedirects functionality */
-	bool					bFixupExportMapDone;
-
-#if WITH_EDITOR
-	/** Check to avoid multiple export duplicate fixups in case we don't save asset. */
-	bool bExportsDuplicatesFixed;
-#endif // WITH_EDITOR
-	/** Id of the thread that created this linker. This is to guard against using this linker on other threads than the one it was created on **/
-	int32					OwnerThread;
-
-	/**
-	 * Helper struct to keep track of background file reads
-	 */
-	struct FPackagePrecacheInfo
+	/** Gets the class name for the specified index in the export map. */
+	COREUOBJECT_API FName GetExportClassName(int32 ExportIdx);
+	/** Gets the class name for the specified index in the import map. */
+	FName GetExportClassName(FPackageIndex PackageIndex)
 	{
-		/** Synchronization object used to wait for completion of async read. Pointer so it can be copied around, etc */
-		FThreadSafeCounter* SynchronizationObject;
-
-		/** Memory that contains the package data read off disk */
-		void* PackageData;
-
-		/** Size of the buffer pointed to by PackageData */
-		int64 PackageDataSize;
-
-		/**
-		 * Basic constructor
-		 */
-		FPackagePrecacheInfo()
-		: SynchronizationObject(NULL)
-		, PackageData(NULL)
-		, PackageDataSize(0)
-		{
-		}
-		/**
-		 * Destructor that will free the sync object
-		 */
-		~FPackagePrecacheInfo()
+		if (PackageIndex.IsExport())
 		{
-			delete SynchronizationObject;
+			return GetExportClassName(PackageIndex.ToExport());
 		}
-	};
-
-	/** Map that keeps track of any precached full package reads															*/
-	static TMap<FString, FPackagePrecacheInfo> PackagePrecacheMap;
-
-	/** Allows access to UTexture2D::StaticClass() without linking Core with Engine											*/
-	static UClass* UTexture2DStaticClass;
-
-	static FName NAME_LoadErrors;
-
-#if WITH_EDITOR
-	/** Feedback scope that is created to house the slow task of an asynchronous linker load. Raw ptr so we don't pull in TUniquePtr for everything. */
-	FScopedSlowTask* LoadProgressScope;
-
-	/** Test whether we should report progress or not */
-	FORCEINLINE bool ShouldReportProgress() const
-	{
-		return !IsAsyncLoading() && ( LoadFlags & ( LOAD_Quiet | LOAD_SeekFree ) ) == 0;
-	}
-#endif 
-
-public:
-
-	/**
-	 * Initialize the static variables
-	 */
-	COREUOBJECT_API static void StaticInit(UClass* InUTexture2DStaticClass);
-
-	/**
-	 * Add redirects to FLinkerLoad static map
-	 */
-	static void CreateActiveRedirectsMap(const FString& GEngineIniName);
-
-	/**
-	 * Locates the class adjusted index and its package adjusted index for a given class name in the import map
-	 */
-	COREUOBJECT_API bool FindImportClassAndPackage( FName ClassName, FPackageIndex& ClassIdx, FPackageIndex& PackageIdx );
-	
-	/**
-	 * Attempts to find the index for the given class object in the import list and adds it + its package if it does not exist
-	 */
-	COREUOBJECT_API bool CreateImportClassAndPackage( FName ClassName, FName PackageName, FPackageIndex& ClassIdx, FPackageIndex& PackageIdx );
-
-	/**
-	 * Allows object instances to be converted to other classes upon loading a package
-	 */
-	COREUOBJECT_API ELinkerStatus FixupExportMap();
-
-
-	/**
- 	 * Fills in the passed in TArray with the packages that are in its PrecacheMap
-	 *
-	 * @param TArray<FString> to be populated
-	 */
-	COREUOBJECT_API static void GetListOfPackagesInPackagePrecacheMap( TArray<FString>& ListOfPackages );
-
-	/**
-	 * Returns whether linker has finished (potentially) async initialization.
-	 *
-	 * @return true if initialized, false if pending.
-	 */
-	FORCEINLINE bool HasFinishedInitialization() const
-	{
-        return bHasFinishedInitialization;
+		return NAME_None;
 	}
-
-	/** Returns ID of the thread that created this linker */
-	FORCEINLINE int32 GetOwnerThreadId() const
+	/** Gets the class name for the specified index in the import map. */
+	FName GetImportClassName(int32 ImportIdx)
 	{
-		return OwnerThread;
-	}
-
-	/**
-	 * If this archive is a FLinkerLoad or FLinkerSave, returns a pointer to the FLinker portion.
-	 */
-	virtual FLinker* GetLinker() override { return this; }
-
-	/** Flush Loader Cache */
-	virtual void FlushCache() override;
-
-	/**
-	 * Creates and returns a FLinkerLoad object.
-	 *
-	 * @param	Parent		Parent object to load into, can be NULL (most likely case)
-	 * @param	Filename	Name of file on disk to load
-	 * @param	LoadFlags	Load flags determining behavior
-	 *
-	 * @return	new FLinkerLoad object for Parent/ Filename
-	 */
-	COREUOBJECT_API static FLinkerLoad* CreateLinker( UPackage* Parent, const TCHAR* Filename, uint32 LoadFlags);
-
-	void Verify();
-
-	COREUOBJECT_API FName GetExportClassPackage( int32 i );
-	virtual FString GetArchiveName() const override;
-
-	/**
-	 * Recursively gathers the dependencies of a given export (the recursive chain of imports
-	 * and their imports, and so on)
-
-	 * @param ExportIndex Index into the linker's ExportMap that we are checking dependencies
-	 * @param Dependencies Set of all dependencies needed
-	 * @param bSkipLoadedObjects Whether to skip already loaded objects when gathering dependencies
-	 */
-	COREUOBJECT_API void GatherExportDependencies(int32 ExportIndex, TSet<FDependencyRef>& Dependencies, bool bSkipLoadedObjects=true);
-
-	/**
-	 * Recursively gathers the dependencies of a given import (the recursive chain of imports
-	 * and their imports, and so on)
-
-	 * @param ImportIndex Index into the linker's ImportMap that we are checking dependencies
-	 * @param Dependencies Set of all dependencies needed
-	 * @param bSkipLoadedObjects Whether to skip already loaded objects when gathering dependencies
-	 */
-	COREUOBJECT_API void GatherImportDependencies(int32 ImportIndex, TSet<FDependencyRef>& Dependencies, bool bSkipLoadedObjects=true);
-
-	/**
-	 * A wrapper around VerifyImportInner. If the VerifyImportInner (previously VerifyImport) fails, this function
-	 * will look for a UObjectRedirector that will point to the real location of the object. You will see this if
-	 * an object was renamed to a different package or group, but something that was referencing the object was not
-	 * not currently open. (Rename fixes up references of all loaded objects, but naturally not for ones that aren't
-	 * loaded).
-	 *
-	 * @param	ImportIndex	The index into this package's ImportMap to verify
-	 * @return Verify import result
-	 */
-	EVerifyResult VerifyImport(int32 ImportIndex);
-	
-	/**
-	 * Loads all objects in package.
-	 *
-	 * @param bForcePreload	Whether to explicitly call Preload (serialize) right away instead of being
-	 *						called from EndLoad()
-	 */
-	COREUOBJECT_API void LoadAllObjects(bool bForcePreload = false);
-
-	/**
-	 * Returns the ObjectName associated with the resource indicated.
-	 * 
-	 * @param	ResourceIndex	location of the object resource
-	 *
-	 * @return	ObjectName for the FObjectResource at ResourceIndex, or NAME_None if not found
-	 */
-	FName ResolveResourceName( FPackageIndex ResourceIndex );
-	
-	int32 FindExportIndex( FName ClassName, FName ClassPackage, FName ObjectName, FPackageIndex ExportOuterIndex );
-	
-	/**
-	 * Function to create the instance of, or verify the presence of, an object as found in this Linker.
-	 *
-	 * @param ObjectClass	The class of the object
-	 * @param ObjectName	The name of the object
-	 * @param Outer			Optional outer that this object must be in (for finding objects in a specific group when there are multiple groups with the same name)
-	 * @param LoadFlags		Flags used to determine if the object is being verified or should be created
-	 * @param Checked		Whether or not a failure will throw an error
-	 * @return The created object, or (UObject*)-1 if this is just verifying
-	 */
-	UObject* Create( UClass* ObjectClass, FName ObjectName, UObject* Outer, uint32 InLoadFlags, bool Checked );
-
-	/**
-	 * Serialize the object data for the specified object from the unreal package file.  Loads any
-	 * additional resources required for the object to be in a valid state to receive the loaded
-	 * data, such as the object's Outer, Class, or ObjectArchetype.
-	 *
-	 * When this function exits, Object is guaranteed to contain the data stored that was stored on disk.
-	 *
-	 * @param	Object	The object to load data for.  If the data for this object isn't stored in this
-	 *					FLinkerLoad, routes the call to the appropriate linker.  Data serialization is 
-	 *					skipped if the object has already been loaded (as indicated by the RF_NeedLoad flag
-	 *					not set for the object), so safe to call on objects that have already been loaded.
-	 *					Note that this function assumes that Object has already been initialized against
-	 *					its template object.
-	 *					If Object is a UClass and the class default object has already been created, calls
-	 *					Preload for the class default object as well.
-	 */
-	void Preload( UObject* Object ) override;
-
-	/**
-	 * Before loading a persistent object from disk, this function can be used to discover
-	 * the object in memory. This could happen in the editor when you save a package (which
-	 * destroys the linker) and then play PIE, which would cause the Linker to be
-	 * recreated. However, the objects are still in memory, so there is no need to reload
-	 * them.
-	 *
-	 * @param ExportIndex	The index of the export to hunt down
-	 * @return The object that was found, or NULL if it wasn't found
-	 */
-	UObject* FindExistingExport(int32 ExportIndex);
-
-	/**
-	 * Builds a string containing the full path for a resource in the export table.
-	 *
-	 * @param OutPathName		[out] Will contain the full path for the resource
-	 * @param ResourceIndex		Index of a resource in the export table
-	 */
-	void BuildPathName( FString& OutPathName, FPackageIndex ExportIndex ) const;
-
-	/**
-	 * Checks if the specified export should be loaded or not.
-	 * Performs similar checks as CreateExport().
-	 *
-	 * @param ExportIndex	Index of the export to check
-	 * @return				true of the export should be loaded
-	 */
-	COREUOBJECT_API bool WillTextureBeLoaded( UClass* Class, int32 ExportIndex );
-
-	/**
-	 * Kick off an async load of a package file into memory
-	 * 
-	 * @param PackageName Name of package to read in. Must be the same name as passed into LoadPackage/CreateLinker
-	 */
-	COREUOBJECT_API static void AsyncPreloadPackage(const TCHAR* PackageName);
-
-	/**
-	 * Called when an object begins serializing property data using script serialization.
-	 */
-	virtual void MarkScriptSerializationStart( const UObject* Obj ) override;
-
-	/**
-	 * Called when an object stops serializing property data using script serialization.
-	 */
-	virtual void MarkScriptSerializationEnd( const UObject* Obj ) override;
-
-	/**
-	 * Looks for an existing linker for the given package, without trying to make one if it doesn't exist
-	 */
-	COREUOBJECT_API static FLinkerLoad* FindExistingLinkerForPackage(const UPackage* Package);
-
-	/**
-	 * Replaces OldObject's entry in its linker with NewObject, so that all subsequent loads of OldObject will return NewObject.
-	 * This is used to update instanced components that were serialized out, but regenerated during compile-on-load
-	 *
-	 * OldObject will be consigned to oblivion, and NewObject will take its place.
-	 *
-	 * WARNING!!!	This function is potentially very dangerous!  It should only be used at very specific times, and in very specific cases.
-	 *				If you're unsure, DON'T TRY TO USE IT!!!
-	 */
-	COREUOBJECT_API static void PRIVATE_PatchNewObjectIntoExport(UObject* OldObject, UObject* NewObject);
-
-	/**
-	 * Wraps a call to the package linker's ResolveAllImports().
-	 * 
-	 * @param  Package    The package whose imports you want all loaded.
-	 * 
-	 * WARNING!!!	This function shouldn't be used carelessly, and serves as a
-	 *				hacky entrypoint to FLinkerLoad's privates. It should only 
-	 *				be used at very specific times, and in very specific cases.
-	 *				If you're unsure, DON'T TRY TO USE IT!!!
-	 */
-	COREUOBJECT_API static void PRIVATE_ForceLoadAllDependencies(UPackage* Package);
-
-	/**
-	 * Invalidates the future loading of a specific object, so that subsequent loads will fail
-	 * This is used to invalidate sub objects of a replaced object that may no longer be valid
-	 */
-	COREUOBJECT_API static void InvalidateExport(UObject* OldObject);
-
-	/** Used by Matinee to fixup component renaming */
-	COREUOBJECT_API static FName FindSubobjectRedirectName(const FName& Name);
-
-private:
-#if WITH_EDITOR
-
-	/**
-	 * Nulls duplicated exports and fixes indexes in ExportMap to point to original objects instead of duplicates.
-	 */
-	void FixupDuplicateExports();
-
-	/**
-	 * Replaces all instances of OldIndex in ExportMap with NewIndex.
-	 */
-	void ReplaceExportIndexes(const FPackageIndex& OldIndex, const FPackageIndex& NewIndex);
-
-#endif // WITH_EDITOR
-
-	UObject* CreateExport( int32 Index );
-
-	/**
-	 * Creates export and preload if requested.
-	 *
-	 * @param Index Index of the export in export map.
-	 * @param bForcePreload Whether to explicitly call Preload (serialize)
-	 *        right away instead of being called from EndLoad().
-	 *
-	 * @return Created object.
-	 */
-	UObject* CreateExportAndPreload(int32 ExportIndex, bool bForcePreload = false);
-
-	/**
-	 * Utility function for easily retrieving the specified export's UClass.
-	 * 
-	 * @param  ExportIndex    Index of the export you want a class for.
-	 * @return The class that the specified export's ClassIndex references.
-	 */
-	UClass* GetExportLoadClass(int32 ExportIndex);
-
-	/** 
-	 * Looks for and loads meta data object from export map.
-	 *
-	 * @param bForcePreload Whether to explicitly call Preload (serialize)
-	 *        right away instead of being called from EndLoad().
-	 * 
-	 * @return If found returns index of meta data object in the export map,
-	 *         INDEX_NONE otherwise.
-	 */
-	int32 LoadMetaDataFromExportMap(bool bForcePreload = false);
-
-	UObject* CreateImport( int32 Index );
-
-	/**
-	 * Determines if the specified import belongs to a native "compiled in"
-	 * package (as opposed to an asset-file package). Recursive if the
-	 * specified import is not a package itself.
-	 * 
-	 * @param  ImportIndex    An index into the ImportMap, defining the import you wish to check.
-	 * @return True if the specified import comes from (or is) a "compiled in" package, otherwise false (it is an asset import).
-	 */
-	bool IsImportNative(const int32 ImportIndex) const;
-
-	/**
-	 * Attempts to lookup and return the corresponding FLinkerLoad object for 
-	 * the specified import WITHOUT invoking  a load, or continuing to load 
-	 * the import package (will only return one if it has already been 
-	 * created... could still be in the process of loading).
-	 * 
-	 * @param  ImportIndex    Specifies the import that you would like a linker for.
-	 * @return The imports associated linker (null if it hasn't been created yet).
-	 */
-	FLinkerLoad* FindExistingLinkerForImport(int32 ImportIndex) const;
-
-	UObject* IndexToObject( FPackageIndex Index );
-
-	void DetachExport( int32 i );
-
-	// FArchive interface.
-	/**
-	 * Hint the archive that the region starting at passed in offset and spanning the passed in size
-	 * is going to be read soon and should be precached.
-	 *
-	 * The function returns whether the precache operation has completed or not which is an important
-	 * hint for code knowing that it deals with potential async I/O. The archive is free to either not 
-	 * implement this function or only partially precache so it is required that given sufficient time
-	 * the function will return true. Archives not based on async I/O should always return true.
-	 *
-	 * This function will not change the current archive position.
-	 *
-	 * @param	PrecacheOffset	Offset at which to begin precaching.
-	 * @param	PrecacheSize	Number of bytes to precache
-	 * @return	false if precache operation is still pending, true otherwise
-	 */
-	virtual bool Precache( int64 PrecacheOffset, int64 PrecacheSize ) override;
-
-#if WITH_EDITOR
-	/**
-	 * Attaches/ associates the passed in bulk data object with the linker.
-	 *
-	 * @param	Owner		UObject owning the bulk data
-	 * @param	BulkData	Bulk data object to associate
-	 */
-	virtual void AttachBulkData(UObject* Owner, FUntypedBulkData* BulkData) override;
-	/**
-	 * Detaches the passed in bulk data object from the linker.
-	 *
-	 * @param	BulkData	Bulk data object to detach
-	 * @param	bEnsureBulkDataIsLoaded	Whether to ensure that the bulk data is loaded before detaching
-	 */
-	virtual void DetachBulkData(FUntypedBulkData* BulkData, bool bEnsureBulkDataIsLoaded) override;
-	/**
-	 * Detaches all attached bulk  data objects.
-	 *
-	 * @param	bEnsureBulkDataIsLoaded	Whether to ensure that the bulk data is loaded before detaching
-	 */
-#endif
-	void DetachAllBulkData(bool bEnsureBulkDataIsLoaded);
-public:
-
-	/**
-	 * Detaches linker from bulk data.
-	 */
-	void LoadAndDetachAllBulkData();
-
-	/**
-	* Detaches linker from bulk data/ exports and removes itself from array of loaders.
-	*/
-	COREUOBJECT_API void Detach();
-
-private:
-
-	void Seek( int64 InPos ) override;
-	int64 Tell() override;
-	int64 TotalSize() override;
-	// this fixes the warning : 'ULinkerSave::Serialize' hides overloaded virtual function
-	using FLinker::Serialize;
-	virtual void Serialize( void* V, int64 Length ) override;
-	using FArchiveUObject::operator<<; // For visibility of the overloads we don't override
-	virtual FArchive& operator<<( UObject*& Object ) override;
-	virtual FArchive& operator<<( FLazyObjectPtr& LazyObjectPtr) override;
-	virtual FArchive& operator<<( FAssetPtr& AssetPtr) override;
-	virtual FArchive& operator<<( FName& Name ) override;
-
-	/**
-	 * Safely verify that an import in the ImportMap points to a good object. This decides whether or not
-	 * a failure to load the object redirector in the wrapper is a fatal error or not (return value)
-	 *
-	 * @param	i				The index into this packages ImportMap to verify
-	 * @param	WarningSuffix	[out] additional information about the load failure that should be appended to
-	 *							the name of the object in the load failure dialog.
-	 *
-	 * @return true if the wrapper should crash if it can't find a good object redirector to load
-	 */
-	bool VerifyImportInner(const int32 ImportIndex, FString& WarningSuffix);
-
-	//
-	// FLinkerLoad creation helpers BEGIN
-	//
-
-	/**
-	 * Creates a FLinkerLoad object for async creation. Tick has to be called manually till it returns
-	 * true in which case the returned linker object has finished the async creation process.
-	 *
-	 * @param	Parent		Parent object to load into, can be NULL (most likely case)
-	 * @param	Filename	Name of file on disk to load
-	 * @param	LoadFlags	Load flags determining behavior
-	 *
-	 * @return	new FLinkerLoad object for Parent/ Filename
-	 */
-	COREUOBJECT_API static FLinkerLoad* CreateLinkerAsync( UPackage* Parent, const TCHAR* Filename, uint32 LoadFlags );
-
-	/**
-	 * Ticks an in-flight linker and spends InTimeLimit seconds on creation. This is a soft time limit used
-	 * if bInUseTimeLimit is true.
-	 *
-	 * @param	InTimeLimit		Soft time limit to use if bInUseTimeLimit is true
-	 * @param	bInUseTimeLimit	Whether to use a (soft) timelimit
-	 * @param	bInUseFullTimeLimit	Whether to use the entire time limit, even if blocked on I/O
-	 * 
-	 * @return	true if linker has finished creation, false if it is still in flight
-	 */
-	ELinkerStatus Tick( float InTimeLimit, bool bInUseTimeLimit, bool bInUseFullTimeLimit);
-
-protected: // Daniel L: Made this protected so I can override the constructor and create a custom loader to load the header of the linker in the DiffFilesCommandlet
-	/**
-	 * Private constructor, passing arguments through from CreateLinker.
-	 *
-	 * @param	Parent		Parent object to load into, can be NULL (most likely case)
-	 * @param	Filename	Name of file on disk to load
-	 * @param	LoadFlags	Load flags determining behavior
-	 */
-	FLinkerLoad(UPackage* InParent, const TCHAR* InFilename, uint32 InLoadFlags);
-private:
-	/**
-	 * Returns whether the time limit allotted has been exceeded, if enabled.
-	 *
-	 * @param CurrentTask	description of current task performed for logging spilling over time limit
-	 * @param Granularity	Granularity on which to check timing, useful in cases where FPlatformTime::Seconds is slow (e.g. PC)
-	 *
-	 * @return true if time limit has been exceeded (and is enabled), false otherwise (including if time limit is disabled)
-	 */
-	COREUOBJECT_API bool IsTimeLimitExceeded( const TCHAR* CurrentTask, int32 Granularity = 1 );
-protected: // Daniel L: Made this protected so I can override the constructor and create a custom loader to load the header of the linker in the DiffFilesCommandlet
-	/**
-	 * Creates loader used to serialize content.
-	 */
-	ELinkerStatus CreateLoader();
-private:
-	/**
-	 * Serializes the package file summary.
-	 */
-	ELinkerStatus SerializePackageFileSummary();
-
-	/**
-	 * Serializes the name map.
-	 */
-	ELinkerStatus SerializeNameMap();
-
-	/**
-	 * Serializes the import map.
-	 */
-	ELinkerStatus SerializeImportMap();
-
-	/**
-	 * Fixes up the import map, performing remapping for backward compatibility and such.
-	 */
-	ELinkerStatus FixupImportMap();
-
-	/**
-	 * Changes imports before the objects are created.
-	 */
-	ELinkerStatus RemapImports();
-
-	/**
-	 * Serializes the export map.
-	 */
-	ELinkerStatus SerializeExportMap();
-
-#if WITH_ENGINE
-	/**
-	 * Kicks off async memory allocations for all textures that will be loaded from this package.
-	 */
-	ELinkerStatus StartTextureAllocation();
-#endif
-
-	/**
-	 * Serializes the import map.
-	 */
-	ELinkerStatus SerializeDependsMap();
-
-public:
-	/**
-	 * Serializes the gatherable text data container.
-	 */
-	COREUOBJECT_API ELinkerStatus SerializeGatherableTextDataMap(bool bForceEnableForCommandlet = false);
+		return ImportMap[ImportIdx].ClassName;
+	}
+	/** Gets the class name for the specified index in the import map. */
+	FName GetImportClassName(FPackageIndex PackageIndex)
+	{
+		if (PackageIndex.IsImport())
+		{
+			return GetImportClassName(PackageIndex.ToImport());
+		}
+		return NAME_None;
+	}
+	/** Gets the class name for the specified package index. */
+	FName GetClassName(FPackageIndex PackageIndex)
+	{
+		if (PackageIndex.IsImport())
+		{
+			return GetImportClassName(PackageIndex);
+		}
+		else if (PackageIndex.IsExport())
+		{
+			return GetExportClassName(PackageIndex);
+		}
+		return NAME_None;
+	}
 
-	/**
-	 * Serializes thumbnails
-	 */
-	COREUOBJECT_API ELinkerStatus SerializeThumbnails( bool bForceEnableForCommandlet=false );
 
-	/** Inform the archive that blueprint finalization is pending. */
-	virtual void ForceBlueprintFinalization() override;
+	FORCEINLINE ELinkerType::Type GetType() const
+	{
+		return LinkerType;
+	}
 
 	/**
-	* Query method to help handle recursive behavior. When this returns true,
-	* this linker is in the middle of, or is about to call FinalizeBlueprint()
-	* (for a blueprint class somewhere in the current callstack). Needed when
-	* we get to finalizing a sub-class before we've finished finalizing its
-	* super (so we know we need to finish finalizing the super first).
+	 * I/O function
 	*
-	* @return True if FinalizeBlueprint() is currently being ran (or about to be ran) for an export (Blueprint) class.
-	*/
-	bool IsBlueprintFinalizationPending() const;
-
-	/**
-	 * Gives external code the ability to create FLinkerPlaceholderBase objects
-	 * in place of loads that may violate the LOAD_DeferDependencyLoads state.
-	 * This will only produce a placeholder if LOAD_DeferDependencyLoads is set
-	 * for this linker.
-	 *
-	 * NOTE: For now, this will only produce UClass placeholders, as that is the 
-	 *       only type we've identified needing.
-	 * 
-	 * @param  ObjectType    The expected type of the object you want to defer loading of.
-	 * @param  ObjectPath    The full object/package path for the expected object.
-	 * @return A FLinkerPlaceholderBase UObject that can be used in place of the import dependency.
+	 * @param	Ar	the archive to read/write into
 	 */
-	UObject* RequestPlaceholderValue(UClass* ObjectType, const TCHAR* ObjectPath);
-
-private:
+	void Serialize(FArchive& Ar);
+	void AddReferencedObjects(FReferenceCollector& Collector);
 	/**
-	 * Regenerates/Refreshes a blueprint class
+	 * Return the path name of the UObject represented by the specified import. 
+	 * (can be used with StaticFindObject)
+	 * 
+	 * @param	ImportIndex	index into the ImportMap for the resource to get the name for
 	 *
-	 * @param	LoadClass		Instance of the class currently being loaded and which is the parent for the blueprint
-	 * @param	ExportObject	Current object bein exported
-	 * @return	Returns true if regeneration was successful, otherwise false
+	 * @return	the path name of the UObject represented by the resource at ImportIndex
 	 */
-	bool RegenerateBlueprintClass(UClass* LoadClass, UObject* ExportObject);
-
+	COREUOBJECT_API FString GetImportPathName(int32 ImportIndex);
 	/**
-	 * Determines if the specified import should be deferred. If so, it will 
-	 * instantiate a placeholder object in its place.
+	 * Return the path name of the UObject represented by the specified import. 
+	 * (can be used with StaticFindObject)
 	 * 
-	 * @param  ImportIndex    An index into this linker's ImportMap, specifying which import to check.
-	 * @return True if the specified import was deferred, other wise false (it is ok to load it).
+	 * @param	PackageIndex	package index for the resource to get the name for
+	 *
+	 * @return	the path name of the UObject represented by the resource at PackageIndex, or the empty string if this isn't an import
 	 */
-	bool DeferPotentialCircularImport(const int32 ImportIndex);
-	
-#if WITH_EDITOR
-	/**
-	* Determines if the Object Import error should be suppressed
-	*
-	* @param  ImportIndex    Internal index into this linker's ImportMap, references the import to check for suppression.
-	* @return True if the import error should be suppressed
-	*/
-	bool IsSuppressableBlueprintImportError(int32 ImportIndex) const;
-#endif // WITH_EDITOR
-
+	FString GetImportPathName(FPackageIndex PackageIndex)
+	{
+		if (PackageIndex.IsImport())
+		{
+			return GetImportPathName(PackageIndex.ToImport());
+		}
+		return FString();
+	}
 	/**
-	 * Stubs in a ULinkerPlaceholderExportObject for the specified export (if 
-	 * one is required, meaning: the export's LoadClass is not fully formed). 
-	 * This should rarely happen, but has been seen in cyclic Blueprint 
-	 * scenarios involving Blueprinted components.
+	 * Return the path name of the UObject represented by the specified export.
+	 * (can be used with StaticFindObject)
 	 * 
-	 * @param  ExportIndex    Identifies the export you want deferred.
-	 * @return True if the export has been deferred (and should not be loaded).
+	 * @param	ExportIndex				index into the ExportMap for the resource to get the name for
+	 * @param	FakeRoot				Optional name to replace use as the root package of this object instead of the linker
+	 * @param	bResolveForcedExports	if true, the package name part of the return value will be the export's original package,
+	 *									not the name of the package it's currently contained within.
+	 *
+	 * @return	the path name of the UObject represented by the resource at ExportIndex
 	 */
-	bool DeferExportCreation(const int32 ExportIndex);
-
+	COREUOBJECT_API FString GetExportPathName(int32 ExportIndex, const TCHAR* FakeRoot=NULL,bool bResolveForcedExports=false);
 	/**
-	 * Iterates through this linker's ExportMap, looking for the corresponding
-	 * class-default-object for the specified class (assumes that the supplied 
-	 * class is an export itself, making this a Blueprint package).
+	 * Return the path name of the UObject represented by the specified export.
+	 * (can be used with StaticFindObject)
 	 * 
-	 * @param  LoadClass    The Blueprint class that this linker is in charge of loading (also belonging to its ExportMap).
-	 * @return An index into this linker's ExportMap array (INDEX_NONE if the CDO wasn't found).
+	 * @param	PackageIndex			package index for the resource to get the name for
+	 * @param	FakeRoot				Optional name to replace use as the root package of this object instead of the linker
+	 * @param	bResolveForcedExports	if true, the package name part of the return value will be the export's original package,
+	 *									not the name of the package it's currently contained within.
+	 *
+	 * @return	the path name of the UObject represented by the resource at PackageIndex, or the empty string if this isn't an export
 	 */
-	int32 FindCDOExportIndex(UClass* LoadClass);
+	FString GetExportPathName(FPackageIndex PackageIndex, const TCHAR* FakeRoot=NULL,bool bResolveForcedExports=false)
+	{
+		if (PackageIndex.IsExport())
+		{
+			return GetExportPathName(PackageIndex.ToExport(), FakeRoot, bResolveForcedExports);
+		}
+		return FString();
+	}
 
 	/**
-	 * Combs the ImportMap for any imports that were deferred, and then creates 
-	 * them (via CreateImport).
+	 * Return the path name of the UObject represented by the specified import. 
+	 * (can be used with StaticFindObject)
 	 * 
-	 * @param  LoadStruct    The (Blueprint) class or struct that you want resolved (so that it no longer contains dependency placeholders).
+	 * @param	PackageIndex	package index
+	 *
+	 * @return	the path name of the UObject represented by the resource at PackageIndex, or the empty string if this is null
 	 */
-	void ResolveDeferredDependencies(UStruct* LoadStruct);
-
+	FString GetPathName(FPackageIndex PackageIndex)
+	{
+		if (PackageIndex.IsImport())
+		{
+			return GetImportPathName(PackageIndex);
+		}
+		else if (PackageIndex.IsExport())
+		{
+			return GetExportPathName(PackageIndex);
+		}
+		return FString();
+	}
 	/**
-	 * Loads the import that the Placeholder was initially stubbed in for (NOTE:
-	 * this could cause recursive behavior), and then replaces all known 
-	 * placeholder references with the proper class.
+	 * Return the full name of the UObject represented by the specified import.
 	 * 
-	 * @param  Placeholder		A ULinkerPlaceholderClass that was substituted in place of a deferred dependency.
-	 * @param  ReferencingClass	The (Blueprint) class that was loading, while we deferred dependencies (now referencing the placeholder).
-	 * @param  ObjectPath		Optional param that denotes the full object/package path for the object the placeholder is supposed to represent 
-	 *                          (used when the passed placeholder is not tied to an import in the linker's ImportMap).
-	 * @return The number of placeholder references replaced (could be none, if this was recursively resolved).
+	 * @param	ImportIndex	index into the ImportMap for the resource to get the name for
+	 *
+	 * @return	the full name of the UObject represented by the resource at ImportIndex
 	 */
-	int32 ResolveDependencyPlaceholder(class FLinkerPlaceholderBase* Placeholder, UClass* ReferencingClass = nullptr, const FName ObjectPath = NAME_None);
-
+	COREUOBJECT_API FString GetImportFullName(int32 ImportIndex);
 	/**
-	 * Query method to help catch recursive behavior. When this returns true, a 
-	 * dependency placeholder is in the middle of being resolved by 
-	 * ResolveDependencyPlaceholder(). Used so a nested call would know to 
-	 * complete that placeholder before continuing.
+	 * Return the full name of the UObject represented by the specified package index
 	 * 
-	 * @return True if ResolveDependencyPlaceholder() is being ran on a placeholder that has yet to be resolved. 
-	 */
-	bool HasUnresolvedDependencies() const;
-
-	/**
-	 * Iterates through the ImportMap and calls CreateImport() for every entry, 
-	 * creating/loading each import as we go. This also makes sure that class 
-	 * imports have had ResolveDeferredDependencies() completely executed for 
-	 * them (even those already running through it earlier in the callstack).
+	 * @param	PackageIndex	package index for the resource to get the name for
+	 *
+	 * @return	the full name of the UObject represented by the resource at PackageIndex
 	 */
-	void ResolveAllImports();
+	FString GetImportFullName(FPackageIndex PackageIndex)
+	{
+		if (PackageIndex.IsImport())
+		{
+			return GetImportFullName(PackageIndex.ToImport());
+		}
+		return FString();
+	}
 
 	/**
-	 * Takes the supplied serialized class and serializes in its CDO, then 
-	 * regenerates both.
+	 * Return the full name of the UObject represented by the specified export.
 	 * 
-	 * @param  LoadClass    The loaded blueprint class (assumes that it has been fully loaded/serialized).
-	 */
-	void FinalizeBlueprint(UClass* LoadClass);
-
-	/**
-	 * Combs the ExportMap for any stubbed in ULinkerPlaceholderExportObjects,
-	 * and finalizes the real export's class before actually creating it
-	 * (exports are deferred when their class isn't fully formed at the time
-	 * CreateExport() is called). Also, this function ensures that deferred CDO  
-	 * serialization is executed (expects its class to be fully resolved at this
-	 * point).
+	 * @param	ExportIndex				index into the ExportMap for the resource to get the name for
+	 * @param	FakeRoot				Optional name to replace use as the root package of this object instead of the linker
+	 * @param	bResolveForcedExports	if true, the package name part of the return value will be the export's original package,
+	 *									not the name of the package it's currently contained within.
 	 *
-	 * @param  LoadClass    A fully loaded/serialized class that may have property references to placeholder export objects (in need of fix-up).
+	 * @return	the full name of the UObject represented by the resource at ExportIndex
 	 */
-	void ResolveDeferredExports(UClass* LoadClass);
-
+	COREUOBJECT_API FString GetExportFullName(int32 ExportIndex, const TCHAR* FakeRoot=NULL,bool bResolveForcedExports=false);
 	/**
-	 * Sometimes we have to instantiate an export object that is of an imported 
-	 * type, and sometimes in those scenarios (thanks to cyclic dependencies) 
-	 * the type class could be a Blueprint type that is half resolved. To avoid
-	 * having to re-instance objects on load, we have to ensure that the class
-	 * is fully regenerated before we spawn any instances of it. That's where 
-	 * this function comes in. It will make sure that the specified class is 
-	 * fully loaded, finalized, and regenerated.
-	 *
-	 * NOTE: be wary, if called in the wrong place, then this could introduce 
-	 *       nasty infinite recursion!
+	 * Return the full name of the UObject represented by the specified package index
 	 * 
-	 * @param  ImportClass    The class you want to make sure is fully regenerated.
-	 * @return True if the class could be regenerated (false if it didn't have its linker set).
+	 * @param	PackageIndex			package index for the resource to get the name for
+	 * @param	FakeRoot				Optional name to replace use as the root package of this object instead of the linker
+	 * @param	bResolveForcedExports	if true, the package name part of the return value will be the export's original package,
+	 *									not the name of the package it's currently contained within.
+	 *
+	 * @return	the full name of the UObject represented by the resource at PackageIndex
 	 */
-	bool ForceRegenerateClass(UClass* ImportClass);
+	FString GetExportFullName(FPackageIndex PackageIndex, const TCHAR* FakeRoot=NULL,bool bResolveForcedExports=false)
+	{
+		if (PackageIndex.IsExport())
+		{
+			return GetExportFullName(PackageIndex.ToExport(), FakeRoot, bResolveForcedExports);
+		}
+		return FString();
+	}
 
 	/**
-	 * Checks to see if an export (or one up its outer chain) is currently 
-	 * in the middle of having its class dependency force-regenerated. This 
-	 * function is meant to help avoid unnecessary recursion, as 
-	 * ForceRegenerateClass() does nothing itself to stave off infinite 
-	 * recursion.
+	 * Return the full name of the UObject represented by the specified export.
 	 * 
-	 * @param  ExportIndex    Identifies the export you're about to call CreateExport() on.
-	 * @return True if the specified export's class (or one up its outer chain) is currently being force-regenerated.
-	 */
-	bool IsExportBeingResolved(int32 ExportIndex);
-
-	void ResetDeferredLoadingState();
-
-	bool HasPerformedFullExportResolvePass();
-
-	/** Finds import, tries to fall back to dynamic class if the object could not be found */
-	UObject* FindImport(UClass* ImportClass, UObject* ImportOuter, const TCHAR* Name);
-	/** Finds import, tries to fall back to dynamic class if the object could not be found */
-	UObject* FindImportFast(UClass* ImportClass, UObject* ImportOuter, FName Name);
-
-	/** Fills all necessary information for constructing dynamic type package linker */
-	void CreateDynamicTypeLoader();
-
-#if	USE_CIRCULAR_DEPENDENCY_LOAD_DEFERRING
-	/** 
-	 * For deferring dependency loads, we block CDO serialization until the 
-	 * class if complete. If we attempt to serialize the CDO while that is 
-	 * happening, we instead defer it and record the export's index here (so we 
-	 * can return to it later).
-	 */
-	bool bForceBlueprintFinalization;
-
-	/** 
-	 * Index of the CDO that should be used for blueprint finalization, may be INDEX_NONE
-	 * in the case of some legacy content.
-	 */
-	int32 DeferredCDOIndex;
-
-	/** 
-	 * Used to track dependency placeholders currently being resolved inside of 
-	 * ResolveDependencyPlaceholder()... utilized for nested reentrant behavior, 
-	 * to make sure this placeholder is completely resolved before continuing on 
-	 * to the next.
+	 * @param	PackageIndex	package index
+	 *
+	 * @return	the path name of the UObject represented by the resource at PackageIndex, or the empty string if this is null
 	 */
-	class FLinkerPlaceholderBase* ResolvingDeferredPlaceholder;
+	FString GetFullImpExpName(FPackageIndex PackageIndex)
+	{
+		if (PackageIndex.IsImport())
+		{
+			return GetImportFullName(PackageIndex);
+		}
+		else if (PackageIndex.IsExport())
+		{
+			return GetExportFullName(PackageIndex);
+		}
+		return FString();
+	}
 
-	/** 
-	 * Internal list to track imports that were deferred, but don't belong to 
-	 * the ImportMap (thinks ones loaded through config files via UProperty::ImportText).
+	/**
+	 * Tell this linker to start SHA calculations
 	 */
-	TMap<FName, FLinkerPlaceholderBase*> ImportPlaceholders;
-#endif // USE_CIRCULAR_DEPENDENCY_LOAD_DEFERRING
-
+	void StartScriptSHAGeneration();
 
-	/** 
-	 * Creates the export hash.
+	/**
+	 * If generating a script SHA key, update the key with this script code
+	 *
+	 * @param ScriptCode Code to SHAify
 	 */
-	ELinkerStatus CreateExportHash();
+	void UpdateScriptSHAKey(const TArray<uint8>& ScriptCode);
 
 	/**
-	 * Finds existing exports in memory and matches them up with this linker. This is required for PIE to work correctly
-	 * and also for script compilation as saving a package will reset its linker and loading will reload/ replace existing
-	 * objects without a linker.
+	 * After generating the SHA key for all of the 
+	 *
+	 * @param OutKey Storage for the key bytes (20 bytes)
 	 */
-	ELinkerStatus FindExistingExports();
+	void GetScriptSHAKey(uint8* OutKey);
 
 	/**
-	 * Finalizes linker creation, adding linker to loaders array and potentially verifying imports.
+	 * Test and object against the load flag filters
+	 *
+	 * @return	true if the object should be filtered and not loaded
 	 */
-	ELinkerStatus FinalizeCreation();
-
-	//
-	// FLinkerLoad creation helpers END
-	//
-};
-
-/*----------------------------------------------------------------------------
-	FLinkerSave.
-----------------------------------------------------------------------------*/
-
-/**
- * Handles saving Unreal package files.
- */
-class FLinkerSave : public FLinker, public FArchiveUObject
-{
-public:
-
-	FORCEINLINE static ELinkerType::Type StaticType()
-	{
-		return ELinkerType::Save;
-	}
-
-	virtual ~FLinkerSave();
-
-	// Variables.
-	/** The archive that actually writes the data to disk. */
-	FArchive* Saver;
-
-	/** Index array - location of the resource for a UObject is stored in the ObjectIndices array using the UObject's Index */
-	TMap<UObject *,FPackageIndex> ObjectIndicesMap;
-
-	/** Index array - location of the name in the NameMap array for each FName is stored in the NameIndices array using the FName's Index */
-	TMap<FName, int32, FDefaultSetAllocator, TLinkerNameMapKeyFuncs<int32>> NameIndices;
-
-	/** List of bulkdata that needs to be stored at the end of the file */
-	struct FBulkDataStorageInfo
+	bool FilterExport(const FObjectExport& Export)
 	{
-		/** Offset to the location where the payload offset is stored */
-		int64 BulkDataOffsetInFilePos;
-		/** Offset to the location where the payload size is stored */
-		int64 BulkDataSizeOnDiskPos;
-		/** Offset to the location where the bulk data flags are stored */
-		int64 BulkDataFlagsPos;
-		/** Bulk data flags at the time of serialization */
-		uint32 BulkDataFlags;
-		/** The bulkdata */
-		FUntypedBulkData* BulkData;
-	};
-	TArray<FBulkDataStorageInfo> BulkDataToAppend;
-
-	/** A mapping of package name to generated script SHA keys */
-	COREUOBJECT_API static TMap<FString, TArray<uint8> > PackagesToScriptSHAMap;
-
-	/** Constructor */
-	FLinkerSave(UPackage* InParent, const TCHAR* InFilename, bool bForceByteSwapping, bool bInSaveUnversioned = false );
-	/** Constructor for memory writer */
-	FLinkerSave(UPackage* InParent, bool bForceByteSwapping, bool bInSaveUnversioned = false );
-	/** Constructor for custom savers */
-	FLinkerSave(UPackage* InParent, FArchive *InSaver, bool bForceByteSwapping, bool bInSaveUnversioned = false);
-
-	/** Returns the appropriate name index for the source name, or 0 if not found in NameIndices */
-	int32 MapName(const FName& Name) const;
-
-	/** Returns the appropriate package index for the source object, or default value if not found in ObjectIndicesMap */
-	FPackageIndex MapObject(const UObject* Object) const;
-
-	// FArchive interface.
-	using FArchiveUObject::operator<<; // For visibility of the overloads we don't override
-	FArchive& operator<<( FName& InName );
-	FArchive& operator<<( UObject*& Obj );
-	FArchive& operator<<( FLazyObjectPtr& LazyObjectPtr );
-	FArchive& operator<<( FAssetPtr& AssetPtr );
-
+		if (Export.bExportLoadFailed || Export.bWasFiltered)
+		{
+			return true;
+		}
 #if WITH_EDITOR
-	// proxy for debugdata
-	virtual void PushDebugDataString(const FName& DebugData) override { Saver->PushDebugDataString(DebugData); }
-	virtual void PopDebugDataString() override { Saver->PopDebugDataString(); }
+		if (!Export.bNotForEditorGame)
+		{
+			return false;
+		}
 #endif
+		if (FilterClientButNotServer && Export.bNotForServer) // "we are a dedicated server"
+		{
+			return true;
+		}
+		if (FilterServerButNotClient && Export.bNotForClient) // "we are a client only"
+		{
+			return true;
+		}
+		if (Export.ThisIndex.IsNull()) // Export is invalid and shouldn't be processed.
+		{
+			return true;
+		}
+		return false;
+	}
 
+	FORCEINLINE bool IsCompressed() const { return !!(Summary.PackageFlags & PKG_StoreCompressed); }
+};
 
-	virtual FString GetArchiveName() const override;
-
-	/**
-	 * If this archive is a FLinkerLoad or FLinkerSave, returns a pointer to the FLinker portion.
-	 */
-	virtual FLinker* GetLinker() { return this; }
-
-	void Seek( int64 InPos );
-	int64 Tell();
-	// this fixes the warning : 'FLinkerSave::Serialize' hides overloaded virtual function
-	using FLinker::Serialize;
-	void Serialize( void* V, int64 Length );
-
-	/**
-	 * Detaches file saver and hence file handle.
-	 */
-	void Detach();
 
-	/**
-	 * Sets a flag indicating that this archive contains data required to be gathered for localization.
-	 */
-	void ThisRequiresLocalizationGather();
-};
 
 template<typename T> 
 FORCEINLINE T* Cast(FLinker* Src)
