Copying //UE4/Dev-Build to //UE4/Dev-Main (Source: //UE4/Dev-Build @ 3104200)

diff --git a/Engine/Source/Runtime/CoreUObject/Public/UObject/LinkerSave.h b/Engine/Source/Runtime/CoreUObject/Public/UObject/LinkerSave.h
new file mode 100644
--- /dev/null
+++ b/Engine/Source/Runtime/CoreUObject/Public/UObject/LinkerSave.h
@@ -0,0 +1,106 @@
+// Copyright 1998-2016 Epic Games, Inc. All Rights Reserved.
+
+#pragma once
+
+#include "ObjectMacros.h"
+#include "EngineVersion.h"
+#include "GatherableTextData.h"
+#include "Linker.h"
+
+/*----------------------------------------------------------------------------
+	FLinkerSave.
+----------------------------------------------------------------------------*/
+
+/**
+ * Handles saving Unreal package files.
+ */
+class FLinkerSave : public FLinker, public FArchiveUObject
+{
+public:
+
+	FORCEINLINE static ELinkerType::Type StaticType()
+	{
+		return ELinkerType::Save;
+	}
+
+	virtual ~FLinkerSave();
+
+	// Variables.
+	/** The archive that actually writes the data to disk. */
+	FArchive* Saver;
+
+	/** Index array - location of the resource for a UObject is stored in the ObjectIndices array using the UObject's Index */
+	TMap<UObject *,FPackageIndex> ObjectIndicesMap;
+
+	/** Index array - location of the name in the NameMap array for each FName is stored in the NameIndices array using the FName's Index */
+	TMap<FName, int32, FDefaultSetAllocator, TLinkerNameMapKeyFuncs<int32>> NameIndices;
+
+	/** List of bulkdata that needs to be stored at the end of the file */
+	struct FBulkDataStorageInfo
+	{
+		/** Offset to the location where the payload offset is stored */
+		int64 BulkDataOffsetInFilePos;
+		/** Offset to the location where the payload size is stored */
+		int64 BulkDataSizeOnDiskPos;
+		/** Offset to the location where the bulk data flags are stored */
+		int64 BulkDataFlagsPos;
+		/** Bulk data flags at the time of serialization */
+		uint32 BulkDataFlags;
+		/** The bulkdata */
+		FUntypedBulkData* BulkData;
+	};
+	TArray<FBulkDataStorageInfo> BulkDataToAppend;
+
+	/** A mapping of package name to generated script SHA keys */
+	COREUOBJECT_API static TMap<FString, TArray<uint8> > PackagesToScriptSHAMap;
+
+	/** Constructor */
+	FLinkerSave(UPackage* InParent, const TCHAR* InFilename, bool bForceByteSwapping, bool bInSaveUnversioned = false );
+	/** Constructor for memory writer */
+	FLinkerSave(UPackage* InParent, bool bForceByteSwapping, bool bInSaveUnversioned = false );
+	/** Constructor for custom savers */
+	FLinkerSave(UPackage* InParent, FArchive *InSaver, bool bForceByteSwapping, bool bInSaveUnversioned = false);
+
+	/** Returns the appropriate name index for the source name, or 0 if not found in NameIndices */
+	int32 MapName(const FName& Name) const;
+
+	/** Returns the appropriate package index for the source object, or default value if not found in ObjectIndicesMap */
+	FPackageIndex MapObject(const UObject* Object) const;
+
+	// FArchive interface.
+	using FArchiveUObject::operator<<; // For visibility of the overloads we don't override
+	FArchive& operator<<( FName& InName );
+	FArchive& operator<<( UObject*& Obj );
+	FArchive& operator<<( FLazyObjectPtr& LazyObjectPtr );
+	FArchive& operator<<( FAssetPtr& AssetPtr );
+
+#if WITH_EDITOR
+	// proxy for debugdata
+	virtual void PushDebugDataString(const FName& DebugData) override { Saver->PushDebugDataString(DebugData); }
+	virtual void PopDebugDataString() override { Saver->PopDebugDataString(); }
+#endif
+
+
+	virtual FString GetArchiveName() const override;
+
+	/**
+	 * If this archive is a FLinkerLoad or FLinkerSave, returns a pointer to the FLinker portion.
+	 */
+	virtual FLinker* GetLinker() { return this; }
+
+	void Seek( int64 InPos );
+	int64 Tell();
+	// this fixes the warning : 'FLinkerSave::Serialize' hides overloaded virtual function
+	using FLinker::Serialize;
+	void Serialize( void* V, int64 Length );
+
+	/**
+	 * Detaches file saver and hence file handle.
+	 */
+	void Detach();
+
+	/**
+	 * Sets a flag indicating that this archive contains data required to be gathered for localization.
+	 */
+	void ThisRequiresLocalizationGather();
+};
diff --git a/Engine/Source/Runtime/CoreUObject/Public/UObject/ObjectMacros.h b/Engine/Source/Runtime/CoreUObject/Public/UObject/ObjectMacros.h
new file mode 100644
--- /dev/null
+++ b/Engine/Source/Runtime/CoreUObject/Public/UObject/ObjectMacros.h
@@ -0,0 +1,1642 @@
+// Copyright 1998-2016 Epic Games, Inc. All Rights Reserved.
+
+/*=============================================================================
+	ObjectBase.h: Unreal object base class.
+=============================================================================*/
+
+#pragma once
+
+/** Represents a serializable object pointer in blueprint bytecode. This is always 64-bits, even on 32-bit platforms. */
+typedef	uint64 ScriptPointerType;
+
+
+#if PLATFORM_VTABLE_AT_END_OF_CLASS
+#error "not supported in UE4"
+#endif
+
+#if HACK_HEADER_GENERATOR 
+#define USE_COMPILED_IN_NATIVES	0
+#else
+#define USE_COMPILED_IN_NATIVES	1
+#endif
+
+/** Set this to 0 to disable UObject thread safety features */
+#define THREADSAFE_UOBJECTS 1
+
+// 1 = old IsA behavior
+// 2 = new IsA behavior
+// 3 = old IsA behavior with checks against the new behavior
+#if 0//!(UE_BUILD_SHIPPING || UE_BUILD_TEST)
+	#define UCLASS_FAST_ISA_IMPL 3
+#else
+	#define UCLASS_FAST_ISA_IMPL 2
+#endif
+
+/*-----------------------------------------------------------------------------
+	Core enumerations.
+-----------------------------------------------------------------------------*/
+
+//
+// Flags for loading objects.
+//
+enum ELoadFlags
+{
+	LOAD_None						= 0x00000000,	// No flags.
+	LOAD_SeekFree					= 0x00000001,	// Loads the package via the seek free loading path/ reader
+	LOAD_NoWarn						= 0x00000002,	// Don't display warning if load fails.
+	LOAD_EditorOnly					= 0x00000004, // Load for editor-only purposes and by editor-only code
+	LOAD_ResolvingDeferredExports	= 0x00000008,	// Denotes that we should not defer export loading (as we're resolving them)
+	LOAD_Verify						= 0x00000010,	// Only verify existance; don't actually load.
+	LOAD_AllowDll					= 0x00000020,	// Allow plain DLLs.
+//	LOAD_Unused						= 0x00000040
+	LOAD_NoVerify					= 0x00000080,   // Don't verify imports yet.
+	LOAD_IsVerifying			= 0x00000100,		// Is verifying imports
+//	LOAD_Unused						= 0x00000200,
+//	LOAD_Unused						= 0x00000400,
+//	LOAD_Unused						= 0x00000800,
+	LOAD_DisableDependencyPreloading = 0x00001000,	// Bypass dependency preloading system
+	LOAD_Quiet						= 0x00002000,   // No log warnings.
+	LOAD_FindIfFail					= 0x00004000,	// Tries FindObject if a linker cannot be obtained (e.g. package is currently being compiled)
+	LOAD_MemoryReader				= 0x00008000,	// Loads the file into memory and serializes from there.
+	LOAD_NoRedirects				= 0x00010000,	// Never follow redirects when loading objects; redirected loads will fail
+	LOAD_ForDiff					= 0x00020000,	// Loading for diffing.
+	LOAD_PackageForPIE				= 0x00080000,   // This package is being loaded for PIE, it must be flagged as such immediately
+	LOAD_DeferDependencyLoads       = 0x00100000,   // Do not load external (blueprint) dependencies (instead, track them for deferred loading)
+	LOAD_ForFileDiff				= 0x00200000,	// Load the package (not for diffing in the editor), instead verify at the two packages serialized output are the same, if they are not then debug break so that you can get the callstack and object information
+};
+
+//
+// Flags for saving packages
+//
+enum ESaveFlags
+{
+	SAVE_None			= 0x00000000,	// No flags
+	SAVE_NoError		= 0x00000001,	// Don't generate errors on save
+	SAVE_FromAutosave	= 0x00000002,   // Used to indicate this save was initiated automatically
+	SAVE_KeepDirty		= 0x00000004,	// Do not clear the dirty flag when saving
+	SAVE_KeepGUID		= 0x00000008,	// Keep the same guid, used to save cooked packages
+	SAVE_Async			= 0x00000010,	// Save to a memory writer, then actually write to disk async
+	SAVE_Unversioned	= 0x00000020,	// Save all versions as zero. Upon load this is changed to the current version. This is only reasonable to use with full cooked builds for distribution.
+	SAVE_CutdownPackage	= 0x00000040,	// Saving cutdown packages in a temp location WITHOUT renaming the package.
+	SAVE_KeepEditorOnlyCookedPackages = 0x00000080,  // keep packages which are marked as editor only even though we are cooking
+};
+
+//
+// Package flags.
+//
+enum EPackageFlags
+{
+	PKG_None						= 0x00000000,	// No flags
+	PKG_NewlyCreated				= 0x00000001,	// Newly created package, not saved yet. In editor only.
+	PKG_ClientOptional				= 0x00000002,	// Purely optional for clients.
+	PKG_ServerSideOnly				= 0x00000004,   // Only needed on the server side.
+	PKG_CompiledIn					= 0x00000010,   // This package is from "compiled in" classes.
+	PKG_ForDiffing					= 0x00000020,	// This package was loaded just for the purposes of diffing
+	PKG_EditorOnly					= 0x00000040, // This is editor-only package (for example: editor module script package)
+	PKG_Developer					= 0x00000080,	// Developer module
+//	PKG_Unused						= 0x00000100,
+//	PKG_Unused						= 0x00000200,
+//	PKG_Unused						= 0x00000400,
+//	PKG_Unused						= 0x00000800,
+//	PKG_Unused						= 0x00001000,
+//	PKG_Unused						= 0x00002000,
+//	PKG_Unused						= 0x00004000,
+	PKG_Need						= 0x00008000,	// Client needs to download this package.
+	PKG_Compiling					= 0x00010000,	// package is currently being compiled
+	PKG_ContainsMap					= 0x00020000,	// Set if the package contains a ULevel/ UWorld object
+	PKG_RequiresLocalizationGather		= 0x00040000,	// Set if the package contains any data to be gathered by localization
+	PKG_DisallowLazyLoading			= 0x00080000,	// Set if the archive serializing this package cannot use lazy loading
+	PKG_PlayInEditor				= 0x00100000,	// Set if the package was created for the purpose of PIE
+	PKG_ContainsScript				= 0x00200000,	// Package is allowed to contain UClass objects
+//	PKG_Unused						= 0x00400000,
+//	PKG_Unused						= 0x00800000,
+//	PKG_Unused						= 0x01000000,	
+	PKG_StoreCompressed				= 0x02000000,	// Package is being stored compressed, requires archive support for compression
+	PKG_StoreFullyCompressed		= 0x04000000,	// Package is serialized normally, and then fully compressed after (must be decompressed before LoadPackage is called)
+//	PKG_Unused						= 0x08000000,	
+//	PKG_Unused						= 0x10000000,	
+//	PKG_Unused						= 0x20000000,
+	PKG_ReloadingForCooker			= 0x40000000,   // this package is reloading in the cooker, try to avoid getting data we will never need. We won't save this package.
+	PKG_FilterEditorOnly			= 0x80000000,	// Package has editor-only data filtered
+};
+
+#define PKG_InMemoryOnly	(EPackageFlags)(PKG_CompiledIn | PKG_NewlyCreated) // Flag mask that indicates if this package is a package that exists in memory only.
+
+ENUM_CLASS_FLAGS(EPackageFlags);
+
+//
+// Internal enums.
+//
+enum EStaticConstructor				{EC_StaticConstructor};
+enum EInternal						{EC_InternalUseOnlyConstructor};
+enum ECppProperty					{EC_CppProperty};
+
+#if WITH_HOT_RELOAD_CTORS
+/** DO NOT USE. Helper class to invoke specialized hot-reload constructor. */
+class FVTableHelper
+{
+public:
+	/** DO NOT USE. This constructor is for internal usage only for hot-reload purposes. */
+	COREUOBJECT_API FVTableHelper()
+	{
+		EnsureRetrievingVTablePtrDuringCtor(TEXT("FVTableHelper()"));
+	}
+};
+#endif // WITH_HOT_RELOAD_CTORS
+
+/** Empty API definition.  Used as a placeholder parameter when no DLL export/import API is needed for a UObject class */
+#define NO_API
+
+
+/**
+ * Flags describing a class.
+ */
+enum EClassFlags
+{
+	/** @name Base flags */
+	//@{
+	CLASS_None				  = 0x00000000,
+	/** Class is abstract and can't be instantiated directly. */
+	CLASS_Abstract            = 0x00000001,
+	/** Save object configuration only to Default INIs, never to local INIs. Must be combined with CLASS_Config */
+	CLASS_DefaultConfig		  = 0x00000002,
+	/** Load object configuration at construction time. */
+	CLASS_Config			  = 0x00000004,
+	/** This object type can't be saved; null it out at save time. */
+	CLASS_Transient			  = 0x00000008,
+	/** Successfully parsed. */
+	CLASS_Parsed              = 0x00000010,
+	/** */
+	//CLASS_                  = 0x00000020,
+	/** All the properties on the class are shown in the advanced section (which is hidden by default) unless SimpleDisplay is specified on the property */
+	CLASS_AdvancedDisplay	  = 0x00000040,
+	/** Class is a native class - native interfaces will have CLASS_Native set, but not RF_MarkAsNative */
+	CLASS_Native			  = 0x00000080,
+	/** Don't export to C++ header. */
+	CLASS_NoExport            = 0x00000100,
+	/** Do not allow users to create in the editor. */
+	CLASS_NotPlaceable        = 0x00000200,
+	/** Handle object configuration on a per-object basis, rather than per-class. */
+	CLASS_PerObjectConfig     = 0x00000400,
+	
+	/** pointers to this class default to weak. */
+	CLASS_PointersDefaultToWeak = 0x00000800,
+	
+	/** Class can be constructed from editinline New button. */
+	CLASS_EditInlineNew		  = 0x00001000,
+	/** Display properties in the editor without using categories. */
+	CLASS_CollapseCategories  = 0x00002000,
+	/** Class is an interface **/
+	CLASS_Interface           = 0x00004000,
+	/**  Do not export a constructor for this class, assuming it is in the cpptext **/
+	CLASS_CustomConstructor   = 0x00008000,
+	/** all properties and functions in this class are const and should be exported as const */
+	CLASS_Const			      = 0x00010000,
+
+	/** pointers to this class default to weak. */
+	CLASS_PointersDefaultToAutoWeak = 0x00020000,
+	
+	/** Indicates that the class was created from blueprint source material */
+	CLASS_CompiledFromBlueprint  = 0x00040000,
+
+	/** Indicates that only the bare minimum bits of this class should be DLL exported/imported */
+	CLASS_MinimalAPI	      = 0x00080000,
+	
+	/** Indicates this class must be DLL exported/imported (along with all of it's members) */
+	CLASS_RequiredAPI	      = 0x00100000,
+
+	/** Indicates that references to this class default to instanced. Used to be subclasses of UComponent, but now can be any UObject */
+	CLASS_DefaultToInstanced  = 0x00200000,
+
+	/** Indicates that the parent token stream has been merged with ours. */
+	CLASS_TokenStreamAssembled  = 0x00400000,
+	/** Class has component properties. */
+	CLASS_HasInstancedReference= 0x00800000,
+	/** Don't show this class in the editor class browser or edit inline new menus. */
+	CLASS_Hidden			  = 0x01000000,
+	/** Don't save objects of this class when serializing */
+	CLASS_Deprecated		  = 0x02000000,
+	/** Class not shown in editor drop down for class selection */
+	CLASS_HideDropDown		  = 0x04000000,
+	/** Class settings are saved to <AppData>/..../Blah.ini (as opposed to CLASS_DefaultConfig) */
+	CLASS_GlobalUserConfig	  = 0x08000000,
+	/** Class was declared directly in C++ and has no boilerplate generated by UnrealHeaderTool */
+	CLASS_Intrinsic			  = 0x10000000,
+	/** Class has already been constructed (maybe in a previous DLL version before hot-reload). */
+	CLASS_Constructed		  = 0x20000000,
+	/** Indicates that object configuration will not check against ini base/defaults when serialized */
+	CLASS_ConfigDoNotCheckDefaults = 0x40000000,
+	/** Class has been consigned to oblivion as part of a blueprint recompile, and a newer version currently exists. */
+	CLASS_NewerVersionExists  = 0x80000000,
+
+	//@}
+
+
+	/** @name Flags to inherit from base class */
+	//@{
+	CLASS_Inherit           = CLASS_Transient | CLASS_DefaultConfig | CLASS_Config | CLASS_PerObjectConfig | CLASS_ConfigDoNotCheckDefaults | CLASS_NotPlaceable
+							| CLASS_PointersDefaultToAutoWeak | CLASS_PointersDefaultToWeak | CLASS_Const
+							| CLASS_HasInstancedReference | CLASS_Deprecated | CLASS_DefaultToInstanced | CLASS_GlobalUserConfig,
+
+	/** these flags will be cleared by the compiler when the class is parsed during script compilation */
+	CLASS_RecompilerClear   = CLASS_Inherit | CLASS_Abstract | CLASS_NoExport | CLASS_Native | CLASS_Intrinsic | CLASS_TokenStreamAssembled,
+
+	/** these flags will be cleared by the compiler when the class is parsed during script compilation */
+	CLASS_ShouldNeverBeLoaded   = CLASS_Native | CLASS_Intrinsic | CLASS_TokenStreamAssembled,
+
+	/** these flags will be inherited from the base class only for non-intrinsic classes */
+	CLASS_ScriptInherit		= CLASS_Inherit | CLASS_EditInlineNew | CLASS_CollapseCategories,
+	//@}
+
+	/** This is used as a mask for the flags put into generated code for "compiled in" classes. */
+	CLASS_SaveInCompiledInClasses = 
+		CLASS_Abstract | 
+		CLASS_DefaultConfig |
+		CLASS_GlobalUserConfig |
+		CLASS_Config |
+		CLASS_Transient |
+		CLASS_Native |
+		CLASS_NotPlaceable |
+		CLASS_PerObjectConfig |
+		CLASS_ConfigDoNotCheckDefaults |
+		CLASS_PointersDefaultToWeak |
+		CLASS_EditInlineNew |
+		CLASS_CollapseCategories |
+		CLASS_Interface |
+		CLASS_PointersDefaultToAutoWeak |
+		CLASS_DefaultToInstanced |
+		CLASS_HasInstancedReference |
+		CLASS_Hidden |
+		CLASS_Deprecated |
+		CLASS_HideDropDown |
+		CLASS_Intrinsic |
+		CLASS_AdvancedDisplay |
+		CLASS_Const |
+		CLASS_MinimalAPI |
+		CLASS_RequiredAPI,
+
+	CLASS_AllFlags			= 0xFFFFFFFF,
+};
+
+
+
+
+
+/**
+ * Flags used for quickly casting classes of certain types; all class cast flags are inherited
+ */
+typedef uint64 EClassCastFlags;
+
+#define CASTCLASS_None							DECLARE_UINT64(0x0000000000000000)
+#define CASTCLASS_UField						DECLARE_UINT64(0x0000000000000001)
+#define CASTCLASS_UInt8Property					DECLARE_UINT64(0x0000000000000002)
+#define CASTCLASS_UEnum							DECLARE_UINT64(0x0000000000000004)
+#define CASTCLASS_UStruct						DECLARE_UINT64(0x0000000000000008)
+#define CASTCLASS_UScriptStruct					DECLARE_UINT64(0x0000000000000010)
+#define CASTCLASS_UClass						DECLARE_UINT64(0x0000000000000020)
+#define CASTCLASS_UByteProperty					DECLARE_UINT64(0x0000000000000040)
+#define CASTCLASS_UIntProperty					DECLARE_UINT64(0x0000000000000080)
+#define CASTCLASS_UFloatProperty				DECLARE_UINT64(0x0000000000000100)
+#define CASTCLASS_UUInt64Property				DECLARE_UINT64(0x0000000000000200)
+#define CASTCLASS_UClassProperty				DECLARE_UINT64(0x0000000000000400)
+#define CASTCLASS_UUInt32Property				DECLARE_UINT64(0x0000000000000800)
+#define CASTCLASS_UInterfaceProperty			DECLARE_UINT64(0x0000000000001000)
+#define CASTCLASS_UNameProperty					DECLARE_UINT64(0x0000000000002000)
+#define CASTCLASS_UStrProperty					DECLARE_UINT64(0x0000000000004000)
+#define CASTCLASS_UProperty						DECLARE_UINT64(0x0000000000008000)
+#define CASTCLASS_UObjectProperty				DECLARE_UINT64(0x0000000000010000)
+#define CASTCLASS_UBoolProperty					DECLARE_UINT64(0x0000000000020000)
+#define CASTCLASS_UUInt16Property				DECLARE_UINT64(0x0000000000040000)
+#define CASTCLASS_UFunction						DECLARE_UINT64(0x0000000000080000)
+#define CASTCLASS_UStructProperty				DECLARE_UINT64(0x0000000000100000)
+#define CASTCLASS_UArrayProperty				DECLARE_UINT64(0x0000000000200000)
+#define CASTCLASS_UInt64Property				DECLARE_UINT64(0x0000000000400000)
+#define CASTCLASS_UDelegateProperty				DECLARE_UINT64(0x0000000000800000)
+#define CASTCLASS_UNumericProperty				DECLARE_UINT64(0x0000000001000000)
+#define CASTCLASS_UMulticastDelegateProperty	DECLARE_UINT64(0x0000000002000000)
+#define CASTCLASS_UObjectPropertyBase			DECLARE_UINT64(0x0000000004000000)
+#define CASTCLASS_UWeakObjectProperty			DECLARE_UINT64(0x0000000008000000)
+#define CASTCLASS_ULazyObjectProperty			DECLARE_UINT64(0x0000000010000000)
+#define CASTCLASS_UAssetObjectProperty			DECLARE_UINT64(0x0000000020000000)
+#define CASTCLASS_UTextProperty					DECLARE_UINT64(0x0000000040000000)
+#define CASTCLASS_UInt16Property				DECLARE_UINT64(0x0000000080000000)
+#define CASTCLASS_UDoubleProperty				DECLARE_UINT64(0x0000000100000000)
+#define CASTCLASS_UAssetClassProperty			DECLARE_UINT64(0x0000000200000000)
+#define CASTCLASS_UPackage						DECLARE_UINT64(0x0000000400000000)
+#define CASTCLASS_ULevel						DECLARE_UINT64(0x0000000800000000)
+#define CASTCLASS_AActor						DECLARE_UINT64(0x0000001000000000)
+#define CASTCLASS_APlayerController				DECLARE_UINT64(0x0000002000000000)
+#define CASTCLASS_APawn							DECLARE_UINT64(0x0000004000000000)
+#define CASTCLASS_USceneComponent				DECLARE_UINT64(0x0000008000000000)
+#define CASTCLASS_UPrimitiveComponent			DECLARE_UINT64(0x0000010000000000)
+#define CASTCLASS_USkinnedMeshComponent			DECLARE_UINT64(0x0000020000000000)
+#define CASTCLASS_USkeletalMeshComponent		DECLARE_UINT64(0x0000040000000000)
+#define CASTCLASS_UBlueprint					DECLARE_UINT64(0x0000080000000000)
+#define CASTCLASS_UDelegateFunction				DECLARE_UINT64(0x0000100000000000)
+#define CASTCLASS_UStaticMeshComponent			DECLARE_UINT64(0x0000200000000000)
+#define CASTCLASS_UMapProperty					DECLARE_UINT64(0x0000400000000000)
+#define CASTCLASS_USetProperty					DECLARE_UINT64(0x0000800000000000)
+
+#define CASTCLASS_AllFlags						DECLARE_UINT64(0xFFFFFFFFFFFFFFFF)
+
+//
+// Flags associated with each property in a class, overriding the
+// property's default behavior.
+// NOTE: When adding one here, please update ParsePropertyFlags
+//
+
+// For compilers that don't support 64 bit enums.
+#define CPF_Edit							DECLARE_UINT64(0x0000000000000001)		// Property is user-settable in the editor.
+#define CPF_ConstParm						DECLARE_UINT64(0x0000000000000002)		// This is a constant function parameter
+#define CPF_BlueprintVisible				DECLARE_UINT64(0x0000000000000004)		// This property can be read by blueprint code
+#define CPF_ExportObject					DECLARE_UINT64(0x0000000000000008)		// Object can be exported with actor.
+#define CPF_BlueprintReadOnly				DECLARE_UINT64(0x0000000000000010)		// This property cannot be modified by blueprint code
+#define CPF_Net								DECLARE_UINT64(0x0000000000000020)		// Property is relevant to network replication.
+#define CPF_EditFixedSize					DECLARE_UINT64(0x0000000000000040)		// Indicates that elements of an array can be modified, but its size cannot be changed.
+#define CPF_Parm							DECLARE_UINT64(0x0000000000000080)		// Function/When call parameter.
+#define CPF_OutParm							DECLARE_UINT64(0x0000000000000100)		// Value is copied out after function call.
+#define CPF_ZeroConstructor					DECLARE_UINT64(0x0000000000000200)		// memset is fine for construction
+#define CPF_ReturnParm						DECLARE_UINT64(0x0000000000000400)		// Return value.
+#define CPF_DisableEditOnTemplate			DECLARE_UINT64(0x0000000000000800)		// Disable editing of this property on an archetype/sub-blueprint
+//#define CPF_      						DECLARE_UINT64(0x0000000000001000)		// 
+#define CPF_Transient   					DECLARE_UINT64(0x0000000000002000)		// Property is transient: shouldn't be saved, zero-filled at load time.
+#define CPF_Config      					DECLARE_UINT64(0x0000000000004000)		// Property should be loaded/saved as permanent profile.
+//#define CPF_								DECLARE_UINT64(0x0000000000008000)		// 
+#define CPF_DisableEditOnInstance			DECLARE_UINT64(0x0000000000010000)		// Disable editing on an instance of this class
+#define CPF_EditConst   					DECLARE_UINT64(0x0000000000020000)		// Property is uneditable in the editor.
+#define CPF_GlobalConfig					DECLARE_UINT64(0x0000000000040000)		// Load config from base class, not subclass.
+#define CPF_InstancedReference				DECLARE_UINT64(0x0000000000080000)		// Property is a component references.
+//#define CPF_								DECLARE_UINT64(0x0000000000100000)
+#define CPF_DuplicateTransient				DECLARE_UINT64(0x0000000000200000)		// Property should always be reset to the default value during any type of duplication (copy/paste, binary duplication, etc.)
+#define CPF_SubobjectReference				DECLARE_UINT64(0x0000000000400000)		// Property contains subobject references (TSubobjectPtr)
+//#define CPF_    							DECLARE_UINT64(0x0000000000800000)		// 
+#define CPF_SaveGame						DECLARE_UINT64(0x0000000001000000)		// Property should be serialized for save games
+#define CPF_NoClear							DECLARE_UINT64(0x0000000002000000)		// Hide clear (and browse) button.
+//#define CPF_  							DECLARE_UINT64(0x0000000004000000)		//
+#define CPF_ReferenceParm					DECLARE_UINT64(0x0000000008000000)		// Value is passed by reference; CPF_OutParam and CPF_Param should also be set.
+#define CPF_BlueprintAssignable				DECLARE_UINT64(0x0000000010000000)		// MC Delegates only.  Property should be exposed for assigning in blueprint code
+#define CPF_Deprecated  					DECLARE_UINT64(0x0000000020000000)		// Property is deprecated.  Read it from an archive, but don't save it.
+#define CPF_IsPlainOldData					DECLARE_UINT64(0x0000000040000000)		// If this is set, then the property can be memcopied instead of CopyCompleteValue / CopySingleValue
+#define CPF_RepSkip							DECLARE_UINT64(0x0000000080000000)		// Not replicated. For non replicated properties in replicated structs 
+#define CPF_RepNotify						DECLARE_UINT64(0x0000000100000000)		// Notify actors when a property is replicated
+#define CPF_Interp							DECLARE_UINT64(0x0000000200000000)		// interpolatable property for use with matinee
+#define CPF_NonTransactional				DECLARE_UINT64(0x0000000400000000)		// Property isn't transacted
+#define CPF_EditorOnly						DECLARE_UINT64(0x0000000800000000)		// Property should only be loaded in the editor
+#define CPF_NoDestructor					DECLARE_UINT64(0x0000001000000000)		// No destructor
+//#define CPF_								DECLARE_UINT64(0x0000002000000000)		//
+#define CPF_AutoWeak						DECLARE_UINT64(0x0000004000000000)		// Only used for weak pointers, means the export type is autoweak
+#define CPF_ContainsInstancedReference		DECLARE_UINT64(0x0000008000000000)		// Property contains component references.
+#define CPF_AssetRegistrySearchable			DECLARE_UINT64(0x0000010000000000)		// asset instances will add properties with this flag to the asset registry automatically
+#define CPF_SimpleDisplay					DECLARE_UINT64(0x0000020000000000)		// The property is visible by default in the editor details view
+#define CPF_AdvancedDisplay					DECLARE_UINT64(0x0000040000000000)		// The property is advanced and not visible by default in the editor details view
+#define CPF_Protected						DECLARE_UINT64(0x0000080000000000)		// property is protected from the perspective of script
+#define CPF_BlueprintCallable				DECLARE_UINT64(0x0000100000000000)		// MC Delegates only.  Property should be exposed for calling in blueprint code
+#define CPF_BlueprintAuthorityOnly			DECLARE_UINT64(0x0000200000000000)		// MC Delegates only.  This delegate accepts (only in blueprint) only events with BlueprintAuthorityOnly.
+#define CPF_TextExportTransient				DECLARE_UINT64(0x0000400000000000)		// Property shouldn't be exported to text format (e.g. copy/paste)
+#define CPF_NonPIEDuplicateTransient		DECLARE_UINT64(0x0000800000000000)		// Property should only be copied in PIE
+#define CPF_ExposeOnSpawn					DECLARE_UINT64(0x0001000000000000)		// Property is exposed on spawn
+#define CPF_PersistentInstance				DECLARE_UINT64(0x0002000000000000)		// A object referenced by the property is duplicated like a component. (Each actor should have an own instance.)
+#define CPF_UObjectWrapper					DECLARE_UINT64(0x0004000000000000)		// Property was parsed as a wrapper class like TSubobjectOf<T>, FScriptInterface etc., rather than a USomething*
+#define CPF_HasGetValueTypeHash				DECLARE_UINT64(0x0008000000000000)		// This property can generate a meaningful hash value.
+#define CPF_NativeAccessSpecifierPublic		DECLARE_UINT64(0x0010000000000000)		// Public native access specifier
+#define CPF_NativeAccessSpecifierProtected	DECLARE_UINT64(0x0020000000000000)		// Protected native access specifier
+#define CPF_NativeAccessSpecifierPrivate	DECLARE_UINT64(0x0040000000000000)		// Private native access specifier
+#define CPF_SkipSerialization				DECLARE_UINT64(0x0080000000000000)		// Property shouldn't be serialized, can still be exported to text
+
+#define CPF_NonPIETransient \
+	EMIT_DEPRECATED_WARNING_MESSAGE("CPF_NonPIETransient is deprecated. Please use CPF_NonPIEDuplicateTransient instead.") \
+	CPF_NonPIEDuplicateTransient
+
+/** @name Combinations flags */
+//@{
+#define CPF_NativeAccessSpecifiers	(CPF_NativeAccessSpecifierPublic | CPF_NativeAccessSpecifierProtected | CPF_NativeAccessSpecifierPrivate)
+
+#define CPF_ParmFlags				(CPF_Parm | CPF_OutParm | CPF_ReturnParm | CPF_ReferenceParm | CPF_ConstParm)
+#define CPF_PropagateToArrayInner	(CPF_ExportObject | CPF_PersistentInstance | CPF_InstancedReference | CPF_ContainsInstancedReference | CPF_Config | CPF_EditConst | CPF_Deprecated | CPF_EditorOnly | CPF_AutoWeak | CPF_UObjectWrapper )
+#define CPF_PropagateToMapValue		(CPF_ExportObject | CPF_PersistentInstance | CPF_InstancedReference | CPF_ContainsInstancedReference | CPF_Config | CPF_EditConst | CPF_Deprecated | CPF_EditorOnly | CPF_AutoWeak | CPF_UObjectWrapper | CPF_Edit )
+#define CPF_PropagateToMapKey		(CPF_ExportObject | CPF_PersistentInstance | CPF_InstancedReference | CPF_ContainsInstancedReference | CPF_Config | CPF_EditConst | CPF_Deprecated | CPF_EditorOnly | CPF_AutoWeak | CPF_UObjectWrapper | CPF_Edit )
+#define CPF_PropagateToSetElement	(CPF_ExportObject | CPF_PersistentInstance | CPF_InstancedReference | CPF_ContainsInstancedReference | CPF_Config | CPF_EditConst | CPF_Deprecated | CPF_EditorOnly | CPF_AutoWeak | CPF_UObjectWrapper | CPF_Edit )
+
+/** the flags that should never be set on interface properties */
+#define CPF_InterfaceClearMask		(CPF_ExportObject|CPF_InstancedReference|CPF_ContainsInstancedReference)
+
+/** all the properties that can be stripped for final release console builds */
+#define CPF_DevelopmentAssets		(CPF_EditorOnly)
+
+/** all the properties that should never be loaded or saved */
+#define CPF_ComputedFlags			(CPF_IsPlainOldData | CPF_NoDestructor | CPF_ZeroConstructor)
+
+#define CPF_AllFlags				DECLARE_UINT64(0xFFFFFFFFFFFFFFFF)
+
+//@}
+
+/** 
+ * Flags describing an object instance
+ */
+enum EObjectFlags
+{
+	// Do not add new flags unless they truly belong here. There are alternatives.
+	// if you change any the bit of any of the RF_Load flags, then you will need legacy serialization
+	RF_NoFlags						= 0x00000000,	///< No flags, used to avoid a cast
+
+	// This first group of flags mostly has to do with what kind of object it is. Other than transient, these are the persistent object flags.
+	// The garbage collector also tends to look at these.
+	RF_Public					=0x00000001,	///< Object is visible outside its package.
+	RF_Standalone				=0x00000002,	///< Keep object around for editing even if unreferenced.
+	RF_MarkAsNative					=0x00000004,	///< Object (UField) will be marked as native on construction (DO NOT USE THIS FLAG in HasAnyFlags() etc)
+	RF_Transactional			=0x00000008,	///< Object is transactional.
+	RF_ClassDefaultObject		=0x00000010,	///< This object is its class's default object
+	RF_ArchetypeObject			=0x00000020,	///< This object is a template for another object - treat like a class default object
+	RF_Transient				=0x00000040,	///< Don't save object.
+
+	// This group of flags is primarily concerned with garbage collection.
+	RF_MarkAsRootSet					=0x00000080,	///< Object will be marked as root set on construction and not be garbage collected, even if unreferenced (DO NOT USE THIS FLAG in HasAnyFlags() etc)
+	//RF_Unused				=0x00000100,	///
+	RF_TagGarbageTemp			=0x00000200,	///< This is a temp user flag for various utilities that need to use the garbage collector. The garbage collector itself does not interpret it.
+
+	// The group of flags tracks the stages of the lifetime of a uobject
+	RF_NeedLoad					=0x00000400,	///< During load, indicates object needs loading.
+	RF_KeepForCooker			=0x00000800,	///< Keep this object during garbage collection because it's still being used by the cooker
+	RF_NeedPostLoad				=0x00001000,	///< Object needs to be postloaded.
+	RF_NeedPostLoadSubobjects	=0x00002000,	///< During load, indicates that the object still needs to instance subobjects and fixup serialized component references
+	//RF_Unused				=0x00004000,	///
+	RF_BeginDestroyed			=0x00008000,	///< BeginDestroy has been called on the object.
+	RF_FinishDestroyed			=0x00010000,	///< FinishDestroy has been called on the object.
+
+	// Misc. Flags
+	RF_BeingRegenerated			=0x00020000,	///< Flagged on UObjects that are used to create UClasses (e.g. Blueprints) while they are regenerating their UClass on load (See FLinkerLoad::CreateExport())
+	RF_DefaultSubObject			=0x00040000,	///< Flagged on subobjects that are defaults
+	RF_WasLoaded				=0x00080000,	///< Flagged on UObjects that were loaded
+	RF_TextExportTransient		=0x00100000,	///< Do not export object to text form (e.g. copy/paste). Generally used for sub-objects that can be regenerated from data in their parent object.
+	RF_LoadCompleted			=0x00200000,	///< Object has been completely serialized by linkerload at least once. DO NOT USE THIS FLAG, It should be replaced with RF_WasLoaded.
+	RF_InheritableComponentTemplate = 0x00400000, ///< Archetype of the object can be in its super class
+	//RF_Unused = 0x00800000, ///
+	RF_StrongRefOnFrame			= 0x01000000,	///< References to this object from persistent function frame are handled as strong ones.
+	//RF_Unused		= 0x02000000,  ///
+	RF_Dynamic = 0x04000000, // Field Only. Dynamic field - doesn't get constructed during static initialization, can be constructed multiple times
+};
+
+	// Special all and none masks
+#define RF_AllFlags				(EObjectFlags)0x07ffffff	///< All flags, used mainly for error checking
+
+	// Predefined groups of the above
+#define RF_Load						((EObjectFlags)(RF_Public | RF_Standalone | RF_Transactional | RF_ClassDefaultObject | RF_ArchetypeObject | RF_DefaultSubObject | RF_TextExportTransient | RF_InheritableComponentTemplate)) // Flags to load from Unrealfiles.
+#define RF_PropagateToSubObjects	((EObjectFlags)(RF_Public | RF_ArchetypeObject | RF_Transactional | RF_Transient))		// Sub-objects will inherit these flags from their SuperObject.
+
+FORCEINLINE EObjectFlags operator|(EObjectFlags Arg1,EObjectFlags Arg2)
+{
+	return EObjectFlags(uint32(Arg1) | uint32(Arg2));
+}
+
+FORCEINLINE EObjectFlags operator&(EObjectFlags Arg1,EObjectFlags Arg2)
+{
+	return EObjectFlags(uint32(Arg1) & uint32(Arg2));
+}
+
+FORCEINLINE EObjectFlags operator~(EObjectFlags Arg)
+{
+	return EObjectFlags(RF_AllFlags & ~uint32(Arg));
+}
+
+FORCEINLINE void operator&=(EObjectFlags& Dest,EObjectFlags Arg)
+{
+	Dest = EObjectFlags(Dest & Arg);
+}
+FORCEINLINE void operator|=(EObjectFlags& Dest,EObjectFlags Arg)
+{
+	Dest = EObjectFlags(Dest | Arg);
+}
+
+//@}
+
+/** Objects flags for internal use (GC, low level UObject code) */
+enum class EInternalObjectFlags : int32
+{
+	None = 0,
+	// All the other bits are reserved, DO NOT ADD NEW FLAGS HERE!
+	ReachableInCluster = 1 << 23, /// External reference to object in cluster exists
+	ClusterRoot = 1 << 24, ///< Root of a cluster
+	Native = 1 << 25, ///< Native (UClass only).
+	Async = 1 << 26, ///< Object exists only on a different thread than the game thread.
+	AsyncLoading = 1 << 27, ///< Object is being asynchronously loaded.
+	Unreachable = 1 << 28, ///< Object is not reachable on the object graph.
+	PendingKill = 1 << 29, ///< Objects that are pending destruction (invalid for gameplay but valid objects)
+	RootSet = 1 << 30, ///< Object will not be garbage collected, even if unreferenced.
+	NoStrongReference = 1 << 31, ///< The object is not referenced by any strong reference. The flag is used by GC.
+
+	GarbageCollectionKeepFlags = Native | Async | AsyncLoading,
+	// Make sure this is up to date!
+	AllFlags = ReachableInCluster | ClusterRoot | Native | Async | AsyncLoading | Unreachable | PendingKill | RootSet | NoStrongReference
+};
+ENUM_CLASS_FLAGS(EInternalObjectFlags);
+
+/*----------------------------------------------------------------------------
+	Core types.
+----------------------------------------------------------------------------*/
+
+// forward declarations
+class UObject;
+class UProperty;
+class FObjectInitializer; 
+
+struct COREUOBJECT_API FReferencerInformation 
+{
+	/** the object that is referencing the target */
+	UObject*				Referencer;
+
+	/** the total number of references from Referencer to the target */
+	int32						TotalReferences;
+
+	/** the array of UProperties in Referencer which hold references to target */
+	TArray<const UProperty*>		ReferencingProperties;
+
+	FReferencerInformation( UObject* inReferencer );
+	FReferencerInformation( UObject* inReferencer, int32 InReferences, const TArray<const UProperty*>& InProperties );
+};
+
+struct COREUOBJECT_API FReferencerInformationList
+{
+	TArray<FReferencerInformation>		InternalReferences;
+	TArray<FReferencerInformation>		ExternalReferences;
+
+	FReferencerInformationList();
+	FReferencerInformationList( const TArray<FReferencerInformation>& InternalRefs, const TArray<FReferencerInformation>& ExternalRefs );
+};
+
+/*----------------------------------------------------------------------------
+	Core macros.
+----------------------------------------------------------------------------*/
+
+// Special canonical package for FindObject, ParseObject.
+#define ANY_PACKAGE ((UPackage*)-1)
+
+// Special prefix for default objects (the UObject in a UClass containing the default values, etc)
+#define DEFAULT_OBJECT_PREFIX TEXT("Default__")
+
+///////////////////////////////
+/// UObject definition macros
+///////////////////////////////
+
+// These macros wrap metadata parsed by the Unreal Header Tool, and are otherwise
+// ignored when code containing them is compiled by the C++ compiler
+#define UPROPERTY(...)
+#define UFUNCTION(...)
+#define USTRUCT(...)
+#define UMETA(...)
+#define UPARAM(...)
+#define UENUM(...)
+#define UDELEGATE(...)
+
+// This pair of macros is used to help implement GENERATED_BODY() and GENERATED_USTRUCT_BODY()
+#define BODY_MACRO_COMBINE_INNER(A,B,C,D) A##B##C##D
+#define BODY_MACRO_COMBINE(A,B,C,D) BODY_MACRO_COMBINE_INNER(A,B,C,D)
+
+#define GENERATED_BODY_LEGACY(...) BODY_MACRO_COMBINE(CURRENT_FILE_ID,_,__LINE__,_GENERATED_BODY_LEGACY)
+#define GENERATED_BODY(...) BODY_MACRO_COMBINE(CURRENT_FILE_ID,_,__LINE__,_GENERATED_BODY)
+
+#define GENERATED_USTRUCT_BODY(...) GENERATED_BODY()
+#define GENERATED_UCLASS_BODY(...) GENERATED_BODY_LEGACY()
+#define GENERATED_UINTERFACE_BODY(...) GENERATED_BODY_LEGACY()
+#define GENERATED_IINTERFACE_BODY(...) GENERATED_BODY_LEGACY()
+
+#if UE_BUILD_DOCS
+#define UCLASS(...)
+#else
+#define UCLASS(...) BODY_MACRO_COMBINE(CURRENT_FILE_ID,_,__LINE__,_PROLOG)
+#endif
+
+#define UINTERFACE(...) UCLASS()
+
+// This macro is used to declare a thunk function in autogenerated boilerplate code
+#define DECLARE_FUNCTION(func) void func( FFrame& Stack, RESULT_DECL )
+
+// These are used for syntax highlighting and to allow autocomplete hints
+
+namespace UC
+{
+	// valid keywords for the UCLASS macro
+	enum 
+	{
+		/// This keyword is used to set the actor group that the class is show in, in the editor.
+		classGroup,
+
+		/// This keyword specifies that this class cannot have its code generated until the given class has been processed.
+		DependsOn,
+
+		/// Declares that instances of this class should always have an outer of the specified class.  This is inherited by subclasses unless overridden.
+		Within, /* =OuterClassName */
+
+		/// Exposes this class as a type that can be used for variables in blueprints
+		BlueprintType,
+
+		/// Prevents this class from being used for variables in blueprints
+		NotBlueprintType,
+
+		/// Exposes this class as an acceptable base class for creating blueprints. The default is NotBlueprintable, unless inherited otherwise. This is inherited by subclasses.
+		Blueprintable,
+
+		/// Specifies that this class is *NOT* an acceptable base class for creating blueprints. The default is NotBlueprintable, unless inherited otherwise. This is inherited by subclasses.
+		NotBlueprintable,
+
+		/// This keyword indicates that the class should be accessible outside of it's module, but does not need all methods exported.
+		/// It exports only the autogenerated methods required for dynamic_cast<>, etc... to work.
+		MinimalAPI,
+
+		/// Prevents automatic generation of the constructor declaration.
+		customConstructor,
+
+		/// Class was declared directly in C++ and has no boilerplate generated by UnrealHeaderTool.
+		/// DO NOT USE THIS FLAG ON NEW CLASSES.
+		Intrinsic,
+
+		/// No autogenerated code will be created for this class; the header is only provided to parse metadata from.
+		/// DO NOT USE THIS FLAG ON NEW CLASSES.
+		noexport,
+
+		/// Allow users to create and place this class in the editor.  This flag is inherited by subclasses.
+		placeable,
+
+		/// This class cannot be placed in the editor (it cancels out an inherited placeable flag).
+		notplaceable,
+
+		/// All instances of this class are considered "instanced". Instanced classes (components) are duplicated upon construction. This flag is inherited by subclasses. 
+		DefaultToInstanced,
+
+		/// All properties and functions in this class are const and should be exported as const.  This flag is inherited by subclasses.
+		Const,
+
+		/// Class is abstract and can't be instantiated directly.
+		Abstract,
+
+		/// This class is deprecated and objects of this class won't be saved when serializing.  This flag is inherited by subclasses.
+		deprecated,
+
+		/// This class can't be saved; null it out at save time.  This flag is inherited by subclasses.
+		Transient,
+
+		/// This class should be saved normally (it cancels out an inherited transient flag).
+		nonTransient,
+
+		/// Load object configuration at construction time.  These flags are inherited by subclasses.
+		/// Class containing config properties. Usage config=ConfigName or config=inherit (inherits config name from base class).
+		config,
+		/// Handle object configuration on a per-object basis, rather than per-class. 
+		perObjectConfig,
+		/// Determine whether on serialize to configs a check should be done on the base/defaults ini's
+		configdonotcheckdefaults,
+
+		/// Save object config only to Default INIs, never to local INIs.
+		defaultconfig,
+
+		/// These affect the behavior of the property editor.
+		/// Class can be constructed from editinline New button.
+		editinlinenew,
+		/// Class can't be constructed from editinline New button.
+		noteditinlinenew,
+		/// Class not shown in editor drop down for class selection.
+		hidedropdown,
+
+		/// Shows the specified categories in a property viewer. Usage: showCategories=CategoryName or showCategories=(category0, category1, ...)
+		showCategories,
+		/// Hides the specified categories in a property viewer. Usage: hideCategories=CategoryName or hideCategories=(category0, category1, ...)
+		hideCategories,
+		/// Indicates that this class is a wrapper class for a component with little intrinsic functionality (this causes things like hideCategories and showCategories to be ignored if the class is subclassed in a Blueprint)
+		ComponentWrapperClass,
+		/// Shows the specified function in a property viewer. Usage: showFunctions=FunctionName or showFunctions=(category0, category1, ...)
+		showFunctions,
+		/// Hides the specified function in a property viewer. Usage: hideFunctions=FunctionName or hideFunctions=(category0, category1, ...)
+		hideFunctions,
+		/// Specifies which categories should be automatically expanded in a property viewer.
+		autoExpandCategories,
+		/// Specifies which categories should be automatically collapsed in a property viewer.
+		autoCollapseCategories,
+		/// Clears the list of auto collapse categories.
+		dontAutoCollapseCategories,
+		/// Display properties in the editor without using categories.
+		collapseCategories,
+		/// Display properties in the editor using categories (default behaviour).
+		dontCollapseCategories,
+
+		/// All the properties of the class are hidden in the main display by default, and are only shown in the advanced details section.
+		AdvancedClassDisplay,
+
+		/// A root convert limits a sub-class to only be able to convert to child classes of the first root class going up the hierarchy.
+		ConversionRoot,
+
+		/// Marks this class as 'experimental' (a totally unsupported and undocumented prototype)
+		Experimental,
+
+		// Marks this class as an 'early access' preview (while not considered production-ready, it's a step beyond 'experimental' and is being provided as a preview of things to come)
+		EarlyAccessPreview,
+	};
+}
+
+namespace UI
+{
+	// valid keywords for the UINTERFACE macro, see the UCLASS versions, above
+	enum 
+	{
+		/// This keyword specifies that this class cannot have its code generated until the given class has been processed.
+		DependsOn,
+
+		/// This keyword indicates that the interface should be accessible outside of it's module, but does not need all methods exported.
+		/// It exports only the autogenerated methods required for dynamic_cast<>, etc... to work.
+		MinimalAPI,
+
+		/// Exposes this interface as an acceptable base class for creating blueprints.  The default is NotBlueprintable, unless inherited otherwise. This is inherited by subclasses.
+		Blueprintable,
+
+		/// Specifies that this interface is *NOT* an acceptable base class for creating blueprints.  The default is NotBlueprintable, unless inherited otherwise. This is inherited by subclasses.
+		NotBlueprintable,
+
+		/// Sets IsConversionRoot metadata flag for this interface.
+		ConversionRoot,
+	};
+}
+
+namespace UF
+{
+	// valid keywords for the UFUNCTION and UDELEGATE macros
+	enum 
+	{
+		/// This function is designed to be overridden by a blueprint.  Do not provide a body for this function;
+		/// the autogenerated code will include a thunk that calls ProcessEvent to execute the overridden body.
+		BlueprintImplementableEvent,
+
+		/// This function is designed to be overridden by a blueprint, but also has a native implementation.
+		/// Provide a body named [FunctionName]_Implementation instead of [FunctionName]; the autogenerated
+		/// code will include a thunk that calls the implementation method when necessary.
+		BlueprintNativeEvent,
+
+		/// This function is sealed and cannot be overridden in subclasses.
+		/// It is only a valid keyword for events; declare other methods as static or final to indicate that they are sealed.
+		SealedEvent,
+
+		/// This function is executable from the command line.
+		Exec,
+
+		/// This function is replicated, and executed on servers.  Provide a body named [FunctionName]_Implementation instead of [FunctionName];
+		/// the autogenerated code will include a thunk that calls the implementation method when necessary.
+		Server,
+
+		/// This function is replicated, and executed on clients.  Provide a body named [FunctionName]_Implementation instead of [FunctionName];
+		/// the autogenerated code will include a thunk that calls the implementation method when necessary.
+		Client,
+
+		/// This function is both executed locally on the server and replicated to all clients, regardless of the Actor's NetOwner
+		NetMulticast,
+
+		/// Replication of calls to this function should be done on a reliable channel.
+		/// Only valid when used in conjunction with Client or Server
+		Reliable,
+
+		/// Replication of calls to this function can be done on an unreliable channel.
+		/// Only valid when used in conjunction with Client or Server
+		Unreliable,
+
+		/// This function fulfills a contract of producing no side effects, and additionally implies BlueprintCallable.
+		BlueprintPure,
+
+		/// This function can be called from blueprint code and should be exposed to the user of blueprint editing tools.
+		BlueprintCallable,
+
+		/// This function will not execute from blueprint code if running on something without network authority
+		BlueprintAuthorityOnly,
+
+		/// This function is cosmetic and will not run on dedicated servers
+		BlueprintCosmetic,
+
+		/// The UnrealHeaderTool code generator will not produce a execFoo thunk for this function; it is up to the user to provide one.
+		CustomThunk,
+
+		/// Specifies the category of the function when displayed in blueprint editing tools.
+		/// Usage: Category=CategoryName or Category="MajorCategory,SubCategory"
+		Category,
+
+		/// This function must supply a _Validate implementation
+		WithValidation,
+
+		/// This function is RPC service request
+		ServiceRequest,
+
+		/// This function is RPC service response
+		ServiceResponse
+	};
+}
+
+namespace UP
+{
+	// valid keywords for the UPROPERTY macro
+	enum 
+	{
+		/// This property is const and should be exported as const.
+		Const,
+
+		/// Property should be loaded/saved to ini file as permanent profile.
+		Config,
+
+		/// Same as above but load config from base class, not subclass.
+		GlobalConfig,
+
+		/// Property should be loaded as localizable text. Implies ReadOnly.
+		Localized,
+
+		/// Property is transient: shouldn't be saved, zero-filled at load time.
+		Transient,
+
+		/// Property should always be reset to the default value during any type of duplication (copy/paste, binary duplication, etc.)
+		DuplicateTransient,
+
+		/// Property should always be reset to the default value unless it's being duplicated for a PIE session - deprecated, use NonPIEDuplicateTransient instead
+		NonPIETransient,
+
+		/// Property should always be reset to the default value unless it's being duplicated for a PIE session
+		NonPIEDuplicateTransient,
+
+		/// Value is copied out after function call. Only valid on function param declaration.
+		Ref,
+
+		/// Object property can be exported with it's owner.
+		Export,
+
+		/// Hide clear (and browse) button in the editor.
+		NoClear,
+
+		/// Indicates that elements of an array can be modified, but its size cannot be changed.
+		EditFixedSize,
+
+		/// Property is relevant to network replication.
+		Replicated,
+
+		/// Property is relevant to network replication. Notify actors when a property is replicated (usage: ReplicatedUsing=FunctionName).
+		ReplicatedUsing,
+
+		/// Skip replication (only for struct members and parameters in service request functions).
+		NotReplicated,
+
+		/// Interpolatable property for use with matinee. Always user-settable in the editor.
+		Interp,
+
+		/// Property isn't transacted.
+		NonTransactional,
+
+		/// Property is a component reference. Implies EditInline and Export.
+		Instanced,
+
+		/// MC Delegates only.  Property should be exposed for assigning in blueprints.
+		BlueprintAssignable,
+
+		/// Specifies the category of the property. Usage: Category=CategoryName.
+		Category,
+
+		/// Properties appear visible by default in a details panel
+		SimpleDisplay,
+
+		/// Properties are in the advanced dropdown in a details panel
+		AdvancedDisplay,
+
+		/// Indicates that this property can be edited by property windows in the editor
+		EditAnywhere,
+
+		/// Indicates that this property can be edited by property windows, but only on instances, not on archetypes
+		EditInstanceOnly,
+
+		/// Indicates that this property can be edited by property windows, but only on archetypes
+		EditDefaultsOnly,
+
+		/// Indicates that this property is visible in property windows, but cannot be edited at all
+		VisibleAnywhere,
+		
+		/// Indicates that this property is only visible in property windows for instances, not for archetypes, and cannot be edited
+		VisibleInstanceOnly,
+
+		/// Indicates that this property is only visible in property windows for archetypes, and cannot be edited
+		VisibleDefaultsOnly,
+
+		/// This property can be read by blueprints, but not modified.
+		BlueprintReadOnly,
+
+		/// This property can be read or written from a blueprint.
+		BlueprintReadWrite,
+
+		/// The AssetRegistrySearchable keyword indicates that this property and it's value will be automatically added
+		/// to the asset registry for any asset class instances containing this as a member variable.  It is not legal
+		/// to use on struct properties or parameters.
+		AssetRegistrySearchable,
+
+		/// Property should be serialized for save game.
+		SaveGame,
+
+		/// MC Delegates only.  Property should be exposed for calling in blueprint code
+		BlueprintCallable,
+
+		/// MC Delegates only. This delegate accepts (only in blueprint) only events with BlueprintAuthorityOnly.
+		BlueprintAuthorityOnly,
+
+		/// Property shouldn't be exported to text format (e.g. copy/paste)
+		TextExportTransient,
+
+		/// Property shouldn't be serialized, can still be exported to text
+		SkipSerialization,
+	};
+}
+
+namespace US
+{
+	// valid keywords for the USTRUCT macro
+	enum 
+	{
+		/// No autogenerated code will be created for this class; the header is only provided to parse metadata from.
+		NoExport,
+
+		/// Indicates that this struct should always be serialized as a single unit
+		Atomic,
+
+		/// Immutable is only legal in Object.h and is being phased out, do not use on new structs!
+		Immutable,
+
+		/// Exposes this struct as a type that can be used for variables in blueprints
+		BlueprintType,
+	};
+}
+
+// Metadata specifiers
+namespace UM
+{
+	// Metadata usable in any UField (UCLASS(), USTRUCT(), UPROPERTY(), UFUNCTION(), etc...)
+	enum
+	{
+		/// Overrides the automatically generated tooltip from the class comment
+		ToolTip,
+
+		/// A short tooltip that is used in some contexts where the full tooltip might be overwhelming (such as the parent class picker dialog)
+		ShortTooltip,
+	};
+
+	// Metadata usable in UCLASS
+	enum
+	{
+		/// [ClassMetadata] Used for Actor Component classes. If present indicates that it can be spawned by a Blueprint.
+		BlueprintSpawnableComponent,
+
+		/// [ClassMetadata] Used for Actor and Component classes. If the native class cannot tick, Blueprint generated classes based this Actor or Component can have bCanEverTick flag overridden even if bCanBlueprintsTickByDefault is false.
+		ChildCanTick,
+
+		/// [ClassMetadata] Used for Actor and Component classes. If the native class cannot tick, Blueprint generated classes based this Actor or Component can never tick even if bCanBlueprintsTickByDefault is true.
+		ChildCannotTick,
+
+		/// [ClassMetadata] Used to make the first subclass of a class ignore all inherited showCategories and hideCategories commands
+		IgnoreCategoryKeywordsInSubclasses,
+
+		/// [ClassMetadata] For BehaviorTree nodes indicates that the class is deprecated and will display a warning when compiled.
+		DeprecatedNode,
+
+		/// [ClassMetadata] [FunctionMetadata] Used in conjunction with DeprecatedNode or DeprecatedFunction to customize the warning message displayed to the user.
+		DeprecationMessage,
+
+		/// [ClassMetadata] [PropertyMetadata] [FunctionMetadata] The name to display for this class, property, or function instead of auto-generating it from the name.
+		DisplayName,
+
+		/// [ClassMetadata] Specifies that this class is an acceptable base class for creating blueprints.
+		IsBlueprintBase,
+
+		/// [ClassMetadata] Comma delimited list of blueprint events that are not be allowed to be overridden in classes of this type
+		KismetHideOverrides,
+
+		/// [ClassMetadata] Specifies interfaces that are not compatible with the class.
+		ProhibitedInterfaces,
+
+		/// [ClassMetadata] Used by BlueprintFunctionLibrary classes to restrict the graphs the functions in the library can be used in to the classes specified.
+		RestrictedToClasses,
+
+		/// [ClassMetadata] Indicates that when placing blueprint nodes in graphs owned by this class that the hidden world context pin should be visible because the self context of the class cannot
+		///                 provide the world context and it must be wired in manually
+		ShowWorldContextPin,
+
+		//[ClassMetadata] Do not spawn an object of the class using Generic Create Object node in Blueprint. It makes sense only for a BluprintType class, that is neither Actor, nor ActorComponent.
+		DontUseGenericSpawnObject,
+
+		//[ClassMetadata] Expose a proxy object of this class in Async Task node.
+		ExposedAsyncProxy,
+	};
+
+	// Metadata usable in USTRUCT
+	enum
+	{
+		/// [StructMetadata] Indicates that the struct has a custom break node (and what the path to the BlueprintCallable UFunction is) that should be used instead of the default BreakStruct node.  
+		HasNativeBreak,
+
+		/// [StructMetadata] Indicates that the struct has a custom make node (and what the path to the BlueprintCallable UFunction is) that should be used instead of the default MakeStruct node.  
+		HasNativeMake,
+
+		/// [StructMetadata] Pins in Make and Break nodes are hidden by default.
+		HiddenByDefault,
+	};
+
+	// Metadata usable in UPROPERTY for customizing the behavior when displaying the property in a property panel
+	enum
+	{
+		/// [PropertyMetadata] Used for FStringClassReference properties.  Indicates whether abstract class types should be shown in the class picker.
+		AllowAbstract,
+
+		/// [PropertyMetadata] Used for FStringAssetReference properties.  Comma delimited list that indicates the class type(s) of assets to be displayed in the asset picker.
+		AllowedClasses,
+
+		/// [PropertyMetadata] Used for FVector properties.  It causes a ratio lock to be added when displaying this property in details panels.
+		AllowPreserveRatio,
+
+		/// [PropertyMetadata] Used for integer properties.  Clamps the valid values that can be entered in the UI to be between 0 and the length of the array specified.
+		ArrayClamp,
+
+		/// [PropertyMetadata] Used for float and integer properties.  Specifies the minimum value that may be entered for the property.
+		ClampMin,
+
+		/// [PropertyMetadata] Used for float and integer properties.  Specifies the maximum value that may be entered for the property.
+		ClampMax,
+
+		/// [ClassMetadata] [PropertyMetadata] [FunctionMetadata] The name to display for this class, property, or function instead of auto-generating it from the name.
+		// DisplayName, (Commented out so as to avoid duplicate name with version in the Class section, but still show in the property section)
+
+		/// [PropertyMetadata] Indicates that the property is an asset type and it should display the thumbnail of the selected asset.
+		DisplayThumbnail,	
+	
+		/// [PropertyMetadata] Specifies a boolean property that is used to indicate whether editing of this property is disabled.
+		EditCondition,
+		
+		/// [PropertyMetadata] Used for FStringAssetReference properties in conjunction with AllowedClasses. Indicates whether only the exact classes specified in AllowedClasses can be used or whether subclasses are valid.
+		ExactClass,
+
+		/// [PropertyMetadata] Specifies a list of categories whose functions should be exposed when building a function list in the Blueprint Editor.
+		ExposeFunctionCategories,
+
+		/// [PropertyMetadata] Specifies whether the property should be exposed on a Spawn Actor for the class type.
+		ExposeOnSpawn,
+
+		/// [PropertyMetadata] Deprecated.
+		FixedIncrement,
+
+		/// [PropertyMetadata] Used for FColor and FLinearColor properties. Indicates that the Alpha property should be hidden when displaying the property widget in the details.
+		HideAlphaChannel,
+
+		/// [PropertyMetadata] Signifies that the bool property is only displayed inline as an edit condition toggle in other properties, and should not be shown on its own row.
+		InlineEditConditionToggle,
+
+		/// [PropertyMetadata] Used for FStringClassReference properties.  Indicates whether only blueprint classes should be shown in the class picker.
+		IsBlueprintBaseOnly,
+
+		/// [PropertyMetadata] Used for Subclass properties. Indicates whether only placeable classes should be shown in the class picker.
+		OnlyPlaceable,
+
+		/// [PropertyMetadata] Used for Transform/Rotator properties (also works on arrays of them). Indicates that the property should be exposed in the viewport as a movable widget.
+		MakeEditWidget,
+
+		/// [PropertyMetadata] For properties in a structure indicates the default value of the property in a blueprint make structure node.
+		MakeStructureDefaultValue,
+
+		/// [PropertyMetadata] Used FStringClassReference properties. Indicates the parent class that the class picker will use when filtering which classes to display.
+		MetaClass,
+
+		/// [PropertyMetadata] Used for numeric properties. Stipulates that the value must be a multiple of the metadata value.
+		Multiple,
+
+		/// [PropertyMetadata] Used for FString and FText properties.  Indicates that the edit field should be multi-line, allowing entry of newlines.
+		MultiLine,
+
+		/// [PropertyMetadata] Used for FString and FText properties.  Indicates that the edit field is a secret field (e.g a password) and entered text will be replaced with dots. Do not use this as your only security measure.  The property data is still stored as plain text. 
+		PasswordField,
+
+		/// [PropertyMetadata] Used for array properties. Indicates that the duplicate icon should not be shown for entries of this array in the property panel.
+		NoElementDuplicate,
+
+		/// [PropertyMetadata] Used for integer and float properties. Indicates that the spin box element of the number editing widget should not be displayed.
+		NoSpinbox,
+
+		/// [PropertyMetadata] Used by FFilePath properties. Indicates the path filter to display in the file picker.
+		FilePathFilter,
+
+		/// [PropertyMetadta] Used by FDirectoryPath properties. Indicates that the directory dialog will output a relative path when setting the property.
+		RelativePath,
+
+		/// [PropertyMetadta] Used by FDirectoryPath properties. Indicates that the directory dialog will output a path relative to the game content directory when setting the property.
+		RelativeToGameContentDir,
+
+		/// [PropertyMetadta] Used by FDirectoryPath properties. Indicates that the path will be picked using the Slate-style directory picker inside the game Content dir.
+		ContentDir,
+
+		// [PropertyMetadata] Used by struct properties. Indicates that the inner properties will not be shown inside an expandable struct, but promoted up a level.
+		ShowOnlyInnerProperties,
+
+		// [PropertyMetadata] Used by numeric properties. Indicates how rapidly the value will grow when moving an unbounded slider.
+		SliderExponent,
+
+		/// [PropertyMetadata] Used for float and integer properties.  Specifies the lowest that the value slider should represent.
+		UIMin,
+
+		/// [PropertyMetadata] Used for float and integer properties.  Specifies the highest that the value slider should represent.
+		UIMax,
+
+		/// [PropertyMetadata] Property is serialized to config and we should be able to set it anywhere along the config hierarchy.
+		ConfigHierarchyEditable,
+
+		/// [PropertyMetadata] Property defaults are generated by the Blueprint compiler and will not be copied when CopyPropertiesForUnrelatedObjects is called post-compile.
+		BlueprintCompilerGeneratedDefaults,
+
+		/// [PropertyMetadata] Used by FDirectoryPath properties.  Converts the path to a long package name
+		LongPackageName,
+	};
+
+	// Metadata usable in UPROPERTY for customizing the behavior of Persona and UMG
+	// TODO: Move this to be contained in those modules specifically?
+	enum 
+	{
+		/// [PropertyMetadata] The property is not exposed as a data pin and is only be editable in the details panel. Applicable only to properties that will be displayed in Persona and UMG.
+		NeverAsPin, 
+
+		/// [PropertyMetadata] The property can be exposed as a data pin, but is hidden by default. Applicable only to properties that will be displayed in Persona and UMG.
+		PinHiddenByDefault, 
+
+		/// [PropertyMetadata] The property can be exposed as a data pin and is visible by default. Applicable only to properties that will be displayed in Persona and UMG.
+		PinShownByDefault, 
+
+		/// [PropertyMetadata] The property is always exposed as a data pin. Applicable only to properties that will be displayed in Persona and UMG.
+		AlwaysAsPin, 
+
+		/// [PropertyMetadata] Indicates that the property has custom code to display and should not generate a standard property widget int he details panel. Applicable only to properties that will be displayed in Persona.
+		CustomizeProperty,
+	};
+
+	// Metadata usable in UPROPERTY for customizing the behavior of Material Expressions
+	// TODO: Move this to be contained in that module?
+	enum
+	{
+		/// [PropertyMetadata] Used for float properties in MaterialExpression classes. If the specified FMaterialExpression pin is not connected, this value is used instead.
+		OverridingInputProperty,
+
+		/// [PropertyMetadata] Used for FMaterialExpression properties in MaterialExpression classes. If specified the pin need not be connected and the value of the property marked as OverridingInputProperty will be used instead.
+		RequiredInput,
+	};
+
+
+	// Metadata usable in UFUNCTION
+	enum
+	{
+		/// [FunctionMetadata] Used with a comma-separated list of parameter names that should show up as advanced pins (requiring UI expansion).
+		/// Alternatively you can set a number, which is the number of paramaters from the start that should *not* be marked as advanced (eg 'AdvancedDisplay="2"' will mark all but the first two advanced).
+		AdvancedDisplay,
+
+		/// [FunctionMetadata] Indicates that a BlueprintCallable function should use a Call Array Function node and that the parameters specified in the comma delimited list should be treated as wild card array properties.
+		ArrayParm,
+
+		/// [FunctionMetadata] Used when ArrayParm has been specified to indicate other function parameters that should be treated as wild card properties linked to the type of the array parameter.
+		ArrayTypeDependentParams,
+
+		/// [FunctionMetadata]
+		AutoCreateRefTerm,
+
+		/// [FunctionMetadata] This function is an internal implementation detail, used to implement another function or node.  It is never directly exposed in a graph.
+		BlueprintInternalUseOnly,
+
+		/// [FunctionMetadata] This function can only be called on 'this' in a blueprint. It cannot be called on another instance.
+		BlueprintProtected,
+
+		/// [FunctionMetadata] Used for BlueprintCallable functions that have a WorldContext pin to indicate that the function can be called even if the class does not implement the virtual function GetWorld().
+		CallableWithoutWorldContext,
+
+		/// [FunctionMetadata] Indicates that a BlueprintCallable function should use the Commutative Associative Binary node.
+		CommutativeAssociativeBinaryOperator,
+
+		/// [FunctionMetadata] Indicates that a BlueprintCallable function should display in the compact display mode and the name to use in that mode.
+		CompactNodeTitle,
+
+		/// [FunctionMetadata]
+		CustomStructureParam,
+
+		/// [FunctionMetadata] For BlueprintCallable functions indicates that the object property named's default value should be the self context of the node
+		DefaultToSelf,
+
+		/// [FunctionMetadata] This function is deprecated, any blueprint references to it cause a compilation warning.
+		DeprecatedFunction,
+
+		/// [ClassMetadata] [FunctionMetadata] Used in conjunction with DeprecatedNode or DeprecatedFunction to customize the warning message displayed to the user.
+		// DeprecationMessage, (Commented out so as to avoid duplicate name with version in the Class section, but still show in the function section)
+
+		/// [FunctionMetadata] For BlueprintCallable functions indicates that an input exec pin should be created for each entry in the enum specified.
+		ExpandEnumAsExecs,
+
+		/// [ClassMetadata] [PropertyMetadata] [FunctionMetadata] The name to display for this class, property, or function instead of auto-generating it from the name.
+		// DisplayName, (Commented out so as to avoid duplicate name with version in the Class section, but still show in the function section)
+
+		/// [FunctionMetadata] For BlueprintCallable functions indicates that the parameter pin should be hidden from the user's view.
+		HidePin,
+
+		/// [FunctionMetadata]
+		HideSpawnParms,
+
+		/// [FunctionMetadata] For BlueprintCallable functions provides additional keywords to be associated with the function for search purposes.
+		Keywords,
+
+		/// [FunctionMetadata] Indicates that a BlueprintCallable function is Latent
+		Latent,
+
+		/// [FunctionMetadata] For Latent BlueprintCallable functions indicates which parameter is the LatentInfo parameter
+		LatentInfo,
+
+		/// [FunctionMetadata] For BlueprintCallable functions indicates that the material override node should be used
+		MaterialParameterCollectionFunction,
+
+		/// [FunctionMetadata] For BlueprintCallable functions indicates that the function should be displayed the same as the implicit Break Struct nodes
+		NativeBreakFunc,
+
+		/// [FunctionMetadata] For BlueprintCallable functions indicates that the function should be displayed the same as the implicit Make Struct nodes
+		NativeMakeFunc,
+
+		/// [FunctionMetadata] Used by BlueprintCallable functions to indicate that this function is not to be allowed in the Construction Script.
+		UnsafeDuringActorConstruction,
+
+		/// [FunctionMetadta] Used by BlueprintCallable functions to indicate which parameter is used to determine the World that the operation is occurring within.
+		WorldContext,
+
+		/// [FunctionMetadta] Used only by static BlueprintPure functions from BlueprintLibrary. A cast node will be automatically added for the return type and the type of the first parameter of the function.
+		BlueprintAutocast,
+	};
+
+	// Metadata usable in UINTERFACE
+	enum
+	{
+		/// [InterfaceMetadata] This interface cannot be implemented by a blueprint (e.g., it has only non-exposed C++ member methods)
+		CannotImplementInterfaceInBlueprint,
+	};
+}
+
+#define RELAY_CONSTRUCTOR(TClass, TSuperClass) TClass(const FObjectInitializer& ObjectInitializer = FObjectInitializer::Get()) : TSuperClass(ObjectInitializer) {}
+
+#if !USE_COMPILED_IN_NATIVES
+#define COMPILED_IN_FLAGS(TStaticFlags) (TStaticFlags& ~(CLASS_Intrinsic))
+#define COMPILED_IN_INTRINSIC 0
+#else
+#define COMPILED_IN_FLAGS(TStaticFlags) (TStaticFlags | CLASS_Intrinsic)
+#define COMPILED_IN_INTRINSIC 1
+#endif
+
+#define DECLARE_SERIALIZER( TClass ) \
+	friend FArchive &operator<<( FArchive& Ar, TClass*& Res ) \
+	{ \
+		return Ar << (UObject*&)Res; \
+	} 
+
+/*-----------------------------------------------------------------------------
+Class declaration macros.
+-----------------------------------------------------------------------------*/
+
+#define DECLARE_CLASS( TClass, TSuperClass, TStaticFlags, TStaticCastFlags, TPackage, TRequiredAPI  ) \
+private: \
+    TClass& operator=(TClass&&);   \
+    TClass& operator=(const TClass&);   \
+	TRequiredAPI static UClass* GetPrivateStaticClass(const TCHAR* Package); \
+public: \
+	/** Bitwise union of #EClassFlags pertaining to this class.*/ \
+	enum {StaticClassFlags=TStaticFlags}; \
+	/** Typedef for the base class ({{ typedef-type }}) */ \
+	typedef TSuperClass Super;\
+	/** Typedef for {{ typedef-type }}. */ \
+	typedef TClass ThisClass;\
+	/** Returns a UClass object representing this class at runtime */ \
+	inline static UClass* StaticClass() \
+	{ \
+		return GetPrivateStaticClass(TPackage); \
+	} \
+	/** Returns the StaticClassFlags for this class */ \
+	inline static EClassCastFlags StaticClassCastFlags() \
+	{ \
+		return TStaticCastFlags; \
+	} \
+	DEPRECATED(4.7, "operator new has been deprecated for UObjects - please use NewObject or NewNamedObject instead") \
+	inline void* operator new( const size_t InSize, UObject* InOuter=(UObject*)GetTransientPackage(), FName InName=NAME_None, EObjectFlags InSetFlags=RF_NoFlags ) \
+	{ \
+		return StaticAllocateObject( StaticClass(), InOuter, InName, InSetFlags ); \
+	} \
+	/** For internal use only; use StaticConstructObject() to create new objects. */ \
+	inline void* operator new(const size_t InSize, EInternal InInternalOnly, UObject* InOuter = (UObject*)GetTransientPackage(), FName InName = NAME_None, EObjectFlags InSetFlags = RF_NoFlags) \
+	{ \
+		return StaticAllocateObject(StaticClass(), InOuter, InName, InSetFlags); \
+} \
+	/** For internal use only; use StaticConstructObject() to create new objects. */ \
+	inline void* operator new( const size_t InSize, EInternal* InMem ) \
+	{ \
+		return (void*)InMem; \
+	}
+
+#define DEFINE_FORBIDDEN_DEFAULT_CONSTRUCTOR_CALL(TClass) \
+	static_assert(false, "You have to define " #TClass "::" #TClass "() or " #TClass "::" #TClass "(const FObjectInitializer&). This is required by UObject system to work correctly.");
+
+#define DEFINE_DEFAULT_CONSTRUCTOR_CALL(TClass) \
+	static void __DefaultConstructor(const FObjectInitializer& X) { new((EInternal*)X.GetObj())TClass(); }
+
+#define DEFINE_DEFAULT_OBJECT_INITIALIZER_CONSTRUCTOR_CALL(TClass) \
+	static void __DefaultConstructor(const FObjectInitializer& X) { new((EInternal*)X.GetObj())TClass(X); }
+
+#if WITH_HOT_RELOAD_CTORS
+	#define DECLARE_VTABLE_PTR_HELPER_CTOR(API, TClass) \
+		/** DO NOT USE. This constructor is for internal usage only for hot-reload purposes. */ \
+		API TClass(FVTableHelper& Helper);
+
+	#define DEFINE_VTABLE_PTR_HELPER_CTOR(TClass) \
+		TClass::TClass(FVTableHelper& Helper) : Super(Helper) {};
+
+	#define DEFINE_VTABLE_PTR_HELPER_CTOR_CALLER_DUMMY() \
+		static UObject* __VTableCtorCaller(FVTableHelper& Helper) \
+		{ \
+			return nullptr; \
+		}
+
+	#if WITH_HOT_RELOAD
+		#define DEFINE_VTABLE_PTR_HELPER_CTOR_CALLER(TClass) \
+			static UObject* __VTableCtorCaller(FVTableHelper& Helper) \
+			{ \
+				return new (EC_InternalUseOnlyConstructor, (UObject*)GetTransientPackage(), NAME_None, RF_NeedLoad | RF_ClassDefaultObject | RF_TagGarbageTemp) TClass(Helper); \
+			}
+	#else // WITH_HOT_RELOAD
+		#define DEFINE_VTABLE_PTR_HELPER_CTOR_CALLER(TClass) \
+			DEFINE_VTABLE_PTR_HELPER_CTOR_CALLER_DUMMY()
+	#endif // WITH_HOT_RELOAD
+
+	#define DECLARE_CLASS_INTRINSIC_NO_CTOR(TClass,TSuperClass,TStaticFlags,TPackage) \
+		DECLARE_CLASS(TClass, TSuperClass, TStaticFlags | CLASS_Intrinsic, CASTCLASS_None, TPackage, NO_API) \
+		enum { IsIntrinsic = 1 }; \
+		static void StaticRegisterNatives##TClass() {} \
+		DECLARE_SERIALIZER(TClass) \
+		DEFINE_DEFAULT_OBJECT_INITIALIZER_CONSTRUCTOR_CALL(TClass) \
+		static UObject* __VTableCtorCaller(FVTableHelper& Helper) \
+		{ \
+			return new (EC_InternalUseOnlyConstructor, (UObject*)GetTransientPackage(), NAME_None, RF_NeedLoad | RF_ClassDefaultObject | RF_TagGarbageTemp) TClass(Helper); \
+		}
+
+	#define DECLARE_CLASS_INTRINSIC(TClass,TSuperClass,TStaticFlags,TPackage) \
+		DECLARE_CLASS(TClass,TSuperClass,TStaticFlags|CLASS_Intrinsic,CASTCLASS_None,TPackage,NO_API ) \
+		RELAY_CONSTRUCTOR(TClass, TSuperClass) \
+		/** DO NOT USE. This constructor is for internal usage only for hot-reload purposes. */ \
+		TClass(FVTableHelper& Helper) : Super(Helper) {}; \
+		enum {IsIntrinsic=1}; \
+		static void StaticRegisterNatives##TClass() {} \
+		DECLARE_SERIALIZER(TClass) \
+		DEFINE_DEFAULT_OBJECT_INITIALIZER_CONSTRUCTOR_CALL(TClass) \
+		static UObject* __VTableCtorCaller(FVTableHelper& Helper) \
+		{ \
+			return new (EC_InternalUseOnlyConstructor, (UObject*)GetTransientPackage(), NAME_None, RF_NeedLoad | RF_ClassDefaultObject | RF_TagGarbageTemp) TClass(Helper); \
+		}
+
+	#define DECLARE_CASTED_CLASS_INTRINSIC_WITH_API_NO_CTOR( TClass, TSuperClass, TStaticFlags, TPackage, TStaticCastFlags, TRequiredAPI ) \
+		DECLARE_CLASS(TClass, TSuperClass, TStaticFlags | CLASS_Intrinsic, TStaticCastFlags, TPackage, TRequiredAPI) \
+		enum { IsIntrinsic = 1 }; \
+		static void StaticRegisterNatives##TClass() {} \
+		DECLARE_SERIALIZER(TClass) \
+		DEFINE_DEFAULT_OBJECT_INITIALIZER_CONSTRUCTOR_CALL(TClass) \
+		static UObject* __VTableCtorCaller(FVTableHelper& Helper) \
+		{ \
+			return new (EC_InternalUseOnlyConstructor, (UObject*)GetTransientPackage(), NAME_None, RF_NeedLoad | RF_ClassDefaultObject | RF_TagGarbageTemp) TClass(Helper); \
+		}
+
+	#define DECLARE_CASTED_CLASS_INTRINSIC_WITH_API( TClass, TSuperClass, TStaticFlags, TPackage, TStaticCastFlags, TRequiredAPI ) \
+		DECLARE_CLASS(TClass,TSuperClass,TStaticFlags|CLASS_Intrinsic,TStaticCastFlags,TPackage,TRequiredAPI ) \
+		RELAY_CONSTRUCTOR(TClass, TSuperClass) \
+		/** DO NOT USE. This constructor is for internal usage only for hot-reload purposes. */ \
+		TClass(FVTableHelper& Helper) : Super(Helper) {}; \
+		enum {IsIntrinsic=1}; \
+		static void StaticRegisterNatives##TClass() {} \
+		DECLARE_SERIALIZER(TClass) \
+		DEFINE_DEFAULT_OBJECT_INITIALIZER_CONSTRUCTOR_CALL(TClass) \
+		static UObject* __VTableCtorCaller(FVTableHelper& Helper) \
+		{ \
+			return new (EC_InternalUseOnlyConstructor, (UObject*)GetTransientPackage(), NAME_None, RF_NeedLoad | RF_ClassDefaultObject | RF_TagGarbageTemp) TClass(Helper); \
+		}
+
+	#define DECLARE_CASTED_CLASS_INTRINSIC_NO_CTOR_NO_VTABLE_CTOR( TClass, TSuperClass, TStaticFlags, TPackage, TStaticCastFlags, TRequiredAPI ) \
+		DECLARE_CLASS(TClass,TSuperClass,TStaticFlags|CLASS_Intrinsic,TStaticCastFlags,TPackage, TRequiredAPI ) \
+		enum {IsIntrinsic=1}; \
+		static void StaticRegisterNatives##TClass() {} \
+		DECLARE_SERIALIZER(TClass) \
+		DEFINE_DEFAULT_OBJECT_INITIALIZER_CONSTRUCTOR_CALL(TClass) \
+		static UObject* __VTableCtorCaller(FVTableHelper& Helper) \
+		{ \
+			return new (EC_InternalUseOnlyConstructor, (UObject*)GetTransientPackage(), NAME_None, RF_NeedLoad | RF_ClassDefaultObject | RF_TagGarbageTemp) TClass(Helper); \
+		}
+
+	#define DECLARE_CASTED_CLASS_INTRINSIC_NO_CTOR( TClass, TSuperClass, TStaticFlags, TPackage, TStaticCastFlags, TRequiredAPI ) \
+		DECLARE_CASTED_CLASS_INTRINSIC_NO_CTOR_NO_VTABLE_CTOR( TClass, TSuperClass, TStaticFlags, TPackage, TStaticCastFlags, TRequiredAPI ) \
+		/** DO NOT USE. This constructor is for internal usage only for hot-reload purposes. */ \
+		TClass(FVTableHelper& Helper) : Super(Helper) {}; \
+
+#else // WITH_HOT_RELOAD_CTORS
+	#define DECLARE_VTABLE_PTR_HELPER_CTOR(API, TClass)
+	#define DEFINE_VTABLE_PTR_HELPER_CTOR(TClass)
+	#define DEFINE_VTABLE_PTR_HELPER_CTOR_CALLER(TClass)
+	#define DEFINE_VTABLE_PTR_HELPER_CTOR_CALLER_DUMMY()
+
+	#define DECLARE_CLASS_INTRINSIC_NO_CTOR(TClass,TSuperClass,TStaticFlags,TPackage) \
+		DECLARE_CLASS(TClass, TSuperClass, TStaticFlags | CLASS_Intrinsic, CASTCLASS_None, TPackage, NO_API) \
+		enum { IsIntrinsic = 1 }; \
+		static void StaticRegisterNatives##TClass() {} \
+		DECLARE_SERIALIZER(TClass) \
+		DEFINE_DEFAULT_OBJECT_INITIALIZER_CONSTRUCTOR_CALL(TClass) \
+
+	#define DECLARE_CLASS_INTRINSIC(TClass,TSuperClass,TStaticFlags,TPackage) \
+		DECLARE_CLASS(TClass,TSuperClass,TStaticFlags|CLASS_Intrinsic,CASTCLASS_None,TPackage,NO_API ) \
+		RELAY_CONSTRUCTOR(TClass, TSuperClass) \
+		enum {IsIntrinsic=1}; \
+		static void StaticRegisterNatives##TClass() {} \
+		DECLARE_SERIALIZER(TClass) \
+		DEFINE_DEFAULT_OBJECT_INITIALIZER_CONSTRUCTOR_CALL(TClass) \
+
+	#define DECLARE_CASTED_CLASS_INTRINSIC_WITH_API_NO_CTOR( TClass, TSuperClass, TStaticFlags, TPackage, TStaticCastFlags, TRequiredAPI ) \
+		DECLARE_CLASS(TClass, TSuperClass, TStaticFlags | CLASS_Intrinsic, TStaticCastFlags, TPackage, TRequiredAPI) \
+		enum { IsIntrinsic = 1 }; \
+		static void StaticRegisterNatives##TClass() {} \
+		DECLARE_SERIALIZER(TClass) \
+		DEFINE_DEFAULT_OBJECT_INITIALIZER_CONSTRUCTOR_CALL(TClass) \
+
+	#define DECLARE_CASTED_CLASS_INTRINSIC_WITH_API( TClass, TSuperClass, TStaticFlags, TPackage, TStaticCastFlags, TRequiredAPI ) \
+		DECLARE_CLASS(TClass,TSuperClass,TStaticFlags|CLASS_Intrinsic,TStaticCastFlags,TPackage,TRequiredAPI ) \
+		RELAY_CONSTRUCTOR(TClass, TSuperClass) \
+		enum {IsIntrinsic=1}; \
+		static void StaticRegisterNatives##TClass() {} \
+		DECLARE_SERIALIZER(TClass) \
+		DEFINE_DEFAULT_OBJECT_INITIALIZER_CONSTRUCTOR_CALL(TClass) \
+
+	#define DECLARE_CASTED_CLASS_INTRINSIC_NO_CTOR( TClass, TSuperClass, TStaticFlags, TPackage, TStaticCastFlags, TRequiredAPI ) \
+		DECLARE_CLASS(TClass,TSuperClass,TStaticFlags|CLASS_Intrinsic,TStaticCastFlags,TPackage, TRequiredAPI ) \
+		enum {IsIntrinsic=1}; \
+		static void StaticRegisterNatives##TClass() {} \
+		DECLARE_SERIALIZER(TClass) \
+		DEFINE_DEFAULT_OBJECT_INITIALIZER_CONSTRUCTOR_CALL(TClass) \
+
+#endif // WITH_HOT_RELOAD_CTORS
+
+
+#define DECLARE_CASTED_CLASS_INTRINSIC( TClass, TSuperClass, TStaticFlags, TPackage, TStaticCastFlags ) \
+	DECLARE_CASTED_CLASS_INTRINSIC_WITH_API( TClass, TSuperClass, TStaticFlags, TPackage, TStaticCastFlags, NO_API) \
+
+// Declare that objects of class being defined reside within objects of the specified class.
+#define DECLARE_WITHIN( TWithinClass ) \
+	/** The required type of this object's outer ({{ typedef-type }}) */ \
+	typedef class TWithinClass WithinClass;  \
+	TWithinClass* GetOuter##TWithinClass() const { return (TWithinClass*)GetOuter(); }
+
+// Register a class at startup time.
+#define IMPLEMENT_CLASS(TClass, TClassCrc) \
+	static TClassCompiledInDefer<TClass> AutoInitialize##TClass(TEXT(#TClass), sizeof(TClass), TClassCrc); \
+	UClass* TClass::GetPrivateStaticClass(const TCHAR* Package) \
+	{ \
+		static UClass* PrivateStaticClass = NULL; \
+		if (!PrivateStaticClass) \
+		{ \
+			/* this could be handled with templates, but we want it external to avoid code bloat */ \
+			GetPrivateStaticClassBody( \
+				Package, \
+				(TCHAR*)TEXT(#TClass) + 1 + ((StaticClassFlags & CLASS_Deprecated) ? 11 : 0), \
+				PrivateStaticClass, \
+				StaticRegisterNatives##TClass, \
+				sizeof(TClass), \
+				TClass::StaticClassFlags, \
+				TClass::StaticClassCastFlags(), \
+				TClass::StaticConfigName(), \
+				(UClass::ClassConstructorType)InternalConstructor<TClass>, \
+				(UClass::ClassVTableHelperCtorCallerType)InternalVTableHelperCtorCaller<TClass>, \
+				&TClass::AddReferencedObjects, \
+				&TClass::Super::StaticClass, \
+				&TClass::WithinClass::StaticClass \
+			); \
+		} \
+		return PrivateStaticClass; \
+	}
+
+// Used for intrinsics, this sets up the boiler plate, plus an initialization singleton, which can create properties and GC tokens
+#define IMPLEMENT_INTRINSIC_CLASS(TClass, TRequiredAPI, TSuperClass, TSuperRequiredAPI, InitCode) \
+	IMPLEMENT_CLASS(TClass, 0) \
+	TRequiredAPI UClass* Z_Construct_UClass_##TClass(); \
+	UClass* Z_Construct_UClass_##TClass() \
+	{ \
+		static UClass* Class = NULL; \
+		if (!Class) \
+		{ \
+			extern TSuperRequiredAPI UClass* Z_Construct_UClass_##TSuperClass(); \
+			UClass* SuperClass = Z_Construct_UClass_##TSuperClass(); \
+			Class = TClass::StaticClass(); \
+			UObjectForceRegistration(Class); \
+			check(Class->GetSuperClass() == SuperClass); \
+			InitCode \
+			Class->StaticLink(); \
+		} \
+		check(Class->GetClass()); \
+		return Class; \
+	} \
+	static FCompiledInDefer Z_CompiledInDefer_UClass_##TClass(Z_Construct_UClass_##TClass, &TClass::StaticClass, TEXT(#TClass), false);
+
+#define IMPLEMENT_CORE_INTRINSIC_CLASS(TClass, TSuperClass, InitCode) \
+	IMPLEMENT_INTRINSIC_CLASS(TClass, COREUOBJECT_API, TSuperClass, COREUOBJECT_API, InitCode)
+
+// Register a dynamic class (created at runtime, not startup). Explicit ClassName parameter because Blueprint types can have names that can't be used natively:
+#define IMPLEMENT_DYNAMIC_CLASS(TClass, ClassName, TClassCrc) \
+	UClass* TClass::GetPrivateStaticClass(const TCHAR* Package) \
+	{ \
+		UPackage* PrivateStaticClassOuter = FindOrConstructDynamicTypePackage(Package); \
+		UClass* PrivateStaticClass = Cast<UClass>(StaticFindObjectFast(UClass::StaticClass(), PrivateStaticClassOuter, (TCHAR*)ClassName)); \
+		if (!PrivateStaticClass) \
+		{ \
+			/* this could be handled with templates, but we want it external to avoid code bloat */ \
+			GetPrivateStaticClassBody( \
+			Package, \
+			(TCHAR*)ClassName, \
+			PrivateStaticClass, \
+			StaticRegisterNatives##TClass, \
+			sizeof(TClass), \
+			TClass::StaticClassFlags, \
+			TClass::StaticClassCastFlags(), \
+			TClass::StaticConfigName(), \
+			(UClass::ClassConstructorType)InternalConstructor<TClass>, \
+			(UClass::ClassVTableHelperCtorCallerType)InternalVTableHelperCtorCaller<TClass>, \
+			&TClass::AddReferencedObjects, \
+			&TClass::Super::StaticClass, \
+			&TClass::WithinClass::StaticClass, \
+			true \
+			); \
+		} \
+		return PrivateStaticClass; \
+	}
+
+/*-----------------------------------------------------------------------------
+	ERenameFlags.
+
+	Options to the UObject::Rename function, bit flag
+-----------------------------------------------------------------------------*/
+
+typedef uint32 ERenameFlags;
+
+#define REN_None				(0x0000)
+#define REN_ForceNoResetLoaders	(0x0001) // Rename won't call ResetLoaders - most likely you should never specify this option (unless you are renaming a UPackage possibly)
+#define REN_Test				(0x0002) // Just test to make sure that the rename is guaranteed to succeed if an non test rename immediately follows
+#define REN_DoNotDirty			(0x0004) // Indicates that the object (and new outer) should not be dirtied.
+#define REN_DontCreateRedirectors (0x0010) // Don't create an object redirector, even if the class is marked RF_Public
+#define REN_NonTransactional	(0x0020) // Don't call Modify() on the objects, so they won't be stored in the transaction buffer
+#define REN_ForceGlobalUnique	(0x0040) // Force unique names across all packages not just while the scope of the new outer
+#define REN_SkipGeneratedClasses (0x0080) // Prevent renaming of any child generated classes and CDO's in blueprints
+
+/*-----------------------------------------------------------------------------
+	Misc.
+-----------------------------------------------------------------------------*/
+
+typedef void (*FAsyncCompletionCallback)( UObject* LinkerRoot, void* CallbackUserData );
+
+/*-----------------------------------------------------------------------------
+	UObject.
+-----------------------------------------------------------------------------*/
+
+namespace UE4
+{
+	/**
+	 * Controls how calls to LoadConfig() should be propagated
+	 */
+	enum ELoadConfigPropagationFlags
+	{
+		LCPF_None					=	0x0,
+		/**
+		 * Indicates that the object should read ini values from each section up its class's hierarchy chain;
+		 * Useful when calling LoadConfig on an object after it has already been initialized against its archetype
+		 */
+		LCPF_ReadParentSections		=	0x1,
+
+		/**
+		 * Indicates that LoadConfig() should be also be called on the class default objects for all children of the original class.
+		 */
+		LCPF_PropagateToChildDefaultObjects		=	0x2,
+
+		/**
+		 * Indicates that LoadConfig() should be called on all instances of the original class.
+		 */
+		LCPF_PropagateToInstances	=	0x4,
+
+		/**
+		 * Indicates that this object is reloading its config data
+		 */
+		LCPF_ReloadingConfigData	=	0x8,
+
+		// Combination flags
+		LCPF_PersistentFlags		=	LCPF_ReloadingConfigData,
+	};
+}
+
+
+/**
+ * Helper class used to save and restore information across a StaticAllocateObject over the top of an existing object.
+ * Currently only used by UClass
+ */
+class FRestoreForUObjectOverwrite
+{
+public:
+	/** virtual destructor **/
+	virtual ~FRestoreForUObjectOverwrite() {}
+	/** Called once the new object has been reinitialized 
+	**/
+	virtual void Restore() const=0;
+};
+
+
diff --git a/Engine/Source/Runtime/CoreUObject/Public/UObject/PackageFileSummary.h b/Engine/Source/Runtime/CoreUObject/Public/UObject/PackageFileSummary.h
new file mode 100644
--- /dev/null
+++ b/Engine/Source/Runtime/CoreUObject/Public/UObject/PackageFileSummary.h
@@ -0,0 +1,386 @@
+// Copyright 1998-2016 Epic Games, Inc. All Rights Reserved.
+
+#pragma once
+
+#include "UObjectBaseUtility.h"
+
+/*----------------------------------------------------------------------------
+	Items stored in Unrealfiles.
+----------------------------------------------------------------------------*/
+
+/**
+ * Revision data for an Unreal package file.
+ */
+//@todo: shouldn't need ExportCount/NameCount with the linker free package map; if so, clean this stuff up
+struct FGenerationInfo
+{
+	/**
+	 * Number of exports in the linker's ExportMap for this generation.
+	 */
+	int32 ExportCount;
+
+	/**
+	 * Number of names in the linker's NameMap for this generation.
+	 */
+	int32 NameCount;
+
+	/** Constructor */
+	FGenerationInfo( int32 InExportCount, int32 InNameCount );
+
+	/** I/O function
+	 * we use a function instead of operator<< so we can pass in the package file summary for version tests, since archive version hasn't been set yet
+	 */
+	void Serialize(FArchive& Ar, const struct FPackageFileSummary& Summary);
+};
+
+#if WITH_ENGINE
+/**
+ * Information about the textures stored in the package.
+ */
+struct FTextureAllocations
+{
+	/**
+	 * Stores an export index for each texture of a certain type (size, format, etc).
+	 */
+	struct FTextureType
+	{
+		FTextureType();
+		COREUOBJECT_API FTextureType( int32 SizeX, int32 SizeY, int32 NumMips, uint32 Format, uint32 TexCreateFlags );
+
+		/** Width of the largest mip-level stored in the package. */
+		int32				SizeX;
+		/** Height of the largest mip-level stored in the package. */
+		int32				SizeY;
+		/** Number of mips */
+		int32				NumMips;
+		/** Texture format */
+		uint32			Format;
+		/** ETextureCreateFlags bit flags */
+		uint32			TexCreateFlags;
+		/** Index into the package ExportMap. */
+		TArray<int32>		ExportIndices;
+
+		/** Not serialized. ResourceMems are constructed on load. */
+		TArray<FTexture2DResourceMem*>	Allocations;
+		/** Note serialized. Number of ExportIndices processed during load. */
+		int32				NumExportIndicesProcessed;
+
+		/**
+		 * Checks whether all potential allocations for this TextureType have been considered yet (during load).
+		 *
+		 * @return true if all allocations have been started
+		 */
+		bool	HaveAllAllocationsBeenConsidered() const
+		{
+			return NumExportIndicesProcessed == ExportIndices.Num();
+		}
+
+		/**
+		 * Serializes an FTextureType
+		 */
+		friend FArchive& operator<<( FArchive& Ar, FTextureAllocations::FTextureType& TextureType );
+	};
+
+	FTextureAllocations()
+	:	PendingAllocationSize(0)
+	,	NumTextureTypesConsidered(0)
+	{
+	}
+
+	/** Array of texture types in the package. */
+	TArray<FTextureType>	TextureTypes;
+	/** Number of allocations that haven't been completed yet. */
+	FThreadSafeCounter		PendingAllocationCount;
+	/** Number of allocated bytes that has yet to be claimed by UTexture2D serialization. */
+	int32						PendingAllocationSize;
+	/** Number of TextureTypes that have been considered for potential allocations so far, during load. */
+	int32						NumTextureTypesConsidered;
+
+	/**
+	 * Finds a suitable ResourceMem allocation, removes it from this container and returns it to the user.
+	 *
+	 * @param SizeX				Width of texture
+	 * @param SizeY				Height of texture
+	 * @param NumMips			Number of mips
+	 * @param Format			Texture format (EPixelFormat)
+	 * @param TexCreateFlags	ETextureCreateFlags bit flags
+	 **/
+	FTexture2DResourceMem*	FindAndRemove( int32 SizeX, int32 SizeY, int32 NumMips, uint32 Format, uint32 TexCreateFlags );
+
+	/**
+	 * Cancels any pending ResourceMem allocation that hasn't been claimed by a texture yet,
+	 * just in case there are any mismatches at run-time.
+	 *
+	 * @param bCancelEverything		If true, cancels all allocations. If false, only cancels allocations that haven't been completed yet.
+	 */
+	void	CancelRemainingAllocations( bool bCancelEverything );
+
+	/**
+	 * Checks if all allocations that should be started have been started (during load).
+	 *
+	 * @return true if all allocations have been started
+	 */
+	bool	HaveAllAllocationsBeenConsidered() const
+	{
+		return NumTextureTypesConsidered == TextureTypes.Num();
+	}
+
+	/**
+	 * Checks if all ResourceMem allocations has completed.
+	 *
+	 * @return true if all ResourceMem allocations has completed
+	 */
+	bool	HasCompleted() const
+	{
+		return PendingAllocationCount.GetValue() == 0;
+	}
+
+	/**
+	 * Checks if all ResourceMem allocations have been claimed by a texture.
+	 *
+	 * @return true if there are no more pending ResourceMem allocations waiting for a texture to claim it
+	 */
+	bool	HasBeenFullyClaimed() const
+	{
+		return PendingAllocationSize == 0;
+	}
+
+	/**
+	 * Serializes an FTextureType
+	 */
+	friend FArchive& operator<<( FArchive& Ar, FTextureAllocations::FTextureType& TextureAllocationType );
+
+	/**
+	 * Serializes an FTextureAllocations struct
+	 */
+	friend FArchive& operator<<( FArchive& Ar, FTextureAllocations& TextureAllocations );
+
+	FTextureAllocations( const FTextureAllocations& Other );
+	void operator=(const FTextureAllocations& Other);
+
+private:
+	/**
+	 * Finds a texture type that matches the given specifications.
+	 *
+	 * @param SizeX				Width of the largest mip-level stored in the package
+	 * @param SizeY				Height of the largest mip-level stored in the package
+	 * @param NumMips			Number of mips
+	 * @param Format			Texture format (EPixelFormat)
+	 * @param TexCreateFlags	ETextureCreateFlags bit flags
+	 * @return					Matching texture type, or NULL if none was found
+	 */
+	FTextureType*	FindTextureType( int32 SizeX, int32 SizeY, int32 NumMips, uint32 Format, uint32 TexCreateFlags );
+
+	/**
+	 * Adds a dummy export index (-1) for a specified texture type.
+	 * Creates the texture type entry if needed.
+	 *
+	 * @param SizeX				Width of the largest mip-level stored in the package
+	 * @param SizeY				Height of the largest mip-level stored in the package
+	 * @param NumMips			Number of mips
+	 * @param Format			Texture format (EPixelFormat)
+	 * @param TexCreateFlags	ETextureCreateFlags bit flags
+	 */
+	void			AddResourceMemInfo( int32 SizeX, int32 SizeY, int32 NumMips, uint32 Format, uint32 TexCreateFlags );
+};
+#endif		// WITH_ENGINE
+
+/**
+ * A "table of contents" for an Unreal package file.  Stored at the top of the file.
+ */
+struct FPackageFileSummary
+{
+	/**
+	 * Magic tag compared against PACKAGE_FILE_TAG to ensure that package is an Unreal package.
+	 */
+	int32		Tag;
+
+private:
+	/* UE4 file version */
+	int32		FileVersionUE4;
+	/* Licensee file version */
+	int32		FileVersionLicenseeUE4;
+	/* Custom version numbers. Keyed off a unique tag for each custom component. */
+	FCustomVersionContainer CustomVersionContainer;
+
+public:
+	/**
+	 * Total size of all information that needs to be read in to create a FLinkerLoad. This includes
+	 * the package file summary, name table and import & export maps.
+	 */
+	int32		TotalHeaderSize;
+
+	/**
+	 * The flags for the package
+	 */
+	uint32	PackageFlags;
+
+	/**
+	 * The Generic Browser folder name that this package lives in
+	 */
+	FString	FolderName;
+
+	/**
+	 * Number of names used in this package
+	 */
+	int32		NameCount;
+
+	/**
+	 * Location into the file on disk for the name data
+	 */
+	int32 	NameOffset;
+
+	/**
+	 * Number of gatherable text data items in this package
+	 */
+	int32	GatherableTextDataCount;
+
+	/**
+	 * Location into the file on disk for the gatherable text data items
+	 */
+	int32 	GatherableTextDataOffset;
+	
+	/**
+	 * Number of exports contained in this package
+	 */
+	int32		ExportCount;
+
+	/**
+	 * Location into the file on disk for the ExportMap data
+	 */
+	int32		ExportOffset;
+
+	/**
+	 * Number of imports contained in this package
+	 */
+	int32		ImportCount;
+
+	/**
+	 * Location into the file on disk for the ImportMap data
+	 */
+	int32		ImportOffset;
+
+	/**
+	* Location into the file on disk for the DependsMap data
+	*/
+	int32		DependsOffset;
+
+	/**
+	 * Number of references contained in this package
+	 */
+	int32		StringAssetReferencesCount;
+
+	/**
+	 * Location into the file on disk for the string asset references map data
+	 */
+	int32		StringAssetReferencesOffset;
+
+	/**
+	 * Thumbnail table offset
+	 */
+	int32		ThumbnailTableOffset;
+
+	/**
+	 * Current id for this package
+	 */
+	FGuid	Guid;
+
+	/**
+	 * Data about previous versions of this package
+	 */
+	TArray<FGenerationInfo> Generations;
+
+	/**
+	 * Engine version this package was saved with. For hotfix releases and engine versions which maintain strict binary compatibility with another version, this may differ from CompatibleWithEngineVersion.
+	 */
+	FEngineVersion SavedByEngineVersion;
+
+	/**
+	 * Engine version this package is compatible with. See SavedByEngineVersion.
+	 */
+	FEngineVersion CompatibleWithEngineVersion;
+
+	/**
+	 * Flags used to compress the file on save and uncompress on load.
+	 */
+	uint32	CompressionFlags;
+
+	/**
+	 * Value that is used to determine if the package was saved by Epic (or licensee) or by a modder, etc
+	 */
+	uint32	PackageSource;
+
+	/**
+	 * Array of compressed chunks in case this package was stored compressed.
+	 */
+	TArray<FCompressedChunk> CompressedChunks;
+
+	/**
+	 * List of additional packages that are needed to be cooked for this package (ie streaming levels)
+	 */
+	TArray<FString>	AdditionalPackagesToCook;
+
+	/** 
+	 * If true, this file will not be saved with version numbers or was saved without version numbers. In this case they are assumed to be the current version. 
+	 * This is only used for full cooks for distribution because it is hard to guarantee correctness 
+	 **/
+	bool bUnversioned;
+
+#if WITH_ENGINE
+	/**
+	 * Information about the textures stored in the package.
+	 */
+	FTextureAllocations	TextureAllocations;
+#endif		// WITH_ENGINE
+
+	/**
+	 * Location into the file on disk for the asset registry tag data
+	 */
+	int32 	AssetRegistryDataOffset;
+
+	/** Offset to the location in the file where the bulkdata starts */
+	int64	BulkDataStartOffset;
+	/**
+	 * Offset to the location in the file where the FWorldTileInfo data starts
+	 */
+	int32 	WorldTileInfoDataOffset;
+
+	/**
+	 * Streaming install ChunkIDs
+	 */
+	TArray<int32>	ChunkIDs;
+
+
+	/** Constructor */
+	COREUOBJECT_API FPackageFileSummary();
+
+	int32 GetFileVersionUE4() const
+	{
+		return FileVersionUE4;
+	}
+
+	int32 GetFileVersionLicenseeUE4() const
+	{
+		return FileVersionLicenseeUE4;
+	}
+
+	const FCustomVersionContainer& GetCustomVersionContainer() const
+	{
+		return CustomVersionContainer;
+	}
+
+	void SetCustomVersionContainer(const FCustomVersionContainer& InContainer)
+	{
+		CustomVersionContainer = InContainer;
+	}
+
+	void SetFileVersions(const int32 EpicUE4, const int32 LicenseeUE4, const bool bInSaveUnversioned = false)
+	{
+		FileVersionUE4 = EpicUE4;
+		FileVersionLicenseeUE4 = LicenseeUE4;
+		bUnversioned = bInSaveUnversioned;
+	}
+
+	/** I/O function */
+	friend COREUOBJECT_API FArchive& operator<<( FArchive& Ar, FPackageFileSummary& Sum );
+};
diff --git a/Engine/Source/Runtime/CoreUObject/Public/UObject/StructOnScope.h b/Engine/Source/Runtime/CoreUObject/Public/UObject/StructOnScope.h
--- a/Engine/Source/Runtime/CoreUObject/Public/UObject/StructOnScope.h
+++ b/Engine/Source/Runtime/CoreUObject/Public/UObject/StructOnScope.h
@@ -2,7 +2,7 @@
 
 #pragma once
 
-#include "CoreObject.h"
+#include "LevelGuids.h"
 
 class FStructOnScope
 {
