Copying //UE4/Dev-Build to //UE4/Dev-Main (Source: //UE4/Dev-Build @ 3104200)

diff --git a/Engine/Source/Runtime/CoreUObject/Public/UObject/TextProperty.h b/Engine/Source/Runtime/CoreUObject/Public/UObject/TextProperty.h
new file mode 100644
--- /dev/null
+++ b/Engine/Source/Runtime/CoreUObject/Public/UObject/TextProperty.h
@@ -0,0 +1,44 @@
+// Copyright 1998-2016 Epic Games, Inc. All Rights Reserved.
+#pragma once
+
+#include "ObjectMacros.h"
+
+// need to break this out a different type so that the DECLARE_CASTED_CLASS_INTRINSIC macro can digest the comma
+typedef TProperty<FText, UProperty> UTextProperty_Super;
+
+class COREUOBJECT_API UTextProperty : public UTextProperty_Super
+{
+	DECLARE_CASTED_CLASS_INTRINSIC(UTextProperty, UTextProperty_Super, 0, TEXT("/Script/CoreUObject"), CASTCLASS_UTextProperty)
+
+public:
+
+	typedef UTextProperty_Super::TTypeFundamentals TTypeFundamentals;
+	typedef TTypeFundamentals::TCppType TCppType;
+
+	UTextProperty(ECppProperty, int32 InOffset, uint64 InFlags)
+		: TProperty(FObjectInitializer::Get(), EC_CppProperty, InOffset, InFlags)
+	{
+	}
+
+	UTextProperty( const FObjectInitializer& ObjectInitializer, ECppProperty, int32 InOffset, uint64 InFlags )
+		:	TProperty( ObjectInitializer, EC_CppProperty, InOffset, InFlags)
+	{
+	}
+
+	// UProperty interface
+	virtual bool ConvertFromType(const FPropertyTag& Tag, FArchive& Ar, uint8* Data, UStruct* DefaultsStruct, bool& bOutAdvanceProperty) override;
+	virtual bool Identical( const void* A, const void* B, uint32 PortFlags ) const override;
+	virtual void SerializeItem( FArchive& Ar, void* Value, void const* Defaults ) const override;
+	virtual void ExportTextItem( FString& ValueStr, const void* PropertyValue, const void* DefaultValue, UObject* Parent, int32 PortFlags, UObject* ExportRootScope ) const override;
+	virtual const TCHAR* ImportText_Internal( const TCHAR* Buffer, void* Data, int32 PortFlags, UObject* OwnerObject, FOutputDevice* ErrorText ) const override;
+	virtual FString GetCPPTypeForwardDeclaration() const override
+	{
+		return FString();
+	}
+	// End of UProperty interface
+
+	/** Generate the correct C++ code for the given text value */
+	static FString GenerateCppCodeForTextValue(const FText& InValue, const FString& Indent);
+
+	static bool Identical_Implementation(const FText& A, const FText& B, uint32 PortFlags);
+};
\ No newline at end of file
diff --git a/Engine/Source/Runtime/CoreUObject/Public/UObject/UnrealType.h b/Engine/Source/Runtime/CoreUObject/Public/UObject/UnrealType.h
--- a/Engine/Source/Runtime/CoreUObject/Public/UObject/UnrealType.h
+++ b/Engine/Source/Runtime/CoreUObject/Public/UObject/UnrealType.h
@@ -6,10 +6,11 @@
 
 #pragma once
 
-#include "ObjectBase.h"
+#include "ObjectMacros.h"
 #include "PropertyPortFlags.h"
 #include "PropertyTag.h"
 #include "Templates/IsTriviallyDestructible.h"
+#include "Serialization/SerializedPropertyScope.h"
 
 COREUOBJECT_API DECLARE_LOG_CATEGORY_EXTERN(LogType, Log, All);
 
diff --git a/Engine/Source/Runtime/Engine/Classes/Engine/BlueprintGeneratedClass.h b/Engine/Source/Runtime/Engine/Classes/Engine/BlueprintGeneratedClass.h
--- a/Engine/Source/Runtime/Engine/Classes/Engine/BlueprintGeneratedClass.h
+++ b/Engine/Source/Runtime/Engine/Classes/Engine/BlueprintGeneratedClass.h
@@ -3,7 +3,7 @@
 #pragma once
 
 #include "MeshBatch.h"
-#include "ArchiveBase.h"
+#include "Archive.h"
 #include "BlueprintGeneratedClass.generated.h"
 
 DECLARE_MEMORY_STAT_EXTERN(TEXT("Persistent Uber Graph Frame memory"), STAT_PersistentUberGraphFrameMemory, STATGROUP_Memory, );
diff --git a/Engine/Source/Runtime/Engine/Classes/Engine/Level.h b/Engine/Source/Runtime/Engine/Classes/Engine/Level.h
--- a/Engine/Source/Runtime/Engine/Classes/Engine/Level.h
+++ b/Engine/Source/Runtime/Engine/Classes/Engine/Level.h
@@ -11,6 +11,7 @@ class UNavigationDataChunk;
 class AInstancedFoliageActor;
 class AWorldSettings;
 class UWorld;
+class FSceneInterface;
 
 /**
  * Structure containing all information needed for determining the screen space
diff --git a/Engine/Source/Runtime/Engine/Public/AudioDevice.h b/Engine/Source/Runtime/Engine/Public/AudioDevice.h
--- a/Engine/Source/Runtime/Engine/Public/AudioDevice.h
+++ b/Engine/Source/Runtime/Engine/Public/AudioDevice.h
@@ -15,7 +15,12 @@ class ICompressedAudioInfo;
 class IAudioSpatializationPlugin;
 class IAudioSpatializationAlgorithm;
 class UReverbEffect;
+class USoundConcurrency;
+class FViewport;
 class FViewportClient;
+class FCanvas;
+
+struct FAudioComponentParam;
 
 /** 
  * Debug state of the audio system
diff --git a/Engine/Source/Runtime/Mac/UnrealAudioCoreAudio/Private/UnrealAudioCoreAudio.cpp b/Engine/Source/Runtime/Mac/UnrealAudioCoreAudio/Private/UnrealAudioCoreAudio.cpp
--- a/Engine/Source/Runtime/Mac/UnrealAudioCoreAudio/Private/UnrealAudioCoreAudio.cpp
+++ b/Engine/Source/Runtime/Mac/UnrealAudioCoreAudio/Private/UnrealAudioCoreAudio.cpp
@@ -22,8 +22,6 @@ See https://developer.apple.com/library/mac/documentation/MusicAudio/Conceptual/
 // Used to toggle on white noise testing for CoreAudio output
 #define UNREAL_AUDIO_TEST_WHITE_NOISE 0
 
-DEFINE_LOG_CATEGORY(LogUnrealAudioDevice);
-
 static const TCHAR* GetCoreAudioError(OSStatus Result)
 {
 	switch (Result)
diff --git a/Engine/Source/Runtime/SlateCore/Public/Fonts/FontCache.h b/Engine/Source/Runtime/SlateCore/Public/Fonts/FontCache.h
--- a/Engine/Source/Runtime/SlateCore/Public/Fonts/FontCache.h
+++ b/Engine/Source/Runtime/SlateCore/Public/Fonts/FontCache.h
@@ -16,9 +16,10 @@ class FCompositeFontCache;
 class FSlateFontRenderer;
 class FSlateTextShaper;
 class FSlateFontCache;
+class FShapedGlyphFaceData;
 
 struct FCharacterRenderData;
-class FShapedGlyphFaceData;
+struct FSlateFontKey;
 
 /** 
  * Methods that can be used to shape text.
diff --git a/Engine/Source/Runtime/SlateCore/Public/Styling/ISlateStyle.h b/Engine/Source/Runtime/SlateCore/Public/Styling/ISlateStyle.h
--- a/Engine/Source/Runtime/SlateCore/Public/Styling/ISlateStyle.h
+++ b/Engine/Source/Runtime/SlateCore/Public/Styling/ISlateStyle.h
@@ -2,6 +2,8 @@
 #pragma once
 #include "Styling/StyleDefaults.h"
 
+class UTexture2D;
+
 struct FMargin;
 struct FSlateColor;
 struct FSlateDynamicImageBrush;
diff --git a/Engine/Source/Runtime/UnrealAudio/Public/UnrealAudioDeviceModule.h b/Engine/Source/Runtime/UnrealAudio/Public/UnrealAudioDeviceModule.h
--- a/Engine/Source/Runtime/UnrealAudio/Public/UnrealAudioDeviceModule.h
+++ b/Engine/Source/Runtime/UnrealAudio/Public/UnrealAudioDeviceModule.h
@@ -9,7 +9,7 @@
 
 #if ENABLE_UNREAL_AUDIO
 
-DECLARE_LOG_CATEGORY_EXTERN(LogUnrealAudioDevice, Log, All);
+UNREALAUDIO_API DECLARE_LOG_CATEGORY_EXTERN(LogUnrealAudioDevice, Log, All);
 
 #define UA_DEVICE_PLATFORM_ERROR(INFO)	(UAudio::OnDeviceError(EDeviceError::PLATFORM, INFO, FString(__FILE__), __LINE__))
 #define AU_DEVICE_PARAM_ERROR(INFO)		(UAudio::OnDeviceError(EDeviceError::INVALID_PARAMETER, INFO, FString(__FILE__), __LINE__))
diff --git a/Engine/Source/Runtime/Windows/UnrealAudioWasapi/Private/UnrealAudioDeviceWasapi.cpp b/Engine/Source/Runtime/Windows/UnrealAudioWasapi/Private/UnrealAudioDeviceWasapi.cpp
--- a/Engine/Source/Runtime/Windows/UnrealAudioWasapi/Private/UnrealAudioDeviceWasapi.cpp
+++ b/Engine/Source/Runtime/Windows/UnrealAudioWasapi/Private/UnrealAudioDeviceWasapi.cpp
@@ -19,8 +19,6 @@
 
 #if ENABLE_UNREAL_AUDIO
 
-DEFINE_LOG_CATEGORY(LogUnrealAudioDevice);
-
 // See MSDN documentation for what these error codes mean in the context of the API call
 static const TCHAR* GetWasapiError(HRESULT Result)
 {
