Don't add dynamic delegates via import if they are already in the list

diff --git a/Engine/Source/Runtime/CoreUObject/Private/UObject/PropertyMulticastDelegate.cpp b/Engine/Source/Runtime/CoreUObject/Private/UObject/PropertyMulticastDelegate.cpp
--- a/Engine/Source/Runtime/CoreUObject/Private/UObject/PropertyMulticastDelegate.cpp
+++ b/Engine/Source/Runtime/CoreUObject/Private/UObject/PropertyMulticastDelegate.cpp
@@ -231,7 +231,7 @@ const TCHAR* UMulticastDelegateProperty::ImportText_Internal( const TCHAR* Buffe
 		}
 
 		// Add this delegate to our multicast delegate's invocation list
-		MulticastDelegate.Add( ImportedDelegate );
+		MulticastDelegate.AddUnique( ImportedDelegate );
 
 		SkipWhitespace(Buffer);
 	}
