UE4 - Restore bulk data to fix crash during launch on, as textures lost their SourceMips due to this

diff --git a/Engine/Source/Runtime/CoreUObject/Private/UObject/SavePackage.cpp b/Engine/Source/Runtime/CoreUObject/Private/UObject/SavePackage.cpp
--- a/Engine/Source/Runtime/CoreUObject/Private/UObject/SavePackage.cpp
+++ b/Engine/Source/Runtime/CoreUObject/Private/UObject/SavePackage.cpp
@@ -4760,7 +4760,6 @@ ESavePackageResult UPackage::Save(UPackage* InOuter, UObject* Base, EObjectFlags
 						BulkDataStorageInfo.BulkData->ClearBulkDataFlags(0xFFFFFFFF);
 						BulkDataStorageInfo.BulkData->SetBulkDataFlags(OldBulkDataFlags);
 						BulkDataStorageInfo.BulkData->Unlock();
-						BulkDataStorageInfo.BulkData->RemoveBulkData();
 					}
 
 					if (BulkArchive)
