Copying //UE4/Dev-Core to //UE4/Dev-Main (Source: //UE4/Dev-Core @ 3151653)

diff --git a/Engine/Source/Runtime/CoreUObject/Private/Serialization/AsyncLoading.cpp b/Engine/Source/Runtime/CoreUObject/Private/Serialization/AsyncLoading.cpp
--- a/Engine/Source/Runtime/CoreUObject/Private/Serialization/AsyncLoading.cpp
+++ b/Engine/Source/Runtime/CoreUObject/Private/Serialization/AsyncLoading.cpp
@@ -21,6 +21,12 @@
 
 #define FIND_MEMORY_STOMPS (1 && (PLATFORM_WINDOWS || PLATFORM_LINUX) && !WITH_EDITORONLY_DATA)
 
+#ifndef USE_EVENT_DRIVEN_ASYNC_LOAD
+#error "USE_EVENT_DRIVEN_ASYNC_LOAD must be defined""
+#endif
+
+DEFINE_LOG_CATEGORY(LogLoadingDev);
+
 //#pragma clang optimize off
 
 /*-----------------------------------------------------------------------------
@@ -38,7 +44,6 @@ DECLARE_CYCLE_STAT(TEXT("CreateLinker AsyncPackage"),STAT_FAsyncPackage_CreateLi
 DECLARE_CYCLE_STAT(TEXT("FinishLinker AsyncPackage"),STAT_FAsyncPackage_FinishLinker,STATGROUP_AsyncLoad);
 DECLARE_CYCLE_STAT(TEXT("LoadImports AsyncPackage"),STAT_FAsyncPackage_LoadImports,STATGROUP_AsyncLoad);
 DECLARE_CYCLE_STAT(TEXT("CreateImports AsyncPackage"),STAT_FAsyncPackage_CreateImports,STATGROUP_AsyncLoad);
-DECLARE_CYCLE_STAT(TEXT("FinishTextureAllocations AsyncPackage"),STAT_FAsyncPackage_FinishTextureAllocations,STATGROUP_AsyncLoad);
 DECLARE_CYCLE_STAT(TEXT("CreateMetaData AsyncPackage"),STAT_FAsyncPackage_CreateMetaData,STATGROUP_AsyncLoad);
 DECLARE_CYCLE_STAT(TEXT("CreateExports AsyncPackage"),STAT_FAsyncPackage_CreateExports,STATGROUP_AsyncLoad);
 DECLARE_CYCLE_STAT(TEXT("FreeReferencedImports AsyncPackage"), STAT_FAsyncPackage_FreeReferencedImports, STATGROUP_AsyncLoad);
@@ -66,170 +71,31 @@ bool IsGarbageCollectionLocked();
 /** Global request ID counter */
 static FThreadSafeCounter GPackageRequestID;
 
-/** 
- * Keeps a reference to all objects created during async load until streaming has finished 
- *
- * ASSUMPTION: AddObject can't be called while GC is running and we don't want to lock when calling AddReferencedObjects
+/** DEBUG only. @todo remove */
+static TMap<FName, TArray<FAsyncPackage*>> PackagesThatDecrementedCounter;
+
+/**
+ * Updates FUObjectThreadContext with the current package when processing it.
+ * FUObjectThreadContext::AsyncPackage is used by NotifyConstructedDuringAsyncLoading.
  */
-class FAsyncObjectsReferencer : FGCObject
+struct FAsyncPackageScope
 {
-	/** Private constructor */
-	FAsyncObjectsReferencer() {}
-
-	/** List of referenced objects */
-	TSet<UObject*> ReferencedObjects;
-#if THREADSAFE_UOBJECTS
-	/** Critical section for referenced objects list */
-	FCriticalSection ReferencedObjectsCritical;
-#endif
-
-public:
-#if !UE_BUILD_SHIPPING
-	bool Contains(UObject* InObj)
-	{
-#if THREADSAFE_UOBJECTS
-		FScopeLock ReferencedObjectsLock(&ReferencedObjectsCritical);
-#endif
-		return ReferencedObjects.Contains(InObj);
-	}
-#endif
-
-	/** Returns the one and only instance of this object */
-	static FAsyncObjectsReferencer& Get();
-	/** FGCObject interface */
-	virtual void AddReferencedObjects(FReferenceCollector& Collector) override
-	{
-		// Note we don't lock here as we're guaranteed that AddObject can only be called from
-		// within FGCScopeGuard scope where GC does not run
-		Collector.AllowEliminatingReferences(false);
-		{
-			Collector.AddReferencedObjects(ReferencedObjects);
-		}
-		Collector.AllowEliminatingReferences(true);
-	}
-	/** 
-	 * Adds an object to be referenced 
-	 * The assumption here is that this can only happen from inside of FGCScopeGuard (@see IsGarbageCollectionLocked()) where we're sure GC is not currently running,
-	 * unless we're on the game thread where atm GC can run simultaneously with async loading.
-	 */
-	FORCEINLINE void AddObject(UObject* InObject)
-	{
-		if (InObject)
-		{
-			UE_CLOG(!IsInGameThread() && !IsGarbageCollectionLocked(), LogStreaming, Fatal, TEXT("Trying to add an object %s to FAsyncObjectsReferencer outside of a FGCScopeLock."), *InObject->GetFullName());
-			{
-#if THREADSAFE_UOBJECTS
-				// Still want to lock as AddObject may be called on the game thread and async loading thread,
-				// but in any case it may not happen when GC runs.
-				FScopeLock ReferencedObjectsLock(&ReferencedObjectsCritical);
-#else
-				check(IsInGameThread());
-#endif
-				if (!ReferencedObjects.Contains(InObject))
-				{
-					ReferencedObjects.Add(InObject);
-				}
-			}
-			InObject->ThisThreadAtomicallyClearedRFUnreachable();
-		}
-	}
-	/** Removes all objects from the list and clears async loading flags */
-	void EmptyReferencedObjects()
-	{
-		check(IsInGameThread());
-#if THREADSAFE_UOBJECTS
-		FScopeLock ReferencedObjectsLock(&ReferencedObjectsCritical);
-#endif
-		const EInternalObjectFlags AsyncFlags = EInternalObjectFlags::Async | EInternalObjectFlags::AsyncLoading;
-		for (UObject* Obj : ReferencedObjects)
-		{
-			// Temporary fatal messages instead of checks to find the cause for a one-time crash in shipping config
-			UE_CLOG(Obj == nullptr, LogStreaming, Fatal, TEXT("NULL object in Async Objects Referencer"));
-			UE_CLOG(!Obj->IsValidLowLevelFast(), LogStreaming, Fatal, TEXT("Invalid object in Async Objects Referencer"));
-			Obj->AtomicallyClearInternalFlags(AsyncFlags);
-			check(!Obj->HasAnyInternalFlags(AsyncFlags))
-		}
-		ReferencedObjects.Reset();
-	}
-	/** Removes all referenced objects and markes them for GC */
-	void EmptyReferencedObjectsAndCancelLoading()
-	{
-		const EObjectFlags LoadFlags = RF_NeedLoad | RF_NeedPostLoad | RF_NeedPostLoadSubobjects;
-		const EInternalObjectFlags AsyncFlags = EInternalObjectFlags::Async | EInternalObjectFlags::AsyncLoading;
-
-#if THREADSAFE_UOBJECTS
-		FScopeLock ReferencedObjectsLock(&ReferencedObjectsCritical);
-#endif
-
-		// All of the referenced objects have been created by async loading code and may be in an invalid state so mark them for GC
-		for (auto Object : ReferencedObjects)
-		{
-			Object->ClearInternalFlags(AsyncFlags);
-			if (Object->HasAnyFlags(LoadFlags))
-			{
-				Object->AtomicallyClearFlags(LoadFlags);
-				Object->MarkPendingKill();
-			}
-			check(!Object->HasAnyInternalFlags(AsyncFlags) && !Object->HasAnyFlags(LoadFlags));
-		}
-		ReferencedObjects.Reset();
-	}
+	/** Outer scope package */
+	FAsyncPackage* PreviousPackage;
+	/** Cached ThreadContext so we don't have to access it again */
+	FUObjectThreadContext& ThreadContext;
 
-#if !UE_BUILD_SHIPPING
-	/** Verifies that no object exists that has either EInternalObjectFlags::AsyncLoading and EInternalObjectFlags::Async set and is NOT being referenced by FAsyncObjectsReferencer */
-	FORCENOINLINE void VerifyAssumptions()
+	FAsyncPackageScope(FAsyncPackage* InPackage)
+		: ThreadContext(FUObjectThreadContext::Get())
 	{
-		const EInternalObjectFlags AsyncFlags = EInternalObjectFlags::Async | EInternalObjectFlags::AsyncLoading;
-		for (FRawObjectIterator It; It; ++It)
-		{
-			FUObjectItem* ObjItem = *It;
-			checkSlow(ObjItem);
-			UObject* Object = static_cast<UObject*>(ObjItem->Object);
-			if (Object->HasAnyInternalFlags(AsyncFlags))
-			{
-				if (!Contains(Object))
-				{
-					UE_LOG(LogStreaming, Error, TEXT("%s has AsyncLoading|Async set but is not referenced by FAsyncObjectsReferencer"), *Object->GetPathName());
-				}
-			}
-		}
+		PreviousPackage = ThreadContext.AsyncPackage;
+		ThreadContext.AsyncPackage = InPackage;
 	}
-#endif
-};
-FAsyncObjectsReferencer& FAsyncObjectsReferencer::Get()
-{
-	static FAsyncObjectsReferencer Singleton;
-	return Singleton;
-}
-
-#if !UE_BUILD_SHIPPING
-class FAsyncLoadingExec : private FSelfRegisteringExec
-{
-public:
-
-	FAsyncLoadingExec()
-	{}
-
-	/** Console commands **/
-	virtual bool Exec(UWorld* InWorld, const TCHAR* Cmd, FOutputDevice& Ar) override
+	~FAsyncPackageScope()
 	{
-		if (FParse::Command(&Cmd, TEXT("VerifyAsyncLoadAssumptions")))
-		{
-			if (!IsAsyncLoading())
-			{
-				FAsyncObjectsReferencer::Get().VerifyAssumptions();
-			}
-			else
-			{
-				Ar.Logf(TEXT("Unable to verify async loading assumptions while streaming."));
-			}
-			return true;
-		}
-		return false;
+		ThreadContext.AsyncPackage = PreviousPackage;
 	}
 };
-static TAutoPtr<FAsyncLoadingExec> GAsyncLoadingExec;
-#endif
 
 static int32 GAsyncLoadingThreadEnabled;
 static FAutoConsoleVariableRef CVarAsyncLoadingThreadEnabledg(
@@ -265,6 +131,7 @@ static FAutoConsoleVariableRef CVarTimeLimitExceededMinTime(
 	);
 
 static int32 GPreloadPackageDependencies = 0;
+#if !USE_EVENT_DRIVEN_ASYNC_LOAD // we don't want preloading of dependencies with event driven loading, we would rather find the dependencies naturally.
 static FAutoConsoleVariableRef CVarPreloadPackageDependencies(
 	TEXT("s.PreloadPackageDependencies"),
 	GPreloadPackageDependencies,
@@ -273,35 +140,54 @@ static FAutoConsoleVariableRef CVarPreloadPackageDependencies(
 	TEXT("1 - Preload package dependencies. Faster but requires asset registry data to be loaded into memory\n"),
 	ECVF_Default
 	);
+#endif
+
+static int32 GEventDrivenLoaderEnabled = 0;
+static FAutoConsoleVariableRef CVarEventDrivenLoaderEnabled(
+	TEXT("s.EventDrivenLoaderEnabled"),
+	GEventDrivenLoaderEnabled,
+	TEXT("Placeholder console variable, currently not used in runtime."),
+	ECVF_Default
+);
 
 uint32 FAsyncLoadingThread::AsyncLoadingThreadID = 0;
 
-static void IsTimeLimitExceededPrint(double InTickStartTime, double CurrentTime, float InTimeLimit, const TCHAR* InLastTypeOfWorkPerformed = nullptr, UObject* InLastObjectWorkWasPerformedOn = nullptr)
+static void IsTimeLimitExceededPrint(double InTickStartTime, double CurrentTime, double LastTestTime, float InTimeLimit, const TCHAR* InLastTypeOfWorkPerformed = nullptr, UObject* InLastObjectWorkWasPerformedOn = nullptr)
 {
+	static double LastPrintStartTime = -1.0;
 	// Log single operations that take longer than time limit (but only in cooked builds)
-	if (
+	if (LastPrintStartTime != InTickStartTime &&
 		(CurrentTime - InTickStartTime) > GTimeLimitExceededMinTime &&
 		(CurrentTime - InTickStartTime) > (GTimeLimitExceededMultiplier * InTimeLimit))
 	{
-		UE_LOG(LogStreaming, Warning, TEXT("IsTimeLimitExceeded: %s %s took (less than) %5.2f ms"),
+		float EstimatedTimeForThisStep = (CurrentTime - InTickStartTime) * 1000;
+		if (LastTestTime > InTickStartTime)
+		{
+			EstimatedTimeForThisStep = (CurrentTime - LastTestTime) * 1000;
+		}
+		LastPrintStartTime = InTickStartTime;
+		UE_LOG(LogStreaming, Warning, TEXT("IsTimeLimitExceeded: %s %s Load Time %5.2fms   Last Step Time %5.2fms"),
 			InLastTypeOfWorkPerformed ? InLastTypeOfWorkPerformed : TEXT("unknown"),
 			InLastObjectWorkWasPerformedOn ? *InLastObjectWorkWasPerformedOn->GetFullName() : TEXT("nullptr"),
-			(CurrentTime - InTickStartTime) * 1000);
+			(CurrentTime - InTickStartTime) * 1000,
+			EstimatedTimeForThisStep);
 	}
 }
 
 static FORCEINLINE bool IsTimeLimitExceeded(double InTickStartTime, bool bUseTimeLimit, float InTimeLimit, const TCHAR* InLastTypeOfWorkPerformed = nullptr, UObject* InLastObjectWorkWasPerformedOn = nullptr)
 {
+	static double LastTestTime = -1.0;
 	bool bTimeLimitExceeded = false;
 	if (bUseTimeLimit)
 	{
 		double CurrentTime = FPlatformTime::Seconds();
 		bTimeLimitExceeded = CurrentTime - InTickStartTime > InTimeLimit;
 
-		if (GWarnIfTimeLimitExceeded)
+		if (bTimeLimitExceeded && GWarnIfTimeLimitExceeded)
 		{
-			IsTimeLimitExceededPrint(InTickStartTime, CurrentTime, InTimeLimit, InLastTypeOfWorkPerformed, InLastObjectWorkWasPerformedOn);
+			IsTimeLimitExceededPrint(InTickStartTime, CurrentTime, LastTestTime, InTimeLimit, InLastTypeOfWorkPerformed, InLastObjectWorkWasPerformedOn);
 		}
+		LastTestTime = CurrentTime;
 	}
 	return bTimeLimitExceeded;
 }
@@ -310,6 +196,53 @@ FORCEINLINE bool FAsyncPackage::IsTimeLimitExceeded()
 	return AsyncLoadingThread.IsAsyncLoadingSuspended() || ::IsTimeLimitExceeded(TickStartTime, bUseTimeLimit, TimeLimit, LastTypeOfWorkPerformed, LastObjectWorkWasPerformedOn);
 }
 
+#if USE_NEW_ASYNC_IO
+
+DEFINE_LOG_CATEGORY_STATIC(LogAsyncArchive, Display, All);
+DECLARE_MEMORY_STAT(TEXT("FArchiveAsync2 Buffers"), STAT_FArchiveAsync2Mem, STATGROUP_Memory);
+
+//#define ASYNC_WATCH_FILE "SM_Boots_Wings.uasset"
+#define TRACK_SERIALIZE (0)
+#define MIN_REMAIN_TIME (0.00101f)   // wait(0) is very different than wait(tiny) so we cut things off well before roundoff error could cause us to block when we didn't intend to. Also the granularity of the event API is 1ms.
+
+
+
+
+FORCEINLINE void FArchiveAsync2::LogItem(const TCHAR* Item, int64 Offset, int64 Size, double StartTime)
+{
+	if (UE_LOG_ACTIVE(LogAsyncArchive, Verbose)
+#if defined(ASYNC_WATCH_FILE)
+		|| FileName.Contains(TEXT(ASYNC_WATCH_FILE))
+#endif
+		)
+	{
+		static double GlobalStartTime(FPlatformTime::Seconds());
+		double Now(FPlatformTime::Seconds());
+
+		float ThisTime = (StartTime != 0.0) ? float(1000.0 * (Now - StartTime)) : 0.0f;
+
+		if (!UE_LOG_ACTIVE(LogAsyncArchive, VeryVerbose) && ThisTime < 1.0f
+#if defined(ASYNC_WATCH_FILE)
+			&& !FileName.Contains(TEXT(ASYNC_WATCH_FILE))
+#endif
+			)
+		{
+			return;
+		}
+
+		FPlatformMisc::LowLevelOutputDebugStringf(TEXT("%32s%3s    %12lld %12lld    %6.2fms    (+%9.2fms)      %s\r\n"),
+			Item,
+			(ThisTime > 1.0f) ? TEXT("***") : TEXT(""),
+			Offset,
+			((Size == MAX_int64) ? TotalSize() : Offset + Size),
+			ThisTime,
+			float(1000.0 * (Now - GlobalStartTime)),
+			*FileName);
+	}
+}
+#endif
+
+
 #if LOOKING_FOR_PERF_ISSUES
 FThreadSafeCounter FAsyncLoadingThread::BlockingCycles = 0;
 #endif
@@ -320,25 +253,25 @@ FAsyncLoadingThread& FAsyncLoadingThread::Get()
 	return GAsyncLoader;
 }
 
-/** Just like TGuardValue for FAsyncLoadingThread::bIsInAsyncLoadingTick but only works for the game thread */
+/** Just like TGuardValue for FAsyncLoadingThread::AsyncLoadingTickCounter but only works for the game thread */
 struct FAsyncLoadingTickScope
 {
-	bool bWasInTick;
+	bool bNeedsToLeaveAsyncTick;
 	FAsyncLoadingTickScope()
-		: bWasInTick(false)
+		: bNeedsToLeaveAsyncTick(false)
 	{
 		if (IsInGameThread())
 		{
 			FAsyncLoadingThread& AsyncLoadingThread = FAsyncLoadingThread::Get();
-			bWasInTick = AsyncLoadingThread.GetIsInAsyncLoadingTick();
-			AsyncLoadingThread.SetIsInAsyncLoadingTick(true);
+			AsyncLoadingThread.EnterAsyncLoadingTick();
+			bNeedsToLeaveAsyncTick = true;
 		}
 	}
 	~FAsyncLoadingTickScope()
 	{
-		if (IsInGameThread())
+		if (bNeedsToLeaveAsyncTick)
 		{
-			FAsyncLoadingThread::Get().SetIsInAsyncLoadingTick(bWasInTick);
+			FAsyncLoadingThread::Get().LeaveAsyncLoadingTick();
 		}
 	}
 };
@@ -350,11 +283,6 @@ void FAsyncLoadingThread::InitializeAsyncThread()
 
 void FAsyncLoadingThread::CancelAsyncLoadingInternal()
 {
-	{
-		FGCScopeGuard GCGuard;
-		FAsyncObjectsReferencer::Get().EmptyReferencedObjectsAndCancelLoading();
-	}
-
 	{
 		// Packages we haven't yet started processing.
 #if THREADSAFE_UOBJECTS
@@ -429,7 +357,18 @@ void FAsyncLoadingThread::QueuePackage(const FAsyncPackageDesc& Package)
 	QueuedRequestsEvent->Trigger();
 }
 
-FAsyncPackage* FAsyncLoadingThread::FindExistingPackageAndAddCompletionCallback(FAsyncPackageDesc* PackageRequest, TMap<FName, FAsyncPackage*>& PackageList)
+void FAsyncPackage::PopulateFlushTree(struct FFlushTree* FlushTree)
+{
+	if (FlushTree->AddPackage(GetPackageName()))
+	{
+		for (FAsyncPackage* PendingImport : PendingImportedPackages)
+		{
+			PendingImport->PopulateFlushTree(FlushTree);
+		}
+	}
+}
+
+FAsyncPackage* FAsyncLoadingThread::FindExistingPackageAndAddCompletionCallback(FAsyncPackageDesc* PackageRequest, TMap<FName, FAsyncPackage*>& PackageList, FFlushTree* FlushTree)
 {
 	checkSlow(IsInAsyncLoadThread());
 	FAsyncPackage* Result = PackageList.FindRef(PackageRequest->Name);
@@ -439,7 +378,11 @@ FAsyncPackage* FAsyncLoadingThread::FindExistingPackageAndAddCompletionCallback(
 		{
 			const bool bInternalCallback = false;
 			Result->AddCompletionCallback(PackageRequest->PackageLoadedDelegate, bInternalCallback);
+		}
 			Result->AddRequestID(PackageRequest->RequestID);
+		if (FlushTree)
+		{
+			Result->PopulateFlushTree(FlushTree);
 		}
 		const int32 QueuedPackagesCount = QueuedPackagesCounter.Decrement();
 		check(QueuedPackagesCount >= 0);
@@ -456,7 +399,7 @@ void FAsyncLoadingThread::UpdateExistingPackagePriorities(FAsyncPackage* InPacka
 		// always inserted anyway AsyncPackageNameLookup.Remove(InPackage->GetPackageName());
 		InPackage->SetPriority(InNewPriority);
 
-		InsertPackage(InPackage, EAsyncPackageInsertMode::InsertBeforeMatchingPriorities);
+		InsertPackage(InPackage, false, EAsyncPackageInsertMode::InsertBeforeMatchingPriorities);
 
 		// Reduce loading counters as InsertPackage incremented them again
 		ExistingAsyncPackagesCounter.Decrement();
@@ -488,9 +431,9 @@ void FAsyncLoadingThread::UpdateExistingPackagePriorities(FAsyncPackage* InPacka
 #endif
 }
 
-void FAsyncLoadingThread::ProcessAsyncPackageRequest(FAsyncPackageDesc* InRequest, FAsyncPackage* InRootPackage, IAssetRegistryInterface* InAssetRegistry)
+void FAsyncLoadingThread::ProcessAsyncPackageRequest(FAsyncPackageDesc* InRequest, FAsyncPackage* InRootPackage, IAssetRegistryInterface* InAssetRegistry, FFlushTree* FlushTree)
 {
-	FAsyncPackage* Package = FindExistingPackageAndAddCompletionCallback(InRequest, AsyncPackageNameLookup);
+	FAsyncPackage* Package = FindExistingPackageAndAddCompletionCallback(InRequest, AsyncPackageNameLookup, FlushTree);
 
 	if (Package)
 	{
@@ -504,7 +447,7 @@ void FAsyncLoadingThread::ProcessAsyncPackageRequest(FAsyncPackageDesc* InReques
 #if THREADSAFE_UOBJECTS
 		FScopeLock LoadedLock(&LoadedPackagesCritical);
 #endif
-		Package = FindExistingPackageAndAddCompletionCallback(InRequest, LoadedPackagesNameLookup);
+		Package = FindExistingPackageAndAddCompletionCallback(InRequest, LoadedPackagesNameLookup, FlushTree);
 	}
 
 	if (!Package)
@@ -513,23 +456,11 @@ void FAsyncLoadingThread::ProcessAsyncPackageRequest(FAsyncPackageDesc* InReques
 #if THREADSAFE_UOBJECTS
 		FScopeLock LoadedLock(&LoadedPackagesToProcessCritical);
 #endif
-		Package = FindExistingPackageAndAddCompletionCallback(InRequest, LoadedPackagesToProcessNameLookup);
+		Package = FindExistingPackageAndAddCompletionCallback(InRequest, LoadedPackagesToProcessNameLookup, FlushTree);
 	}
 
 	if (!Package)
 	{
-		// we probably don't need this bit (we can't cover much latency anyway) but it prevents false positives
-		// probably would be better to start this on the game thread when the request is initiated.
-#if 0 // left here as a note that this is probably not prestreamed and maybe could be
-		if (InAssetRegistry && !InRootPackage)
-		{
-			FString PrestreamFilename = GetPrestreamPackageLinkerName(*InRequest->Name.ToString());
-			if (PrestreamFilename.Len() > 0)
-			{
-				HintFutureRead(*PrestreamFilename); // should set the priority on the IO requests here
-			}
-		}
-#endif
 		// New package that needs to be loaded or a package has already been loaded long time ago
 		Package = new FAsyncPackage(*InRequest);
 		if (InRequest->PackageLoadedDelegate.IsBound())
@@ -538,6 +469,10 @@ void FAsyncLoadingThread::ProcessAsyncPackageRequest(FAsyncPackageDesc* InReques
 			Package->AddCompletionCallback(InRequest->PackageLoadedDelegate, bInternalCallback);
 		}
 		Package->SetDependencyRootPackage(InRootPackage);
+		if (FlushTree)
+		{
+			Package->PopulateFlushTree(FlushTree);
+		}
 
 #if !WITH_EDITOR
 		if (InAssetRegistry && !InRootPackage)
@@ -548,16 +483,6 @@ void FAsyncLoadingThread::ProcessAsyncPackageRequest(FAsyncPackageDesc* InReques
 			ScanPackageDependenciesForLoadOrder(*InRequest->Name.ToString(), OrderTracker, Order, InAssetRegistry);
 
 			OrderTracker.ValueSort(TLess<int32>());
-#if USE_NEW_ASYNC_IO
-			/*for (auto& Dependency : OrderTracker)
-			{
-				FString PrestreamFilename = GetPrestreamPackageLinkerName(*Dependency.Key.ToString());
-				if (PrestreamFilename.Len() > 0)
-				{
-					HintFutureRead(*PrestreamFilename); // should set the priority on the IO requests here
-			}
-			}*/
-#endif
 			for (auto& Dependency : OrderTracker)
 			{
 				if (Dependency.Value < OrderTracker.Num() - 1) // we are already handling this one
@@ -566,13 +491,13 @@ void FAsyncLoadingThread::ProcessAsyncPackageRequest(FAsyncPackageDesc* InReques
 							const int32 RequestID = GPackageRequestID.Increment();
 							FAsyncLoadingThread::Get().AddPendingRequest(RequestID);
 					FAsyncPackageDesc DependencyPackageRequest(RequestID, Dependency.Key, NAME_None, FGuid(), FLoadPackageAsyncDelegate(), InRequest->PackageFlags, INDEX_NONE, InRequest->Priority);
-					ProcessAsyncPackageRequest(&DependencyPackageRequest, InRootPackage, nullptr);
+					ProcessAsyncPackageRequest(&DependencyPackageRequest, InRootPackage, nullptr, FlushTree);
 				}
 			}
 		}
 #endif
 		// Add to queue according to priority.
-		InsertPackage(Package, EAsyncPackageInsertMode::InsertAfterMatchingPriorities);
+		InsertPackage(Package, false, EAsyncPackageInsertMode::InsertAfterMatchingPriorities);
 
 		// For all other cases this is handled in FindExistingPackageAndAddCompletionCallback
 		const int32 QueuedPackagesCount = QueuedPackagesCounter.Decrement();
@@ -580,7 +505,7 @@ void FAsyncLoadingThread::ProcessAsyncPackageRequest(FAsyncPackageDesc* InReques
 	}
 }
 
-int32 FAsyncLoadingThread::CreateAsyncPackagesFromQueue()
+int32 FAsyncLoadingThread::CreateAsyncPackagesFromQueue(bool bUseTimeLimit, bool bUseFullTimeLimit, float TimeLimit, FFlushTree* FlushTree)
 {
 	SCOPE_CYCLE_COUNTER(STAT_FAsyncPackage_CreateAsyncPackagesFromQueue);
 	SCOPED_LOADTIMER(CreateAsyncPackagesFromQueueTime);
@@ -590,136 +515,3145 @@ int32 FAsyncLoadingThread::CreateAsyncPackagesFromQueue()
 	int32 NumCreated = 0;
 	checkSlow(IsInAsyncLoadThread());
 
+	int32 TimeSliceGranularity = 1; // do 4 packages at a time
+
+	if (!bUseTimeLimit)
+	{
+		TimeSliceGranularity = MAX_int32; // no point in taking small steps
+	}
+
 	TArray<FAsyncPackageDesc*> QueueCopy;
+	double TickStartTime = FPlatformTime::Seconds();
+	do
 	{
+		{
 #if THREADSAFE_UOBJECTS
-		FScopeLock QueueLock(&QueueCritical);
+			FScopeLock QueueLock(&QueueCritical);
 #endif
-		QueueCopy = QueuedPackages;
-		QueuedPackages.Reset();
+			QueueCopy.Reset();
+			QueueCopy.Reserve(FMath::Min(TimeSliceGranularity, QueuedPackages.Num()));
+
+			int32 NumCopied = 0;
+
+			for (FAsyncPackageDesc* PackageRequest : QueuedPackages)
+			{
+				if (NumCopied < TimeSliceGranularity)
+				{
+					NumCopied++;
+					QueueCopy.Add(PackageRequest);
+				}
+				else
+				{
+					break;
+				}
+			}
+			if (NumCopied)
+			{
+				QueuedPackages.RemoveAt(0, NumCopied, false);
+			}
+			else
+			{
+				break;
+			}
+		}
+
+		if (QueueCopy.Num() > 0)
+		{
+			IAssetRegistryInterface* AssetRegistry = nullptr;
+
+			if (GPreloadPackageDependencies && IsPlatformFileCompatibleWithDependencyPreloading())
+			{
+				AssetRegistry = IAssetRegistryInterface::GetPtr();
+			}
+
+			double Timer = 0;
+			{
+				SCOPE_SECONDS_COUNTER(Timer);
+				for (FAsyncPackageDesc* PackageRequest : QueueCopy)
+				{
+					ProcessAsyncPackageRequest(PackageRequest, nullptr, AssetRegistry, FlushTree);
+					delete PackageRequest;
+				}
+			}
+			UE_LOG(LogStreaming, Verbose, TEXT("Async package requests inserted in %fms"), Timer * 1000.0);
+		}
+
+		NumCreated += QueueCopy.Num();
+	} while(!IsTimeLimitExceeded(TickStartTime, bUseTimeLimit, TimeLimit, TEXT("CreateAsyncPackagesFromQueue")));
+
+	return NumCreated;
+}
+
+#if USE_EVENT_DRIVEN_ASYNC_LOAD
+
+static FThreadSafeCounter AsyncPackageSerialNumber;
+
+FName FUnsafeWeakAsyncPackagePtr::HumanReadableStringForDebugging() const
+{
+	return Package ? Package->GetPackageName() : FName();
+}
+
+FWeakAsyncPackagePtr::FWeakAsyncPackagePtr(FAsyncPackage* Package)
+	: SerialNumber(0)
+{
+	if (Package)
+	{
+		PackageName = Package->GetPackageName();
+		SerialNumber = Package->SerialNumber;
+	}
+}
+
+FAsyncPackage& FWeakAsyncPackagePtr::GetPackage() const
+{
+	FAsyncPackage* Result = FAsyncLoadingThread::Get().GetPackage(*this);
+	check(Result);
+	return *Result;
+}
+
+FString FAsyncPackage::GetDebuggingPath(FPackageIndex Idx)
+{
+	if (!Linker)
+	{
+		return TEXT("Null linker");
+	}
+	FString Details;
+	FString Prefix;
+	if (Idx.IsExport() && Linker->LinkerRoot)
+	{
+		Prefix = Linker->LinkerRoot->GetName();
+	}
+	while (!Idx.IsNull())
+	{
+		FObjectResource& Res = Linker->ImpExp(Idx);
+		Details = Res.ObjectName.ToString() / Details;
+		Idx = Res.OuterIndex;
+	}
+	return Prefix / Details;
+}
+
+FString FEventLoadNodePtr::HumanReadableStringForDebugging() const
+{
+	const TCHAR* NodeName = TEXT("Unknown");
+	FString Details;
+
+	FAsyncPackage* Pkg = &WaitingPackage.GetPackage();
+	if (ImportOrExportIndex.IsNull())
+	{
+		switch (Phase)
+		{
+		case EEventLoadNode::Package_LoadSummary:
+			NodeName = TEXT("Package_LoadSummary");
+			break;
+		case EEventLoadNode::Package_SetupImports:
+			NodeName = TEXT("Package_SetupImports");
+			break;
+		case EEventLoadNode::Package_ExportsSerialized:
+			NodeName = TEXT("Package_ExportsSerialized");
+			break;
+		default:
+			check(0);
+		}
+	}
+	else
+	{
+		switch (Phase)
+		{
+		case EEventLoadNode::ImportOrExport_Create:
+			if (ImportOrExportIndex.IsImport())
+			{
+				NodeName = TEXT("Import_Create");
+			}
+			else
+			{
+				NodeName = TEXT("Export_Create");
+			}
+			break;
+		case EEventLoadNode::Export_StartIO:
+			NodeName = TEXT("Export_StartIO");
+			break;
+		case EEventLoadNode::ImportOrExport_Serialize:
+			if (ImportOrExportIndex.IsImport())
+			{
+				NodeName = TEXT("Import_Serialize");
+			}
+			else
+			{
+				NodeName = TEXT("Export_Serialize");
+			}
+			break;
+		default:
+			check(0);
+		}
+		Details = TEXT("Package or Linker Not Found");
+
+		if (Pkg)
+		{
+			Details = Pkg->GetDebuggingPath(ImportOrExportIndex);
+		}
 	}
+	return FString::Printf(TEXT("%s %d %s   %s"), *WaitingPackage.HumanReadableStringForDebugging().ToString(), ImportOrExportIndex.ForDebugging(), NodeName, *Details);
+}
+
+void FEventLoadNodeArray::Init(int32 InNumImports, int32 InNumExports)
+{
+	check(InNumExports && !NumExports && TotalNumberOfNodesAdded <= int32(EEventLoadNode::Package_NumPhases) && !TotalNumberOfImportExportNodes);
+	NumImports = InNumImports;
+	NumExports = InNumExports;
+	OffsetToImports = 0;
+	OffsetToExports = OffsetToImports + NumImports * int32(EEventLoadNode::Import_NumPhases);
+	TotalNumberOfImportExportNodes = OffsetToExports + NumExports * int32(EEventLoadNode::Export_NumPhases);
+	check(TotalNumberOfImportExportNodes);
+	Array = new FEventLoadNode[TotalNumberOfImportExportNodes];
+}
+
+void FEventLoadNodeArray::Shutdown()
+{
+	check(!TotalNumberOfNodesAdded);
+	delete[] Array;
+	Array = nullptr;
+}
+void FEventLoadNodeArray::GetAddedNodes(TArray<FEventLoadNodePtr>& OutAddedNodes, FAsyncPackage* Owner)
+{
+	if (TotalNumberOfNodesAdded)
+	{
+		FEventLoadNodePtr Node;
+		Node.WaitingPackage = FCheckedWeakAsyncPackagePtr(Owner);
+		for (int32 Index = 0; Index < int32(EEventLoadNode::Package_NumPhases); Index++)
+		{
+			Node.Phase = EEventLoadNode(Index);
+			FEventLoadNode& NodeRef(PtrToNode(Node));
+			if (NodeRef.bAddedToGraph)
+			{
+				OutAddedNodes.Add(Node);
+			}
+		}
+		for (int32 ImportIndex = 0; ImportIndex < NumImports; ImportIndex++)
+		{
+			Node.ImportOrExportIndex = FPackageIndex::FromImport(ImportIndex);
+			for (int32 Index = 0; Index < int32(EEventLoadNode::Import_NumPhases); Index++)
+			{
+				Node.Phase = EEventLoadNode(Index);
+				FEventLoadNode& NodeRef(PtrToNode(Node));
+				if (NodeRef.bAddedToGraph)
+				{
+					OutAddedNodes.Add(Node);
+				}
+			}
+
+		}
+		for (int32 ExportIndex = 0; ExportIndex < NumExports; ExportIndex++)
+		{
+			Node.ImportOrExportIndex = FPackageIndex::FromExport(ExportIndex);
+			for (int32 Index = 0; Index < int32(EEventLoadNode::Export_NumPhases); Index++)
+			{
+				Node.Phase = EEventLoadNode(Index);
+				FEventLoadNode& NodeRef(PtrToNode(Node));
+				if (NodeRef.bAddedToGraph)
+				{
+					OutAddedNodes.Add(Node);
+				}
+			}
+
+		}
+	}
+}
+
+FORCEINLINE FEventLoadNodeArray& FEventLoadGraph::GetArray(FEventLoadNodePtr& Node)
+{
+	return Node.WaitingPackage.GetPackage().EventNodeArray;
+}
+
+FORCEINLINE FEventLoadNode& FEventLoadGraph::GetNode(FEventLoadNodePtr& NodeToGet)
+{
+	return GetArray(NodeToGet).GetNode(NodeToGet);
+}
+
+
+
+void FEventLoadGraph::AddNode(FEventLoadNodePtr& NewNode, bool bHoldForLater, int32 NumImplicitPrereqs)
+{
+	SCOPED_LOADTIMER_CNT(Graph_AddNode);
+
+	FEventLoadNodeArray& Array = GetArray(NewNode);
+	if (Array.AddNode(NewNode))
+	{
+		check(!PackagesWithNodes.Contains(NewNode.WaitingPackage));
+		PackagesWithNodes.Add(NewNode.WaitingPackage);
+	}
+	int32 NumAddPrereq = (bHoldForLater ? 1 : 0) + NumImplicitPrereqs;
+	if (NumAddPrereq)
+	{
+		Array.GetNode(NewNode).NumPrerequistes += NumAddPrereq;
+	}
+}
+
+void FEventLoadGraph::AddArc(FEventLoadNodePtr& PrereqisitePtr, FEventLoadNodePtr& DependentPtr)
+{
+	SCOPED_LOADTIMER_CNT(Graph_AddArc);
+	FEventLoadNode* PrereqisiteNode(&GetNode(PrereqisitePtr));
+	FEventLoadNode* DependentNode(&GetNode(DependentPtr));
+	check(!DependentNode->bFired);
+	DependentNode->NumPrerequistes++;
+	PrereqisiteNode->NodesWaitingForMe.Add(DependentPtr);
+}
+
+void FEventLoadGraph::RemoveNode(FEventLoadNodePtr& InNodeToRemove)
+{
+	FEventLoadNodePtr NodeToRemove(InNodeToRemove); // make a copy of this so we don't end up destroying in indirectly
+	SCOPED_LOADTIMER_CNT(Graph_RemoveNode);
+	check(FAsyncLoadingThread::IsInAsyncLoadThread());
+	static TArray<int32> IndicesToFire;
+	check(!IndicesToFire.Num());
+
+	FEventLoadNode::TNodesWaitingForMeArray NodesToFire;
+	{
+		FEventLoadNodeArray& Array = GetArray(NodeToRemove);
+		FEventLoadNode* PrereqisiteNode(&Array.GetNode(NodeToRemove));
+		check(PrereqisiteNode->bFired);
+		check(PrereqisiteNode->NumPrerequistes == 0);
+		Swap(NodesToFire, PrereqisiteNode->NodesWaitingForMe);
+
+
+		for (FEventLoadNodePtr& Target : NodesToFire)
+		{
+			FEventLoadNode* DependentNode(&GetNode(Target));
+			check(DependentNode->NumPrerequistes > 0);
+			if (!--DependentNode->NumPrerequistes)
+			{
+				DependentNode->bFired = true;
+				IndicesToFire.Add(&Target - &NodesToFire[0]);
+			}
+		}
+		if (Array.RemoveNode(NodeToRemove))
+		{
+			PackagesWithNodes.Remove(NodeToRemove.WaitingPackage);
+			Array.Shutdown();
+		}
+	}
+
+#if 0
+
+	if (IndicesToFire.Num() > 2) // sorting is not required, don't bother if there isn't much
+	{
+		IndicesToFire.Sort(
+			[&NodesToFire](int32 A, int32 B) -> bool
+			{
+				return NodesToFire[A] < NodesToFire[B];
+			}
+		);
+	}
+
+	int32 CurrentWeakPtrSerialNumber = -1;
+	FAsyncPackage* CurrentTarget = nullptr;
+	for (int32 Index : IndicesToFire)
+	{
+		FEventLoadNodePtr& Target = NodesToFire[Index];
+		if (CurrentWeakPtrSerialNumber != Target.WaitingPackage.SerialNumber)
+		{
+			CurrentWeakPtrSerialNumber = Target.WaitingPackage.SerialNumber;
+			check(CurrentWeakPtrSerialNumber);
+			CurrentTarget = FAsyncLoadingThread::Get().GetPackage(Target.WaitingPackage);
+		}
+		check(CurrentTarget && CurrentTarget->SerialNumber == Target.WaitingPackage.SerialNumber);
+		SCOPED_LOADTIMER_CNT(Graph_RemoveNodeFire);
+		CurrentTarget->FireNode(Target);
+	}
+#else
+#if USE_IMPLICIT_ARCS
+	int32 NumImplicitArcs = NodeToRemove.NumImplicitArcs();
+	if (NumImplicitArcs)
+	{
+		check(NumImplicitArcs == 1); // would need different code otherwise
+		FEventLoadNodePtr Target = NodeToRemove.GetImplicitArc();
+		FEventLoadNode* DependentNode(&GetNode(Target));
+		check(DependentNode->NumPrerequistes > 0);
+		if (!--DependentNode->NumPrerequistes)
+		{
+			DependentNode->bFired = true;
+			FAsyncPackage& CurrentTarget = Target.WaitingPackage.GetPackage();
+			CurrentTarget.FireNode(Target);
+		}
+	}
+#endif
+
+	for (int32 Index : IndicesToFire)
+	{
+		FEventLoadNodePtr& Target = NodesToFire[Index];
+		FAsyncPackage& CurrentTarget = Target.WaitingPackage.GetPackage();
+#if VERIFY_WEAK_ASYNC_PACKAGE_PTRS
+		check(CurrentTarget.SerialNumber == Target.WaitingPackage.SerialNumber);
+#else
+		check(CurrentTarget.SerialNumber);
+#endif
+		SCOPED_LOADTIMER_CNT(Graph_RemoveNodeFire);
+		CurrentTarget.FireNode(Target);
+	}
+#endif
+	IndicesToFire.Reset();
+}
+
+
+void FEventLoadGraph::NodeWillBeFiredExternally(FEventLoadNodePtr& NodeThatWasFired)
+{
+	SCOPED_LOADTIMER_CNT(Graph_Misc);
+	FEventLoadNode* DependentNode(&GetNode(NodeThatWasFired));
+	check(!DependentNode->bFired);
+	DependentNode->bFired = true;
+}
+
+void FEventLoadGraph::DoneAddingPrerequistesFireIfNone(FEventLoadNodePtr& NewNode, bool bWasHeldForLater)
+{
+	SCOPED_LOADTIMER_CNT(Graph_DoneAddingPrerequistesFireIfNone);
+	FEventLoadNode* DependentNode(&GetNode(NewNode));
+	check(!DependentNode->bFired);
+	if (bWasHeldForLater)
+	{
+		check(DependentNode->NumPrerequistes > 0);
+		DependentNode->NumPrerequistes--;
+	}
+	if (!DependentNode->NumPrerequistes)
+	{
+		DependentNode->bFired = true;
+		FAsyncPackage& CurrentTarget = NewNode.WaitingPackage.GetPackage();
+		SCOPED_LOADTIMER_CNT(Graph_DoneAddingPrerequistesFireIfNoneFire);
+		CurrentTarget.FireNode(NewNode);
+	}
+}
+
+#if !UE_BUILD_SHIPPING
+bool FEventLoadGraph::CheckForCyclesInner(const TMultiMap<FEventLoadNodePtr, FEventLoadNodePtr>& Arcs, TSet<FEventLoadNodePtr>& Visited, TSet<FEventLoadNodePtr>& Stack, const FEventLoadNodePtr& Visit)
+{
+	bool bResult = false;
+	if (Stack.Contains(Visit))
+	{
+		bResult = true;
+	}
+	else
+	{
+		bool bWasAlreadyTested = false;
+		Visited.Add(Visit, &bWasAlreadyTested);
+		if (!bWasAlreadyTested)
+		{
+			Stack.Add(Visit);
+			for (auto It = Arcs.CreateConstKeyIterator(Visit); !bResult && It; ++It)
+			{
+				bResult = CheckForCyclesInner(Arcs, Visited, Stack, It.Value());
+			}
+			Stack.Remove(Visit);
+		}
+	}
+	UE_CLOG(bResult, LogStreaming, Error, TEXT("Cycle Node %s"), *Visit.HumanReadableStringForDebugging());
+	return bResult;
+}
+#endif
+
+void FEventLoadGraph::CheckForCycles()
+{
+#if !UE_BUILD_SHIPPING && !UE_BUILD_TEST
+	TMultiMap<FEventLoadNodePtr, FEventLoadNodePtr> Arcs;
+	TArray<FEventLoadNodePtr> AddedNodes;
+	for (FCheckedWeakAsyncPackagePtr &Ptr : PackagesWithNodes)
+	{
+		FAsyncPackage* Pkg = &Ptr.GetPackage();
+		Pkg->EventNodeArray.GetAddedNodes(AddedNodes, Pkg);
+	}
+	for (FEventLoadNodePtr &Ptr : AddedNodes)
+	{
+		FEventLoadNode& Node(GetNode(Ptr));
+
+		if (!Node.NumPrerequistes)
+		{
+			UE_LOG(LogStreaming, Fatal, TEXT("Node %s has zero prerequisites, but has not been queued."), *Ptr.HumanReadableStringForDebugging());
+		}
+		for (FEventLoadNodePtr Other : Node.NodesWaitingForMe)
+		{
+			Arcs.Add(Other, Ptr);
+		}
+#if USE_IMPLICIT_ARCS
+		int32 NumImplicitArcs = Ptr.NumImplicitArcs();
+		if (NumImplicitArcs)
+		{
+			check(NumImplicitArcs == 1); // would need different code otherwise
+			FEventLoadNodePtr Target = Ptr.GetImplicitArc();
+			Arcs.Add(Target, Ptr);
+		}
+#endif
+	}
+	TSet<FEventLoadNodePtr> Visited;
+	TSet<FEventLoadNodePtr> Stack;
+	for (FEventLoadNodePtr &Ptr : AddedNodes)
+	{
+		if (CheckForCyclesInner(Arcs, Visited, Stack, Ptr))
+		{
+			UE_LOG(LogStreaming, Fatal, TEXT("Async loading event graph contained a cycle, see above."));
+		}
+	}
+	if (AddedNodes.Num())
+	{
+		UE_LOG(LogStreaming, Fatal, TEXT("No outstanding IO, no nodes in the queue, yet we still have 'AddedNodes' in the graph."));
+	}
+#endif
+	if (PackagesWithNodes.Num())
+	{
+		UE_LOG(LogStreaming, Fatal, TEXT("No outstanding IO, no nodes in the queue, yet we still have 'PackagesWithNodes' in the graph."));
+	}
+}
+
+//@todoio this won't work with UT etc, need *_API or a real api to work with the pak precacher.
+extern bool GPakCache_AcceptPrecacheRequests;
+
+struct FPrecacheCallbackHandler
+{
+	FAsyncFileCallBack PrecacheCallBack;
+
+	FCriticalSection IncomingLock;
+	TArray<IAsyncReadRequest *> Incoming;
+	TArray<FWeakAsyncPackagePtr> IncomingSummaries;
+	bool bFireIncomingEvent; 
+	FEvent* PermanentIncomingEvent;
+
+	TMap<IAsyncReadRequest *, FWeakAsyncPackagePtr> WaitingPackages;
+	TSet<FWeakAsyncPackagePtr> WaitingSummaries;
+
+	int64 UnprocessedMemUsed;
+
+	FPrecacheCallbackHandler()
+		: bFireIncomingEvent(false)
+		, PermanentIncomingEvent(nullptr)
+		, UnprocessedMemUsed(0)
+	{
+		PrecacheCallBack =
+			[this](bool bWasCanceled, IAsyncReadRequest* Request)
+		{
+			RequestComplete(bWasCanceled, Request);
+		};
+	}
+	~FPrecacheCallbackHandler()
+	{
+		FScopeLock Lock(&IncomingLock);
+		check(!bFireIncomingEvent);
+		check(!Incoming.Num() && !IncomingSummaries.Num() && !WaitingPackages.Num() && !WaitingSummaries.Num());
+		if (PermanentIncomingEvent)
+		{
+			FPlatformProcess::ReturnSynchEventToPool(PermanentIncomingEvent);
+			PermanentIncomingEvent = nullptr;
+		}
+	}
+
+	FAsyncFileCallBack* GetCompletionCallback()
+	{
+		return &PrecacheCallBack;
+	}
+
+	void RequestComplete(bool bWasCanceled, IAsyncReadRequest* Precache)
+	{
+		check(!bWasCanceled); // not handled yet
+		FScopeLock Lock(&IncomingLock);
+		Incoming.Add(Precache);
+		if (bFireIncomingEvent)
+		{
+			bFireIncomingEvent = false; // only trigger once
+			PermanentIncomingEvent->Trigger();
+		}
+		else if (Incoming.Num() > 100)
+		{
+			if (GPakCache_AcceptPrecacheRequests)
+			{
+				UE_LOG(LogStreaming, Log, TEXT("Throttling off (async)"));
+				GPakCache_AcceptPrecacheRequests = false;
+			}
+		}
+	}
+
+	void SummaryComplete(const FWeakAsyncPackagePtr& Pkg)
+	{
+		FScopeLock Lock(&IncomingLock);
+		IncomingSummaries.Add(Pkg);
+		if (bFireIncomingEvent)
+		{
+			bFireIncomingEvent = false; // only trigger once
+			PermanentIncomingEvent->Trigger();
+		}
+	}
+
+	bool ProcessIncoming()
+	{
+		TArray<IAsyncReadRequest *> LocalIncoming;
+		TArray<FWeakAsyncPackagePtr> LocalIncomingSummaries;
+		{
+			FScopeLock Lock(&IncomingLock);
+			Swap(LocalIncoming, Incoming);
+			Swap(LocalIncomingSummaries, IncomingSummaries);
+		}
+		for (IAsyncReadRequest * Req : LocalIncoming)
+		{
+			check(Req);
+			FWeakAsyncPackagePtr Found = WaitingPackages.FindAndRemoveChecked(Req);
+			FAsyncPackage* Pkg = FAsyncLoadingThread::Get().GetPackage(Found);
+			check(Pkg);
+			UnprocessedMemUsed += Pkg->PrecacheRequestReady(Req);
+		}
+		for (FWeakAsyncPackagePtr& Sum : LocalIncomingSummaries)
+		{
+			FAsyncLoadingThread* LocalAsyncLoadingThread = &FAsyncLoadingThread::Get();
+			LocalAsyncLoadingThread->QueueEvent_FinishLinker(Sum);
+			check(WaitingSummaries.Contains(Sum));
+			WaitingSummaries.Remove(Sum);
+		}
+		if (LocalIncoming.Num())
+		{
+			CheckThottleIOState();
+		}
+		return LocalIncoming.Num() || LocalIncomingSummaries.Num();
+	}
+
+	bool AnyIOOutstanding()
+	{
+		return WaitingPackages.Num() || WaitingSummaries.Num();
+	}
+
+	bool WaitForIO(float SecondsToWait = 0.0f)
+	{
+		check(AnyIOOutstanding());
+		check(SecondsToWait >= 0.0f);
+		{
+			FScopeLock Lock(&IncomingLock);
+			if (Incoming.Num() || IncomingSummaries.Num())
+			{
+				return true;
+			}
+			if (!PermanentIncomingEvent)
+			{
+				PermanentIncomingEvent = FPlatformProcess::GetSynchEventFromPool();
+			}
+			bFireIncomingEvent = true;
+		}
+		if (SecondsToWait == 0.0f)
+		{
+			PermanentIncomingEvent->Wait();
+			check(!bFireIncomingEvent);
+			return true;
+		}
+		uint32 Ms = FMath::Max<uint32>(uint32(SecondsToWait * 1000.0f), 1);
+		if (PermanentIncomingEvent->Wait(Ms))
+		{
+			check(!bFireIncomingEvent);
+			return true;
+		}
+		FScopeLock Lock(&IncomingLock);
+		if (bFireIncomingEvent)
+		{
+			// nobody triggered it
+			bFireIncomingEvent = false;
+			return false;
+		}
+		else
+		{
+			// We timed out and then it was triggered, so we have data and we need to reset the event
+			PermanentIncomingEvent->Reset();
+			return true;
+		}
+	}
+
+	void RegisterNewPrecacheRequest(IAsyncReadRequest* Precache, FAsyncPackage* Package)
+	{
+		WaitingPackages.Add(Precache, FWeakAsyncPackagePtr(Package));
+	}
+	void RegisterNewSummaryRequest(FAsyncPackage* Package)
+	{
+		WaitingSummaries.Add(FWeakAsyncPackagePtr(Package));
+	}
+
+	void CheckThottleIOState()
+	{
+		if (GPakCache_AcceptPrecacheRequests && UnprocessedMemUsed > 30 * 1024 * 1024)
+		{
+			UE_LOG(LogStreaming, Log, TEXT("Throttling off"));
+			GPakCache_AcceptPrecacheRequests = false;
+		}
+		else if (!GPakCache_AcceptPrecacheRequests && UnprocessedMemUsed <= 30 * 1024 * 1024)
+		{
+			UE_LOG(LogStreaming, Log, TEXT("Throttling on"));
+			GPakCache_AcceptPrecacheRequests = true;
+		}
+	}
+
+	void FinishRequest(int64 Size)
+	{
+		UnprocessedMemUsed -= Size;
+		check(UnprocessedMemUsed >= 0);
+		CheckThottleIOState();
+	}
+};
+
+static FPrecacheCallbackHandler GPrecacheCallbackHandler;
+
+int32 GRandomizeLoadOrder = 0;
+static FAutoConsoleVariableRef CVarRandomizeLoadOrder(
+	TEXT("s.RandomizeLoadOrder"),
+	GRandomizeLoadOrder,
+	TEXT("If > 0, will randomize the load order of pending packages using this seed instead of using the most efficient order. This can be used to find bugs."),
+	ECVF_Default
+);
+
+static int32 GetRandomSerialNumber(int32 MaxVal = MAX_int32)
+{
+	static FRandomStream RandomStream(GRandomizeLoadOrder);
+
+	return RandomStream.RandHelper(MaxVal);
+}
+
+void FImportOrImportIndexArray::HeapPop(int32& OutItem, bool bAllowShrinking)
+{
+	if (GRandomizeLoadOrder)
+	{
+		int32 Index = FMath::Clamp<int32>(GetRandomSerialNumber(Num() - 1), 0, Num() - 1);
+		OutItem = (*this)[Index];
+		RemoveAt(Index, 1, false);
+		return;
+	}
+	TArray<int32>::HeapPop(OutItem, bAllowShrinking);
+}
+
+
+FScopedAsyncPackageEvent::FScopedAsyncPackageEvent(FAsyncPackage* InPackage)
+	:Package(InPackage)
+{
+	check(Package);
+
+	// Update the thread context with the current package. This is used by NotifyConstructedDuringAsyncLoading.
+	FUObjectThreadContext& ThreadContext = FUObjectThreadContext::Get();
+	PreviousPackage = ThreadContext.AsyncPackage;
+	ThreadContext.AsyncPackage = Package;
+
+	Package->BeginAsyncLoad();
+	FExclusiveLoadPackageTimeTracker::PushLoadPackage(Package->Desc.NameToLoad);
+}
+
+FScopedAsyncPackageEvent::~FScopedAsyncPackageEvent()
+{
+	FExclusiveLoadPackageTimeTracker::PopLoadPackage(Package->Linker ? Package->Linker->LinkerRoot : nullptr);
+	Package->EndAsyncLoad();
+	Package->LastObjectWorkWasPerformedOn = nullptr;
+	Package->LastTypeOfWorkPerformed = nullptr;
+
+	// Restore the package from the outer scope
+	FUObjectThreadContext& ThreadContext = FUObjectThreadContext::Get();
+	ThreadContext.AsyncPackage = PreviousPackage;
+}
+
+FORCENOINLINE static bool CheckForFileOpenLogCommandLine()
+{
+	if (FParse::Param(FCommandLine::Get(), TEXT("RandomizeLoadOrder")))
+	{
+		GRandomizeLoadOrder = 1;
+	}
+
+	return FParse::Param(FCommandLine::Get(), TEXT("FileOpenLog"));
+}
+
+FORCEINLINE static bool FileOpenLogActive()
+{
+#if 1
+	static bool bDoingLoadOrder = CheckForFileOpenLogCommandLine();
+	return bDoingLoadOrder;
+#else
+	return false;
+#endif
+}
+
+FORCEINLINE static int32 SummariesFirstPriorityForFileOpenLog(int32 UsualPriority, bool bSummaryRelated)
+{
+	return FileOpenLogActive() ? (bSummaryRelated ? 1 : 0) : UsualPriority;
+}
+
+void FAsyncLoadingThread::QueueEvent_CreateLinker(FAsyncPackage* Package, int32 EventSystemPriority)
+{
+	check(Package);
+	Package->AddNode(EEventLoadNode::Package_LoadSummary);
+	FWeakAsyncPackagePtr WeakPtr(Package);
+	EventQueue.AddAsyncEvent(Package->GetPriority(), GRandomizeLoadOrder ? GetRandomSerialNumber() : Package->SerialNumber, SummariesFirstPriorityForFileOpenLog(EventSystemPriority, true),
+		TFunction<void(FAsyncLoadEventArgs& Args)>(
+			[WeakPtr, this](FAsyncLoadEventArgs& Args)
+			{
+				FAsyncPackage* Pkg = GetPackage(WeakPtr);
+				check(Pkg);
+				if (Pkg)
+				{
+					Pkg->SetTimeLimit(Args, TEXT("Create Linker"));
+					Pkg->Event_CreateLinker();
+					Args.OutLastObjectWorkWasPerformedOn = Pkg->GetLinkerRoot();
+				}
+			}
+	));
+}
+
+void FAsyncPackage::Event_CreateLinker()
+{
+	// Keep track of time when we start loading.
+	if (LoadStartTime == 0.0)
+	{
+		double Now = FPlatformTime::Seconds();
+		LoadStartTime = Now;
+
+		// If we are a dependency of another package, we need to tell that package when its first dependent started loading,
+		// otherwise because that package loads last it'll not include the entire load time of all its dependencies
+		if (DependencyRootPackage)
+		{
+			// Only the first dependent needs to register the start time
+			if (DependencyRootPackage->GetLoadStartTime() == 0.0)
+			{
+				DependencyRootPackage->LoadStartTime = Now;
+			}
+		}
+	}
+	FScopedAsyncPackageEvent Scope(this);
+	SCOPED_LOADTIMER(Package_CreateLinker);
+	check(!Linker);
+	NodeWillBeFiredExternally(EEventLoadNode::Package_LoadSummary);
+	CreateLinker();
+	check(AsyncPackageLoadingState == EAsyncPackageLoadingState::NewPackage);
+	if (Linker)
+	{
+		AsyncPackageLoadingState = EAsyncPackageLoadingState::WaitingForSummary;
+		Linker->bLockoutLegacyOperations = true;
+	}
+	else
+	{
+		RemoveNode(EEventLoadNode::Package_LoadSummary);
+		EventDrivenLoadingComplete();
+		AsyncPackageLoadingState = EAsyncPackageLoadingState::PostLoad_Etc;
+		check(!FAsyncLoadingThread::Get().AsyncPackagesReadyForTick.Contains(this));
+		FAsyncLoadingThread::Get().AsyncPackagesReadyForTick.Add(this);
+	}
+}
+
+
+void FAsyncLoadingThread::QueueEvent_FinishLinker(FWeakAsyncPackagePtr WeakPtr, int32 EventSystemPriority)
+{
+	FAsyncPackage* Pkg = GetPackage(WeakPtr);
+	if (Pkg)
+	{
+		EventQueue.AddAsyncEvent(Pkg->GetPriority(), GRandomizeLoadOrder ? GetRandomSerialNumber() : Pkg->SerialNumber, SummariesFirstPriorityForFileOpenLog(EventSystemPriority, true),
+			TFunction<void(FAsyncLoadEventArgs& Args)>(
+				[WeakPtr, this](FAsyncLoadEventArgs& Args)
+		{
+			FAsyncPackage* PkgInner = GetPackage(WeakPtr);
+			check(PkgInner);
+			if (PkgInner)
+			{
+				PkgInner->SetTimeLimit(Args, TEXT("Finish Linker"));
+				PkgInner->Event_FinishLinker();
+			}
+		}));
+	}
+
+}
+
+
+void FAsyncPackage::Event_FinishLinker()
+{
+	FScopedAsyncPackageEvent Scope(this);
+	SCOPED_LOADTIMER(Package_FinishLinker);
+	EAsyncPackageState::Type Result = FinishLinker();
+	if (Result == EAsyncPackageState::TimeOut && !bLoadHasFailed)
+	{
+		AsyncLoadingThread.QueueEvent_FinishLinker(FWeakAsyncPackagePtr(this), FAsyncLoadEvent::EventSystemPriority_MAX);
+		return;
+	}
+
+	if (!bLoadHasFailed)
+	{
+		check(Linker && Linker->HasFinishedInitialization());
+
+		// Add nodes for all imports and exports.
+
+		{
+			LastTypeOfWorkPerformed = TEXT("ImportAddNode");
+			int32 NumImplicitForImportExport = 0;
+#if USE_IMPLICIT_ARCS
+			NumImplicitForImportExport = 1;
+#endif
+
+			if (ImportAddNodeIndex == 0 && ExportAddNodeIndex == 0) // one time only
+			{
+				int32 NumImplicit = 0;
+#if USE_IMPLICIT_ARCS
+				NumImplicit = Linker->ImportMap.Num() + Linker->ExportMap.Num();
+#endif
+
+				FEventLoadNodePtr MyDoneNode = AddNode(EEventLoadNode::Package_ExportsSerialized, FPackageIndex(), false, NumImplicit);
+
+				// There are things waiting to link imports. I need to not finish until those links are made
+				// Package_ExportsSerialized is actually earlier than we need. We just need to make sure the linker is not destroyed before the other packages link
+				for (FCheckedWeakAsyncPackagePtr& Waiter : PackagesWaitingToLinkImports)
+				{
+					FEventLoadNodePtr Prereq;
+					Prereq.WaitingPackage = Waiter;
+					Prereq.Phase = EEventLoadNode::Package_SetupImports;
+					AddArc(Prereq, MyDoneNode);
+				}
+				PackagesWaitingToLinkImports.Empty();
+
+				AddNode(EEventLoadNode::Package_SetupImports, FPackageIndex(), true);
+				EventNodeArray.Init(Linker->ImportMap.Num(), Linker->ExportMap.Num());
+			}
+			FEventLoadNodePtr MyDependentExportsSerializedNode;
+			MyDependentExportsSerializedNode.WaitingPackage = FCheckedWeakAsyncPackagePtr(this);
+			MyDependentExportsSerializedNode.Phase = EEventLoadNode::Package_ExportsSerialized;
+
+			for (int32 LocalImportIndex = ImportAddNodeIndex; LocalImportIndex < Linker->ImportMap.Num(); LocalImportIndex++)
+			{
+				//optimization: could avoid creating all of these nodes, in the common event that they are already done
+				FEventLoadNodePtr MyDependentCreateNode = AddNode(EEventLoadNode::ImportOrExport_Create, FPackageIndex::FromImport(LocalImportIndex));
+				FEventLoadNodePtr MyDependentSerializeNode = AddNode(EEventLoadNode::ImportOrExport_Serialize, FPackageIndex::FromImport(LocalImportIndex), false, NumImplicitForImportExport);
+
+#if !USE_IMPLICIT_ARCS
+				// Can't consider this import serialized until we hook it up after creation
+				AddArc(MyDependentCreateNode, MyDependentSerializeNode);
+
+
+				// Can't consider the package done with event driven loading until all imports are serialized
+				AddArc(MyDependentSerializeNode, MyDependentExportsSerializedNode);
+#endif
+				ImportAddNodeIndex = LocalImportIndex + 1;
+				if (LocalImportIndex % 50 == 0 && IsTimeLimitExceeded())
+				{
+					AsyncLoadingThread.QueueEvent_FinishLinker(FWeakAsyncPackagePtr(this), FAsyncLoadEvent::EventSystemPriority_MAX);
+					return;
+				}
+			}
+
+			LastTypeOfWorkPerformed = TEXT("ExportAddNode");
+			for (int32 LocalExportIndex = ExportAddNodeIndex; LocalExportIndex < Linker->ExportMap.Num(); LocalExportIndex++)
+			{
+				//optimization: could avoid creating all of these nodes, in the (less) common event that they are already done
+
+				FEventLoadNodePtr MyDependentCreateNode = AddNode(EEventLoadNode::ImportOrExport_Create, FPackageIndex::FromExport(LocalExportIndex));
+				FEventLoadNodePtr MyDependentIONode = AddNode(EEventLoadNode::Export_StartIO, FPackageIndex::FromExport(LocalExportIndex), false, NumImplicitForImportExport);
+				FEventLoadNodePtr MyDependentSerializeNode = AddNode(EEventLoadNode::ImportOrExport_Serialize, FPackageIndex::FromExport(LocalExportIndex), false, NumImplicitForImportExport);
+
+#if !USE_IMPLICIT_ARCS
+				// can't do the IO request until it is created
+				AddArc(MyDependentCreateNode, MyDependentIONode);
+
+				// can't serialize until the IO request is ready
+				AddArc(MyDependentIONode, MyDependentSerializeNode);
+
+				// Can't consider the package done with event driven loading until all exports are serialized
+				AddArc(MyDependentSerializeNode, MyDependentExportsSerializedNode);
+#endif
+				ExportAddNodeIndex = LocalExportIndex + 1;
+
+				if (LocalExportIndex % 30 == 0 && IsTimeLimitExceeded())
+				{
+					AsyncLoadingThread.QueueEvent_FinishLinker(FWeakAsyncPackagePtr(this), FAsyncLoadEvent::EventSystemPriority_MAX);
+					return;
+				}
+			}
+		}
+
+		check(AsyncPackageLoadingState == EAsyncPackageLoadingState::WaitingForSummary);
+		AsyncPackageLoadingState = EAsyncPackageLoadingState::StartImportPackages;
+		AsyncLoadingThread.QueueEvent_StartImportPackages(this);
+	}
+	RemoveNode(EEventLoadNode::Package_LoadSummary);
+	if (bLoadHasFailed)
+	{
+		EventDrivenLoadingComplete();
+		AsyncPackageLoadingState = EAsyncPackageLoadingState::PostLoad_Etc;
+		check(!FAsyncLoadingThread::Get().AsyncPackagesReadyForTick.Contains(this));
+		FAsyncLoadingThread::Get().AsyncPackagesReadyForTick.Add(this);
+	}
+}
+
+void FAsyncLoadingThread::QueueEvent_StartImportPackages(FAsyncPackage* Package, int32 EventSystemPriority)
+{
+	check(Package);
+	FWeakAsyncPackagePtr WeakPtr(Package);
+	EventQueue.AddAsyncEvent(Package->GetPriority(), GRandomizeLoadOrder ? GetRandomSerialNumber() : Package->SerialNumber, SummariesFirstPriorityForFileOpenLog(EventSystemPriority, true),
+		TFunction<void(FAsyncLoadEventArgs& Args)>(
+			[WeakPtr, this](FAsyncLoadEventArgs& Args)
+	{
+		FAsyncPackage* Pkg = GetPackage(WeakPtr);
+		if (Pkg)
+		{
+			Pkg->SetTimeLimit(Args, TEXT("Start Import Packages"));
+			Pkg->Event_StartImportPackages();
+		}
+	}
+	));
+}
+
+void FAsyncPackage::Event_StartImportPackages()
+{
+
+	{
+		FScopedAsyncPackageEvent Scope(this);
+		if (LoadImports_Event() == EAsyncPackageState::TimeOut)
+		{
+			AsyncLoadingThread.QueueEvent_StartImportPackages(this, FAsyncLoadEvent::EventSystemPriority_MAX); // start here next frame
+			return;
+		}
+	}
+
+	check(AsyncPackageLoadingState == EAsyncPackageLoadingState::StartImportPackages);
+	AsyncPackageLoadingState = EAsyncPackageLoadingState::WaitingForImportPackages;
+	DoneAddingPrerequistesFireIfNone(EEventLoadNode::Package_SetupImports, FPackageIndex(), true);
+}
+
+//@todoio we should sort the imports at cook time so this recursive procedure is not needed.
+FObjectImport* FAsyncPackage::FindExistingImport(int32 LocalImportIndex)
+{
+	FObjectImport* Import = &Linker->ImportMap[LocalImportIndex];
+	if (!Import->XObject && !Import->bImportSearchedFor)
+	{
+		Import->bImportSearchedFor = true;
+		if (Import->OuterIndex.IsNull())
+		{
+			Import->XObject = StaticFindObjectFast(UPackage::StaticClass(), nullptr, Import->ObjectName, true);
+			if (Import->XObject)
+			{
+				AddObjectReference(Import->XObject);
+			}
+			check(!Import->XObject || CastChecked<UPackage>(Import->XObject));
+		}
+		else
+		{
+			check(Import->OuterIndex.IsImport()); // can't see how an import could have an export as an outer
+			FObjectImport* ImportOuter = FindExistingImport(Import->OuterIndex.ToImport());
+			if (ImportOuter->XObject)
+			{
+				Import->XObject = StaticFindObjectFast(UObject::StaticClass(), ImportOuter->XObject, Import->ObjectName, false, true);
+				if (Import->XObject)
+				{
+					check(Import->XObject->GetClass()->GetFName() == Import->ClassName);
+					AddObjectReference(Import->XObject);
+				}
+			}
+		}
+	}
+	return Import;
+}
+
+EAsyncPackageState::Type FAsyncPackage::LoadImports_Event()
+{
+	SCOPE_CYCLE_COUNTER(STAT_FAsyncPackage_LoadImports);
+	LastObjectWorkWasPerformedOn = LinkerRoot;
+	LastTypeOfWorkPerformed = TEXT("loading imports event");
+
+	// GC can't run in here
+	FGCScopeGuard GCGuard;
+
+	FCheckedWeakAsyncPackagePtr WeakThis(this);
+	FEventLoadNodePtr MyDependentNode;
+	MyDependentNode.WaitingPackage = WeakThis;
+	MyDependentNode.Phase = EEventLoadNode::Package_SetupImports;
+
+	bool bDidSomething = false;
+	// Create imports.
+	while (LoadImportIndex < Linker->ImportMap.Num() && (FileOpenLogActive() || !IsTimeLimitExceeded()))
+	{
+		// Get the package for this import
+		int32 LocalImportIndex = LoadImportIndex++;
+		FObjectImport* Import = FindExistingImport(LocalImportIndex);
+		FObjectImport* OriginalImport = Import;
+
+		if (Import->XObject)
+		{
+			continue; // we already have this thing
+		}
+
+		bool bForcePackageLoad = false;
+		if (!Import->OuterIndex.IsNull())
+		{
+			// we didn't find an object, so we need to stream the package in because it might have been GC'd and we need to reload it (unless we have already done that according to bImportPackageHandled)
+			FObjectImport* ImportOutermost = Import;
+
+			// set the already handled flag as we go down..by the time we are done, they will all be handled
+			while (!ImportOutermost->bImportPackageHandled && ImportOutermost->OuterIndex.IsImport())
+			{
+				ImportOutermost->bImportPackageHandled = true;
+				ImportOutermost = &Linker->Imp(ImportOutermost->OuterIndex);
+			}
+			if (ImportOutermost->bImportPackageHandled)
+			{
+				continue;
+			}
+			check(ImportOutermost->OuterIndex.IsNull());
+			ImportOutermost->bImportPackageHandled = true;
+			bForcePackageLoad = true;
+			Import = ImportOutermost; // just do the rest of the package code, but start the async package even if we find the upackage
+		}
+		// else don't set handled because bForcePackageLoad is false, meaning we might not set the thing anyway
+
+		// @todoio: why do we need this? some UFunctions have null outer in the linker.
+		if (Import->ClassName != NAME_Package)
+		{
+			check(0);
+			continue;
+		}
+
+		// Don't try to import a package that is in an import table that we know is an invalid entry
+		if (FLinkerLoad::KnownMissingPackages.Contains(Import->ObjectName))
+		{
+			continue;
+		}
+		UPackage* ExistingPackage = nullptr;
+		FAsyncPackage* PendingPackage = nullptr;
+		if (Import->XObject)
+		{
+			ExistingPackage = CastChecked<UPackage>(Import->XObject);
+			PendingPackage = ExistingPackage->LinkerLoad ? ExistingPackage->LinkerLoad->AsyncRoot : nullptr;
+		}
+		// Our import package name is the import name
+		const FName ImportPackageFName(Import->ObjectName);
+		if (!PendingPackage && (!ExistingPackage || !ExistingPackage->HasAnyPackageFlags(PKG_CompiledIn)))
+		{
+			PendingPackage = FAsyncLoadingThread::Get().FindAsyncPackage(ImportPackageFName);
+		}
+		if (!PendingPackage)
+		{
+			if (ExistingPackage && ExistingPackage->HasAnyPackageFlags(PKG_CompiledIn))
+			{
+				// This can happen with editor only classes, not sure if this should be a warning or a silent continue
+				UE_LOG(LogStreaming, Warning, TEXT("FAsyncPackage::LoadImports for %s: Skipping import %s, depends on missing native class"), *Desc.NameToLoad.ToString(), *OriginalImport->ObjectName.ToString());
+			}
+			else if (!ExistingPackage || bForcePackageLoad)
+			{
+				// The package doesn't exist and this import is not in the dependency list so add it now.
+				check(!FPackageName::IsShortPackageName(Import->ObjectName));
+				UE_LOG(LogStreaming, Verbose, TEXT("FAsyncPackage::LoadImports for %s: Loading %s"), *Desc.NameToLoad.ToString(), *Import->ObjectName.ToString());
+				const FAsyncPackageDesc Info(INDEX_NONE, Import->ObjectName);
+				PendingPackage = new FAsyncPackage(Info);
+				PendingPackage->Desc.Priority = Desc.Priority;
+				AsyncLoadingThread.InsertPackage(PendingPackage);
+				bDidSomething = true;
+			}
+			else
+			{
+				// it would be nice to make sure it is actually loaded as we expect
+			}
+		}
+		if (PendingPackage)
+		{
+			if (int32(PendingPackage->AsyncPackageLoadingState) <= int32(EAsyncPackageLoadingState::WaitingForSummary))
+			{
+				FEventLoadNodePtr PrereqisiteNode;
+				PrereqisiteNode.WaitingPackage = FCheckedWeakAsyncPackagePtr(PendingPackage);
+				PrereqisiteNode.Phase = EEventLoadNode::Package_LoadSummary;
+
+				// we can't set up our imports until all packages we are importing have loaded their summary
+				AddArc(PrereqisiteNode, MyDependentNode);
+
+				// The other package should not leave the event driven loader until we have linked our imports, this just keeps it until we setup our imports...and at that time we will add more arcs
+				// we can't do that just yet, so make a note of it to do it when the node is actually added (if it is ever added, might be a missing file or something)
+				PendingPackage->PackagesWaitingToLinkImports.Add(WeakThis);
+				bDidSomething = true;
+			}
+			else if (int32(PendingPackage->AsyncPackageLoadingState) <= int32(EAsyncPackageLoadingState::WaitingForPostLoad))
+			{
+				FEventLoadNodePtr MyPrerequisiteNode;
+				MyPrerequisiteNode.WaitingPackage = WeakThis;
+				MyPrerequisiteNode.Phase = EEventLoadNode::Package_SetupImports;
+
+				FEventLoadNodePtr DependentNode;
+				DependentNode.WaitingPackage = FCheckedWeakAsyncPackagePtr(PendingPackage);
+				DependentNode.Phase = EEventLoadNode::Package_ExportsSerialized;  // this could be much later...really all we care about is that the linker isn't destroyed.
+
+				
+				AddArc(MyPrerequisiteNode, DependentNode);
+				bDidSomething = true;
+			}
+		}
+		UpdateLoadPercentage();
+	}
+
+	if (FileOpenLogActive() && !bDidSomething && LoadImportIndex == Linker->ImportMap.Num())
+	{
+		check(GetPriority() < MAX_int32);
+		SetPriority(GetPriority() + 1);
+	}
+	return LoadImportIndex == Linker->ImportMap.Num() ? EAsyncPackageState::Complete : EAsyncPackageState::TimeOut;
+}
+
+void FAsyncLoadingThread::QueueEvent_SetupImports(FAsyncPackage* Package, int32 EventSystemPriority)
+{
+	check(Package->AsyncPackageLoadingState == EAsyncPackageLoadingState::WaitingForImportPackages);
+	Package->AsyncPackageLoadingState = EAsyncPackageLoadingState::SetupImports;
+	check(Package);
+	FWeakAsyncPackagePtr WeakPtr(Package);
+	EventQueue.AddAsyncEvent(Package->GetPriority(), GRandomizeLoadOrder ? GetRandomSerialNumber() : Package->SerialNumber, SummariesFirstPriorityForFileOpenLog(EventSystemPriority, false),
+		TFunction<void(FAsyncLoadEventArgs& Args)>(
+			[WeakPtr, this](FAsyncLoadEventArgs& Args)
+	{
+		FAsyncPackage* Pkg = GetPackage(WeakPtr);
+		if (Pkg)
+		{
+			Pkg->SetTimeLimit(Args, TEXT("Setup Imports"));
+			Pkg->Event_SetupImports();
+		}
+	}
+	));
+}
+
+void FAsyncPackage::Event_SetupImports()
+{
+	{
+		FScopedAsyncPackageEvent Scope(this);
+		//@todo we need to time slice this, it runs to completion at the moment
+		verify(SetupImports_Event() == EAsyncPackageState::Complete);
+	}
+	check(AsyncPackageLoadingState == EAsyncPackageLoadingState::SetupImports);
+	check(ImportIndex == Linker->ImportMap.Num());
+	AsyncPackageLoadingState = EAsyncPackageLoadingState::SetupExports;
+	RemoveNode(EEventLoadNode::Package_SetupImports);
+	AsyncLoadingThread.QueueEvent_SetupExports(this);
+}
+
+
+EAsyncPackageState::Type FAsyncPackage::SetupImports_Event()
+{
+	SCOPED_LOADTIMER(CreateImportsTime);
+	SCOPE_CYCLE_COUNTER(STAT_FAsyncPackage_CreateImports);
+
+	// GC can't run in here
+	FGCScopeGuard GCGuard;
+
+	FCheckedWeakAsyncPackagePtr WeakThis(this);
+
+	if (!ImportIndex)
+	{
+		for (int32 InnerImportIndex = 0; InnerImportIndex < Linker->ImportMap.Num(); InnerImportIndex++)
+		{
+			Linker->ImportMap[InnerImportIndex].bImportSearchedFor = false; // we need to clear these if we have to call FindExistingImport below
+		}
+	}
+
+	// Create imports.
+	while (ImportIndex < Linker->ImportMap.Num()) //todo timeslice && !IsTimeLimitExceeded())
+	{
+		int32 LocalImportIndex = ImportIndex++;
+		FScopedCreateImportCounter ScopedCounter(Linker, LocalImportIndex);
+		FObjectImport& Import = Linker->ImportMap[LocalImportIndex];
+
+		if (Import.OuterIndex.IsNull())
+		{
+			UPackage* ImportPackage = (UPackage*)Import.XObject;
+			if (!ImportPackage)
+			{
+				ImportPackage = FindObjectFast<UPackage>(NULL, Import.ObjectName, false, false);
+				check(ImportPackage || FLinkerLoad::IsKnownMissingPackage(Import.ObjectName)); // We should have packages created for all imports by now
+				Import.XObject = ImportPackage; // this is an optimization to avoid looking up import packages multiple times, also, later we assume these are already filled in
+				if (Import.XObject)
+				{
+					AddObjectReference(Import.XObject);
+				}
+
+			}
+			if (ImportPackage)
+			{
+				FLinkerLoad* ImportLinker = ImportPackage->LinkerLoad;
+				if (ImportLinker && ImportLinker->AsyncRoot)
+				{
+					check(ImportLinker->AsyncRoot != this);
+					// make sure we wait for this package to serialize (and all of its dependents) before we start doing postloads
+					if (int32(ImportLinker->AsyncRoot->AsyncPackageLoadingState) <= int32(EAsyncPackageLoadingState::WaitingForPostLoad)) // no need to clutter dependencies with things that are already done
+					{
+						PackagesIMayBeWaitingForBeforePostload.Add(FWeakAsyncPackagePtr(ImportLinker->AsyncRoot));
+					}
+				}
+			}
+		}
+		else if (!Import.XObject || !Import.XObject->HasAnyFlags(RF_LoadCompleted)) // if this export was created by some external process (i.e. CDO), even if it doesn't have RF_NeedLoad, it will be loaded later
+		{
+			FPackageIndex OuterMostIndex = FPackageIndex::FromImport(LocalImportIndex);
+			while (true)
+			{
+				check(!OuterMostIndex.IsNull() && OuterMostIndex.IsImport());
+				FObjectImport& OuterMostImport = Linker->Imp(OuterMostIndex);
+
+				if (OuterMostImport.OuterIndex.IsNull())
+				{
+					break;
+				}
+				OuterMostIndex = OuterMostImport.OuterIndex;
+			}
+			FObjectImport& OuterMostImport = Linker->Imp(OuterMostIndex);
+			check(OuterMostImport.OuterIndex.IsNull());
+			UPackage* ImportPackage = (UPackage*)OuterMostImport.XObject;
+			if (!ImportPackage)
+			{
+				ImportPackage = FindObjectFast<UPackage>(NULL, OuterMostImport.ObjectName, false, false);
+				check(ImportPackage || FLinkerLoad::IsKnownMissingPackage(OuterMostImport.ObjectName)); // We should have packages created for all imports by now
+				OuterMostImport.XObject = ImportPackage; // this is an optimization to avoid looking up import packages multiple times, also, later we assume these are already filled in
+				if (OuterMostImport.XObject)
+				{
+					AddObjectReference(OuterMostImport.XObject);
+				}
+			}
+
+			bool bWaitingForImport = false;
+			if (ImportPackage)
+			{
+				FLinkerLoad* ImportLinker = ImportPackage->LinkerLoad;
+				if (!ImportLinker || !ImportLinker->AsyncRoot)
+				{
+					FindExistingImport(LocalImportIndex);
+					if (!Import.XObject || (!Import.XObject->HasAnyFlags(RF_LoadCompleted) && Import.XObject->HasAnyFlags(RF_WasLoaded)))
+					{
+						UE_LOG(LogStreaming, Error, TEXT("Found package without a linker and could not find %s in %s."), *Import.ObjectName.ToString(), *ImportPackage->GetName());
+					}
+
+				}
+				if (ImportLinker && ImportLinker->AsyncRoot)
+				{
+					check(ImportLinker->AsyncRoot != this);
+					check(!Import.OuterIndex.IsNull());
+					check(Import.OuterIndex.IsImport());
+					FName OuterName = Import.OuterIndex == OuterMostIndex ? NAME_None : // if the outer is the package, then that is just a null outer in the export table
+						Linker->Imp(Import.OuterIndex).ObjectName;
+					FPackageIndex LocalExportIndex;
+					for (auto It = ImportLinker->AsyncRoot->ObjectNameToImportOrExport.CreateKeyIterator(Import.ObjectName); It; ++It)
+					{
+						FPackageIndex PotentialExport = It.Value();
+						if (PotentialExport.IsExport())
+						{
+							FObjectExport& Export = ImportLinker->Exp(PotentialExport);
+							// this logic might not cover all cases
+							bool bMatch = false;
+							if (Export.OuterIndex.IsNull())
+							{
+								bMatch = OuterName == NAME_None;
+							}
+							else
+							{
+								check(Export.OuterIndex.IsExport());
+								bMatch = ImportLinker->Exp(Export.OuterIndex).ObjectName == OuterName;
+							}
+							if (bMatch)
+							{
+								check(LocalExportIndex.IsNull()); // otherwise we have two exports that match and our criteria is not good enough
+								LocalExportIndex = PotentialExport;
+							}
+						}
+					}
+					UE_CLOG(LocalExportIndex.IsNull(), LogStreaming, Warning, TEXT("Could not find import %s.%s in package %s, this is probably related to adding stuff we don't need to import map under the 'IsEventDrivenLoaderEnabledInCookedBuilds'."), *OuterName.ToString(), *Import.ObjectName.ToString(), *ImportPackage->GetName());
+					if (!LocalExportIndex.IsNull())
+					{
+						FObjectExport& Export = ImportLinker->Exp(LocalExportIndex);
+						Import.SourceIndex = LocalExportIndex.ToExport();
+						Import.SourceLinker = ImportLinker;
+						{
+							if (!Export.Object)
+							{
+								FEventLoadNodePtr MyDependentNode;
+								MyDependentNode.WaitingPackage = WeakThis;
+								MyDependentNode.ImportOrExportIndex = FPackageIndex::FromImport(LocalImportIndex);
+								MyDependentNode.Phase = EEventLoadNode::ImportOrExport_Create;
+
+								{
+									FEventLoadNodePtr PrereqisiteNode;
+									PrereqisiteNode.WaitingPackage = FCheckedWeakAsyncPackagePtr(ImportLinker->AsyncRoot);
+									PrereqisiteNode.ImportOrExportIndex = LocalExportIndex;
+									PrereqisiteNode.Phase = EEventLoadNode::ImportOrExport_Create;
+
+									// can't create an import until the corresponding export is created
+									AddArc(PrereqisiteNode, MyDependentNode);
+								}
+
+
+								{
+									FEventLoadNodePtr DependentNode;
+									DependentNode.WaitingPackage = FCheckedWeakAsyncPackagePtr(ImportLinker->AsyncRoot);
+									DependentNode.Phase = EEventLoadNode::Package_ExportsSerialized; // this could be much later...really all we care about is that the linker isn't destroyed.
+
+									// The other package should not leave the event driven loader until we have linked this import
+									AddArc(MyDependentNode, DependentNode);
+								}
+							}
+							else
+							{
+								check(!Import.XObject || Import.XObject == Export.Object);
+								Import.XObject = Export.Object;
+								if (Import.XObject)
+								{
+									AddObjectReference(Import.XObject);
+								}
+							}
+							if (!Export.Object || Export.Object->HasAnyFlags(RF_NeedLoad) 
+								|| !Export.Object->HasAnyFlags(RF_LoadCompleted)) // if this export was created by some external process (i.e. CDO), even if it doesn't have RF_NeedLoad, it will be loaded later
+							{
+								FEventLoadNodePtr MyDependentNode;
+								MyDependentNode.WaitingPackage = WeakThis;
+								MyDependentNode.ImportOrExportIndex = FPackageIndex::FromImport(LocalImportIndex);
+								MyDependentNode.Phase = EEventLoadNode::ImportOrExport_Serialize;
+
+								FEventLoadNodePtr PrereqisiteNode;
+								PrereqisiteNode.WaitingPackage = FCheckedWeakAsyncPackagePtr(ImportLinker->AsyncRoot);
+								PrereqisiteNode.ImportOrExportIndex = LocalExportIndex;
+								PrereqisiteNode.Phase = EEventLoadNode::ImportOrExport_Serialize;
+
+								// can't consider an import serialized until the corresponding export is serialized
+								AddArc(PrereqisiteNode, MyDependentNode);
+
+							}
+						}
+					}
+				}
+			}
+		}
+		DoneAddingPrerequistesFireIfNone(EEventLoadNode::ImportOrExport_Create, FPackageIndex::FromImport(LocalImportIndex));
+	}
+
+	return ImportIndex == Linker->ImportMap.Num() ? EAsyncPackageState::Complete : EAsyncPackageState::TimeOut;
+}
+
+void FAsyncLoadingThread::QueueEvent_SetupExports(FAsyncPackage* Package, int32 EventSystemPriority)
+{
+	check(Package->AsyncPackageLoadingState == EAsyncPackageLoadingState::SetupExports);
+	check(Package);
+	FWeakAsyncPackagePtr WeakPtr(Package);
+	EventQueue.AddAsyncEvent(Package->GetPriority(), GRandomizeLoadOrder ? GetRandomSerialNumber() : Package->SerialNumber, SummariesFirstPriorityForFileOpenLog(EventSystemPriority, false),
+		TFunction<void(FAsyncLoadEventArgs& Args)>(
+			[WeakPtr, this](FAsyncLoadEventArgs& Args)
+	{
+		FAsyncPackage* Pkg = GetPackage(WeakPtr);
+		if (Pkg)
+		{
+			Pkg->SetTimeLimit(Args, TEXT("Setup Exports"));
+			Pkg->Event_SetupExports();
+		}
+	}
+	));
+}
+
+void FAsyncPackage::Event_SetupExports()
+{
+	{
+		FScopedAsyncPackageEvent Scope(this);
+		if (SetupExports_Event() == EAsyncPackageState::TimeOut)
+		{
+			AsyncLoadingThread.QueueEvent_SetupExports(this, FAsyncLoadEvent::EventSystemPriority_MAX); // start here next frame
+			return;
+		}
+	}
+	check(AsyncPackageLoadingState == EAsyncPackageLoadingState::SetupExports);
+	AsyncPackageLoadingState = EAsyncPackageLoadingState::ProcessNewImportsAndExports;
+	ConditionalQueueProcessImportsAndExports();
+}
+
+void FAsyncLoadingThread::QueueEvent_ProcessImportsAndExports(FAsyncPackage* Package, int32 EventSystemPriority)
+{
+	check(Package->AsyncPackageLoadingState == EAsyncPackageLoadingState::ProcessNewImportsAndExports);
+	check(Package);
+	FWeakAsyncPackagePtr WeakPtr(Package);
+	EventQueue.AddAsyncEvent(Package->GetPriority(), GRandomizeLoadOrder ? GetRandomSerialNumber() : Package->SerialNumber, SummariesFirstPriorityForFileOpenLog(EventSystemPriority, false),
+		TFunction<void(FAsyncLoadEventArgs& Args)>(
+			[WeakPtr, this](FAsyncLoadEventArgs& Args)
+	{
+		FAsyncPackage* Pkg = GetPackage(WeakPtr);
+		if (Pkg)
+		{
+			Pkg->SetTimeLimit(Args, TEXT("ProcessImportsAndExports"));
+			Pkg->Event_ProcessImportsAndExports();
+		}
+	}
+	));
+}
+
+void FAsyncLoadingThread::QueueEvent_ProcessPostloadWait(FAsyncPackage* Package, int32 EventSystemPriority)
+{
+	check(Package->AsyncPackageLoadingState == EAsyncPackageLoadingState::WaitingForPostLoad);
+	check(Package);
+	FWeakAsyncPackagePtr WeakPtr(Package);
+	EventQueue.AddAsyncEvent(Package->GetPriority(), GRandomizeLoadOrder ? GetRandomSerialNumber() : Package->SerialNumber, SummariesFirstPriorityForFileOpenLog(EventSystemPriority, false),
+		TFunction<void(FAsyncLoadEventArgs& Args)>(
+			[WeakPtr, this](FAsyncLoadEventArgs& Args)
+	{
+		FAsyncPackage* Pkg = GetPackage(WeakPtr);
+		if (Pkg)
+		{
+			Pkg->SetTimeLimit(Args, TEXT("Process Process Postload Wait"));
+			Pkg->Event_ProcessPostloadWait();
+		}
+	}
+	));
+}
+
+void FAsyncLoadingThread::QueueEvent_ExportsDone(FAsyncPackage* Package, int32 EventSystemPriority)
+{
+	check(Package->AsyncPackageLoadingState == EAsyncPackageLoadingState::ProcessNewImportsAndExports);
+	check(Package);
+	FWeakAsyncPackagePtr WeakPtr(Package);
+	EventQueue.AddAsyncEvent(Package->GetPriority(), GRandomizeLoadOrder ? GetRandomSerialNumber() : Package->SerialNumber, SummariesFirstPriorityForFileOpenLog(EventSystemPriority, false),
+		TFunction<void(FAsyncLoadEventArgs& Args)>(
+			[WeakPtr, this](FAsyncLoadEventArgs& Args)
+	{
+		FAsyncPackage* Pkg = GetPackage(WeakPtr);
+		if (Pkg)
+		{
+			Pkg->SetTimeLimit(Args, TEXT("Exports Done"));
+			Pkg->Event_ExportsDone();
+		}
+	}
+	));
+}
+
+void FAsyncLoadingThread::QueueEvent_StartPostLoad(FAsyncPackage* Package, int32 EventSystemPriority)
+{
+	check(Package->AsyncPackageLoadingState == EAsyncPackageLoadingState::ReadyForPostLoad);
+	check(Package);
+	FWeakAsyncPackagePtr WeakPtr(Package);
+	EventQueue.AddAsyncEvent(Package->GetPriority(), GRandomizeLoadOrder ? GetRandomSerialNumber() : Package->SerialNumber, SummariesFirstPriorityForFileOpenLog(EventSystemPriority, false),
+		TFunction<void(FAsyncLoadEventArgs& Args)>(
+			[WeakPtr, this](FAsyncLoadEventArgs& Args)
+	{
+		FAsyncPackage* Pkg = GetPackage(WeakPtr);
+		if (Pkg)
+		{
+			Pkg->SetTimeLimit(Args, TEXT("Start Post Load"));
+			Pkg->Event_StartPostload();
+		}
+	}
+	));
+}
+bool FAsyncPackage::AnyImportsAndExportWorkOutstanding()
+{
+	return ImportsThatAreNowCreated.Num() ||
+		ImportsThatAreNowSerialized.Num() ||
+		ExportsThatCanBeCreated.Num() ||
+		ExportsThatCanHaveIOStarted.Num() ||
+		ExportsThatCanBeSerialized.Num() ||
+		ReadyPrecacheRequests.Num();
+}
+
+void FAsyncPackage::ConditionalQueueProcessImportsAndExports(bool bRequeueForTimeout)
+{
+	if (AsyncPackageLoadingState != EAsyncPackageLoadingState::ProcessNewImportsAndExports)
+	{
+		return;
+	}
+	if (!bProcessImportsAndExportsInFlight && AnyImportsAndExportWorkOutstanding())
+	{
+		bProcessImportsAndExportsInFlight = true;
+		int32 Pri = 0;
+		if (ReadyPrecacheRequests.Num())
+		{
+			Pri = FAsyncLoadEvent::EventSystemPriority_MAX;
+		}
+		else if (ExportsThatCanHaveIOStarted.Num() && PrecacheRequests.Num() < 2)
+		{
+			Pri = FAsyncLoadEvent::EventSystemPriority_MAX - 1;
+		}
+		AsyncLoadingThread.QueueEvent_ProcessImportsAndExports(this , Pri);
+	}
+}
+
+void FAsyncPackage::ConditionalQueueProcessPostloadWait()
+{
+	check(AsyncPackageLoadingState == EAsyncPackageLoadingState::WaitingForPostLoad);
+	if (!bProcessPostloadWaitInFlight
+		&& !PackagesIAmWaitingForBeforePostload.Num()) // if there other things we are waiting for, no need to the do the processing now
+	{
+		bProcessPostloadWaitInFlight = true;
+		AsyncLoadingThread.QueueEvent_ProcessPostloadWait(this);
+	}
+}
+
+EAsyncPackageState::Type FAsyncPackage::SetupExports_Event()
+{
+	SCOPED_LOADTIMER(CreateExportsTime);
+	SCOPE_CYCLE_COUNTER(STAT_FAsyncPackage_CreateExports);
+
+	// GC can't run in here
+	FGCScopeGuard GCGuard;
+
+	FCheckedWeakAsyncPackagePtr WeakThis(this);
+	Linker->GetFArchiveAsync2Loader()->LogItem(TEXT("SetupExports_Event"));
+
+	LastTypeOfWorkPerformed = TEXT("SetupExports_Event");
+	LastObjectWorkWasPerformedOn = LinkerRoot;
+	// Create exports.
+	while (ExportIndex < Linker->ExportMap.Num())
+	{
+		if (ExportIndex % 10 == 0 && IsTimeLimitExceeded())
+		{
+			break;
+		}
+		int32 LocalExportIndex = ExportIndex++;
+		FObjectExport& Export = Linker->ExportMap[LocalExportIndex];
+		// Check whether we already loaded the object and if not whether the context flags allow loading it.
+		check(!Export.Object); // we should not have this yet
+		if (!Export.Object)
+		{
+			if (!Linker->FilterExport(Export) && !Export.ClassIndex.IsNull())
+			{
+				check(Export.ObjectName != NAME_None || !(Export.ObjectFlags&RF_Public));
+				check(!Export.bDynamicClass); // @todoio we don't handle dynamic classes
+
+				int32 RunningIndex = Export.FirstExportDependency;
+				if (RunningIndex >= 0)
+				{
+					FEventLoadNodePtr MyDependentNode;
+					MyDependentNode.WaitingPackage = WeakThis;
+					MyDependentNode.ImportOrExportIndex = FPackageIndex::FromExport(LocalExportIndex);
+
+					FEventLoadNodePtr PrereqisiteNode;
+					PrereqisiteNode.WaitingPackage = WeakThis;
+
+					MyDependentNode.Phase = EEventLoadNode::Export_StartIO;
+					PrereqisiteNode.Phase = EEventLoadNode::ImportOrExport_Serialize;
+					for (int32 Index = Export.SerializationBeforeSerializationDependencies; Index > 0; Index--)
+					{
+						FPackageIndex Dep = Linker->PreloadDependencies[RunningIndex++];
+						check(!Dep.IsNull());
+						PrereqisiteNode.ImportOrExportIndex = Dep;
+						// don't request IO for this export until these are serialized
+						AddArc(PrereqisiteNode, MyDependentNode);
+					}
+
+					MyDependentNode.Phase = EEventLoadNode::Export_StartIO;
+					PrereqisiteNode.Phase = EEventLoadNode::ImportOrExport_Create;
+					for (int32 Index = Export.CreateBeforeSerializationDependencies; Index > 0; Index--)
+					{
+						FPackageIndex Dep = Linker->PreloadDependencies[RunningIndex++];
+						check(!Dep.IsNull());
+						PrereqisiteNode.ImportOrExportIndex = Dep;
+						// don't request IO for this export until these are done
+						AddArc(PrereqisiteNode, MyDependentNode);
+					}
+
+					MyDependentNode.Phase = EEventLoadNode::ImportOrExport_Create;
+					PrereqisiteNode.Phase = EEventLoadNode::ImportOrExport_Serialize;
+					for (int32 Index = Export.SerializationBeforeCreateDependencies; Index > 0; Index--)
+					{
+						FPackageIndex Dep = Linker->PreloadDependencies[RunningIndex++];
+						check(!Dep.IsNull());
+						PrereqisiteNode.ImportOrExportIndex = Dep;
+						// can't create this export until these things are serialized
+						AddArc(PrereqisiteNode, MyDependentNode);
+					}
+
+					MyDependentNode.Phase = EEventLoadNode::ImportOrExport_Create;
+					PrereqisiteNode.Phase = EEventLoadNode::ImportOrExport_Create;
+					for (int32 Index = Export.CreateBeforeCreateDependencies; Index > 0; Index--)
+					{
+						FPackageIndex Dep = Linker->PreloadDependencies[RunningIndex++];
+						check(!Dep.IsNull());
+						PrereqisiteNode.ImportOrExportIndex = Dep;
+						// can't create this export until these things are created
+						AddArc(PrereqisiteNode, MyDependentNode);
+					}
+
+				}
+			}
+		}
+		DoneAddingPrerequistesFireIfNone(EEventLoadNode::ImportOrExport_Create, FPackageIndex::FromExport(LocalExportIndex));
+	}
+
+	return ExportIndex == Linker->ExportMap.Num() ? EAsyncPackageState::Complete : EAsyncPackageState::TimeOut;
+}
+
+void FAsyncPackage::Event_ProcessImportsAndExports()
+{
+	if (bAllExportsSerialized)
+	{
+		// we can sometimes get a stray event here caused by the completion of an import that no export was waiting for
+		check(!AnyImportsAndExportWorkOutstanding());
+		return;
+	}
+	check(AsyncPackageLoadingState == EAsyncPackageLoadingState::ProcessNewImportsAndExports);
+	{
+		FScopedAsyncPackageEvent Scope(this);
+		ProcessImportsAndExports_Event();
+		bProcessImportsAndExportsInFlight = false;
+		ConditionalQueueProcessImportsAndExports(true);
+	}
+	check(AsyncPackageLoadingState == EAsyncPackageLoadingState::ProcessNewImportsAndExports);
+}
+
+void FAsyncPackage::LinkImport(int32 LocalImportIndex)
+{
+	check(LocalImportIndex >= 0 && LocalImportIndex < Linker->ImportMap.Num());
+	FObjectImport& Import = Linker->ImportMap[LocalImportIndex];
+	if (!Import.XObject)
+	{
+		Linker->GetFArchiveAsync2Loader()->LogItem(TEXT("LinkImport"));
+		if (Import.SourceLinker)
+		{
+			Import.XObject = Import.SourceLinker->ExportMap[Import.SourceIndex].Object;
+		}
+		else
+		{
+			// this block becomes active when a package completely finishes before we setup our import arcs.
+
+			FPackageIndex OuterMostIndex = FPackageIndex::FromImport(LocalImportIndex);
+			while (true)
+			{
+				check(!OuterMostIndex.IsNull() && OuterMostIndex.IsImport());
+				FObjectImport& OuterMostImport = Linker->Imp(OuterMostIndex);
+
+				if (OuterMostImport.OuterIndex.IsNull())
+				{
+					break;
+				}
+				OuterMostIndex = OuterMostImport.OuterIndex;
+			}
+			FObjectImport& OuterMostImport = Linker->Imp(OuterMostIndex);
+			UPackage* ImportPackage = (UPackage*)OuterMostImport.XObject; // these were filled in a previous step
+
+			check(ImportPackage || FLinkerLoad::IsKnownMissingPackage(OuterMostImport.ObjectName)); // We should have packages created for all imports by now
+			if (ImportPackage)
+			{
+				if (&OuterMostImport == &Import)
+				{
+					check(0); // we should not be here because package imports are already filled in
+				}
+				else
+				{
+					UPackage* ClassPackage = FindObjectFast<UPackage>(NULL, Import.ClassPackage, false, false);
+					if (ClassPackage)
+					{
+						UClass*	FindClass = FindObjectFast<UClass>(ClassPackage, Import.ClassName, false, false);
+						if (FindClass)
+						{
+							Import.XObject = FLinkerLoad::FindImportFast(FindClass, ImportPackage, Import.ObjectName);
+						}
+					}
+				}
+			}
+		}
+		if (Import.XObject)
+		{
+			AddObjectReference(Import.XObject);
+		}
+	}
+}
+
+void FAsyncPackage::DumpDependencies(const TCHAR* Label, UObject* Obj)
+{
+	UE_LOG(LogStreaming, Error, TEXT("****DumpDependencies [%s]:"), Label);
+	if (!Obj)
+	{
+		UE_LOG(LogStreaming, Error, TEXT("    Obj is nullptr"), Label);
+		return;
+	}
+	UE_LOG(LogStreaming, Error, TEXT("    Obj is %s"), *Obj->GetFullName());
+	UPackage* Package = Obj->GetOutermost();
+	if (!Package->LinkerLoad)
+	{
+		UE_LOG(LogStreaming, Error, TEXT("    %s has no linker"), *Package->GetFullName());
+	}
+	for (int32 LocalExportIndex = 0; LocalExportIndex < Package->LinkerLoad->ExportMap.Num(); LocalExportIndex++)
+	{
+		FObjectExport& Export = Package->LinkerLoad->ExportMap[LocalExportIndex];
+		if (Export.Object == Obj || Export.Object == nullptr)
+		{
+			if (Export.ObjectName == Obj->GetFName())
+			{
+				DumpDependencies(TEXT(""), Package->LinkerLoad, FPackageIndex::FromExport(LocalExportIndex));
+			}
+		}
+	}
+
+
+}
+
+void FAsyncPackage::DumpDependencies(const TCHAR* Label, FLinkerLoad* DumpLinker, FPackageIndex DumpExportIndex)
+{
+	FObjectExport& Export = DumpLinker->Exp(DumpExportIndex);
+	if (Label[0])
+	{
+		UE_LOG(LogStreaming, Error, TEXT("****DumpDependencies [%s]:"), Label);
+	}
+	UE_LOG(LogStreaming, Error, TEXT("    Export %d %s"), DumpExportIndex.ForDebugging(), *Export.ObjectName.ToString());
+	UE_LOG(LogStreaming, Error, TEXT("    Linker is %s"), *DumpLinker->GetArchiveName());
+
+
+	auto PrintDep = [DumpLinker](const TCHAR* DepLabel, FPackageIndex Dep)
+	{
+		if (Dep.IsNull())
+		{
+			UE_LOG(LogStreaming, Error, TEXT("        Dep %s null"), DepLabel);
+		}
+		else if (Dep.IsImport())
+		{
+			UE_LOG(LogStreaming, Error, TEXT("        Dep %s Import %5d   %s"), DepLabel, Dep.ToImport(), *DumpLinker->ImpExp(Dep).ObjectName.ToString());
+		}
+		else
+		{
+			UE_LOG(LogStreaming, Error, TEXT("        Dep %s Export %5d    %s     (class %s)"), DepLabel, Dep.ToExport(),
+				*DumpLinker->ImpExp(Dep).ObjectName.ToString(),
+				DumpLinker->Exp(Dep).ClassIndex.IsNull() ? TEXT("null") : *DumpLinker->ImpExp(DumpLinker->Exp(Dep).ClassIndex).ObjectName.ToString()
+				);
+		}
+	};
+
+	int32 RunningIndex = Export.FirstExportDependency;
+	if (RunningIndex >= 0)
+	{
+		for (int32 Index = Export.SerializationBeforeSerializationDependencies; Index > 0; Index--)
+		{
+			FPackageIndex Dep = DumpLinker->PreloadDependencies[RunningIndex++];
+			PrintDep(TEXT("S_BEFORE_S"), Dep);
+		}
+
+		for (int32 Index = Export.CreateBeforeSerializationDependencies; Index > 0; Index--)
+		{
+			FPackageIndex Dep = DumpLinker->PreloadDependencies[RunningIndex++];
+			PrintDep(TEXT("C_BEFORE_S"), Dep);
+		}
+
+		for (int32 Index = Export.SerializationBeforeCreateDependencies; Index > 0; Index--)
+		{
+			FPackageIndex Dep = DumpLinker->PreloadDependencies[RunningIndex++];
+			PrintDep(TEXT("S_BEFORE_C"), Dep);
+		}
+
+		for (int32 Index = Export.CreateBeforeCreateDependencies; Index > 0; Index--)
+		{
+			FPackageIndex Dep = DumpLinker->PreloadDependencies[RunningIndex++];
+			PrintDep(TEXT("C_BEFORE_C"), Dep);
+		}
+	}
+
+}
+
+
+UObject* FAsyncPackage::EventDrivenIndexToObject(FPackageIndex Index, bool bCheckSerialized, FPackageIndex DumpIndex)
+{
+	UObject* Result = nullptr;
+	if (Index.IsNull())
+	{
+		return Result;
+	}
+	if (Index.IsExport())
+	{
+		Result = Linker->Exp(Index).Object;
+	}
+	else if (Index.IsImport())
+	{
+		Result = Linker->Imp(Index).XObject;
+	}
+	if (!Result)
+	{
+		FEventLoadNodePtr MyDependentNode;
+		MyDependentNode.WaitingPackage = FCheckedWeakAsyncPackagePtr(this); //unnecessary
+		MyDependentNode.ImportOrExportIndex = Index;
+		MyDependentNode.Phase = EEventLoadNode::ImportOrExport_Create;
+		if (EventNodeArray.GetNode(MyDependentNode, false).bAddedToGraph || !EventNodeArray.GetNode(MyDependentNode, false).bFired)
+		{
+			FUObjectThreadContext& ThreadContext = FUObjectThreadContext::Get();
+			UClass* SerClass = Cast<UClass>(ThreadContext.SerializedObject);
+			if (!SerClass || Linker->ImpExp(Index).ObjectName != SerClass->GetDefaultObjectName())
+			{
+				DumpDependencies(TEXT("Dependencies"), ThreadContext.SerializedObject);
+				UE_LOG(LogStreaming, Fatal, TEXT("Missing Dependency, request for %s but it was still waiting for creation."), *Linker->ImpExp(Index).ObjectName.ToString());
+			}
+		}
+	}
+	else if (bCheckSerialized && (Result->HasAnyFlags(RF_NeedLoad) 
+		|| (Result->HasAnyFlags(RF_WasLoaded) && !Result->HasAnyFlags(RF_LoadCompleted)))) 
+	{
+		FEventLoadNodePtr MyDependentNode;
+		MyDependentNode.WaitingPackage = FCheckedWeakAsyncPackagePtr(this); //unnecessary
+		MyDependentNode.ImportOrExportIndex = Index;
+		MyDependentNode.Phase = EEventLoadNode::ImportOrExport_Serialize;
+
+		if (DumpIndex.IsNull())
+		{
+			FUObjectThreadContext& ThreadContext = FUObjectThreadContext::Get();
+			DumpDependencies(TEXT("Dependencies"), ThreadContext.SerializedObject);
+		}
+		else
+		{
+			DumpDependencies(TEXT("Dependencies"), Linker, DumpIndex);
+		}
+
+		if (EventNodeArray.GetNode(MyDependentNode, false).bAddedToGraph || !EventNodeArray.GetNode(MyDependentNode, false).bFired)
+		{
+			UE_LOG(LogStreaming, Fatal, TEXT("Missing Dependency, request for %s but it was still waiting for serialization."), *Linker->ImpExp(Index).ObjectName.ToString());
+		}
+		else
+		{
+			UE_LOG(LogStreaming, Fatal, TEXT("Missing Dependency, request for %s but it was still has RF_NeedLoad."), *Linker->ImpExp(Index).ObjectName.ToString());
+		}
+	}
+	return Result;
+}
+
+
+void FAsyncPackage::EventDrivenCreateExport(int32 LocalExportIndex)
+{
+	SCOPED_LOADTIMER(Package_CreateExports);
+	FObjectExport& Export = Linker->ExportMap[LocalExportIndex];
+	// Check whether we already loaded the object and if not whether the context flags allow loading it.
+	check(!Export.Object || Export.Object->HasAnyFlags(RF_ClassDefaultObject)); // we should not have this yet, unless it is a CDO
+	if (!Export.Object)
+	{
+		if (!Export.Object && !Linker->FilterExport(Export)) // for some acceptable position, it was not "not for" 
+		{
+			Linker->GetFArchiveAsync2Loader()->LogItem(TEXT("EventDrivenCreateExport"), Export.SerialOffset, Export.SerialSize);
+			LastTypeOfWorkPerformed = TEXT("EventDrivenCreateExport");
+			LastObjectWorkWasPerformedOn = nullptr;
+			check(Export.ObjectName != NAME_None || !(Export.ObjectFlags&RF_Public));
+			check(IsLoading());
+			check(!Export.bDynamicClass); //@todo support dynamic classes
+
+			UClass* LoadClass = nullptr;
+			if (Export.ClassIndex.IsNull())
+			{
+				LoadClass = UClass::StaticClass();
+			}
+			else
+			{
+				LoadClass = CastEventDrivenIndexToObject<UClass>(Export.ClassIndex, true, FPackageIndex::FromExport(LocalExportIndex));
+			}
+			if (!LoadClass)
+			{
+				UE_LOG(LogStreaming, Error, TEXT("Could not find class %s to create %s"), *Linker->ImpExp(Export.ClassIndex).ObjectName.ToString(), *Export.ObjectName.ToString());
+				Export.bExportLoadFailed = true;
+				return;
+			}
+			UStruct* SuperStruct = nullptr;
+			if (!Export.SuperIndex.IsNull())
+			{
+				SuperStruct = CastEventDrivenIndexToObject<UStruct>(Export.SuperIndex, true, FPackageIndex::FromExport(LocalExportIndex));
+				if (!SuperStruct)
+				{
+					// see FLinkerLoad::CreateExport, there may be some more we can do here
+					UE_LOG(LogStreaming, Error, TEXT("Could not find SuperStruct %s to create %s"), *Linker->ImpExp(Export.SuperIndex).ObjectName.ToString(), *Export.ObjectName.ToString());
+					Export.bExportLoadFailed = true;
+					return;
+				}
+			}
+			UObject* ThisParent = nullptr;
+			if (!Export.OuterIndex.IsNull())
+			{
+				ThisParent = EventDrivenIndexToObject(Export.OuterIndex, false, FPackageIndex::FromExport(LocalExportIndex));
+			}
+			else if (Export.bForcedExport)
+			{
+				// see FLinkerLoad::CreateExport, there may be some more we can do here
+				check(!Export.bForcedExport); // this is leftover from seekfree loading I think
+			}
+			else
+			{
+				check(LinkerRoot);
+				ThisParent = LinkerRoot;
+			}
+			check(!dynamic_cast<UObjectRedirector*>(ThisParent));
+			if (!ThisParent)
+			{
+				UE_LOG(LogStreaming, Error, TEXT("Could not find outer %s to create %s"), *Linker->ImpExp(Export.OuterIndex).ObjectName.ToString(), *Export.ObjectName.ToString());
+				Export.bExportLoadFailed = true;
+				return;
+			}
+
+			// Try to find existing object first in case we're a forced export to be able to reconcile. Also do it for the
+			// case of async loading as we cannot in-place replace objects.
+
+			UObject* ActualObjectWithTheName = StaticFindObjectFastInternal(NULL, ThisParent, Export.ObjectName, true);
+
+			// if we require cooked data, attempt to find exports in memory first...which we always do for event driven loading
+			check(FPlatformProperties::RequiresCookedData()
+				|| IsAsyncLoading()
+				|| Export.bForcedExport
+				|| LinkerRoot->ShouldFindExportsInMemoryFirst()
+				);
+
+			if (ActualObjectWithTheName && (ActualObjectWithTheName->GetClass() == LoadClass))
+			{
+				Export.Object = ActualObjectWithTheName;
+			}
+
+			// Object is found in memory.
+			if (Export.Object)
+			{
+				// Mark that we need to dissociate forced exports later on if we are a forced export.
+				if (Export.bForcedExport)
+				{
+					// see FLinkerLoad::CreateExport, there may be some more we can do here
+					check(!Export.bForcedExport); // this is leftover from seekfree loading I think
+				}
+				// Associate linker with object to avoid detachment mismatches.
+				else
+				{
+					Export.Object->SetLinker(Linker, LocalExportIndex);
+
+					// If this object was allocated but never loaded (components created by a constructor, CDOs, etc) make sure it gets loaded
+					// Do this for all subobjects created in the native constructor.
+					if (!Export.Object->HasAnyFlags(RF_LoadCompleted))
+					{
+						UE_LOG(LogStreaming, Verbose, TEXT("Note: %s was found in memory but created by some other process and needs loading."), *Export.Object->GetFullName());
+						Export.Object->SetFlags(RF_NeedLoad | RF_NeedPostLoad | RF_NeedPostLoadSubobjects | RF_WasLoaded);
+					}
+				}
+			}
+			else
+			{
+				if (ActualObjectWithTheName && !ActualObjectWithTheName->GetClass()->IsChildOf(LoadClass))
+				{
+					UE_LOG(LogLinker, Error, TEXT("Failed import: class '%s' name '%s' outer '%s'. There is another object (of '%s' class) at the path."),
+						*LoadClass->GetName(), *Export.ObjectName.ToString(), *ThisParent->GetName(), *ActualObjectWithTheName->GetClass()->GetName());
+					Export.bExportLoadFailed = true; // I am not sure if this is an actual fail or not...looked like it in the original code
+					return;
+				}
+
+				// Find the Archetype object for the one we are loading.
+				check(!Export.TemplateIndex.IsNull());
+				UObject* Template = EventDrivenIndexToObject(Export.TemplateIndex, true, FPackageIndex::FromExport(LocalExportIndex));
+				check(Template);
+				// we also need to ensure that the template has set up any instances
+				Template->ConditionalPostLoadSubobjects();
+
+
+				check(!GVerifyObjectReferencesOnly); // not supported with the event driven loader
+				// Create the export object, marking it with the appropriate flags to
+				// indicate that the object's data still needs to be loaded.
+				EObjectFlags ObjectLoadFlags = Export.ObjectFlags;
+				ObjectLoadFlags = EObjectFlags(ObjectLoadFlags | RF_NeedLoad | RF_NeedPostLoad | RF_NeedPostLoadSubobjects | RF_WasLoaded);
+
+				FName NewName = Export.ObjectName;
+
+				// If we are about to create a CDO, we need to ensure that all parent sub-objects are loaded
+				// to get default value initialization to work.
+#if DO_CHECK
+				if ((ObjectLoadFlags & RF_ClassDefaultObject) != 0)
+				{
+					UClass* SuperClass = LoadClass->GetSuperClass();
+					UObject* SuperCDO = SuperClass ? SuperClass->GetDefaultObject() : nullptr;
+					check(!SuperCDO || Template == SuperCDO); // the template for a CDO is the CDO of the super
+					if (SuperClass && !SuperClass->IsNative())
+					{
+						check(SuperCDO);
+						if (SuperClass->HasAnyFlags(RF_NeedLoad))
+						{
+							UE_LOG(LogStreaming, Fatal, TEXT("Super %s had RF_NeedLoad while creating %s"), *SuperClass->GetFullName(), *Export.ObjectName.ToString());
+							Export.bExportLoadFailed = true;
+							return;
+						}
+						if (SuperCDO->HasAnyFlags(RF_NeedLoad))
+						{
+							UE_LOG(LogStreaming, Fatal, TEXT("Super CDO %s had RF_NeedLoad while creating %s"), *SuperCDO->GetFullName(), *Export.ObjectName.ToString());
+							Export.bExportLoadFailed = true;
+							return;
+						}
+						TArray<UObject*> SuperSubObjects;
+						GetObjectsWithOuter(SuperCDO, SuperSubObjects, /*bIncludeNestedObjects=*/ false, /*ExclusionFlags=*/ RF_NoFlags, /*InternalExclusionFlags=*/ EInternalObjectFlags::Native);
+
+						for (UObject* SubObject : SuperSubObjects)
+						{
+							if (SubObject->HasAnyFlags(RF_NeedLoad))
+							{
+								UE_LOG(LogStreaming, Fatal, TEXT("Super CDO subobject %s had RF_NeedLoad while creating %s"), *SubObject->GetFullName(), *Export.ObjectName.ToString());
+								Export.bExportLoadFailed = true;
+								return;
+							}
+						}
+					}
+					else
+					{
+						check(Template->IsA(LoadClass));
+					}
+				}
+#endif
+				if (LoadClass->HasAnyFlags(RF_NeedLoad))
+				{
+					UE_LOG(LogStreaming, Fatal, TEXT("LoadClass %s had RF_NeedLoad while creating %s"), *LoadClass->GetFullName(), *Export.ObjectName.ToString());
+					Export.bExportLoadFailed = true;
+					return;
+				}
+				{
+					UObject* LoadCDO = LoadClass->GetDefaultObject();
+					if (LoadCDO->HasAnyFlags(RF_NeedLoad))
+					{
+						UE_LOG(LogStreaming, Fatal, TEXT("Class CDO %s had RF_NeedLoad while creating %s"), *LoadCDO->GetFullName(), *Export.ObjectName.ToString());
+						Export.bExportLoadFailed = true;
+						return;
+					}
+				}
+				if (Template->HasAnyFlags(RF_NeedLoad))
+				{
+					UE_LOG(LogStreaming, Fatal, TEXT("Template %s had RF_NeedLoad while creating %s"), *Template->GetFullName(), *Export.ObjectName.ToString());
+					Export.bExportLoadFailed = true;
+					return;
+				}
+
+				Export.Object = StaticConstructObject_Internal
+					(
+						LoadClass,
+						ThisParent,
+						NewName,
+						ObjectLoadFlags,
+						EInternalObjectFlags::None,
+						Template,
+						false,
+						nullptr,
+						true
+					);
+				if (GIsInitialLoad || GUObjectArray.IsOpenForDisregardForGC())
+				{
+					Export.Object->AddToRoot();
+				}
+				Export.Object->SetLinker(Linker, LocalExportIndex);
+				check(Export.Object->GetClass() == LoadClass);
+				check(NewName == Export.ObjectName);
+
+				// If it's a struct or class, set its parent.
+				if (UStruct* Struct = dynamic_cast<UStruct*>(Export.Object))
+				{
+					if (SuperStruct)
+					{
+						Struct->SetSuperStruct(SuperStruct);
+					}
+
+					// If it's a class, bind it to C++.
+					if (UClass* ClassObject = dynamic_cast<UClass*>(Export.Object))
+					{
+						ClassObject->Bind();
+					}
+				}
+			}
+		}
+	}
+	if (Export.Object)
+	{
+		AddObjectReference(Export.Object);
+	}
+	LastObjectWorkWasPerformedOn = Export.Object;
+	check(Linker->FilterExport(Export) || (Export.Object && !Export.bExportLoadFailed));
+}
+
+void FAsyncPackage::EventDrivenSerializeExport(int32 LocalExportIndex)
+{
+	SCOPED_LOADTIMER(Package_PreLoadObjects);
+
+	FObjectExport& Export = Linker->ExportMap[LocalExportIndex];
+	UObject* Object = Export.Object;
+	if (Object && Object->HasAnyFlags(RF_NeedLoad))
+	{
+		Linker->GetFArchiveAsync2Loader()->LogItem(TEXT("EventDrivenSerializeExport"), Export.SerialOffset, Export.SerialSize);
+
+		LastTypeOfWorkPerformed = TEXT("EventDrivenSerializeExport");
+		LastObjectWorkWasPerformedOn = Object;
+		check(Object->GetLinker() == Linker);
+		check(Object->GetLinkerIndex() == LocalExportIndex);
+		UClass* Cls = nullptr;
+
+		// If this is a struct, make sure that its parent struct is completely loaded
+		if (UStruct* Struct = dynamic_cast<UStruct*>(Object))
+		{
+			Cls = dynamic_cast<UClass*>(Object);
+			UStruct* SuperStruct = Struct->GetSuperStruct();
+			if (SuperStruct && SuperStruct->HasAnyFlags(RF_NeedLoad))
+			{
+				UE_LOG(LogStreaming, Fatal, TEXT("SuperStruct %s had RF_NeedLoad while creating %s"), *SuperStruct->GetFullName(), *Export.ObjectName.ToString());
+				Export.bExportLoadFailed = true;
+				return;
+			}		
+		}
+		check(Export.SerialOffset >= CurrentBlockOffset && Export.SerialOffset + Export.SerialSize <= CurrentBlockOffset + CurrentBlockBytes);
+
+		FArchiveAsync2* FAA2 = Linker->GetFArchiveAsync2Loader();
+		check(FAA2);
+
+		const int64 SavedPos = FAA2->Tell();
+		FAA2->Seek(Export.SerialOffset);
+
+		Object->ClearFlags(RF_NeedLoad);
+
+		FUObjectThreadContext& ThreadContext = FUObjectThreadContext::Get();
+		UObject* PrevSerializedObject = ThreadContext.SerializedObject;
+		ThreadContext.SerializedObject = Object;
+		Linker->bForceSimpleIndexToObject = true;
+
+		// Find the Archetype object for the one we are loading. This is piped to GetArchetypeFromLoader
+		check(!Export.TemplateIndex.IsNull());
+		UObject* Template = EventDrivenIndexToObject(Export.TemplateIndex, true, FPackageIndex::FromExport(LocalExportIndex));
+		check(Template);
+
+		check(!Linker->TemplateForGetArchetypeFromLoader);
+		Linker->TemplateForGetArchetypeFromLoader = Template;
+
+		if (Object->HasAnyFlags(RF_ClassDefaultObject))
+		{
+
+			Object->GetClass()->SerializeDefaultObject(Object, *Linker);
+		}
+		else
+		{
+			Object->Serialize(*Linker);
+		}
+		check(Linker->TemplateForGetArchetypeFromLoader == Template);
+		Linker->TemplateForGetArchetypeFromLoader = nullptr;
+
+
+
+		Object->SetFlags(RF_LoadCompleted);
+		ThreadContext.SerializedObject = PrevSerializedObject;
+		Linker->bForceSimpleIndexToObject = false;
+
+		if (FAA2->Tell() - Export.SerialOffset != Export.SerialSize)
+		{
+			if (Object->GetClass()->HasAnyClassFlags(CLASS_Deprecated))
+			{
+				UE_LOG(LogStreaming, Warning, TEXT("%s"), *FString::Printf(TEXT("%s: Serial size mismatch: Got %d, Expected %d"), *Object->GetFullName(), (int32)(FAA2->Tell() - Export.SerialOffset), Export.SerialSize));
+			}
+			else
+			{
+				UE_LOG(LogStreaming, Fatal, TEXT("%s"), *FString::Printf(TEXT("%s: Serial size mismatch: Got %d, Expected %d"), *Object->GetFullName(), (int32)(FAA2->Tell() - Export.SerialOffset), Export.SerialSize));
+			}
+		}
+
+		FAA2->Seek(SavedPos);
+#if DO_CHECK
+		if (Object->HasAnyFlags(RF_ClassDefaultObject) && Object->GetClass()->HasAnyClassFlags(CLASS_CompiledFromBlueprint))
+		{
+			check(Object->HasAllFlags(RF_NeedPostLoad | RF_WasLoaded));
+			//Object->SetFlags(RF_NeedPostLoad | RF_WasLoaded);
+		}
+#endif
+
+
+	}
+}
+
+#define MAX_EXPORT_PRECACHE_BLOCK (1024*1024)
+#define MAX_EXPORT_COUNT_PRECACHE (20)
+#define MAX_EXPORT_ALLOWED_SKIP (48*1024)
+
+void FAsyncPackage::StartPrecacheRequest()
+{
+	int32 LocalExportIndex = -1;
+	while (true)
+	{
+		ExportsThatCanHaveIOStarted.HeapPop(LocalExportIndex, false);
+		FObjectExport& Export = Linker->ExportMap[LocalExportIndex];
+		bool bReady = false;
+		if (Export.Object && Export.Object->HasAnyFlags(RF_NeedLoad))
+		{
+			// look for an existing request that will cover it
+			if (Export.SerialOffset >= CurrentBlockOffset && Export.SerialOffset + Export.SerialSize <= CurrentBlockOffset + CurrentBlockBytes)
+			{
+				// ready right now
+				bReady = true;
+			}
+			else
+			{
+				IAsyncReadRequest* Precache = ExportIndexToPrecacheRequest.FindRef(LocalExportIndex);
+				if (Precache)
+				{
+					// it is in an outstanding request
+					FExportIORequest* Req = PrecacheRequests.Find(Precache);
+					check(Req);
+					check(Export.SerialOffset >= Req->Offset && Export.SerialOffset + Export.SerialSize <= Req->Offset + Req->BytesToRead);
+					Req->ExportsToRead.Add(LocalExportIndex);
+				}
+				else
+				{
+					break;
+				}
+			}
+		}
+		else
+		{
+			bReady = true;
+		}
+		if (bReady)
+		{
+			RemoveNode(EEventLoadNode::Export_StartIO, FPackageIndex::FromExport(LocalExportIndex));
+		}
+		if (!ExportsThatCanHaveIOStarted.Num())
+		{
+			return;
+		}
+	}
+	// LocalExportIndex will start a new precache request
+	FObjectExport& Export = Linker->ExportMap[LocalExportIndex];
+
+	FExportIORequest NewReq;
+	NewReq.Offset = Export.SerialOffset;
+	NewReq.BytesToRead = Export.SerialSize;
+	check(NewReq.BytesToRead > 0 && NewReq.Offset > 0);
+	NewReq.ExportsToRead.Add(LocalExportIndex);
+
+	int32 LastExportIndex = LocalExportIndex;
+	if (!GRandomizeLoadOrder) // the blow code relies on sorting, which doesn't happen when we use a random load order; we will load export-by-export with no export fusion.
+	{
+		while (ExportsThatCanHaveIOStarted.Num() && NewReq.BytesToRead < MAX_EXPORT_PRECACHE_BLOCK && LastExportIndex - LocalExportIndex <= MAX_EXPORT_COUNT_PRECACHE)
+		{
+			int32 MaybeLastExportIndex = ExportsThatCanHaveIOStarted[0];
+			check(MaybeLastExportIndex > LastExportIndex); 
+
+				FObjectExport& LaterExport = Linker->ExportMap[MaybeLastExportIndex];
+			if (LaterExport.SerialOffset >= CurrentBlockOffset && LaterExport.SerialOffset + LaterExport.SerialSize <= CurrentBlockOffset + CurrentBlockBytes)
+			{
+				// ready right now, release it and remove it from the queue
+				int32 TempExportIndex = -1;
+				ExportsThatCanHaveIOStarted.HeapPop(TempExportIndex, false);
+				check(TempExportIndex == MaybeLastExportIndex);
+				RemoveNode(EEventLoadNode::Export_StartIO, FPackageIndex::FromExport(MaybeLastExportIndex));
+				break;
+			}
+
+			int64 Gap = LaterExport.SerialOffset - (NewReq.Offset + NewReq.BytesToRead);
+			check(Gap >= 0);
+
+			if (Gap > MAX_EXPORT_ALLOWED_SKIP || NewReq.BytesToRead + LaterExport.SerialSize > MAX_EXPORT_PRECACHE_BLOCK)
+			{
+				break; // this is too big of a gap or we already have a big enough read request
+			}
+			bool bAlreadyCovered = false;
+			for (int32 Index = LastExportIndex + 1; Index <= MaybeLastExportIndex; Index++)
+			{
+				if (ExportIndexToPrecacheRequest.Contains(Index))
+				{
+					bAlreadyCovered = true;
+					break;
+				}
+			}
+			if (bAlreadyCovered)
+			{
+				break;
+			}
+			// this export is good to merge into the request
+			ExportsThatCanHaveIOStarted.HeapPop(LastExportIndex, false);
+			check(LastExportIndex == MaybeLastExportIndex);
+			NewReq.BytesToRead = LaterExport.SerialOffset + LaterExport.SerialSize - NewReq.Offset;
+			check(NewReq.BytesToRead > 0);
+			NewReq.ExportsToRead.Add(LastExportIndex);
+		}
+	}
+	check(NewReq.ExportsToRead.Num());
+	FArchiveAsync2* FAA2 = Linker->GetFArchiveAsync2Loader();
+	check(FAA2);
+
+	IAsyncReadRequest* Precache = FAA2->MakeEventDrivenPrecacheRequest(NewReq.Offset, NewReq.BytesToRead,GPrecacheCallbackHandler.GetCompletionCallback());
+
+	NewReq.FirstExportCovered = LocalExportIndex;
+	NewReq.LastExportCovered = LastExportIndex;
+	for (int32 Index = NewReq.FirstExportCovered; Index <= NewReq.LastExportCovered; Index++)
+	{
+		check(!ExportIndexToPrecacheRequest.Contains(Index));
+		ExportIndexToPrecacheRequest.Add(Index, Precache);
+	}
+	check(!PrecacheRequests.Contains(Precache));
+	FExportIORequest& RequestInPlace = PrecacheRequests.Add(Precache);
+	Swap(RequestInPlace, NewReq);
+	GPrecacheCallbackHandler.RegisterNewPrecacheRequest(Precache, this);
+}
+
+int64 FAsyncPackage::PrecacheRequestReady(IAsyncReadRequest * Read)
+{
+	ReadyPrecacheRequests.Add(Read);
+	int64 Size = PrecacheRequests.FindChecked(Read).BytesToRead;
+	ConditionalQueueProcessImportsAndExports();
+	return Size;
+}
+
+void FAsyncPackage::MakeNextPrecacheRequestCurrent()
+{
+	check(ReadyPrecacheRequests.Num());
+	IAsyncReadRequest* Read = ReadyPrecacheRequests.Pop(false);
+	FExportIORequest& Req = PrecacheRequests.FindChecked(Read);
+	CurrentBlockOffset = Req.Offset;
+	CurrentBlockBytes = Req.BytesToRead;
+	ExportsInThisBlock.Reset();
+
+	GPrecacheCallbackHandler.FinishRequest(Req.BytesToRead);
+
+	FArchiveAsync2* FAA2 = Linker->GetFArchiveAsync2Loader();
+	check(FAA2);
+	bool bReady = FAA2->PrecacheForEvent(CurrentBlockOffset, CurrentBlockBytes);
+	UE_CLOG(!bReady, LogStreaming, Warning, TEXT("Preache request should have been hot %s."), *Linker->Filename);
+	for (int32 Index = Req.FirstExportCovered; Index <= Req.LastExportCovered; Index++)
+	{
+		verify(ExportIndexToPrecacheRequest.Remove(Index) == 1);
+		ExportsInThisBlock.Add(Index);
+	}
+	for (int32 LocalExportIndex : Req.ExportsToRead)
+	{
+		RemoveNode(EEventLoadNode::Export_StartIO, FPackageIndex::FromExport(LocalExportIndex));
+	}
+	Read->WaitCompletion();
+	delete Read;
+	PrecacheRequests.Remove(Read);
+}
+
+void FAsyncPackage::FlushPrecacheBuffer()
+{
+	CurrentBlockOffset = -1;
+	CurrentBlockBytes = -1;
+	FArchiveAsync2* FAA2 = Linker->GetFArchiveAsync2Loader();
+	check(FAA2);
+	FAA2->FlushPrecacheBlock();
+}
+
+int32 GCurrentExportIndex = -1;
+
+EAsyncPackageState::Type FAsyncPackage::ProcessImportsAndExports_Event()
+{
+	check(Linker);
+	bool bDidSomething = true;
+	int32 LoopIterations = 0;
+	while (!IsTimeLimitExceeded() && bDidSomething)
+	{
+		if ((LoopIterations && GRandomizeLoadOrder) || ++LoopIterations == 20)
+		{
+			break; // requeue this to give other packages a chance to start IO
+		}
+		bDidSomething = false;
+		if (PrecacheRequests.Num() < 2 && ExportsThatCanHaveIOStarted.Num())
+		{
+			bDidSomething = true;
+			StartPrecacheRequest();
+			LastTypeOfWorkPerformed = TEXT("ProcessImportsAndExports Start IO");
+			LastObjectWorkWasPerformedOn = nullptr;
+		}
+		if (bDidSomething)
+		{
+			continue; // check time limit, and lets do the creates and new IO requests before the serialize checks
+		}
+		if (ImportsThatAreNowCreated.Num())
+		{
+			bDidSomething = true;
+			int32 LocalImportIndex = -1;
+			ImportsThatAreNowCreated.HeapPop(LocalImportIndex, false);
+			{
+				// GC can't run in here
+				FGCScopeGuard GCGuard;
+				LinkImport(LocalImportIndex);
+			}
+			RemoveNode(EEventLoadNode::ImportOrExport_Create, FPackageIndex::FromImport(LocalImportIndex));
+			LastTypeOfWorkPerformed = TEXT("ProcessImportsAndExports LinkImport");
+			LastObjectWorkWasPerformedOn = nullptr;
+		}
+		if (bDidSomething)
+		{
+			continue; // check time limit
+		}
+		if (ImportsThatAreNowSerialized.Num())
+		{
+			bDidSomething = true;
+			int32 LocalImportIndex = -1;
+			ImportsThatAreNowSerialized.HeapPop(LocalImportIndex, false);
+			FObjectImport& Import = Linker->ImportMap[LocalImportIndex];
+			if (Import.XObject)
+			{
+				checkf(!Import.XObject->HasAnyFlags(RF_NeedLoad), TEXT("%s had RF_NeedLoad yet it was marked as serialized."), *Import.XObject->GetFullName()); // otherwise is isn't done serializing, is it?
+			}
+			RemoveNode(EEventLoadNode::ImportOrExport_Serialize, FPackageIndex::FromImport(LocalImportIndex));
+			LastTypeOfWorkPerformed = TEXT("ProcessImportsAndExports ImportsThatAreNowSerialized");
+			LastObjectWorkWasPerformedOn = nullptr;
+		}
+		if (bDidSomething)
+		{
+			continue; // check time limit, and lets do the creates before the serialize checks
+		}
+		if (ExportsThatCanBeCreated.Num())
+		{
+			bDidSomething = true;
+			int32 LocalExportIndex = -1;
+			ExportsThatCanBeCreated.HeapPop(LocalExportIndex, false);
+			{
+				FGCScopeGuard GCGuard;
+				EventDrivenCreateExport(LocalExportIndex);
+			}
+			RemoveNode(EEventLoadNode::ImportOrExport_Create, FPackageIndex::FromExport(LocalExportIndex));
+		}
+		if (bDidSomething)
+		{
+			continue; // check time limit, and lets do the creates before the serialize checks
+		}
+		if (ExportsThatCanHaveIOStarted.Num())
+		{
+			bDidSomething = true;
+			StartPrecacheRequest();
+			LastTypeOfWorkPerformed = TEXT("ProcessImportsAndExports Start IO");
+			LastObjectWorkWasPerformedOn = nullptr;
+		}
+		if (bDidSomething)
+		{
+			continue; // check time limit, and lets do the creates and new IO requests before the serialize checks
+		}
+		if (ExportsThatCanBeSerialized.Num())
+		{
+			bDidSomething = true;
+			int32 LocalExportIndex = -1;
+			ExportsThatCanBeSerialized.HeapPop(LocalExportIndex, false);
+
+			if (ExportsInThisBlock.Remove(LocalExportIndex) == 1)
+			{
+				FGCScopeGuard GCGuard;
+				GCurrentExportIndex = LocalExportIndex;
+				EventDrivenSerializeExport(LocalExportIndex);
+				GCurrentExportIndex = -1;
+			}
+			else
+			{
+				check(!Linker->ExportMap[LocalExportIndex].Object || !Linker->ExportMap[LocalExportIndex].Object->HasAnyFlags(RF_NeedLoad));
+			}
+			RemoveNode(EEventLoadNode::ImportOrExport_Serialize, FPackageIndex::FromExport(LocalExportIndex));
+		}
+		if (bDidSomething)
+		{
+			continue; // This is really important, we want to avoid discarding the current read block at all costs
+		}
+		check(!ExportsThatCanBeSerialized.Num());
+		if (CurrentBlockBytes > 0 && !ExportsInThisBlock.Num())
+		{
+			// we are completely done with this block, so we should explicitly discard it to save memory
+			// this is pretty mediocre because maybe the things left in the list don't need to be load anyway, but it covers the common case of precaching a single thing and precaching a block of things that are all loaded
+			FlushPrecacheBuffer();
+			LastTypeOfWorkPerformed = TEXT("ProcessImportsAndExports FlushPrecacheBuffer");
+			LastObjectWorkWasPerformedOn = nullptr;
+		}
+		// else we might get a new export in this block, so we might as well hang onto it...though it might be discarded anyway for a new request below
+
+		if (ReadyPrecacheRequests.Num())
+		{
+			// this generally takes no time, so we don't consider it doing something
+			MakeNextPrecacheRequestCurrent();
+			LastTypeOfWorkPerformed = TEXT("ProcessImportsAndExports MakeNextPrecacheRequestCurrent");
+			LastObjectWorkWasPerformedOn = nullptr;
+		}
+	}
+	return (!bDidSomething) ? EAsyncPackageState::Complete : EAsyncPackageState::TimeOut;
+}
+
+void FAsyncPackage::Event_ExportsDone()
+{
+	Linker->GetFArchiveAsync2Loader()->LogItem(TEXT("Event_ExportsDone"));
+	check(AsyncPackageLoadingState == EAsyncPackageLoadingState::ProcessNewImportsAndExports);
+	bAllExportsSerialized = true;
+	RemoveNode(EEventLoadNode::Package_ExportsSerialized);
+	check(AsyncPackageLoadingState == EAsyncPackageLoadingState::ProcessNewImportsAndExports);
+	AsyncPackageLoadingState = EAsyncPackageLoadingState::WaitingForPostLoad;
+	check(!AnyImportsAndExportWorkOutstanding());
+	FlushPrecacheBuffer();
+
+	ConditionalQueueProcessPostloadWait();
+
+	FWeakAsyncPackagePtr WeakThis(this);
+	for (FWeakAsyncPackagePtr NotifyPtr : OtherPackagesWaitingForMeBeforePostload)
+	{
+		FAsyncPackage* TestPkg(AsyncLoadingThread.GetPackage(NotifyPtr));
+		if (TestPkg)
+		{
+			check(TestPkg != this);
+			int32 NumRem = TestPkg->PackagesIAmWaitingForBeforePostload.Remove(WeakThis);
+			check(NumRem);
+			TestPkg->PackagesIMayBeWaitingForBeforePostload.Add(WeakThis);
+			TestPkg->ConditionalQueueProcessPostloadWait();
+		}
+	}
+	OtherPackagesWaitingForMeBeforePostload.Empty();
+}
+
+
+void FAsyncPackage::Event_ProcessPostloadWait()
+{
+	Linker->GetFArchiveAsync2Loader()->LogItem(TEXT("Event_ProcessPostloadWait"));
+	check(AsyncPackageLoadingState == EAsyncPackageLoadingState::WaitingForPostLoad);
+	check(bAllExportsSerialized && !OtherPackagesWaitingForMeBeforePostload.Num());
+	bProcessPostloadWaitInFlight = false;
+
+	FWeakAsyncPackagePtr WeakThis(this);
+
+	check(!PackagesIAmWaitingForBeforePostload.Num());
+	TSet<FWeakAsyncPackagePtr> AlreadyHandled;
+	AlreadyHandled.Add(WeakThis); // we never consider ourself a dependent
+
+	// pretty dang complicated incremental algorithm to determine when all dependent packages are loaded....so we can postload our objects
+
+	// remove junk from the wait list and look for anything that isn't ready
+	for (auto It = PackagesIMayBeWaitingForBeforePostload.CreateIterator(); It; ++It)
+	{
+		FWeakAsyncPackagePtr TestPtr = *It;
+		check(TestPtr == WeakThis || !AlreadyHandled.Contains(TestPtr));
+		FAsyncPackage* TestPkg(AsyncLoadingThread.GetPackage(TestPtr));
+		if (!TestPkg || TestPkg == this || int32(TestPkg->AsyncPackageLoadingState) > int32(EAsyncPackageLoadingState::WaitingForPostLoad))
+		{
+			AlreadyHandled.Add(TestPtr);
+			It.RemoveCurrent();
+			continue;
+		}
+		if (!TestPkg->bAllExportsSerialized)
+		{
+			AlreadyHandled.Add(TestPtr);
+			// We need to wait for this package, link it so that we are notified, we will stop exploring on the next iteration because we are definitely waiting for something
+			check(!PackagesIAmWaitingForBeforePostload.Contains(TestPtr));
+			PackagesIAmWaitingForBeforePostload.Add(TestPtr);
+			check(!TestPkg->OtherPackagesWaitingForMeBeforePostload.Contains(WeakThis));
+			TestPkg->OtherPackagesWaitingForMeBeforePostload.Add(WeakThis);
+			It.RemoveCurrent();
+		}
+	}
+
+	while (PackagesIMayBeWaitingForBeforePostload.Num() && !PackagesIAmWaitingForBeforePostload.Num())
+	{
+		// flatten the dependency tree looking for something that isn't finished
+		FWeakAsyncPackagePtr PoppedPtr;
+		{
+			auto First = PackagesIMayBeWaitingForBeforePostload.CreateIterator();
+			PoppedPtr = *First;
+			First.RemoveCurrent();
+		}
+		if (AlreadyHandled.Contains(PoppedPtr))
+		{
+			continue;
+		}
+		AlreadyHandled.Add(PoppedPtr);
+		FAsyncPackage* TestPkg(AsyncLoadingThread.GetPackage(PoppedPtr));
+		check(TestPkg != this);
+		if (!TestPkg || int32(TestPkg->AsyncPackageLoadingState) > int32(EAsyncPackageLoadingState::WaitingForPostLoad))
+		{
+			continue;
+		}
+		check(TestPkg->bAllExportsSerialized); // we should have already handled these
+		// this package and all _direct_ dependents are ready, but lets collapse the tree here and deal with indirect dependents
+		for (FWeakAsyncPackagePtr MaybeRecursePtr : TestPkg->PackagesIAmWaitingForBeforePostload)
+		{
+			check(!(MaybeRecursePtr == WeakThis));
+			FAsyncPackage* MaybeRecursePkg(AsyncLoadingThread.GetPackage(MaybeRecursePtr));
+			check(MaybeRecursePkg && !MaybeRecursePkg->bAllExportsSerialized);
+
+			check(!PackagesIAmWaitingForBeforePostload.Contains(MaybeRecursePtr));
+			PackagesIAmWaitingForBeforePostload.Add(MaybeRecursePtr);
+			check(!MaybeRecursePkg->OtherPackagesWaitingForMeBeforePostload.Contains(WeakThis));
+			MaybeRecursePkg->OtherPackagesWaitingForMeBeforePostload.Add(WeakThis);
+		}
+		for (FWeakAsyncPackagePtr MaybeRecursePtr : TestPkg->PackagesIMayBeWaitingForBeforePostload)
+		{
+			if (!AlreadyHandled.Contains(MaybeRecursePtr))
+			{
+				FAsyncPackage* MaybeRecursePkg(AsyncLoadingThread.GetPackage(MaybeRecursePtr));
+				check(MaybeRecursePkg != this);
+				if (!MaybeRecursePkg || int32(MaybeRecursePkg->AsyncPackageLoadingState) > int32(EAsyncPackageLoadingState::WaitingForPostLoad))
+				{
+					continue;
+				}
+				if (MaybeRecursePkg->bAllExportsSerialized)
+				{
+					PackagesIMayBeWaitingForBeforePostload.Add(MaybeRecursePtr);
+				}
+				else
+				{
+					check(!PackagesIAmWaitingForBeforePostload.Contains(MaybeRecursePtr));
+					PackagesIAmWaitingForBeforePostload.Add(MaybeRecursePtr);
+					check(!MaybeRecursePkg->OtherPackagesWaitingForMeBeforePostload.Contains(WeakThis));
+					MaybeRecursePkg->OtherPackagesWaitingForMeBeforePostload.Add(WeakThis);
+				}
+			}
+		}
+	}
+	if (!PackagesIAmWaitingForBeforePostload.Num())
+	{
+		check(!PackagesIMayBeWaitingForBeforePostload.Num());
+		// all done
+		check(AsyncPackageLoadingState == EAsyncPackageLoadingState::WaitingForPostLoad);
+		AsyncPackageLoadingState = EAsyncPackageLoadingState::ReadyForPostLoad;
+		AsyncLoadingThread.QueueEvent_StartPostLoad(this);
+		check(bAllExportsSerialized && !OtherPackagesWaitingForMeBeforePostload.Num());
+	}
+}
+
+void FAsyncPackage::Event_StartPostload()
+{
+	Linker->GetFArchiveAsync2Loader()->LogItem(TEXT("Event_StartPostload"));
+	check(AsyncPackageLoadingState == EAsyncPackageLoadingState::ReadyForPostLoad);
+	check(!PackagesIMayBeWaitingForBeforePostload.Num());
+	check(!PackagesIAmWaitingForBeforePostload.Num());
+	check(!OtherPackagesWaitingForMeBeforePostload.Num());
+	AsyncPackageLoadingState = EAsyncPackageLoadingState::PostLoad_Etc;
+	EventDrivenLoadingComplete();
+	{
+		TArray<UObject*>& ObjLoaded = FUObjectThreadContext::Get().ObjLoaded;
+		ObjLoaded.Reserve(ObjLoaded.Num() + Linker->ExportMap.Num());
+		for (int32 LocalExportIndex = 0; LocalExportIndex < Linker->ExportMap.Num(); LocalExportIndex++)
+		{
+			FObjectExport& Export = Linker->ExportMap[LocalExportIndex];
+			UObject* Object = Export.Object;
+			if (Object && Object->HasAnyFlags(RF_NeedPostLoad))
+			{
+				ObjLoaded.Add(Object);
+			}
+		}
+	}
+	check(!FAsyncLoadingThread::Get().AsyncPackagesReadyForTick.Contains(this));
+	FAsyncLoadingThread::Get().AsyncPackagesReadyForTick.Add(this);
+}
+void FAsyncPackage::EventDrivenLoadingComplete()
+{
+	check(!AnyImportsAndExportWorkOutstanding());
+	bool bAny = false;
+	TArray<FEventLoadNodePtr> AddedNodes;
+	EventNodeArray.GetAddedNodes(AddedNodes, this);
+
+	for (FEventLoadNodePtr& Ptr : AddedNodes)
+	{
+		bAny = true;
+		UE_LOG(LogStreaming, Error, TEXT("Leaked Event Driven Node %s"), *Ptr.HumanReadableStringForDebugging());
+	}
+
+	if (bAny)
+	{
+		check(!bAny);
+		RemoveAllNodes();
+	}
+	check(!AnyImportsAndExportWorkOutstanding());
+
+	PackagesWaitingToLinkImports.Empty(); // usually redundant
+
+}
+
+FEventLoadGraph FAsyncPackage::GlobalEventGraph;
+
+FEventLoadNodePtr FAsyncPackage::AddNode(EEventLoadNode Phase, FPackageIndex ImportOrExportIndex, bool bHoldForLater, int32 NumImplicitPrereqs)
+{
+	FEventLoadNodePtr MyNode;
+	MyNode.WaitingPackage = FCheckedWeakAsyncPackagePtr(this);
+	MyNode.ImportOrExportIndex = ImportOrExportIndex;
+	MyNode.Phase = Phase;
+
+	GetEventGraph().AddNode(MyNode, bHoldForLater, NumImplicitPrereqs);
+	return MyNode;
+}
+
+void FAsyncPackage::DoneAddingPrerequistesFireIfNone(EEventLoadNode Phase, FPackageIndex ImportOrExportIndex, bool bWasHeldForLater)
+{
+	FEventLoadNodePtr MyNode;
+	MyNode.WaitingPackage = FCheckedWeakAsyncPackagePtr(this);
+	MyNode.ImportOrExportIndex = ImportOrExportIndex;
+	MyNode.Phase = Phase;
+
+	GetEventGraph().DoneAddingPrerequistesFireIfNone(MyNode, bWasHeldForLater);
+}
+
+
+void FAsyncPackage::RemoveNode(EEventLoadNode Phase, FPackageIndex ImportOrExportIndex)
+{
+	FEventLoadNodePtr MyNode;
+	MyNode.WaitingPackage = FCheckedWeakAsyncPackagePtr(this);
+	MyNode.ImportOrExportIndex = ImportOrExportIndex;
+	MyNode.Phase = Phase;
+
+	return GetEventGraph().RemoveNode(MyNode);
+}
+
+void FAsyncPackage::NodeWillBeFiredExternally(EEventLoadNode Phase, FPackageIndex ImportOrExportIndex)
+{
+	FEventLoadNodePtr MyNode;
+	MyNode.WaitingPackage = FCheckedWeakAsyncPackagePtr(this);
+	MyNode.ImportOrExportIndex = ImportOrExportIndex;
+	MyNode.Phase = Phase;
+
+	return GetEventGraph().NodeWillBeFiredExternally(MyNode);
+}
+
+void FAsyncPackage::AddArc(FEventLoadNodePtr& PrereqisiteNode, FEventLoadNodePtr& DependentNode)
+{
+	GetEventGraph().AddArc(PrereqisiteNode, DependentNode);
+}
+
+void FAsyncPackage::RemoveAllNodes()
+{
+	FEventLoadGraph& Graph = GetEventGraph();
+	TArray<FEventLoadNodePtr> AddedNodes;
+	EventNodeArray.GetAddedNodes(AddedNodes, this);
+	for (FEventLoadNodePtr& Ptr : AddedNodes)
+	{
+		Graph.RemoveNode(Ptr);
+	}
+}
+
+void FAsyncPackage::FireNode(FEventLoadNodePtr& NodeToFire)
+{
+	check(int32(AsyncPackageLoadingState) < int32(EAsyncPackageLoadingState::PostLoad_Etc));
+	if (NodeToFire.ImportOrExportIndex.IsNull())
+	{
+		switch (NodeToFire.Phase)
+		{
+		case EEventLoadNode::Package_LoadSummary:
+			break;
+		case EEventLoadNode::Package_SetupImports:
+			AsyncLoadingThread.QueueEvent_SetupImports(this);
+			break;
+		case EEventLoadNode::Package_ExportsSerialized:
+			AsyncLoadingThread.QueueEvent_ExportsDone(this);
+			break;
+		default:
+			check(0);
+		}
+	}
+	else
+	{
+		switch (NodeToFire.Phase)
+		{
+		case EEventLoadNode::ImportOrExport_Create:
+			if (NodeToFire.ImportOrExportIndex.IsImport())
+			{
+				ImportsThatAreNowCreated.HeapPush(NodeToFire.ImportOrExportIndex.ToImport());
+			}
+			else
+			{
+				ExportsThatCanBeCreated.HeapPush(NodeToFire.ImportOrExportIndex.ToExport());
+			}
+			break;
+		case EEventLoadNode::Export_StartIO:
+			ExportsThatCanHaveIOStarted.HeapPush(NodeToFire.ImportOrExportIndex.ToExport());
+			break;
+		case EEventLoadNode::ImportOrExport_Serialize:
+			if (NodeToFire.ImportOrExportIndex.IsImport())
+			{
+				ImportsThatAreNowSerialized.HeapPush(NodeToFire.ImportOrExportIndex.ToImport());
+			}
+			else
+			{
+				ExportsThatCanBeSerialized.HeapPush(NodeToFire.ImportOrExportIndex.ToExport());
+			}
+			break;
+		default:
+			check(0);
+		}
+
+		if (AsyncPackageLoadingState == EAsyncPackageLoadingState::ProcessNewImportsAndExports) // this is redundant, but we want to save the function call
+		{
+			ConditionalQueueProcessImportsAndExports();
+		}
+	}
+}
+
+
+
+#endif
+
+void FAsyncLoadingThread::InsertPackage(FAsyncPackage* Package, bool bReinsert, EAsyncPackageInsertMode InsertMode)
+{
+	checkSlow(IsInAsyncLoadThread());
+
+#if USE_EVENT_DRIVEN_ASYNC_LOAD && DO_CHECK
+	FWeakAsyncPackagePtr WeakPtr(Package);
+	check(Package);
+#endif
+
+	if (!bReinsert)
+	{
+		// Incremented on the Async Thread, decremented on the game thread
+		ExistingAsyncPackagesCounter.Increment();
+	}
+
+	{
+#if THREADSAFE_UOBJECTS
+		FScopeLock LockAsyncPackages(&AsyncPackagesCritical);
+#endif
+		if (bReinsert)
+		{
+			AsyncPackages.RemoveSingle(Package);
+		}
+		int32 InsertIndex = -1;
+
+		switch (InsertMode)
+		{
+			case EAsyncPackageInsertMode::InsertAfterMatchingPriorities:
+			{
+				InsertIndex = AsyncPackages.IndexOfByPredicate([Package](const FAsyncPackage* Element)
+				{
+					return Element->GetPriority() < Package->GetPriority();
+				});
+
+				break;
+			}
+
+			case EAsyncPackageInsertMode::InsertBeforeMatchingPriorities:
+			{
+				// Insert new package keeping descending priority order in AsyncPackages
+				InsertIndex = AsyncPackages.IndexOfByPredicate([Package](const FAsyncPackage* Element)
+				{
+					return Element->GetPriority() <= Package->GetPriority();
+				});
+
+				break;
+			}
+		};
+
+		InsertIndex = InsertIndex == INDEX_NONE ? AsyncPackages.Num() : InsertIndex;
+
+		AsyncPackages.InsertUninitialized(InsertIndex);
+		TArray<FAsyncPackage*>& AlreadyDecremented = PackagesThatDecrementedCounter.FindOrAdd(Package->GetPackageName());
+		if (AlreadyDecremented.Contains(Package))
+		{
+			UE_LOG(LogStreaming, Warning, TEXT("Package %s has been removed multiple times"), *Package->GetPackageName().ToString());
+		}
+
+		AsyncPackages[InsertIndex] = Package;
+
+		if (!bReinsert)
+		{
+			AsyncPackageNameLookup.Add(Package->GetPackageName(), Package);
+#if USE_EVENT_DRIVEN_ASYNC_LOAD
+			// @todo If this is a reinsert for some priority thing, well we don't go back and retract the stuff in flight to adjust the priority of events
+			QueueEvent_CreateLinker(Package);
+#endif
+		}
+
+	}
+#if USE_EVENT_DRIVEN_ASYNC_LOAD
+	check(GetPackage(WeakPtr) == Package);
+#endif
+}
+
+void FAsyncLoadingThread::AddToLoadedPackages(FAsyncPackage* Package)
+{
+#if THREADSAFE_UOBJECTS
+	FScopeLock LoadedLock(&LoadedPackagesCritical);
+#endif
+	if (!LoadedPackages.Contains(Package))
+	{
+		TArray<FAsyncPackage*>& AlreadyDecremented = PackagesThatDecrementedCounter.FindOrAdd(Package->GetPackageName());
+		if (AlreadyDecremented.Contains(Package))
+		{
+			UE_LOG(LogStreaming, Warning, TEXT("Package %s has already been removed but is still being loaded."), *Package->GetPackageName().ToString());
+		}
+
+		LoadedPackages.Add(Package);
+		LoadedPackagesNameLookup.Add(Package->GetPackageName(), Package);
+	}
+}
+
+#if USE_EVENT_DRIVEN_ASYNC_LOAD
+static FThreadSafeCounter RecusionNotAllowed;
+struct FScopedRecursionNotAllowed
+{
+#if !UE_BUILD_SHIPPING && !UE_BUILD_TEST
+	FScopedRecursionNotAllowed()
+	{
+		verify(RecusionNotAllowed.Increment() == 1);
+	}
+	~FScopedRecursionNotAllowed()
+	{
+		verify(RecusionNotAllowed.Decrement() == 0);
+	}
+#endif
+};
+#endif
+
+EAsyncPackageState::Type FAsyncLoadingThread::ProcessAsyncLoading(int32& OutPackagesProcessed, bool bUseTimeLimit /*= false*/, bool bUseFullTimeLimit /*= false*/, float TimeLimit /*= 0.0f*/, FFlushTree* FlushTree)
+{
+	SCOPE_CYCLE_COUNTER(STAT_FAsyncLoadingThread_ProcessAsyncLoading);
+	SCOPED_LOADTIMER(AsyncLoadingTime);
+
+	// If we're not multithreaded and flushing async loading, update the thread heartbeat
+	const bool bNeedsHeartbeatTick = !bUseTimeLimit && !FAsyncLoadingThread::IsMultithreaded();
+
+	EAsyncPackageState::Type LoadingState = EAsyncPackageState::Complete;
+	OutPackagesProcessed = 0;
+
+	double TickStartTime = FPlatformTime::Seconds();
+
+#if USE_EVENT_DRIVEN_ASYNC_LOAD
+	FScopedRecursionNotAllowed RecursionGuard;
+	FAsyncLoadingTickScope InAsyncLoadingTick;
+	uint32 LoopIterations = 0;
+
+	while (true)
+	{
+		if (bNeedsHeartbeatTick && (++LoopIterations) % 32 == 31)
+		{
+			// Update heartbeat after 32 events
+			FThreadHeartBeat::Get().HeartBeat();
+		}
+
+		if (FileOpenLogActive() && GPrecacheCallbackHandler.AnyIOOutstanding())
+		{
+			GPrecacheCallbackHandler.WaitForIO(10.0f); // We don't want the ideal load order to be affected by the behavior of the disk...so we always just wait for IO
+		}
+		bool bDidSomething = GPrecacheCallbackHandler.ProcessIncoming();
+		OutPackagesProcessed += (bDidSomething ? 1 : 0);
 
-	if (QueueCopy.Num() > 0)
-	{
-		IAssetRegistryInterface* AssetRegistry = nullptr;
-		
-		if (GPreloadPackageDependencies && IsPlatformFileCompatibleWithDependencyPreloading())
+		if (IsTimeLimitExceeded(TickStartTime, bUseTimeLimit, TimeLimit, TEXT("ProcessIncoming"), nullptr))
 		{
-			AssetRegistry = IAssetRegistryInterface::GetPtr();
+			return EAsyncPackageState::TimeOut;
 		}
-		
-		double Timer = 0;
+
+		if (IsAsyncLoadingSuspended())
+		{
+			return EAsyncPackageState::TimeOut;
+		}
+
 		{
-			SCOPE_SECONDS_COUNTER(Timer);
-			for (FAsyncPackageDesc* PackageRequest : QueueCopy)
+			const float RemainingTimeLimit = FMath::Max(0.0f, TimeLimit - (float)(FPlatformTime::Seconds() - TickStartTime));
+			int32 NumCreated = CreateAsyncPackagesFromQueue(bUseTimeLimit, bUseFullTimeLimit, RemainingTimeLimit);
+			OutPackagesProcessed += NumCreated;
+			bDidSomething = NumCreated > 0 || bDidSomething;
+			if (IsTimeLimitExceeded(TickStartTime, bUseTimeLimit, TimeLimit, TEXT("CreateAsyncPackagesFromQueue")))
 			{
-				ProcessAsyncPackageRequest(PackageRequest, nullptr, AssetRegistry);
-				delete PackageRequest;
+				return EAsyncPackageState::TimeOut;
 			}
 		}
-		UE_LOG(LogStreaming, Verbose, TEXT("Async package requests inserted in %fms"), Timer * 1000.0);
-	}
-
-	NumCreated = QueueCopy.Num();
+		if (bDidSomething)
+		{
+			continue;
+		}
 
-	return NumCreated;
-}
+		{
+			FAsyncLoadEventArgs Args;
+			Args.bUseTimeLimit = bUseTimeLimit;
+			Args.TickStartTime = TickStartTime;
+			Args.TimeLimit = TimeLimit;
 
-void FAsyncLoadingThread::InsertPackage(FAsyncPackage* Package, EAsyncPackageInsertMode InsertMode)
-{
-	checkSlow(IsInAsyncLoadThread());
+			Args.OutLastTypeOfWorkPerformed = nullptr;
+			Args.OutLastObjectWorkWasPerformedOn = nullptr;
+			if (EventQueue.PopAndExecute(Args))
+			{
+				OutPackagesProcessed++;
+				if (IsTimeLimitExceeded(Args.TickStartTime, Args.bUseTimeLimit, Args.TimeLimit, Args.OutLastTypeOfWorkPerformed, Args.OutLastObjectWorkWasPerformedOn))
+				{
+					return EAsyncPackageState::TimeOut;
+				}
+				bDidSomething = true;
+			}
+		}
+		if (bDidSomething)
+		{
+			continue;
+		}
+		if (AsyncPackagesReadyForTick.Num())
+		{
+			OutPackagesProcessed++;
+			bDidSomething = true;
+			FAsyncPackage* Package = AsyncPackagesReadyForTick[0];
+			check(Package->AsyncPackageLoadingState == EAsyncPackageLoadingState::PostLoad_Etc);
+			SCOPED_LOADTIMER(ProcessAsyncLoadingTime);
 
-	// Incremented on the Async Thread, decremented on the game thread
-	ExistingAsyncPackagesCounter.Increment();
+			EAsyncPackageState::Type LocalLoadingState = EAsyncPackageState::Complete;
+			if (Package->HasFinishedLoading() == false)
+			{
+				float RemainingTimeLimit = FMath::Max(0.0f, TimeLimit - (float)(FPlatformTime::Seconds() - TickStartTime));
+				LocalLoadingState = Package->TickAsyncPackage(bUseTimeLimit, bUseFullTimeLimit, RemainingTimeLimit, FlushTree);
+				if (LocalLoadingState == EAsyncPackageState::TimeOut)
+				{
+					if (IsTimeLimitExceeded(TickStartTime, bUseTimeLimit, TimeLimit, TEXT("TickAsyncPackage")))
+					{
+						return EAsyncPackageState::TimeOut;
+					}
+					UE_LOG(LogStreaming, Error, TEXT("Should not have a timeout when the time limit is not exceeded."));
+					continue;
+				}
+			}
+			else
+			{
+				// This package has finished loading but some other package is still holding
+				// a reference to it because it has this package in its dependency list.
+				// move it to the back of the list
+				AsyncPackagesReadyForTick.RemoveAt(0, 1, false); //@todoio this should maybe be a heap or something to avoid the removal cost
+				AsyncPackagesReadyForTick.Add(Package);
+			}
+			if (LocalLoadingState == EAsyncPackageState::Complete)
+			{
 
-	{
+				// We're done, at least on this thread, so we can remove the package now.
+				AddToLoadedPackages(Package);
+				{
 #if THREADSAFE_UOBJECTS
-		FScopeLock LockAsyncPackages(&AsyncPackagesCritical);
+					FScopeLock LockAsyncPackages(&AsyncPackagesCritical);
 #endif
-		int32 InsertIndex = -1;
-
-		switch (InsertMode)
+					AsyncPackageNameLookup.Remove(Package->GetPackageName());
+					int32 PackageIndex = AsyncPackages.Find(Package);
+					AsyncPackages.RemoveAt(PackageIndex);
+					AsyncPackagesReadyForTick.RemoveAt(0, 1, false); //@todoio this should maybe be a heap or something to avoid the removal cost
+				}
+			}
+			if (IsTimeLimitExceeded(TickStartTime, bUseTimeLimit, TimeLimit, TEXT("TickAsyncPackage")))
+			{
+				return EAsyncPackageState::TimeOut;
+			}
+		}
+		if (bDidSomething)
+		{
+			continue;
+		}
+		bool bAnyIOOutstanding = GPrecacheCallbackHandler.AnyIOOutstanding();
+		if (bAnyIOOutstanding)
 		{
-		case EAsyncPackageInsertMode::InsertAfterMatchingPriorities:
+			SCOPED_LOADTIMER(Package_EventIOWait);
+			double StartTime = FPlatformTime::Seconds();
+			if (bUseTimeLimit)
 			{
-				InsertIndex = AsyncPackages.IndexOfByPredicate([Package](const FAsyncPackage* Element)
+				if (bUseFullTimeLimit)
 				{
-					return Element->GetPriority() < Package->GetPriority();
-				});
-
-				break;
+					const float RemainingTimeLimit = FMath::Max(0.0f, TimeLimit - (float)(FPlatformTime::Seconds() - TickStartTime));
+					if (RemainingTimeLimit > 0.0f)
+					{
+						bool bGotIO = GPrecacheCallbackHandler.WaitForIO(RemainingTimeLimit);
+						if (bGotIO)
+						{
+							OutPackagesProcessed++;
+							continue; // we got some IO, so start processing at the top
+						}
+						{
+							float ThisTime = float(FPlatformTime::Seconds() - StartTime);
+						}
+					}
+				}
+				//UE_LOG(LogStreaming, Error, TEXT("Not using full time limit, waiting for IO..."));
+				return EAsyncPackageState::TimeOut;
 			}
-
-		case EAsyncPackageInsertMode::InsertBeforeMatchingPriorities:
+			else
 			{
-		// Insert new package keeping descending priority order in AsyncPackages
-				InsertIndex = AsyncPackages.IndexOfByPredicate([Package](const FAsyncPackage* Element)
-		{
-			return Element->GetPriority() <= Package->GetPriority();
-		});
-
-				break;
+				bool bGotIO = GPrecacheCallbackHandler.WaitForIO(10.0f); // wait "forever"
+				if (!bGotIO)
+				{
+					UE_LOG(LogStreaming, Error, TEXT("Waited for 10 seconds on IO...."));
+				}
+				{
+					OutPackagesProcessed++;
+				}
 			}
-		};
-
-		InsertIndex = InsertIndex == INDEX_NONE ? AsyncPackages.Num() : InsertIndex;
-
-		AsyncPackages.InsertUninitialized(InsertIndex);
-		AsyncPackages[InsertIndex] = Package;
-		AsyncPackageNameLookup.Add(Package->GetPackageName(), Package);
-	}
-}
-
-void FAsyncLoadingThread::JustInTimeHinting()
-{
-#if USE_NEW_ASYNC_IO
-	const int32 NumHints = 1000;
-
-	for (int32 PackageIdx = 0; PackageIdx < AsyncPackages.Num() && PendingPackageHints.Num() < NumHints; ++PackageIdx)
-	{
-		FAsyncPackage* Package = AsyncPackages[PackageIdx];
-		if (!PendingPackageHints.Contains(Package))
+		}
+		else
 		{
-			const FString PrestreamName = GetPrestreamPackageLinkerName(*Package->GetPackageNameToLoad().ToString(), false);
-			if(PrestreamName.Len())
-			{
-				HintFutureRead(*PrestreamName);
-				PendingPackageHints.Add(Package);
-			}			
+			LoadingState = EAsyncPackageState::Complete;
+			break;
 		}
-
 	}
-#endif
-}
 
-EAsyncPackageState::Type FAsyncLoadingThread::ProcessAsyncLoading(int32& OutPackagesProcessed, bool bUseTimeLimit /*= false*/, bool bUseFullTimeLimit /*= false*/, float TimeLimit /*= 0.0f*/)
-{
-	SCOPE_CYCLE_COUNTER(STAT_FAsyncLoadingThread_ProcessAsyncLoading);
-	SCOPED_LOADTIMER(AsyncLoadingTime);
-	
-	// If we're not multithreaded and flushing async loading, update the thread heartbeat
-	const bool bNeedsHeartbeatTick = !bUseTimeLimit && !FAsyncLoadingThread::IsMultithreaded();
 
-	EAsyncPackageState::Type LoadingState = EAsyncPackageState::Complete;
-	OutPackagesProcessed = 0;
+#else
+#if 0 && USE_NEW_ASYNC_IO
+bool bDepthFirst = true;
+#else
+bool bDepthFirst = false;
+#endif
 
 	// We need to loop as the function has to handle finish loading everything given no time limit
 	// like e.g. when called from FlushAsyncLoading.
-	for (int32 PackageIndex = 0; LoadingState != EAsyncPackageState::TimeOut && PackageIndex < AsyncPackages.Num(); ++PackageIndex)
+	for (int32 PackageIndex = 0; ((bDepthFirst && LoadingState == EAsyncPackageState::Complete) || (!bDepthFirst && LoadingState != EAsyncPackageState::TimeOut)) && PackageIndex < AsyncPackages.Num(); ++PackageIndex)
 	{
-		JustInTimeHinting();
-
 		SCOPED_LOADTIMER(ProcessAsyncLoadingTime);
 		OutPackagesProcessed++;
 
 		// Package to be loaded.
 		FAsyncPackage* Package = AsyncPackages[PackageIndex];
-
-		if (Package->HasFinishedLoading() == false)
+		if (FlushTree && !FlushTree->Contains(Package->GetPackageName()))
 		{
-			    // Package tick returns EAsyncPackageState::Complete on completion.
-			    // We only tick packages that have not yet been loaded.
-			    LoadingState = Package->Tick(bUseTimeLimit, bUseFullTimeLimit, TimeLimit);
+			LoadingState = EAsyncPackageState::PendingImports;
+		}
+		else if (Package->HasFinishedLoading() == false)
+		{
+#if USE_EVENT_DRIVEN_ASYNC_LOAD
+			LoadingState = EAsyncPackageState::PendingImports;
+#else
+			// @todo: Guard against recursively re-entering?
+			//if (!Package->IsCurrentlyBeingProcessed())
+			{
+				// Package tick returns EAsyncPackageState::Complete on completion.
+				// We only tick packages that have not yet been loaded.
+				LoadingState = Package->TickAsyncPackage(bUseTimeLimit, bUseFullTimeLimit, TimeLimit, FlushTree);
+			}
+#endif
 		}
 		else
 		{
@@ -727,41 +3661,39 @@ EAsyncPackageState::Type FAsyncLoadingThread::ProcessAsyncLoading(int32& OutPack
 			// a reference to it because it has this package in its dependency list.
 			LoadingState = EAsyncPackageState::Complete;
 		}
-		bool bPackageFullyLoaded = false;
 		if (LoadingState == EAsyncPackageState::Complete)
 		{
-
-#if USE_NEW_ASYNC_IO
-			if (PendingPackageHints.Contains(Package))
-			{
-				PendingPackageHints.Remove(Package);
-				HintFutureReadDone(*Package->GetPackageNameToLoad().ToString());
-			}
-#endif			
-
-
 			// We're done, at least on this thread, so we can remove the package now.
-			AddToLoadedPackages(Package);
+			if (!Package->HasThreadedLoadingFinished())
 			{
+				Package->ThreadedLoadingHasFinished();
+			AddToLoadedPackages(Package);
 #if THREADSAFE_UOBJECTS
 				FScopeLock LockAsyncPackages(&AsyncPackagesCritical);
 #endif
 				AsyncPackageNameLookup.Remove(Package->GetPackageName());
-				AsyncPackages.RemoveAt(PackageIndex);
+				AsyncPackages.Remove(Package);
+
+				// Need to process this index again as we just removed an item
+				PackageIndex--;
 			}
+
+			check(!AsyncPackages.Contains(Package));
 					
-			// Need to process this index again as we just removed an item
-			PackageIndex--;
-			bPackageFullyLoaded = true;
 		}
 		else if (!bUseTimeLimit && !FPlatformProcess::SupportsMultithreading())
 		{
+#if !USE_NEW_ASYNC_IO
 			// Tick async loading when multithreading is disabled.
 			FIOSystem::Get().TickSingleThreaded();
+#endif
 		}
 
-		// Check if there's any new packages in the queue.
-		CreateAsyncPackagesFromQueue();
+		{ // maybe we shouldn't do this if we are already out of time?
+			// Check if there's any new packages in the queue.
+			const float RemainingTimeLimit = FMath::Max(0.0f, TimeLimit - (float)(FPlatformTime::Seconds() - TickStartTime));
+			CreateAsyncPackagesFromQueue(bUseTimeLimit, bUseFullTimeLimit, RemainingTimeLimit);
+		}
 
 		if (bNeedsHeartbeatTick)
 		{
@@ -769,17 +3701,19 @@ EAsyncPackageState::Type FAsyncLoadingThread::ProcessAsyncLoading(int32& OutPack
 			FThreadHeartBeat::Get().HeartBeat();
 		}
 	}
-
+#endif
 	return LoadingState;
 }
 
 
-EAsyncPackageState::Type FAsyncLoadingThread::ProcessLoadedPackages(bool bUseTimeLimit, bool bUseFullTimeLimit, float TimeLimit, int32 WaitForRequestID /*= INDEX_NONE*/)
+EAsyncPackageState::Type FAsyncLoadingThread::ProcessLoadedPackages(bool bUseTimeLimit, bool bUseFullTimeLimit, float TimeLimit, bool& bDidSomething, FFlushTree* FlushTree)
 {
 	SCOPED_LOADTIMER(TickAsyncLoading_ProcessLoadedPackages);
 
 	EAsyncPackageState::Type Result = EAsyncPackageState::Complete;
 
+	// This is for debugging purposes only. @todo remove
+	volatile int32 CurrentAsyncLoadingCounter = AsyncLoadingTickCounter;
 
 	double TickStartTime = FPlatformTime::Seconds();
 	{
@@ -793,15 +3727,11 @@ EAsyncPackageState::Type FAsyncLoadingThread::ProcessLoadedPackages(bool bUseTim
 		LoadedPackagesNameLookup.Reset();
 	}
 		
+	bDidSomething = LoadedPackagesToProcess.Num() > 0;
 	for (int32 PackageIndex = 0; PackageIndex < LoadedPackagesToProcess.Num() && !IsAsyncLoadingSuspended(); ++PackageIndex)
 	{
 		SCOPED_LOADTIMER(ProcessLoadedPackagesTime);
 
-		if (PackageIndex % 20 == 0 && IsTimeLimitExceeded(TickStartTime, bUseTimeLimit, TimeLimit, TEXT("ProcessLoadedPackages")))
-		{
-			break;
-		}
-
 		FAsyncPackage* Package = LoadedPackagesToProcess[PackageIndex];
 
 		Result = Package->PostLoadDeferredObjects(TickStartTime, bUseTimeLimit, TimeLimit);
@@ -812,12 +3742,22 @@ EAsyncPackageState::Type FAsyncLoadingThread::ProcessLoadedPackages(bool bUseTim
 			{
 				FScopeLock LoadedLock(&LoadedPackagesToProcessCritical);					
 				LoadedPackagesToProcess.RemoveAt(PackageIndex--);
+				if (LoadedPackagesToProcess.FindByPredicate([Package](const FAsyncPackage* Pkg)
+				{
+					return Pkg->GetPackageName() == Package->GetPackageName();
+				}))
+				{
+					UE_LOG(LogStreaming, Warning, TEXT("Package %s has already been loaded"), *Package->GetPackageName().ToString());
+				}
 				LoadedPackagesToProcessNameLookup.Remove(Package->GetPackageName());
 					
 				if (FPlatformProperties::RequiresCookedData())
 				{
 					// Emulates ResetLoaders on the package linker's linkerroot.
-					Package->ResetLoader();
+					if (!Package->IsBeingProcessedRecursively())
+					{
+						Package->ResetLoader();
+					}
 				}
 				else
 				{
@@ -826,11 +3766,22 @@ EAsyncPackageState::Type FAsyncLoadingThread::ProcessLoadedPackages(bool bUseTim
 					// while this async package is still bound to it.
 					Package->DetachLinker();
 				}
-					
+				
+				// Close linkers opened by synchronous loads during async loading
+				Package->CloseDelayedLinkers();
 			}
 
 			// Incremented on the Async Thread, now decrement as we're done with this package				
 			const int32 NewExistingAsyncPackagesCounterValue = ExistingAsyncPackagesCounter.Decrement();
+			TArray<FAsyncPackage*>& AlreadyDecremented = PackagesThatDecrementedCounter.FindOrAdd(Package->GetPackageName());
+			if (!AlreadyDecremented.Contains(Package))
+			{
+				AlreadyDecremented.Add(Package);
+			}
+			else
+			{
+				UE_LOG(LogStreaming, Warning, TEXT("Package %s has already been removed"), *Package->GetPackageName().ToString());
+			}
 			UE_CLOG(NewExistingAsyncPackagesCounterValue < 0, LogStreaming, Fatal, TEXT("ExistingAsyncPackagesCounter is negative, this means we loaded more packages then requested so there must be a bug in async loading code."));
 
 			// Call external callbacks
@@ -839,9 +3790,10 @@ EAsyncPackageState::Type FAsyncLoadingThread::ProcessLoadedPackages(bool bUseTim
 			Package->CallCompletionCallbacks(bInternalCallbacks, LoadingResult);
 
 			// We don't need the package anymore
-			PackagesToDelete.Add(Package);
+			PackagesToDelete.AddUnique(Package);
+			Package->MarkRequestIDsAsComplete();
 
-			if (WaitForRequestID != INDEX_NONE && !ContainsRequestID(WaitForRequestID))
+			if (IsTimeLimitExceeded(TickStartTime, bUseTimeLimit, TimeLimit, TEXT("ProcessLoadedPackages Misc")) || (FlushTree && !ContainsRequestID(FlushTree->RequestId)))
 			{
 				// The only package we care about has finished loading, so we're good to exit
 				break;
@@ -853,14 +3805,21 @@ EAsyncPackageState::Type FAsyncLoadingThread::ProcessLoadedPackages(bool bUseTim
 		}
 	}
 
+	bDidSomething = bDidSomething || PackagesToDelete.Num() > 0;
+
 	// Delete packages we're done processing and are no longer dependencies of anything else
 	for (int32 PackageIndex = 0; PackageIndex < PackagesToDelete.Num(); ++PackageIndex)
 	{
 		FAsyncPackage* Package = PackagesToDelete[PackageIndex];
-		if (Package->GetDependencyRefCount() == 0)
+		if (Package->GetDependencyRefCount() == 0 && !Package->IsBeingProcessedRecursively())
 		{
-			delete Package;
 			PackagesToDelete.RemoveAtSwap(PackageIndex--);
+			delete Package;
+			if (IsTimeLimitExceeded(TickStartTime, bUseTimeLimit, TimeLimit, TEXT("ProcessLoadedPackages PackagesToDelete")))
+			{
+				Result = EAsyncPackageState::TimeOut;
+				break;
+			}
 		}
 	}
 
@@ -869,14 +3828,13 @@ EAsyncPackageState::Type FAsyncLoadingThread::ProcessLoadedPackages(bool bUseTim
 		// We're not done until all packages have been deleted
 		Result = PackagesToDelete.Num() ? EAsyncPackageState::PendingImports : EAsyncPackageState::Complete;
 	}
-	
 
 	return Result;
 }
 
 
 
-EAsyncPackageState::Type FAsyncLoadingThread::TickAsyncLoading(bool bUseTimeLimit, bool bUseFullTimeLimit, float TimeLimit, int32 WaitForRequestID /*= INDEX_NONE*/)
+EAsyncPackageState::Type FAsyncLoadingThread::TickAsyncLoading(bool bUseTimeLimit, bool bUseFullTimeLimit, float TimeLimit, FFlushTree* FlushTree)
 {
 	const bool bLoadingSuspended = IsAsyncLoadingSuspended();
 	const bool bIsMultithreaded = FAsyncLoadingThread::IsMultithreaded();
@@ -887,29 +3845,41 @@ EAsyncPackageState::Type FAsyncLoadingThread::TickAsyncLoading(bool bUseTimeLimi
 		double TickStartTime = FPlatformTime::Seconds();
 		double TimeLimitUsedForProcessLoaded = 0;
 
+		bool bDidSomething = false;
 		{
-			Result = ProcessLoadedPackages(bUseTimeLimit, bUseFullTimeLimit, TimeLimit, WaitForRequestID);
+			Result = ProcessLoadedPackages(bUseTimeLimit, bUseFullTimeLimit, TimeLimit, bDidSomething, FlushTree);
 			TimeLimitUsedForProcessLoaded = FPlatformTime::Seconds() - TickStartTime;
+			UE_CLOG(bUseTimeLimit && TimeLimitUsedForProcessLoaded > .1f, LogStreaming, Warning, TEXT("Took %6.2fms to ProcessLoadedPackages"), float(TimeLimitUsedForProcessLoaded) * 1000.0f);
 		}
 
-		if (!bIsMultithreaded && Result != EAsyncPackageState::TimeOut && !IsTimeLimitExceeded(TickStartTime, bUseTimeLimit, TimeLimit, TEXT("Pre-TickAsyncThread")))
+		if (!bIsMultithreaded && Result != EAsyncPackageState::TimeOut && !IsTimeLimitExceeded(TickStartTime, bUseTimeLimit, TimeLimit, TEXT("ProcessLoadedPackages")))
 		{
 			double RemainingTimeLimit = FMath::Max(0.0, TimeLimit - TimeLimitUsedForProcessLoaded);
-			Result = TickAsyncThread(bUseTimeLimit, bUseFullTimeLimit, RemainingTimeLimit);			
+			Result = TickAsyncThread(bUseTimeLimit, bUseFullTimeLimit, RemainingTimeLimit, bDidSomething, FlushTree);
 		}
 
-		if (Result != EAsyncPackageState::TimeOut && !IsTimeLimitExceeded(TickStartTime, bUseTimeLimit, TimeLimit, TEXT("Pre-EmptyReferencedObjects")))
+		if (Result != EAsyncPackageState::TimeOut && !IsTimeLimitExceeded(TickStartTime, bUseTimeLimit, TimeLimit, TEXT("TickAsyncThread")))
 		{
+			{
 #if THREADSAFE_UOBJECTS
-			FScopeLock QueueLock(&QueueCritical);
-			FScopeLock LoadedLock(&LoadedPackagesCritical);
+				FScopeLock QueueLock(&QueueCritical);
+				FScopeLock LoadedLock(&LoadedPackagesCritical);
 #endif
-			// Release references we acquired when async loading when there are no more FAsyncPackages in existence
-			if (ExistingAsyncPackagesCounter.GetValue() == 0)
+				// Release references we acquired when async loading when there are no more FAsyncPackages in existence
+				if (ExistingAsyncPackagesCounter.GetValue() == 0)
+				{
+					bDidSomething = true; // we are all done, no need to check for cycles
+					FDeferredMessageLog::Flush();
+					IsTimeLimitExceeded(TickStartTime, bUseTimeLimit, TimeLimit, TEXT("EmptyReferencedObjects"));
+				}
+			}
+#if USE_EVENT_DRIVEN_ASYNC_LOAD
+			if (!bDidSomething && !bIsMultithreaded)
 			{
-				FDeferredMessageLog::Flush();
-				FAsyncObjectsReferencer::Get().EmptyReferencedObjects();
+				CheckForCycles();
+				IsTimeLimitExceeded(TickStartTime, bUseTimeLimit, TimeLimit, TEXT("CheckForCycles (non-shipping)"));
 			}
+#endif
 		}
 	}
 
@@ -918,9 +3888,14 @@ EAsyncPackageState::Type FAsyncLoadingThread::TickAsyncLoading(bool bUseTimeLimi
 
 FAsyncLoadingThread::FAsyncLoadingThread()
 {
-#if !UE_BUILD_SHIPPING
-	GAsyncLoadingExec = new FAsyncLoadingExec();
+#if USE_EVENT_DRIVEN_ASYNC_LOAD
+	UE_CLOG(!IsEventDrivenLoaderEnabledInCookedBuilds(), LogStreaming, Fatal,
+		TEXT("Event driven async loader is being used but it does NOT seem to be enabled in project settings."));
+#elif !WITH_EDITORONLY_DATA
+	UE_CLOG(IsEventDrivenLoaderEnabledInCookedBuilds(), LogStreaming, Fatal,
+		TEXT("Event driven async loader is NOT being used but it seems to be enabled in project settings."));
 #endif
+
 	QueuedRequestsEvent = FPlatformProcess::GetSynchEventFromPool();
 	CancelLoadingEvent = FPlatformProcess::GetSynchEventFromPool();
 	ThreadSuspendedEvent = FPlatformProcess::GetSynchEventFromPool();
@@ -936,11 +3911,21 @@ FAsyncLoadingThread::FAsyncLoadingThread()
 		Thread = nullptr;
 		Init();
 	}
-	bIsInAsyncLoadingTick = false;
+	AsyncLoadingTickCounter = 0;
 }
 
 FAsyncLoadingThread::~FAsyncLoadingThread()
 {
+#if USE_EVENT_DRIVEN_ASYNC_LOAD
+	// check that event queue is empty
+	const TCHAR* LastTypeOfWorkPerformed = nullptr;
+	UObject* LastObjectWorkWasPerformedOn = nullptr;
+
+	FAsyncLoadEventArgs Args;
+	check(!EventQueue.PopAndExecute(Args));
+
+#endif
+
 	delete Thread;
 	Thread = nullptr;
 	FPlatformProcess::ReturnSynchEventToPool(QueuedRequestsEvent);
@@ -971,8 +3956,9 @@ uint32 FAsyncLoadingThread::Run()
 			{
 				bWasSuspendedLastFrame = false;
 				ThreadResumedEvent->Trigger();
-			}			
-			TickAsyncThread(false, true, 0.0f);
+			}		
+			bool bDidSomething = false;
+			TickAsyncThread(false, true, 0.0f, bDidSomething);
 		}
 		else if (!bWasSuspendedLastFrame)
 		{
@@ -987,27 +3973,74 @@ uint32 FAsyncLoadingThread::Run()
 	return 0;
 }
 
-EAsyncPackageState::Type  FAsyncLoadingThread::TickAsyncThread(bool bUseTimeLimit, bool bUseFullTimeLimit, float TimeLimit)
+#if USE_EVENT_DRIVEN_ASYNC_LOAD
+void FAsyncLoadingThread::CheckForCycles()
+{
+	// no outstanding IO, nothing was done in this iteration, we are done
+	FAsyncPackage::GlobalEventGraph.CheckForCycles();
+
+#if !UE_BUILD_SHIPPING && !UE_BUILD_TEST
+	// lets look over the postload wait stuff and see if that is bugged
+	for (int32 PackageIdx = 0; PackageIdx < AsyncPackages.Num(); ++PackageIdx)
+	{
+		FAsyncPackage* Package = AsyncPackages[PackageIdx];
+		if (Package->AsyncPackageLoadingState == EAsyncPackageLoadingState::WaitingForPostLoad)
+		{
+			UE_CLOG(!Package->PackagesIAmWaitingForBeforePostload.Num(), LogStreaming, Fatal, TEXT("We have nothing to do and there is no IO outstanding, yet %s is waiting for NO other packages to serialize:"), *Package->GetPackageName().ToString());
+			UE_LOG(LogStreaming, Error, TEXT("We have nothing to do and there is no IO outstanding, yet %s is waiting for other packages to serialize:"), *Package->GetPackageName().ToString());
+
+			for (FWeakAsyncPackagePtr Test : Package->PackagesIAmWaitingForBeforePostload)
+			{
+				UE_LOG(LogStreaming, Error, TEXT("    Waiting for %s"), *Test.HumanReadableStringForDebugging().ToString());
+			}
+		}
+	}
+#endif
+}
+#endif
+
+
+
+EAsyncPackageState::Type  FAsyncLoadingThread::TickAsyncThread(bool bUseTimeLimit, bool bUseFullTimeLimit, float TimeLimit, bool& bDidSomething, FFlushTree* FlushTree)
 {
 	EAsyncPackageState::Type Result = EAsyncPackageState::Complete;
 	if (!bShouldCancelLoading)
 	{
 		int32 ProcessedRequests = 0;
+		double TickStartTime = FPlatformTime::Seconds();
 		if (AsyncThreadReady.GetValue())
 		{
-			CreateAsyncPackagesFromQueue();
-			Result = ProcessAsyncLoading(ProcessedRequests, bUseTimeLimit, bUseFullTimeLimit, TimeLimit);
+			CreateAsyncPackagesFromQueue(bUseTimeLimit, bUseFullTimeLimit, TimeLimit, FlushTree);
+			float TimeUsed = (float)(FPlatformTime::Seconds() - TickStartTime);
+			const float RemainingTimeLimit = FMath::Max(0.0f, TimeLimit - TimeUsed);
+			if (RemainingTimeLimit <= 0.0f && bUseTimeLimit && !IsMultithreaded())
+			{
+				Result = EAsyncPackageState::TimeOut;
+			}
+			else
+			{
+				Result = ProcessAsyncLoading(ProcessedRequests, bUseTimeLimit, bUseFullTimeLimit, RemainingTimeLimit, FlushTree);
+				bDidSomething = bDidSomething || ProcessedRequests > 0;
+			}
 		}
 		if (ProcessedRequests == 0 && IsMultithreaded())
 		{
+#if USE_EVENT_DRIVEN_ASYNC_LOAD
+			CheckForCycles();
+			IsTimeLimitExceeded(TickStartTime, bUseTimeLimit, TimeLimit, TEXT("CheckForCycles (non-shipping)"));
+#endif
 			const bool bIgnoreThreadIdleStats = true;
+			SCOPED_LOADTIMER(Package_Temp3);
 			QueuedRequestsEvent->Wait(30, bIgnoreThreadIdleStats);
 		}
+
 	}
 	else
 	{
 		// Blocks main thread
+		double TickStartTime = FPlatformTime::Seconds();
 		CancelAsyncLoadingInternal();
+		IsTimeLimitExceeded(TickStartTime, bUseTimeLimit, TimeLimit, TEXT("CancelAsyncLoadingInternal"));
 		bShouldCancelLoading = false;
 	}
 
@@ -1121,18 +4154,15 @@ void NotifyConstructedDuringAsyncLoading(UObject* Object, bool bSubObject)
 	{
 		Object->SetInternalFlags(EInternalObjectFlags::AsyncLoading);
 	}
-	FAsyncObjectsReferencer::Get().AddObject(Object);
+	FUObjectThreadContext& ThreadContext = FUObjectThreadContext::Get();
+	check(ThreadContext.AsyncPackage); // Otherwise something is wrong and we're creating objects outside of async loading code
+	ThreadContext.AsyncPackage->AddObjectReference(Object);
 }
 
 /*-----------------------------------------------------------------------------
 	FAsyncPackage implementation.
 -----------------------------------------------------------------------------*/
 
-
-int32 FAsyncPackage::PreLoadIndex = 0;
-int32 FAsyncPackage::PreLoadSortIndex = 0;
-int32 FAsyncPackage::PostLoadIndex = 0;
-
 /**
 * Constructor
 */
@@ -1145,19 +4175,37 @@ FAsyncPackage::FAsyncPackage(const FAsyncPackageDesc& InDesc)
 , LoadImportIndex(0)
 , ImportIndex(0)
 , ExportIndex(0)
+, PreLoadIndex(0)
+, PreLoadSortIndex(0)
+, PostLoadIndex(0)
 , DeferredPostLoadIndex(0)
+, DeferredFinalizeIndex(0)
 , TimeLimit(FLT_MAX)
 , bUseTimeLimit(false)
 , bUseFullTimeLimit(false)
 , bTimeLimitExceeded(false)
 , bLoadHasFailed(false)
 , bLoadHasFinished(false)
+, bThreadedLoadingFinished(false)
 , TickStartTime(0)
 , LastObjectWorkWasPerformedOn(nullptr)
 , LastTypeOfWorkPerformed(nullptr)
 , LoadStartTime(0.0)
 , LoadPercentage(0)
+, ReentryCount(0)
 , AsyncLoadingThread(FAsyncLoadingThread::Get())
+#if USE_EVENT_DRIVEN_ASYNC_LOAD
+, AsyncPackageLoadingState(EAsyncPackageLoadingState::NewPackage)
+, SerialNumber(AsyncPackageSerialNumber.Increment())
+, CurrentBlockOffset(-1)
+, CurrentBlockBytes(-1)
+, ImportAddNodeIndex(0)
+, ExportAddNodeIndex(0)
+, bProcessImportsAndExportsInFlight(false)
+, bProcessPostloadWaitInFlight(false)
+, bAllExportsSerialized(false)
+#endif
+
 #if PERF_TRACK_DETAILED_ASYNC_STATS
 , TickCount(0)
 , TickLoopCount(0)
@@ -1183,19 +4231,93 @@ FAsyncPackage::FAsyncPackage(const FAsyncPackageDesc& InDesc)
 
 FAsyncPackage::~FAsyncPackage()
 {
-	AsyncLoadingThread.RemovePendingRequests(RequestIDs);
+#if DO_CHECK && USE_EVENT_DRIVEN_ASYNC_LOAD
+	for (FCompletionCallback& CompletionCallback : CompletionCallbacks)
+	{
+		if (!CompletionCallback.bCalled)
+		{
+			check(0);
+		}
+	}
+#endif
+
+	MarkRequestIDsAsComplete();
 	DetachLinker();
+	check(!FAsyncLoadingThread::Get().IsPackageReferenced(this));
+#if USE_EVENT_DRIVEN_ASYNC_LOAD
+	SerialNumber = 0; // the weak pointer will always fail now
+	check(!EventNodeArray.Array && !EventNodeArray.TotalNumberOfNodesAdded);
+	RemoveAllNodes();
+#endif
+
+	EmptyReferencedObjects();
+}
+
+void FAsyncPackage::AddReferencedObjects(FReferenceCollector& Collector)
+{
+	Collector.AllowEliminatingReferences(false);
+	{
+		Collector.AddReferencedObjects(ReferencedObjects);
+	}
+	Collector.AllowEliminatingReferences(true);
+}
+
+void FAsyncPackage::AddObjectReference(UObject* InObject)
+{
+	if (InObject)
+	{
+		UE_CLOG(!IsInGameThread() && !IsGarbageCollectionLocked(), LogStreaming, Fatal, TEXT("Trying to add an object %s to FAsyncPackage referenced objects list outside of a FGCScopeLock."), *InObject->GetFullName());
+		{
+#if USE_EVENT_DRIVEN_ASYNC_LOAD
+			FScopeLock ReferencedObjectsLock(&ReferencedObjectsCritical);
+#endif
+			if (!ReferencedObjects.Contains(InObject))
+			{
+				ReferencedObjects.Add(InObject);
+			}
+		}
+		InObject->ThisThreadAtomicallyClearedRFUnreachable();
+	}
+}
+
+void FAsyncPackage::EmptyReferencedObjects()
+{
+	check(IsInGameThread());
+	const EInternalObjectFlags AsyncFlags = EInternalObjectFlags::Async | EInternalObjectFlags::AsyncLoading;
+#if USE_EVENT_DRIVEN_ASYNC_LOAD
+	FScopeLock ReferencedObjectsLock(&ReferencedObjectsCritical);
+#endif
+	for (UObject* Obj : ReferencedObjects)
+	{
+		// Temporary fatal messages instead of checks to find the cause for a one-time crash in shipping config
+		UE_CLOG(Obj == nullptr, LogStreaming, Fatal, TEXT("NULL object in Async Objects Referencer"));
+		UE_CLOG(!Obj->IsValidLowLevelFast(), LogStreaming, Fatal, TEXT("Invalid object in Async Objects Referencer"));
+		Obj->AtomicallyClearInternalFlags(AsyncFlags);
+		check(!Obj->HasAnyInternalFlags(AsyncFlags))
+	}
+	ReferencedObjects.Reset();
 }
 
 void FAsyncPackage::AddRequestID(int32 Id)
 {
 	if (Id > 0)
 	{
+		if (Desc.RequestID == INDEX_NONE)
+		{
+			// For debug readability
+			Desc.RequestID = Id;
+		}
 		RequestIDs.Add(Id);
 		AsyncLoadingThread.AddPendingRequest(Id);
 	}
 }
 
+void FAsyncPackage::MarkRequestIDsAsComplete()
+{
+	AsyncLoadingThread.RemovePendingRequests(RequestIDs);
+	RequestIDs.Reset();
+}
+
 /**
  * @return Time load begun. This is NOT the time the load was requested in the case of other pending requests.
  */
@@ -1239,11 +4361,13 @@ void FAsyncPackage::DetachLinker()
  */
 bool FAsyncPackage::GiveUpTimeSlice()
 {
+#if !USE_NEW_ASYNC_IO
 	static const bool bPlatformIsSingleThreaded = !FPlatformProcess::SupportsMultithreading();
 	if (bPlatformIsSingleThreaded)
 	{
 		FIOSystem::Get().TickSingleThreaded();
 	}
+#endif
 
 	if (bUseTimeLimit && !bUseFullTimeLimit)
 	{
@@ -1261,7 +4385,7 @@ void FAsyncPackage::BeginAsyncLoad()
 {
 	if (IsInGameThread())
 	{
-		FAsyncLoadingThread::Get().SetIsInAsyncLoadingTick(true);
+		FAsyncLoadingThread::Get().EnterAsyncLoadingTick();
 	}
 
 	// this won't do much during async loading except increase the load count which causes IsLoading to return true
@@ -1281,7 +4405,7 @@ void FAsyncPackage::EndAsyncLoad()
 
 	if (IsInGameThread())
 	{
-		FAsyncLoadingThread::Get().SetIsInAsyncLoadingTick(false);
+		FAsyncLoadingThread::Get().LeaveAsyncLoadingTick();
 	}
 
 	if (!bLoadHasFailed)
@@ -1301,17 +4425,19 @@ void FAsyncPackage::EndAsyncLoad()
  * @return	true if package has finished loading, false otherwise
  */
 
-EAsyncPackageState::Type FAsyncPackage::Tick(bool InbUseTimeLimit, bool InbUseFullTimeLimit, float& InOutTimeLimit)
+EAsyncPackageState::Type FAsyncPackage::TickAsyncPackage(bool InbUseTimeLimit, bool InbUseFullTimeLimit, float& InOutTimeLimit, FFlushTree* FlushTree)
 {
+#if USE_EVENT_DRIVEN_ASYNC_LOAD
+	check(int32(AsyncPackageLoadingState) > int32(EAsyncPackageLoadingState::ProcessNewImportsAndExports));
+#endif
+	ReentryCount++;
+
 	SCOPE_CYCLE_COUNTER(STAT_FAsyncPackage_Tick);
 	SCOPED_LOADTIMER(Package_Tick);
 
 	// Whether we should execute the next step.
 	EAsyncPackageState::Type LoadingState = EAsyncPackageState::Complete;
 
-	check(LastObjectWorkWasPerformedOn == nullptr);
-	check(LastTypeOfWorkPerformed == nullptr);
-
 	// Set up tick relevant variables.
 	bUseTimeLimit = InbUseTimeLimit;
 	bUseFullTimeLimit = InbUseFullTimeLimit;
@@ -1336,6 +4462,8 @@ EAsyncPackageState::Type FAsyncPackage::Tick(bool InbUseTimeLimit, bool InbUseFu
 		}
 	}
 
+	FAsyncPackageScope PackageScope(this);
+
 	// Make sure we finish our work if there's no time limit. The loop is required as PostLoad
 	// might cause more objects to be loaded in which case we need to Preload them again.
 	do
@@ -1349,6 +4477,7 @@ EAsyncPackageState::Type FAsyncPackage::Tick(bool InbUseTimeLimit, bool InbUseFu
 		// We have begun loading a package that we know the name of. Let the package time tracker know.
 		FExclusiveLoadPackageTimeTracker::PushLoadPackage(Desc.NameToLoad);
 
+#if !USE_EVENT_DRIVEN_ASYNC_LOAD
 		// Create raw linker. Needs to be async created via ticking before it can be used.
 		if (LoadingState == EAsyncPackageState::Complete)
 		{
@@ -1367,7 +4496,7 @@ EAsyncPackageState::Type FAsyncPackage::Tick(bool InbUseTimeLimit, bool InbUseFu
 		if (LoadingState == EAsyncPackageState::Complete)
 		{
 			SCOPED_LOADTIMER(Package_LoadImports);
-			LoadingState = LoadImports();
+			LoadingState = LoadImports(FlushTree);
 		}
 
 		// Create imports from linker import table.
@@ -1377,13 +4506,6 @@ EAsyncPackageState::Type FAsyncPackage::Tick(bool InbUseTimeLimit, bool InbUseFu
 			LoadingState = CreateImports();
 		}
 
-		// Finish all async texture allocations.
-		if (LoadingState == EAsyncPackageState::Complete)
-		{
-			SCOPED_LOADTIMER(Package_FinishTextureAllocations);
-			LoadingState = FinishTextureAllocations();
-		}
-
 #if WITH_EDITORONLY_DATA
 		// Create and preload the package meta-data
 		if (LoadingState == EAsyncPackageState::Complete)
@@ -1406,7 +4528,7 @@ EAsyncPackageState::Type FAsyncPackage::Tick(bool InbUseTimeLimit, bool InbUseFu
 			SCOPED_LOADTIMER(Package_PreLoadObjects);
 			LoadingState = PreLoadObjects();
 		}
-
+#endif
 		// Call PostLoad on objects, this could cause new objects to be loaded that require
 		// another iteration of the PreLoad loop.
 		if (LoadingState == EAsyncPackageState::Complete)
@@ -1432,15 +4554,30 @@ EAsyncPackageState::Type FAsyncPackage::Tick(bool InbUseTimeLimit, bool InbUseFu
 
 	check(bUseTimeLimit || LoadingState != EAsyncPackageState::TimeOut || AsyncLoadingThread.IsAsyncLoadingSuspended());
 
+	if (LinkerRoot && LoadingState == EAsyncPackageState::Complete)
+	{
+		LinkerRoot->MarkAsFullyLoaded();
+	}
+
 	// We can't have a reference to a UObject.
 	LastObjectWorkWasPerformedOn = nullptr;
 	// Reset type of work performed.
 	LastTypeOfWorkPerformed = nullptr;
 	// Mark this package as loaded if everything completed.
 	bLoadHasFinished = (LoadingState == EAsyncPackageState::Complete);
+#if USE_EVENT_DRIVEN_ASYNC_LOAD
+	if (bLoadHasFinished)
+	{
+		check(AsyncPackageLoadingState == EAsyncPackageLoadingState::PostLoad_Etc);
+		AsyncPackageLoadingState = EAsyncPackageLoadingState::PackageComplete;
+	}
+#endif
 	// Subtract the time it took to load this package from the global limit.
 	InOutTimeLimit = (float)FMath::Max(0.0, InOutTimeLimit - (FPlatformTime::Seconds() - TickStartTime));
 
+	ReentryCount--;
+	check(ReentryCount >= 0);
+
 	// true means that we're done loading this package.
 	return LoadingState;
 }
@@ -1465,7 +4602,7 @@ EAsyncPackageState::Type FAsyncPackage::CreateLinker()
 		{
 			FGCScopeGuard GCGuard;
 			Package = CreatePackage(nullptr, *Desc.Name.ToString());
-			FAsyncObjectsReferencer::Get().AddObject(Package);
+			AddObjectReference(Package);
 			LinkerRoot = Package;
 		}
 		FScopeCycleCounterUObject ConstructorScope(Package, GET_STATID(STAT_FAsyncPackage_CreateLinker));
@@ -1481,9 +4618,21 @@ EAsyncPackageState::Type FAsyncPackage::CreateLinker()
 		Package->SetLoadedByEditorPropertiesOnly(false);
 #endif
 
+		LastObjectWorkWasPerformedOn = Package;
 		// if the linker already exists, we don't need to lookup the file (it may have been pre-created with
 		// a different filename)
 		Linker = FLinkerLoad::FindExistingLinkerForPackage(Package);
+#if USE_EVENT_DRIVEN_ASYNC_LOAD
+		if (Linker)
+		{
+			// this almost works, but the EDL does not tolerate redoing steps it already did
+			UE_LOG(LogStreaming, Fatal , TEXT("Package %s was reloaded before it even closed the linker from a previous load. Seems like a waste of time eh?"), *Desc.Name.ToString());
+			check(Package);
+			FWeakAsyncPackagePtr WeakPtr(this);
+			GPrecacheCallbackHandler.RegisterNewSummaryRequest(this);
+			GPrecacheCallbackHandler.SummaryComplete(WeakPtr);
+		}
+#endif
 
 		if (!Linker)
 		{
@@ -1500,6 +4649,9 @@ EAsyncPackageState::Type FAsyncPackage::CreateLinker()
 				!DoesPackageExist))
 			{
 				UE_LOG(LogStreaming, Error, TEXT("Couldn't find file for package %s requested by async loading code. NameToLoad: %s"), *Desc.Name.ToString(), *Desc.NameToLoad.ToString());
+#if !WITH_EDITORONLY_DATA
+				UE_CLOG(bUseTimeLimit, LogStreaming, Error, TEXT("This will hitch streaming because it ends up searching the disk instead of finding the file in the pak file."));
+#endif
 				bLoadHasFailed = true;
 				return EAsyncPackageState::TimeOut;
 			}
@@ -1508,7 +4660,7 @@ EAsyncPackageState::Type FAsyncPackage::CreateLinker()
 			uint32 LinkerFlags = LOAD_None;
 			if (FApp::IsGame() && !GIsEditor)
 			{
-				LinkerFlags |= (LOAD_SeekFree | LOAD_NoVerify);
+				LinkerFlags |= (LOAD_Async | LOAD_NoVerify);
 			}
 #if WITH_EDITOR
 			else if ((Desc.PackageFlags & PKG_PlayInEditor) != 0)
@@ -1516,7 +4668,23 @@ EAsyncPackageState::Type FAsyncPackage::CreateLinker()
 				LinkerFlags |= LOAD_PackageForPIE;
 			}
 #endif
+#if USE_EVENT_DRIVEN_ASYNC_LOAD
+			FWeakAsyncPackagePtr WeakPtr(this);
+			check(Package);
+			Linker = FLinkerLoad::CreateLinkerAsync(Package, *PackageFileName, LinkerFlags
+				, TFunction<void()>(
+					[WeakPtr]()
+					{
+						GPrecacheCallbackHandler.SummaryComplete(WeakPtr);
+					}
+				));
+			if (Linker)
+			{
+				GPrecacheCallbackHandler.RegisterNewSummaryRequest(this);
+			}
+#else
 			Linker = FLinkerLoad::CreateLinkerAsync(Package, *PackageFileName, LinkerFlags);
+#endif
 		}
 
 		// Associate this async package with the linker
@@ -1548,7 +4716,7 @@ EAsyncPackageState::Type FAsyncPackage::FinishLinker()
 		// Operation still pending if Tick returns false
 		FLinkerLoad::ELinkerStatus LinkerResult = Linker->Tick(RemainingTimeLimit, bUseTimeLimit, bUseFullTimeLimit);
 		if (LinkerResult != FLinkerLoad::LINKER_Loaded)
-		{
+		{ 
 			// Give up remainder of timeslice if there is one to give up.
 			GiveUpTimeSlice();
 			Result = EAsyncPackageState::TimeOut;
@@ -1589,9 +4757,11 @@ FORCEINLINE int32 ContainsDependencyPackage(TArray<FAsyncPackage*>& Dependencies
 	*
 	* @param PendingImport Name of the package imported either directly or by one of the imported packages
 	*/
-void FAsyncPackage::AddImportDependency(const FName& PendingImport)
+void FAsyncPackage::AddImportDependency(const FName& PendingImport, FFlushTree* FlushTree)
 {
 	FAsyncPackage* PackageToStream = FAsyncLoadingThread::Get().FindAsyncPackage(PendingImport);
+	const bool bReinsert = PackageToStream != nullptr;
+
 	if (!PackageToStream)
 	{
 		const FAsyncPackageDesc Info(INDEX_NONE, PendingImport);
@@ -1602,9 +4772,18 @@ void FAsyncPackage::AddImportDependency(const FName& PendingImport)
 		{
 			PackageToStream->Desc.Priority = Desc.Priority;
 		}
-		FAsyncLoadingThread::Get().InsertPackage(PackageToStream);
 	}
-	
+#if USE_NEW_ASYNC_IO
+	const bool bDepthFirst = true;
+#else
+	const bool bDepthFirst = false;
+#endif
+
+	if (!bReinsert || bDepthFirst)
+	{
+		FAsyncLoadingThread::Get().InsertPackage(PackageToStream, bReinsert);
+	}
+
 	if (!PackageToStream->HasFinishedLoading() && 
 		!PackageToStream->bLoadHasFailed)
 	{
@@ -1612,6 +4791,10 @@ void FAsyncPackage::AddImportDependency(const FName& PendingImport)
 		PackageToStream->AddCompletionCallback(FLoadPackageAsyncDelegate::CreateRaw(this, &FAsyncPackage::ImportFullyLoadedCallback), bInternalCallback);
 		PackageToStream->DependencyRefCount.Increment();
 		PendingImportedPackages.Add(PackageToStream);
+		if (FlushTree)
+		{
+			PackageToStream->PopulateFlushTree(FlushTree);
+		}
 	}
 	else
 	{
@@ -1626,14 +4809,14 @@ void FAsyncPackage::AddImportDependency(const FName& PendingImport)
  *
  * @param PendingImport Package imported either directly or by one of the imported packages
  */
-bool FAsyncPackage::AddUniqueLinkerDependencyPackage(FAsyncPackage& PendingImport)
+bool FAsyncPackage::AddUniqueLinkerDependencyPackage(FAsyncPackage& PendingImport, FFlushTree* FlushTree)
 {
 	if (ContainsDependencyPackage(PendingImportedPackages, PendingImport.GetPackageName()) == INDEX_NONE)
 	{
 		FLinkerLoad* PendingImportLinker = PendingImport.Linker;
 		if (PendingImportLinker == nullptr || !PendingImportLinker->HasFinishedInitialization())
 		{
-			AddImportDependency(PendingImport.GetPackageName());
+			AddImportDependency(PendingImport.GetPackageName(), FlushTree);
 			UE_LOG(LogStreaming, Verbose, TEXT("  Adding linker dependency %s"), *PendingImport.GetPackageName().ToString());
 		}
 		else if (this != &PendingImport)
@@ -1649,7 +4832,7 @@ bool FAsyncPackage::AddUniqueLinkerDependencyPackage(FAsyncPackage& PendingImpor
  *
  * @param ImportedPackage Package imported either directly or by one of the imported packages
  */
-void FAsyncPackage::AddDependencyTree(FAsyncPackage& ImportedPackage, TSet<FAsyncPackage*>& SearchedPackages)
+void FAsyncPackage::AddDependencyTree(FAsyncPackage& ImportedPackage, TSet<FAsyncPackage*>& SearchedPackages, FFlushTree* FlushTree)
 {
 	if (SearchedPackages.Contains(&ImportedPackage))
 	{
@@ -1659,9 +4842,9 @@ void FAsyncPackage::AddDependencyTree(FAsyncPackage& ImportedPackage, TSet<FAsyn
 	for (int32 Index = 0; Index < ImportedPackage.PendingImportedPackages.Num(); ++Index)
 	{
 		FAsyncPackage& PendingImport = *ImportedPackage.PendingImportedPackages[Index];
-		if (!AddUniqueLinkerDependencyPackage(PendingImport))
+		if (!AddUniqueLinkerDependencyPackage(PendingImport, FlushTree))
 		{
-			AddDependencyTree(PendingImport, SearchedPackages);
+			AddDependencyTree(PendingImport, SearchedPackages, FlushTree);
 		}
 	}
 	// Mark this package as searched
@@ -1673,12 +4856,12 @@ void FAsyncPackage::AddDependencyTree(FAsyncPackage& ImportedPackage, TSet<FAsyn
  *
  * @return true if we finished load all imports, false otherwise
  */
-EAsyncPackageState::Type FAsyncPackage::LoadImports()
+EAsyncPackageState::Type FAsyncPackage::LoadImports(FFlushTree* FlushTree)
 {
 	SCOPE_CYCLE_COUNTER(STAT_FAsyncPackage_LoadImports);
 	LastObjectWorkWasPerformedOn	= LinkerRoot;
 	LastTypeOfWorkPerformed			= TEXT("loading imports");
-	
+
 	// GC can't run in here
 	FGCScopeGuard GCGuard;
 
@@ -1726,7 +4909,7 @@ EAsyncPackageState::Type FAsyncPackage::LoadImports()
 				if (PendingPackageLinker == nullptr || !PendingPackageLinker->HasFinishedInitialization())
 				{
 					// Add this import to the dependency list.
-					AddUniqueLinkerDependencyPackage(*PendingPackage);
+					AddUniqueLinkerDependencyPackage(*PendingPackage, FlushTree);
 				}
 				else
 				{
@@ -1736,7 +4919,7 @@ EAsyncPackageState::Type FAsyncPackage::LoadImports()
 					ReferencedImports.Add(PendingPackage);
 					// Check if we need to add its dependencies too.
 					TSet<FAsyncPackage*> SearchedPackages;
-					AddDependencyTree(*PendingPackage, SearchedPackages);
+					AddDependencyTree(*PendingPackage, SearchedPackages, FlushTree);
 				}
 			}
 		}
@@ -1748,7 +4931,7 @@ EAsyncPackageState::Type FAsyncPackage::LoadImports()
 			if (!FPackageName::IsShortPackageName(ImportPackageName))
 			{
 				UE_LOG(LogStreaming, Verbose, TEXT("FAsyncPackage::LoadImports for %s: Loading %s"), *Desc.NameToLoad.ToString(), *ImportPackageName);
-				AddImportDependency(ImportPackageFName);
+				AddImportDependency(ImportPackageFName, FlushTree);
 #if 0 && USE_NEW_ASYNC_IO
 				LoadImportIndex--;
 				return EAsyncPackageState::TimeOut; // lets go load that package now so we can do it in order.
@@ -1781,11 +4964,13 @@ void FAsyncPackage::ImportFullyLoadedCallback(const FName& InPackageName, UPacka
 	{
 		UE_LOG(LogStreaming, Verbose, TEXT("FAsyncPackage::LoadImports for %s: Loaded %s"), *Desc.NameToLoad.ToString(), *InPackageName.ToString());
 		int32 Index = ContainsDependencyPackage(PendingImportedPackages, InPackageName);
-		check(Index != INDEX_NONE);
+		if (Index != INDEX_NONE)
+		{
 		// Keep a reference to this package so that its linker doesn't go away too soon
 		ReferencedImports.Add(PendingImportedPackages[Index]);
 		PendingImportedPackages.RemoveAt(Index);
 	}
+	}
 }
 
 /** 
@@ -1809,44 +4994,12 @@ EAsyncPackageState::Type FAsyncPackage::CreateImports()
 		LastTypeOfWorkPerformed			= TEXT("creating imports for");
 
 		// Make sure this object is not claimed by GC if it's triggered while streaming.
-		FAsyncObjectsReferencer::Get().AddObject(Object);
+		AddObjectReference(Object);
 	}
 
 	return ImportIndex == Linker->ImportMap.Num() ? EAsyncPackageState::Complete : EAsyncPackageState::TimeOut;
 }
 
-/**
- * Checks if all async texture allocations for this package have been completed.
- *
- * @return true if all texture allocations have been completed, false otherwise
- */
-EAsyncPackageState::Type FAsyncPackage::FinishTextureAllocations()
-{
-	//@TODO: Cancel allocations if they take too long.
-#if WITH_ENGINE
-	bool bHasCompleted = Linker->Summary.TextureAllocations.HasCompleted();
-	if ( !bHasCompleted )
-	{
-		SCOPE_CYCLE_COUNTER(STAT_FAsyncPackage_FinishTextureAllocations);
-		if ( bUseTimeLimit && !bUseFullTimeLimit)
-		{
-			// Try again next Tick instead.
-			GiveUpTimeSlice();
-		}
-		else
-		{
-			// Need to finish right now. Cancel async allocations that haven't finished yet.
-			// Those will be allocated immediately by UTexture2D during serialization instead.
-			Linker->Summary.TextureAllocations.CancelRemainingAllocations( false );
-			bHasCompleted = true;
-		}
-	}
-	return bHasCompleted ? EAsyncPackageState::Complete : EAsyncPackageState::TimeOut;
-#else
-	return EAsyncPackageState::Complete;
-#endif		// WITH_ENGINE
-}
-
 #if WITH_EDITORONLY_DATA
 /**
 * Creates and loads meta-data for the package.
@@ -1900,7 +5053,7 @@ EAsyncPackageState::Type FAsyncPackage::CreateExports()
 		FArchiveAsync2* FAA2 = Linker->GetFArchiveAsync2Loader();
 		if (FAA2)
 		{
-			bReady = FAA2->Precache(Export.SerialOffset, Export.SerialSize, bUseTimeLimit, bUseFullTimeLimit, TickStartTime, TimeLimit, false);
+			bReady = FAA2->Precache(Export.SerialOffset, Export.SerialSize, bUseTimeLimit, bUseFullTimeLimit, TickStartTime, TimeLimit);
 		}
 		else
 #endif
@@ -1917,7 +5070,8 @@ EAsyncPackageState::Type FAsyncPackage::CreateExports()
 				// This will cause the object to be serialized. We do this here for all objects and
 				// not just UClass and template objects, for which this is required in order to ensure
 				// seek free loading, to be able introduce async file I/O.
-				Linker->Preload( Object );
+				Linker->Preload(Object);
+				PackageObjLoaded.Add(Object);
 			}
 			LastObjectWorkWasPerformedOn	= Object;
 			LastTypeOfWorkPerformed = TEXT("creating exports for");
@@ -2112,11 +5266,14 @@ EAsyncPackageState::Type FAsyncPackage::PreLoadObjects()
 	// GC can't run in here
 	FGCScopeGuard GCGuard;
 
-	auto& ObjLoaded = FUObjectThreadContext::Get().ObjLoaded;
+	TArray<UObject*>& ThreadObjLoaded = FUObjectThreadContext::Get().ObjLoaded;	
+	PackageObjLoaded.Append(ThreadObjLoaded);
+	ThreadObjLoaded.Reset();
+
 	// Preload (aka serialize) the objects.
-	while (PreLoadIndex < ObjLoaded.Num() && !IsTimeLimitExceeded())
+	while (PreLoadIndex < PackageObjLoaded.Num() && !IsTimeLimitExceeded())
 	{
-		UObject* Object = ObjLoaded[PreLoadIndex];
+		UObject* Object = PackageObjLoaded[PreLoadIndex];
 		if (Object && Object->HasAnyFlags(RF_NeedLoad))
 		{
 			FLinkerLoad* ObjectLinker = Object->GetLinker();
@@ -2127,24 +5284,31 @@ EAsyncPackageState::Type FAsyncPackage::PreLoadObjects()
 				bool bReady;
 				if (FAA2)
 				{
-					bReady = FAA2->Precache(Export.SerialOffset, Export.SerialSize, bUseTimeLimit, bUseFullTimeLimit, TickStartTime, TimeLimit, false);
+					bReady = FAA2->Precache(Export.SerialOffset, Export.SerialSize, bUseTimeLimit, bUseFullTimeLimit, TickStartTime, TimeLimit);
 				}
 				else
 				{
-					bReady = Linker->Precache(Export.SerialOffset, Export.SerialSize);
+					bReady = ObjectLinker->Precache(Export.SerialOffset, Export.SerialSize);
 				}
-				if (bReady)
+				if (!bReady && bUseTimeLimit)
 				{
-					ObjectLinker->Preload(Object);
-					LastObjectWorkWasPerformedOn = Object;
-					LastTypeOfWorkPerformed = TEXT("preloading");
+					UE_LOG(LogStreaming, Warning, TEXT("Detected streaming object that was not precached, will hitch %s"), *Object->GetFullName());
 				}
+				ObjectLinker->Preload(Object);
+				LastObjectWorkWasPerformedOn = Object;
+				LastTypeOfWorkPerformed = TEXT("preloading");
+
+				PackageObjLoaded.Append(ThreadObjLoaded);
+				ThreadObjLoaded.Reset();
 			}
 		}
 		PreLoadIndex++;
 	}
 
-	return PreLoadIndex == ObjLoaded.Num() ? EAsyncPackageState::Complete : EAsyncPackageState::TimeOut;
+	PackageObjLoaded.Append(ThreadObjLoaded);
+	ThreadObjLoaded.Reset();
+
+	return PreLoadIndex == PackageObjLoaded.Num() ? EAsyncPackageState::Complete : EAsyncPackageState::TimeOut;
 }
 
 #else
@@ -2157,12 +5321,15 @@ EAsyncPackageState::Type FAsyncPackage::PreLoadObjects()
 	// GC can't run in here
 	FGCScopeGuard GCGuard;
 
-	auto& ObjLoaded = FUObjectThreadContext::Get().ObjLoaded;
+	TArray<UObject*>& ThreadObjLoaded = FUObjectThreadContext::Get().ObjLoaded;
+	PackageObjLoaded.Append(ThreadObjLoaded);
+	ThreadObjLoaded.Reset();
+
 	// Preload (aka serialize) the objects.
-	while (PreLoadIndex < ObjLoaded.Num() && !IsTimeLimitExceeded())
+	while (PreLoadIndex < PackageObjLoaded.Num() && !IsTimeLimitExceeded())
 	{
 		//@todo async: make this part async as well.
-		UObject* Object = ObjLoaded[PreLoadIndex++];
+		UObject* Object = PackageObjLoaded[PreLoadIndex++];
 		if (Object && Object->GetLinker())
 		{
 			Object->GetLinker()->Preload(Object);
@@ -2171,7 +5338,10 @@ EAsyncPackageState::Type FAsyncPackage::PreLoadObjects()
 		}
 	}
 
-	return PreLoadIndex == ObjLoaded.Num() ? EAsyncPackageState::Complete : EAsyncPackageState::TimeOut;
+	PackageObjLoaded.Append(ThreadObjLoaded);
+	ThreadObjLoaded.Reset();
+
+	return PreLoadIndex == PackageObjLoaded.Num() ? EAsyncPackageState::Complete : EAsyncPackageState::TimeOut;
 }
 #endif
 
@@ -2189,38 +5359,64 @@ EAsyncPackageState::Type FAsyncPackage::PostLoadObjects()
 	// GC can't run in here
 	FGCScopeGuard GCGuard;
 
-	EAsyncPackageState::Type Result = EAsyncPackageState::Complete;
 	TGuardValue<bool> GuardIsRoutingPostLoad(FUObjectThreadContext::Get().IsRoutingPostLoad, true);
 
-	auto& ObjLoaded = FUObjectThreadContext::Get().ObjLoaded;
+	TArray<UObject*>& ThreadObjLoaded = FUObjectThreadContext::Get().ObjLoaded;
+	if (ThreadObjLoaded.Num())
+	{
+		// New objects have been loaded. They need to go through PreLoad first so exit now and come back after they've been preloaded.
+		PackageObjLoaded.Append(ThreadObjLoaded);
+		ThreadObjLoaded.Reset();
+		return EAsyncPackageState::TimeOut;
+	}
+
+#if USE_EVENT_DRIVEN_ASYNC_LOAD
+	PreLoadIndex = PackageObjLoaded.Num(); // we did preloading in a different way and never incremented this
+#endif
+
 	// PostLoad objects.
-	while (PostLoadIndex < ObjLoaded.Num() && PostLoadIndex < PreLoadIndex && !IsTimeLimitExceeded())
+	while (PostLoadIndex < PackageObjLoaded.Num() && PostLoadIndex < PreLoadIndex && !IsTimeLimitExceeded())
 	{
-		UObject* Object = ObjLoaded[PostLoadIndex++];
-		check(Object);
-		if (!FAsyncLoadingThread::Get().IsMultithreaded() || Object->IsPostLoadThreadSafe())
+		UObject* Object = PackageObjLoaded[PostLoadIndex++];
+		if (Object)
 		{
-			FScopeCycleCounterUObject ConstructorScope(Object, GET_STATID(STAT_FAsyncPackage_PostLoadObjects));
+			if (!FAsyncLoadingThread::Get().IsMultithreaded() || Object->IsPostLoadThreadSafe())
+			{
+				FScopeCycleCounterUObject ConstructorScope(Object, GET_STATID(STAT_FAsyncPackage_PostLoadObjects));
 
-			Object->ConditionalPostLoad();
+#if USE_EVENT_DRIVEN_ASYNC_LOAD
+				check(!Object->HasAnyFlags(RF_NeedLoad));
+#endif
+				Object->ConditionalPostLoad();
 
-			LastObjectWorkWasPerformedOn = Object;
-			LastTypeOfWorkPerformed = TEXT("postloading_async");
-		}
-		else
-		{
-			DeferredPostLoadObjects.Add(Object);
+				LastObjectWorkWasPerformedOn = Object;
+				LastTypeOfWorkPerformed = TEXT("postloading_async");
+
+				if (ThreadObjLoaded.Num())
+				{
+					// New objects have been loaded. They need to go through PreLoad first so exit now and come back after they've been preloaded.
+					PackageObjLoaded.Append(ThreadObjLoaded);
+					ThreadObjLoaded.Reset();
+					return EAsyncPackageState::TimeOut;
+				}
+			}
+			else
+			{
+				DeferredPostLoadObjects.Add(Object);
+			}
+			// All object must be finalized on the game thread
+			DeferredFinalizeObjects.Add(Object);
+			check(Object->IsValidLowLevelFast());
+			// Make sure all objects in DeferredFinalizeObjects are referenced too
+			AddObjectReference(Object);
 		}
-		// All object must be finalized on the game thread
-		DeferredFinalizeObjects.Add(Object);
-		check(Object->IsValidLowLevelFast());
-		// Make sure all objects in DeferredFinalizeObjects are referenced too
-		FAsyncObjectsReferencer::Get().AddObject(Object);
 	}
 
-	// New objects might have been loaded during PostLoad.
-	Result = (PreLoadIndex == ObjLoaded.Num()) && (PostLoadIndex == ObjLoaded.Num()) ? EAsyncPackageState::Complete : EAsyncPackageState::TimeOut;
+	PackageObjLoaded.Append(ThreadObjLoaded);
+	ThreadObjLoaded.Reset();
 
+	// New objects might have been loaded during PostLoad.
+	EAsyncPackageState::Type Result = (PreLoadIndex == PackageObjLoaded.Num()) && (PostLoadIndex == PackageObjLoaded.Num()) ? EAsyncPackageState::Complete : EAsyncPackageState::TimeOut;
 	return Result;
 }
 
@@ -2231,14 +5427,13 @@ EAsyncPackageState::Type FAsyncPackage::PostLoadDeferredObjects(double InTickSta
 	SCOPE_CYCLE_COUNTER(STAT_FAsyncPackage_PostLoadObjectsGameThread);
 	SCOPED_LOADTIMER(PostLoadDeferredObjectsTime);
 
+	FAsyncPackageScope PackageScope(this);
+
 	EAsyncPackageState::Type Result = EAsyncPackageState::Complete;
-	TGuardValue<bool> GuardIsRoutingPostLoad(FUObjectThreadContext::Get().IsRoutingPostLoad, true);
+	TGuardValue<bool> GuardIsRoutingPostLoad(PackageScope.ThreadContext.IsRoutingPostLoad, true);
 	FAsyncLoadingTickScope InAsyncLoadingTick;
 
-	LastObjectWorkWasPerformedOn = nullptr;
-	LastTypeOfWorkPerformed = TEXT("postloading_gamethread");
-
-	TArray<UObject*>& ObjLoadedInPostLoad = FUObjectThreadContext::Get().ObjLoaded;
+	TArray<UObject*>& ObjLoadedInPostLoad = PackageScope.ThreadContext.ObjLoaded;
 	TArray<UObject*> ObjLoadedInPostLoadLocal;
 
 	while (DeferredPostLoadIndex < DeferredPostLoadObjects.Num() && 
@@ -2247,6 +5442,8 @@ EAsyncPackageState::Type FAsyncPackage::PostLoadDeferredObjects(double InTickSta
 	{
 		UObject* Object = DeferredPostLoadObjects[DeferredPostLoadIndex++];
 		check(Object);
+		LastObjectWorkWasPerformedOn = Object;
+		LastTypeOfWorkPerformed = TEXT("postloading_gamethread");
 
 		FScopeCycleCounterUObject ConstructorScope(Object, GET_STATID(STAT_FAsyncPackage_PostLoadObjectsGameThread));
 
@@ -2292,10 +5489,14 @@ EAsyncPackageState::Type FAsyncPackage::PostLoadDeferredObjects(double InTickSta
 	Result = (DeferredPostLoadIndex == DeferredPostLoadObjects.Num()) ? EAsyncPackageState::Complete : EAsyncPackageState::TimeOut;
 	if (Result == EAsyncPackageState::Complete)
 	{
+		LastObjectWorkWasPerformedOn = nullptr;
+		LastTypeOfWorkPerformed = TEXT("DeferredFinalizeObjects");
 		TArray<UObject*> CDODefaultSubobjects;
 		// Clear async loading flags (we still want RF_Async, but EInternalObjectFlags::AsyncLoading can be cleared)
-		for (UObject* Object : DeferredFinalizeObjects)
+		while (DeferredFinalizeIndex < DeferredFinalizeObjects.Num() &&
+			(DeferredPostLoadIndex % 100 != 0 || (!AsyncLoadingThread.IsAsyncLoadingSuspended() && !::IsTimeLimitExceeded(InTickStartTime, bInUseTimeLimit, InOutTimeLimit, LastTypeOfWorkPerformed, LastObjectWorkWasPerformedOn))))
 		{
+			UObject* Object = DeferredFinalizeObjects[DeferredFinalizeIndex++];
 			if (Object)
 			{
 				Object->AtomicallyClearInternalFlags(EInternalObjectFlags::AsyncLoading);
@@ -2336,10 +5537,14 @@ EAsyncPackageState::Type FAsyncPackage::PostLoadDeferredObjects(double InTickSta
 				CDODefaultSubobjects.Reset();
 			}
 		}
+		::IsTimeLimitExceeded(InTickStartTime, bInUseTimeLimit, InOutTimeLimit, LastTypeOfWorkPerformed, LastObjectWorkWasPerformedOn);
+		Result = (DeferredFinalizeIndex == DeferredFinalizeObjects.Num()) ? EAsyncPackageState::Complete : EAsyncPackageState::TimeOut;
 
 		// Mark package as having been fully loaded and update load time.
-		if (LinkerRoot && !bLoadHasFailed)
+		if (Result == EAsyncPackageState::Complete && LinkerRoot && !bLoadHasFailed)
 		{
+			LastObjectWorkWasPerformedOn = LinkerRoot;
+			LastTypeOfWorkPerformed = TEXT("CreateClustersFromPackage");
 			LinkerRoot->AtomicallyClearInternalFlags(EInternalObjectFlags::AsyncLoading);
 			LinkerRoot->MarkAsFullyLoaded();			
 			LinkerRoot->SetLoadTime(FPlatformTime::Seconds() - LoadStartTime);
@@ -2347,10 +5552,12 @@ EAsyncPackageState::Type FAsyncPackage::PostLoadDeferredObjects(double InTickSta
 			if (Linker)
 			{
 				CreateClustersFromPackage(Linker);
-
+#if !USE_NEW_ASYNC_IO
 				// give a hint to the IO system that we are done with this file for now
 				FIOSystem::Get().HintDoneWithFile(Linker->Filename);
+#endif
 			}
+			::IsTimeLimitExceeded(InTickStartTime, bInUseTimeLimit, InOutTimeLimit, LastTypeOfWorkPerformed, LastObjectWorkWasPerformedOn);
 		}
 
 		FStringAssetReference::InvalidateTag();
@@ -2359,11 +5566,6 @@ EAsyncPackageState::Type FAsyncPackage::PostLoadDeferredObjects(double InTickSta
 	return Result;
 }
 
-/**
- * Finish up objects and state, which means clearing the EInternalObjectFlags::AsyncLoading flag on newly created ones
- *
- * @return true
- */
 EAsyncPackageState::Type FAsyncPackage::FinishObjects()
 {
 	SCOPED_LOADTIMER(FinishObjectsTime);
@@ -2372,25 +5574,27 @@ EAsyncPackageState::Type FAsyncPackage::FinishObjects()
 	LastObjectWorkWasPerformedOn	= nullptr;
 	LastTypeOfWorkPerformed			= TEXT("finishing all objects");
 		
-	auto& LoadingGlobals = FUObjectThreadContext::Get();
+	TArray<UObject*>& ThreadObjLoaded = FUObjectThreadContext::Get().ObjLoaded;
 
 	EAsyncLoadingResult::Type LoadingResult;
 	if (!bLoadHasFailed)
 	{
-		LoadingGlobals.ObjLoaded.Empty();
+		ThreadObjLoaded.Reset();
 		LoadingResult = EAsyncLoadingResult::Succeeded;
 	}
 	else
-	{
+	{		
+		PackageObjLoaded.Append(ThreadObjLoaded);
+
 		// Cleanup objects from this package only
-		for (int32 ObjectIndex = LoadingGlobals.ObjLoaded.Num() - 1; ObjectIndex >= 0; --ObjectIndex)
+		for (int32 ObjectIndex = PackageObjLoaded.Num() - 1; ObjectIndex >= 0; --ObjectIndex)
 		{
-			UObject* Object = LoadingGlobals.ObjLoaded[ObjectIndex];
-			if (Object->GetOutermost()->GetFName() == Desc.Name)
+			UObject* Object = PackageObjLoaded[ObjectIndex];
+			if (Object && Object->GetOutermost()->GetFName() == Desc.Name)
 			{
 				Object->ClearFlags(RF_NeedPostLoad | RF_NeedLoad | RF_NeedPostLoadSubobjects);
 				Object->MarkPendingKill();
-				LoadingGlobals.ObjLoaded.RemoveAt(ObjectIndex);
+				PackageObjLoaded[ObjectIndex] = nullptr;
 			}
 		}
 		LoadingResult = EAsyncLoadingResult::Failed;
@@ -2401,29 +5605,9 @@ EAsyncPackageState::Type FAsyncPackage::FinishObjects()
 	PreLoadIndex = 0;
 	PreLoadSortIndex = 0;
 	PostLoadIndex = 0;
-	
-	// Close any linkers that have been open as a result of blocking load while async loading
-	TArray<FLinkerLoad*> PackagesToClose = MoveTemp(FUObjectThreadContext::Get().DelayedLinkerClosePackages);
-	for (FLinkerLoad* LinkerToClose : PackagesToClose)
-	{
-		check(LinkerToClose);
-		check(LinkerToClose->LinkerRoot);
-		if (!LinkerToClose->AsyncRoot)
-		{
-			FLinkerManager::Get().ResetLoaders(LinkerToClose->LinkerRoot);
-			check(LinkerToClose->LinkerRoot == nullptr);
-			check(LinkerToClose->AsyncRoot == nullptr);
-		}
-	}
 
-	// If we successfully loaded
-	if (!bLoadHasFailed && Linker)
-	{
-	#if WITH_ENGINE
-		// Cancel all texture allocations that haven't been claimed yet.
-		Linker->Summary.TextureAllocations.CancelRemainingAllocations( true );
-	#endif		// WITH_ENGINE
-	}
+	// Keep the linkers to close until we finish loading and it's safe to close them too
+	DelayedLinkerClosePackages = MoveTemp(FUObjectThreadContext::Get().DelayedLinkerClosePackages);
 
 	if (Linker)
 	{
@@ -2463,13 +5647,36 @@ EAsyncPackageState::Type FAsyncPackage::FinishObjects()
 	return EAsyncPackageState::Complete;
 }
 
+void FAsyncPackage::CloseDelayedLinkers()
+{
+	// Close any linkers that have been open as a result of blocking load while async loading
+	for (FLinkerLoad* LinkerToClose : DelayedLinkerClosePackages)
+	{
+		check(LinkerToClose);
+		check(LinkerToClose->LinkerRoot);
+	#if USE_EVENT_DRIVEN_ASYNC_LOAD
+		FLinkerLoad* LinkerToReset = FLinkerLoad::FindExistingLinkerForPackage(CastChecked<UPackage>(LinkerToClose->LinkerRoot));
+		check(LinkerToReset == LinkerToClose);
+		if (LinkerToReset && LinkerToReset->AsyncRoot)
+		{
+			UE_LOG(LogTemp, Error, TEXT("Linker cannot be reset right now...leaking %s"), *LinkerToReset->GetArchiveName())
+				continue;
+		}
+	#endif
+		FLinkerManager::Get().ResetLoaders(LinkerToClose->LinkerRoot);
+		check(LinkerToClose->LinkerRoot == nullptr);
+		check(LinkerToClose->AsyncRoot == nullptr);
+	}
+}
+
 void FAsyncPackage::CallCompletionCallbacks(bool bInternal, EAsyncLoadingResult::Type LoadingResult)
 {
 	UPackage* LoadedPackage = (!bLoadHasFailed) ? LinkerRoot : nullptr;
-	for (auto& CompletionCallback : CompletionCallbacks)
+	for (FCompletionCallback& CompletionCallback : CompletionCallbacks)
 	{
-		if (CompletionCallback.bIsInternal == bInternal)
+		if (CompletionCallback.bIsInternal == bInternal && !CompletionCallback.bCalled)
 		{
+			CompletionCallback.bCalled = true;
 			CompletionCallback.Callback.ExecuteIfBound(Desc.Name, LoadedPackage, LoadingResult);
 		}
 	}
@@ -2477,6 +5684,9 @@ void FAsyncPackage::CallCompletionCallbacks(bool bInternal, EAsyncLoadingResult:
 
 void FAsyncPackage::Cancel()
 {
+#if USE_NEW_ASYNC_IO
+	UE_LOG(LogStreaming, Fatal, TEXT("FAsyncPackage::Cancel is not supported with the new loader"));
+#endif
 	// Call any completion callbacks specified.
 	const EAsyncLoadingResult::Type Result = EAsyncLoadingResult::Canceled;
 	for (int32 CallbackIndex = 0; CallbackIndex < CompletionCallbacks.Num(); CallbackIndex++)
@@ -2488,8 +5698,10 @@ void FAsyncPackage::Cancel()
 	{
 		if (Linker)
 		{
+#if !USE_NEW_ASYNC_IO
 			// give a hint to the IO system that we are done with this file for now
 			FIOSystem::Get().HintDoneWithFile(Linker->Filename);
+#endif
 			Linker->FlushCache();
 		}
 		LinkerRoot->ClearFlags(RF_WasLoaded);
@@ -2635,18 +5847,25 @@ void FlushAsyncLoading(int32 PackageID /* = INDEX_NONE */)
 
 		FCoreDelegates::OnAsyncLoadingFlush.Broadcast();
 
+#if !USE_NEW_ASYNC_IO
 		// Disallow low priority requests like texture streaming while we are flushing streaming
 		// in order to avoid excessive seeking.
 		FIOSystem::Get().SetMinPriority( AIOP_Normal );
+#endif
 
 		// Flush async loaders by not using a time limit. Needed for e.g. garbage collection.
 		UE_LOG(LogStreaming, Log,  TEXT("Flushing async loaders.") );
 		{
+			TAutoPtr<FFlushTree> FlushTree;
+			if (PackageID != INDEX_NONE)
+			{
+				FlushTree = new FFlushTree(PackageID);
+			}
 			SCOPE_CYCLE_COUNTER(STAT_FAsyncPackage_TickAsyncLoadingGameThread);
 			while (IsAsyncLoading())
 			{
-				EAsyncPackageState::Type Result = AsyncThread.TickAsyncLoading(false, false, 0, PackageID);
-				if (PackageID != INDEX_NONE && Result == EAsyncPackageState::Complete)
+				EAsyncPackageState::Type Result = AsyncThread.TickAsyncLoading(false, false, 0, FlushTree.GetOwnedPointer());
+				if (PackageID != INDEX_NONE && !AsyncThread.ContainsRequestID(PackageID))
 				{
 					break;
 				}
@@ -2662,9 +5881,10 @@ void FlushAsyncLoading(int32 PackageID /* = INDEX_NONE */)
 
 		check(PackageID != INDEX_NONE || !IsAsyncLoading());
 
+#if !USE_NEW_ASYNC_IO
 		// Reset min priority again.
 		FIOSystem::Get().SetMinPriority( AIOP_MIN );
-		FMemory::Trim();
+#endif
 	}
 }
 
@@ -2693,13 +5913,13 @@ EAsyncPackageState::Type ProcessAsyncLoading(bool bUseTimeLimit, bool bUseFullTi
 bool IsAsyncLoadingCoreUObjectInternal()
 {
 	// GIsInitialLoad guards the async loading thread from being created too early
-	return !GIsInitialLoad && FAsyncLoadingThread::Get().IsAsyncLoadingPackages();
+	return FAsyncLoadingThread::Get().IsAsyncLoadingPackages();
 }
 
 bool IsAsyncLoadingMultithreadedCoreUObjectInternal()
 {
 	// GIsInitialLoad guards the async loading thread from being created too early
-	return !GIsInitialLoad && FAsyncLoadingThread::Get().IsMultithreaded();
+	return FAsyncLoadingThread::Get().IsMultithreaded();
 }
 
 void SuspendAsyncLoadingInternal()
@@ -2714,6 +5934,18 @@ void ResumeAsyncLoadingInternal()
 	FAsyncLoadingThread::Get().ResumeLoading();
 }
 
+bool IsEventDrivenLoaderEnabledInCookedBuilds()
+{
+#if WITH_EDITORONLY_DATA
+	check(GConfig); // Otherwise there's no way we have the correct value of GEventDrivenLoaderEnabled
+	return !!GEventDrivenLoaderEnabled;
+#elif USE_EVENT_DRIVEN_ASYNC_LOAD
+	return true;
+#else
+	return false;
+#endif
+}
+
 /*----------------------------------------------------------------------------
 	FArchiveAsync.
 ----------------------------------------------------------------------------*/
@@ -3225,322 +6457,84 @@ void FArchiveAsync::Serialize(void* Data, int64 Count)
 	}
 #endif // FIND_MEMORY_STOMPS
 
-	// Update stats if we were blocked.
-#if LOOKING_FOR_PERF_ISSUES
-	if( bIOBlocked )
-	{
-		const int32 BlockingCycles = int32(FPlatformTime::Cycles() - StartCycles);
-		FAsyncLoadingThread::BlockingCycles.Add( BlockingCycles );
-
-		UE_LOG(LogStreaming, Warning, TEXT("FArchiveAsync::Serialize: %5.2fms blocking on read from '%s' (Offset: %lld, Size: %lld)"), 
-			FPlatformTime::ToMilliseconds(BlockingCycles), 
-			*FileName, 
-			CurrentPos, 
-			Count );
-	}
-#endif
-
-	// Copy memory to destination.
-	FMemory::Memcpy( Data, PrecacheBuffer[CURRENT] + (CurrentPos - PrecacheStartPos[CURRENT]), Count );
-	// Serialization implicitly increases position in file.
-	CurrentPos += Count;
-}
-
-int64 FArchiveAsync::Tell()
-{
-	return CurrentPos;
-}
-
-int64 FArchiveAsync::TotalSize()
-{
-	return UncompressedFileSize + BulkDataAreaSize;
-}
-
-void FArchiveAsync::Seek( int64 InPos )
-{
-	check( InPos >= 0 && InPos <= TotalSize() );
-	CurrentPos = InPos;
-}
-
-
-#endif // !USE_NEW_ASYNC_IO
-/*----------------------------------------------------------------------------
-	End.
-----------------------------------------------------------------------------*/
-
-#if USE_NEW_ASYNC_IO
-
-DEFINE_LOG_CATEGORY_STATIC(LogAsyncArchive, Display, All);
-
-struct FRequestPair
-{
-	IAsyncReadFileHandle* Handle;
-	IAsyncReadRequest* Request;
-
-	FRequestPair()
-		: Handle(nullptr)
-		, Request(nullptr)
-	{
-	}
-
-	void Clear()
-	{
-		if (Request)
-		{
-			Request->Cancel();
-			Request->WaitCompletion();
-			delete Request;
-			Request = nullptr;
-		}
-		if (Handle)
-		{
-			delete Handle;
-			Handle = nullptr;
-		}
-	}
-};
-
-struct FHintHistory
-{
-	FCriticalSection Crit;
-	TMap<FName, FRequestPair> Requests;
-
-	bool Contains(const TCHAR* Filename)
-	{
-		FName Name(Filename);
-		FScopeLock Lock(&Crit);
-		return Requests.Contains(Name);
-	}
-
-	void Add(const TCHAR* Filename, FRequestPair& Pair)
-	{
-		FName Name(Filename);
-		FScopeLock Lock(&Crit);
-		Requests.Add(Name, Pair);
-	}
-	void Done(const TCHAR* Filename)
-	{
-		FName Name(Filename);
-		FScopeLock Lock(&Crit);
-		FRequestPair Pair;
-		Requests.RemoveAndCopyValue(Name, Pair);
-		if (!Pair.Handle)
-		{
-			//UE_LOG(LogAsyncArchive, Warning, TEXT("FPakPlatformFile::HintFutureReadDone, file was not found %s"), Filename);
-		}
-		else
-		{
-			Pair.Clear();
-		}
-	}
-
-	void Dump()
-	{
-		FScopeLock Lock(&Crit);
-		for (auto& Elem : Requests)
-		{
-			UE_LOG(LogAsyncArchive, Warning, TEXT("%s"), *Elem.Key.ToString());
-		}
-	}
-	void Flush()
-	{
-		FScopeLock Lock(&Crit);
-		if (Requests.Num())
-		{
-			for (auto& Elem : Requests)
-			{
-				Elem.Value.Clear();
-			}
-		}
-		Requests.Empty();
-	}
-};
-
-static FHintHistory GHintHistory;
-
-static void DumpHistory(const TArray<FString>& Args)
-{
-	GHintHistory.Dump();
-	//GHintHistory.Flush();
-}
-
-static FAutoConsoleCommand DumpRequestsCmd(
-	TEXT("pak.DumpAsyncHintHistory"),
-	TEXT("debug command to print the pak file debug history."),
-	FConsoleCommandWithArgsDelegate::CreateStatic(&DumpHistory)
-	);
-
-#define DEFAULT_USE_PRECACHE_HINTS (1)
-
-static int32 GPakCache_EnableHints = DEFAULT_USE_PRECACHE_HINTS;
-static int32 GPakCache_EnableHintsLatched = DEFAULT_USE_PRECACHE_HINTS;
-static FAutoConsoleVariableRef CVar_EnableHints(
-	TEXT("pakcache.EnableHints"),
-	GPakCache_EnableHints,
-	TEXT("If > 0, then process precache hints.")
-	);
-
-static bool LatchEnableHints()
-{
-	if (GPakCache_EnableHints != GPakCache_EnableHintsLatched)
-	{
-		GPakCache_EnableHintsLatched = GPakCache_EnableHints;
-		if (!GPakCache_EnableHintsLatched)
-		{
-			GHintHistory.Flush();
-		}
-	}
-	return !!GPakCache_EnableHintsLatched;
-}
-
-void HintFutureReadDone(const TCHAR * FileName)
-{
-	if (LatchEnableHints())
-	{
-		GHintHistory.Done(FileName);
-	}
-}
-
-#define PRECACHE_SUMMARY_ONLY (0)
-
-void HintFutureRead(const TCHAR * FileName)
-{
-	if (LatchEnableHints() && !GHintHistory.Contains(FileName))
-	{
-		FRequestPair Pair;
-		Pair.Handle = FPlatformFileManager::Get().GetPlatformFile().OpenAsyncRead(FileName);
-		if (Pair.Handle)
-		{
-			Pair.Request = Pair.Handle->ReadRequest(0, PRECACHE_SUMMARY_ONLY ? 1 : MAX_int64, AIOP_Precache);
-		}
-		GHintHistory.Add(FileName, Pair);
-	}
-}
-
-static bool GRecordPakOpens = false;
-static TArray<FName> GPakOpens;
-
-static void AddPakOpen(const TCHAR* Filename)
-{
-	if (GRecordPakOpens)
+	// Update stats if we were blocked.
+#if LOOKING_FOR_PERF_ISSUES
+	if( bIOBlocked )
 	{
-		FName Name(Filename);
-		if (!GPakOpens.Contains(Name))
-		{
-			GPakOpens.Add(Name);
-		}
+		const int32 BlockingCycles = int32(FPlatformTime::Cycles() - StartCycles);
+		FAsyncLoadingThread::BlockingCycles.Add( BlockingCycles );
+
+		UE_LOG(LogStreaming, Warning, TEXT("FArchiveAsync::Serialize: %5.2fms blocking on read from '%s' (Offset: %lld, Size: %lld)"), 
+			FPlatformTime::ToMilliseconds(BlockingCycles), 
+			*FileName, 
+			CurrentPos, 
+			Count );
 	}
+#endif
+
+	// Copy memory to destination.
+	FMemory::Memcpy( Data, PrecacheBuffer[CURRENT] + (CurrentPos - PrecacheStartPos[CURRENT]), Count );
+	// Serialization implicitly increases position in file.
+	CurrentPos += Count;
 }
 
-static void StartOpenLog(const TArray<FString>& Args)
+int64 FArchiveAsync::Tell()
 {
-	GRecordPakOpens = true;
+	return CurrentPos;
 }
 
-static FAutoConsoleCommand StartOpenLogCmd(
-	TEXT("pak.StartOpenLog"),
-	TEXT("debug command to start recording open commands."),
-	FConsoleCommandWithArgsDelegate::CreateStatic(&StartOpenLog)
-	);
-
-static void EndOpenLog(const TArray<FString>& Args)
+int64 FArchiveAsync::TotalSize()
 {
-	GRecordPakOpens = false;
-
-	FString Contents = TEXT("\r\nstatic const TCHAR *PrecacheList[] = {\r\n");
-
-	bool bStarted = false;
-	for (FName& Name : GPakOpens)
-	{
-		Contents += FString::Printf(TEXT("\t%cTEXT(\"%s\")\r\n"), bStarted ? ',' : ' ', *Name.ToString());
-		bStarted = true;
-	}
-
-	Contents += TEXT("};\r\n");
-
-	FPlatformMisc::LowLevelOutputDebugStringf(TEXT("\r\n%s\r\n"), *Contents);
-
-	GPakOpens.Empty();
+	return UncompressedFileSize + BulkDataAreaSize;
 }
 
-static FAutoConsoleCommand EndOpenLogCmd(
-	TEXT("pak.EndOpenLog"),
-	TEXT("debug command to end recording open commands."),
-	FConsoleCommandWithArgsDelegate::CreateStatic(&EndOpenLog)
-	);
-
-
-
-FArchive* NewFArchiveAsync2(const TCHAR* InFileName)
+void FArchiveAsync::Seek( int64 InPos )
 {
-	return new FArchiveAsync2(InFileName);
+	check( InPos >= 0 && InPos <= TotalSize() );
+	CurrentPos = InPos;
 }
 
-//#define ASYNC_WATCH_FILE "Meshes/work/TwinBlast_CylShadows.uasset"
-#define USE_FArchiveAsync2_PRECACHING (1)
-#define USE_FArchiveAsync2_PRECACHE_WHOLE_FILE (1)
-#define TRACK_SERIALIZE (0)
-#define MIN_REMAIN_TIME (0.00101f)   // wait(0) is very different than wait(tiny) so we cut things off well before roundoff error could cause us to block when we didn't intend to. Also the granularity of the event API is 1ms.
 
-FORCEINLINE void FArchiveAsync2::LogItem(const TCHAR* Item, int64 Offset, int64 Size, double StartTime)
-{
-	if (UE_LOG_ACTIVE(LogAsyncArchive, Verbose)
-#if defined(ASYNC_WATCH_FILE)
-	|| FileName.Contains(TEXT(ASYNC_WATCH_FILE))
-#endif
-		)
-	{
-		static double OpenTime(FPlatformTime::Seconds());
-		double Now(FPlatformTime::Seconds());
+#endif // !USE_NEW_ASYNC_IO
+/*----------------------------------------------------------------------------
+	End.
+----------------------------------------------------------------------------*/
 
-		float ThisTime = (StartTime != 0.0) ? float(1000.0 * (Now - StartTime)) : 0.0f;
+#if USE_NEW_ASYNC_IO
+//@todoio Also all of the leaf-first loader can go away
 
-		if (!UE_LOG_ACTIVE(LogAsyncArchive, VeryVerbose) && ThisTime < 1.0f 
-#if defined(ASYNC_WATCH_FILE)
-			&& !FileName.Contains(TEXT(ASYNC_WATCH_FILE))
+FArchiveAsync2::FArchiveAsync2(const TCHAR* InFileName
+#if USE_EVENT_DRIVEN_ASYNC_LOAD
+	, TFunction<void()>&& InSummaryReadyCallback
 #endif
-			)
-		{
-			return;
-		}
-
-		FPlatformMisc::LowLevelOutputDebugStringf(TEXT("%32s%3s    %12lld %12lld    %6.2fms    (+%9.2fms)      %s\r\n"),
-			Item,
-			(ThisTime > 1.0f) ? TEXT("***") : TEXT(""),
-			Offset,
-			((Size == MAX_int64) ? TotalSize() : Offset + Size),
-			ThisTime,
-			float(1000.0 * (Now - OpenTime)),
-			*FileName);
-	}
-}
-
-FArchiveAsync2::FArchiveAsync2(const TCHAR* InFileName)
+	)
 	: Handle(nullptr)
 	, SizeRequestPtr(nullptr)
+#if !(SPLIT_COOKED_FILES && USE_EVENT_DRIVEN_ASYNC_LOAD)
 	, SummaryRequestPtr(nullptr)
 	, SummaryPrecacheRequestPtr(nullptr)
+#endif
 	, ReadRequestPtr(nullptr)
-	, PrecacheRequestPtr(nullptr)
 	, CanceledReadRequestPtr(nullptr)
 	, PrecacheBuffer(nullptr)
-	, NextPrecacheBuffer(nullptr)
 	, FileSize(-1)
 	, CurrentPos(0)
 	, PrecacheStartPos(0)
 	, PrecacheEndPos(0)
 	, ReadRequestOffset(0)
 	, ReadRequestSize(0)
-	, PrecacheOffset(-1)
 	, HeaderSize(0)
+	, HeaderSizeWhenReadingExportsFromSplitFile(0)
 	, LoadPhase(ELoadPhase::WaitingForSize)
 	, FileName(InFileName)
 	, OpenTime(FPlatformTime::Seconds())
 	, SummaryReadTime(0.0)
 	, ExportReadTime(0.0)
-	, PrecacheEOFTime(0.0)
+#if USE_EVENT_DRIVEN_ASYNC_LOAD
+	, SummaryReadyCallback(Forward<TFunction<void()>>(InSummaryReadyCallback))
+#endif
+
 {
-	AddPakOpen(InFileName);
 	LogItem(TEXT("Open"));
 	Handle = FPlatformFileManager::Get().GetPlatformFile().OpenAsyncRead(InFileName);
 	check(Handle); // this generally cannot fail because it is async
@@ -3549,15 +6543,27 @@ FArchiveAsync2::FArchiveAsync2(const TCHAR* InFileName)
 	{
 		ReadCallback(bWasCancelled, Request);
 	};
-	SizeRequestPtr = Handle->SizeRequest(&ReadCallbackFunction);	
+#if USE_EVENT_DRIVEN_ASYNC_LOAD
+	check(SummaryReadyCallback);
+	ReadCallbackFunctionForLinkerLoad = [this](bool bWasCancelled, IAsyncReadRequest* Request)
+	{
+		SummaryReadyCallback();
+	};
+#endif
+	SizeRequestPtr = Handle->SizeRequest(&ReadCallbackFunction);
 
 }
 
 FArchiveAsync2::~FArchiveAsync2()
 {
-	HintFutureReadDone(*FileName); // done
 	// Invalidate any precached data and free memory.
 	FlushCache();
+	if (Handle)
+	{
+		delete Handle;
+		Handle = nullptr;
+	}
+	LogItem(TEXT("~FArchiveAsync2"), 0, 0);
 }
 
 void FArchiveAsync2::ReadCallback(bool bWasCancelled, IAsyncReadRequest* Request)
@@ -3577,15 +6583,27 @@ void FArchiveAsync2::ReadCallback(bool bWasCancelled, IAsyncReadRequest* Request
 		}
 		else
 		{
+#if SPLIT_COOKED_FILES && USE_EVENT_DRIVEN_ASYNC_LOAD
+			// in this case we don't need to serialize the summary because we know the header is the whole file
+			HeaderSize = FileSize;
+			LogItem(TEXT("Starting Split Header"), 0, FileSize);
+			PrecacheInternal(0, HeaderSize);
+			FPlatformMisc::MemoryBarrier();
+			LoadPhase = ELoadPhase::WaitingForHeader;
+#else
 			int64 Size = FMath::Min<int64>(MAX_SUMMARY_SIZE, FileSize);
 			LogItem(TEXT("Starting Summary"), 0, Size);
 			SummaryRequestPtr = Handle->ReadRequest(0, Size, AIOP_Normal, &ReadCallbackFunction);
 			// I need a precache request here to keep the memory alive until I submit the header request
 			SummaryPrecacheRequestPtr = Handle->ReadRequest(0, Size, AIOP_Precache);
+#endif
 		}
 	}
 	else if (LoadPhase == ELoadPhase::WaitingForSummary)
 	{
+#if SPLIT_COOKED_FILES && USE_EVENT_DRIVEN_ASYNC_LOAD
+		check(0);
+#else
 		uint8* Mem = Request->GetReadResults();
 		if (!Mem)
 		{
@@ -3593,7 +6611,7 @@ void FArchiveAsync2::ReadCallback(bool bWasCancelled, IAsyncReadRequest* Request
 		}
 		else
 		{
-			FBufferReader Ar(Mem, FMath::Min<int64>(MAX_SUMMARY_SIZE, FileSize),/*bInFreeOnClose=*/ true, /*bIsPersistent=*/ true);
+			FBufferReader Ar(Mem, FMath::Min<int64>(MAX_SUMMARY_SIZE, FileSize),/*bInFreeOnClose=*/ false, /*bIsPersistent=*/ true);
 			FPackageFileSummary Sum;
 			Ar << Sum;
 			if (Ar.IsError() || Sum.TotalHeaderSize > FileSize)
@@ -3602,14 +6620,18 @@ void FArchiveAsync2::ReadCallback(bool bWasCancelled, IAsyncReadRequest* Request
 			}
 			else
 			{
+				//@todoio change header format to put the TotalHeaderSize at the start of the file
 				check(Ar.Tell() < MAX_SUMMARY_SIZE / 2); // we need to be sure that we can at least get the size from the initial request. This is an early warning that custom versions are starting to get too big, relocate the total size to be at offset 4!
 				HeaderSize = Sum.TotalHeaderSize;
 				LogItem(TEXT("Starting Header"), 0, HeaderSize);
 				PrecacheInternal(0, HeaderSize);
 			}
+			FMemory::Free(Mem);
+			DEC_MEMORY_STAT_BY(STAT_AsyncFileMemory, FMath::Min<int64>(MAX_SUMMARY_SIZE, FileSize));
 		}
 		FPlatformMisc::MemoryBarrier();
 		LoadPhase = ELoadPhase::WaitingForHeader;
+#endif
 	}
 	else
 	{
@@ -3617,106 +6639,27 @@ void FArchiveAsync2::ReadCallback(bool bWasCancelled, IAsyncReadRequest* Request
 	}
 }
 
-void FArchiveAsync2::DispatchPrecache(int64 Offset, bool bForceTrimPrecache)
-{
-	if (!USE_FArchiveAsync2_PRECACHING)
-	{
-		return;
-	}
-	bool bFirstOrSecondTime = PrecacheOffset < HeaderSize;
-	if (!bFirstOrSecondTime && !bForceTrimPrecache)
-	{
-		return;
-	}
-	bool bFirstTime = PrecacheOffset < 0;
-	bool bTrimToHeader = !USE_FArchiveAsync2_PRECACHE_WHOLE_FILE && bFirstTime;
-	if (bTrimToHeader)
-	{
-		Offset = HeaderSize - 1; // we are not yet processing imports, so we want to put the precache request at the first import and we are only going to do one byte.
-	}
-	if (Offset > PrecacheOffset)
-	{
-		if (Offset < TotalSizeOrMaxInt64IfNotReady())
-		{
-			// important to get this in flight before we cancel the old one
-			int64 Size = bTrimToHeader ? 1 : MAX_int64; // MAX_int64 is a legal size for a precache request, if we have finished the header, then we only save the first byte of the first export
-			double StartTime = FPlatformTime::Seconds();
-			IAsyncReadRequest* NewPrecacheRequestPtr = Handle->ReadRequest(Offset, Size, AIOP_Precache);
-			LogItem(TEXT("Precache"), Offset, Size, StartTime);
-			CompletePrecache();
-			PrecacheRequestPtr = NewPrecacheRequestPtr;
-#if defined(ASYNC_WATCH_FILE)
-			if (FileName.Contains(TEXT(ASYNC_WATCH_FILE)) && Offset == 894913)
-			{
-				UE_LOG(LogAsyncArchive, Warning, TEXT("Handy Breakpoint after Precache"));
-			}
-#endif
-		}
-		else
-		{
-			CompletePrecache(); // we have reached the end of the file, so we are done with precaching
-			LogItem(TEXT("Precache EOF"));
-#if defined(ASYNC_WATCH_FILE)
-			if (FileName.Contains(TEXT(ASYNC_WATCH_FILE)))
-			{
-				UE_LOG(LogAsyncArchive, Warning, TEXT("Handy Breakpoint after Precache EOF"));
-			}
-#endif
-			check(PrecacheEOFTime = 0.0);
-			if (PrecacheEOFTime == 0.0)
-			{
-				PrecacheEOFTime = FPlatformTime::Seconds();
-			}
-
-		}
-		PrecacheOffset = Offset;
-		if (bFirstTime)
-		{
-			HintFutureReadDone(*FileName); // we will take over precaching from here
-		}
-		return;
-	}
-	if (Offset < PrecacheOffset)
-	{
-		UE_LOG(LogAsyncArchive, Warning, TEXT("FArchiveAsync2::DispatchPrecache Backwards streaming in %s  NewOffset = %lld   PreviousPrecacheOffset = %lld"), *FileName, Offset, PrecacheOffset);
-	}
-#if defined(ASYNC_WATCH_FILE)
-	if (FileName.Contains(TEXT(ASYNC_WATCH_FILE)))
-	{
-		UE_LOG(LogAsyncArchive, Warning, TEXT("Handy Breakpoint after new precache."));
-	}
-#endif
-}
-
-void FArchiveAsync2::CompletePrecache()
+void FArchiveAsync2::FlushPrecacheBlock()
 {
-	if (PrecacheRequestPtr)
+	DiscardInlineBufferAndUpdateCurrentPos();
+	if (PrecacheBuffer)
 	{
-		double StartTime = FPlatformTime::Seconds();
-		PrecacheRequestPtr->Cancel();
-		CompleteCancel(); // this deals with the cancel request...which we are going to reuse for us
-		CanceledReadRequestPtr = PrecacheRequestPtr;
-		PrecacheRequestPtr = nullptr;
-		LogItem(TEXT("CompletePrecache"), 0, 0, StartTime);
+		DEC_MEMORY_STAT_BY(STAT_FArchiveAsync2Mem, PrecacheEndPos - PrecacheStartPos);
+		FMemory::Free(PrecacheBuffer);
 	}
+	PrecacheBuffer = nullptr;
+	PrecacheStartPos = 0;
+	PrecacheEndPos = 0;
 }
 
 void FArchiveAsync2::FlushCache()
 {
-	bool nNonRedundantFlush = PrecacheEndPos || PrecacheBuffer || PrecacheRequestPtr || ReadRequestPtr;
+	bool nNonRedundantFlush = PrecacheEndPos || PrecacheBuffer || ReadRequestPtr;
 	LogItem(TEXT("Flush"));
 	WaitForIntialPhases();
 	WaitRead(); // this deals with the read request
-	CompletePrecache(); // this deals with the precache request
 	CompleteCancel(); // this deals with the cancel request, important this is last because completing other things leaves cancels to process
-	if (PrecacheBuffer)
-	{
-		FMemory::Free(PrecacheBuffer);
-	}
-	PrecacheBuffer = nullptr;
-	check(!NextPrecacheBuffer);
-	PrecacheStartPos = 0;
-	PrecacheEndPos = 0;
+	FlushPrecacheBlock();
 
 	if ((UE_LOG_ACTIVE(LogAsyncArchive, Verbose) 
 #if defined(ASYNC_WATCH_FILE)
@@ -3736,16 +6679,11 @@ void FArchiveAsync2::FlushCache()
 			return;
 		}
 
-		if (PrecacheEOFTime == 0.0)
-		{
-			PrecacheEOFTime = Now;
-		}
-		FPlatformMisc::LowLevelOutputDebugStringf(TEXT("Flush     Lifeitme %6.2fms   Open->Summary %6.2fms    Summary->Export1 %6.2fms    Export1->ExportN %6.2fms    ExportN->Now %6.2fms       %s\r\n"),
+		FPlatformMisc::LowLevelOutputDebugStringf(TEXT("Flush     Lifeitme %6.2fms   Open->Summary %6.2fms    Summary->Export1 %6.2fms    Export1->Now %6.2fms       %s\r\n"),
 			TotalLifetime,
 			float(1000.0 * (SummaryReadTime - OpenTime)),
 			float(1000.0 * (ExportReadTime - SummaryReadTime)),
-			float(1000.0 * (PrecacheEOFTime - ExportReadTime)),
-			float(1000.0 * (Now - PrecacheEOFTime)),
+			float(1000.0 * (Now - ExportReadTime)),
 			*FileName);
 #if defined(ASYNC_WATCH_FILE)
 		if (FileName.Contains(TEXT(ASYNC_WATCH_FILE)))
@@ -3777,16 +6715,57 @@ int64 FArchiveAsync2::TotalSize()
 	{
 		QUICK_SCOPE_CYCLE_COUNTER(STAT_FArchiveAsync2_TotalSize);
 		SizeRequestPtr->WaitCompletion();
+#if SPLIT_COOKED_FILES && USE_EVENT_DRIVEN_ASYNC_LOAD
+		if (HeaderSizeWhenReadingExportsFromSplitFile)
+		{
+			FileSize = SizeRequestPtr->GetSizeResults();
+		}
+#endif
 		delete SizeRequestPtr;
 		SizeRequestPtr = nullptr;
 	}
-	return FileSize;
+	return FileSize + HeaderSizeWhenReadingExportsFromSplitFile;
+}
+
+#if DEVIRTUALIZE_FLinkerLoad_Serialize
+FORCEINLINE void FArchiveAsync2::SetPosAndUpdatePrecacheBuffer(int64 Pos)
+{
+	check(Pos >= 0 && Pos <= TotalSizeOrMaxInt64IfNotReady());
+	if (Pos < PrecacheStartPos || Pos >= PrecacheEndPos)
+	{
+		ActiveFPLB->Reset();
+		CurrentPos = Pos;
+	}
+	else
+	{
+		check(PrecacheBuffer);
+		ActiveFPLB->OriginalFastPathLoadBuffer = PrecacheBuffer;
+		ActiveFPLB->StartFastPathLoadBuffer = PrecacheBuffer + (Pos - PrecacheStartPos);
+		ActiveFPLB->EndFastPathLoadBuffer = PrecacheBuffer + (PrecacheEndPos - PrecacheStartPos);
+		CurrentPos = PrecacheStartPos;
+	}
+	check(Tell() == Pos);
 }
+#endif
 
 void FArchiveAsync2::Seek(int64 InPos)
 {
+#if SPLIT_COOKED_FILES && USE_EVENT_DRIVEN_ASYNC_LOAD
+	if (LoadPhase < ELoadPhase::ProcessingExports)
+	{
+		check(!HeaderSizeWhenReadingExportsFromSplitFile && HeaderSize && TotalSize() == HeaderSize);
+		if (InPos >= HeaderSize)
+		{
+			FirstExportStarting();
+		}
+	}
+#endif
 	check(InPos >= 0 && InPos <= TotalSizeOrMaxInt64IfNotReady());
+#if DEVIRTUALIZE_FLinkerLoad_Serialize
+	SetPosAndUpdatePrecacheBuffer(InPos);
+#else
 	CurrentPos = InPos;
+#endif
 }
 
 bool FArchiveAsync2::WaitRead(float TimeLimit)
@@ -3813,28 +6792,25 @@ void FArchiveAsync2::CompleteRead()
 	double StartTime = FPlatformTime::Seconds();
 	check(LoadPhase != ELoadPhase::WaitingForSize && LoadPhase != ELoadPhase::WaitingForSummary);
 	check(ReadRequestPtr && ReadRequestPtr->PollCompletion());
+	if (PrecacheBuffer)
+	{
+		FlushPrecacheBlock();
+	}
 	if (!ArIsError)
 	{
-		DispatchPrecache(ReadRequestOffset + ReadRequestSize);
-		if (PrecacheBuffer)
+		uint8* Mem = ReadRequestPtr->GetReadResults();
+		if (!Mem)
 		{
-			FMemory::Free(PrecacheBuffer);
-			PrecacheBuffer = nullptr;
+			ArIsError = true;
+		}
+		else
+		{
+			PrecacheBuffer = Mem;
+			PrecacheStartPos = ReadRequestOffset;
+			PrecacheEndPos = ReadRequestOffset + ReadRequestSize;
+			INC_MEMORY_STAT_BY(STAT_FArchiveAsync2Mem, PrecacheEndPos - PrecacheStartPos);
+			DEC_MEMORY_STAT_BY(STAT_AsyncFileMemory, ReadRequestSize);
 		}
-
-	uint8* Mem = ReadRequestPtr->GetReadResults();
-	if (!Mem)
-	{
-		ArIsError = true;
-			PrecacheStartPos = 0;
-			PrecacheEndPos = 0;
-	}
-	else
-	{
-		PrecacheBuffer = Mem;
-		PrecacheStartPos = ReadRequestOffset;
-		PrecacheEndPos = ReadRequestOffset + ReadRequestSize;
-	}
 	}
 	delete ReadRequestPtr;
 	ReadRequestPtr = nullptr;
@@ -3870,7 +6846,11 @@ void FArchiveAsync2::CancelRead()
 
 bool FArchiveAsync2::WaitForIntialPhases(float InTimeLimit)
 {
-	if (SizeRequestPtr || SummaryRequestPtr || SummaryPrecacheRequestPtr)
+	if (SizeRequestPtr 
+#if !(SPLIT_COOKED_FILES && USE_EVENT_DRIVEN_ASYNC_LOAD)
+		|| SummaryRequestPtr || SummaryPrecacheRequestPtr
+#endif
+		)
 	{
 		QUICK_SCOPE_CYCLE_COUNTER(STAT_FArchiveAsync2_WaitForIntialPhases);
 		double StartTime = FPlatformTime::Seconds();
@@ -3887,6 +6867,7 @@ bool FArchiveAsync2::WaitForIntialPhases(float InTimeLimit)
 				return false;
 			}
 		}
+#if !(SPLIT_COOKED_FILES && USE_EVENT_DRIVEN_ASYNC_LOAD)
 		if (SummaryRequestPtr)
 		{
 			float TimeLimit = 0.0f;
@@ -3931,47 +6912,52 @@ bool FArchiveAsync2::WaitForIntialPhases(float InTimeLimit)
 				return false;
 			}
 		}
+#endif
 		LogItem(TEXT("Wait Summary"), 0, HeaderSize, StartTime);
 	}
 	return true;
 }
 
-bool FArchiveAsync2::PrecacheInternal(int64 RequestOffset, int64 RequestSize)
+bool FArchiveAsync2::PrecacheInternal(int64 RequestOffset, int64 RequestSize, bool bApplyMinReadSize)
 {
 	// CAUTION! This is possibly called the first time from a random IO thread.
 	if (LoadPhase != ELoadPhase::WaitingForSummary)
 	{
-	if (RequestSize == 0 || (RequestOffset >= PrecacheStartPos && RequestOffset + RequestSize <= PrecacheEndPos))
-	{
-		// ready
-		return true;
-	}
-	if (ReadRequestPtr && RequestOffset >= ReadRequestOffset && RequestOffset + RequestSize <= ReadRequestOffset + ReadRequestSize)
-	{
-		// current request contains request
-		if (ReadRequestPtr->PollCompletion())
+		if (RequestSize == 0 || (RequestOffset >= PrecacheStartPos && RequestOffset + RequestSize <= PrecacheEndPos))
 		{
-			CompleteRead();
-			check(RequestOffset >= PrecacheStartPos && RequestOffset + RequestSize <= PrecacheEndPos);
+			// ready
 			return true;
 		}
-		return false;
-	}
-	if (ReadRequestPtr)
-	{
-		// this one does not have what we need
-			UE_LOG(LogStreaming, Warning, TEXT("FArchiveAsync2::DispatchPrecache Canceled read for %s  Offset = %lld   Size = %lld"), *FileName, RequestOffset, ReadRequestSize);
-		CancelRead();
-	}
+		if (ReadRequestPtr && RequestOffset >= ReadRequestOffset && RequestOffset + RequestSize <= ReadRequestOffset + ReadRequestSize)
+		{
+			// current request contains request
+			if (ReadRequestPtr->PollCompletion())
+			{
+				CompleteRead();
+				check(RequestOffset >= PrecacheStartPos && RequestOffset + RequestSize <= PrecacheEndPos);
+				return true;
+			}
+			return false;
+		}
+		if (ReadRequestPtr)
+		{
+			// this one does not have what we need
+			UE_LOG(LogStreaming, Warning, TEXT("FArchiveAsync2::PrecacheInternal Canceled read for %s  Offset = %lld   Size = %lld"), *FileName, RequestOffset, ReadRequestSize);
+			CancelRead();
+		}
 	}
 	check(!ReadRequestPtr);
 	ReadRequestOffset = RequestOffset;
 	ReadRequestSize = RequestSize;
 
 
-	if (LoadPhase == ELoadPhase::ProcessingExports)
+	if (bApplyMinReadSize && LoadPhase != ELoadPhase::WaitingForSummary)
 	{
+#if WITH_EDITOR
+		static int64 MinimumReadSize = 1024 * 1024;
+#else
 		static int64 MinimumReadSize = 65536;
+#endif
 		checkSlow(MinimumReadSize >= 2048 && MinimumReadSize <= 1024 * 1024); // not a hard limit, but we should be loading at least a reasonable amount of data
 		if (ReadRequestSize < MinimumReadSize)
 		{ 
@@ -3992,20 +6978,78 @@ bool FArchiveAsync2::PrecacheInternal(int64 RequestOffset, int64 RequestSize)
 		UE_LOG(LogAsyncArchive, Warning, TEXT("Handy Breakpoint Read"));
 	}
 #endif
-
-	ReadRequestPtr = Handle->ReadRequest(ReadRequestOffset, ReadRequestSize, AIOP_Normal);
+	check(ReadRequestOffset - HeaderSizeWhenReadingExportsFromSplitFile >= 0);
+	ReadRequestPtr = Handle->ReadRequest(ReadRequestOffset - HeaderSizeWhenReadingExportsFromSplitFile, ReadRequestSize, AIOP_Normal
+#if USE_EVENT_DRIVEN_ASYNC_LOAD
+		, (LoadPhase == ELoadPhase::WaitingForSummary) ? &ReadCallbackFunctionForLinkerLoad : nullptr
+#endif
+		);
 	if (LoadPhase != ELoadPhase::WaitingForSummary && ReadRequestPtr->PollCompletion())
 	{
-		LogItem(TEXT("Read Start Hot"), ReadRequestOffset, ReadRequestSize, StartTime);
+		LogItem(TEXT("Read Start Hot"), ReadRequestOffset - HeaderSizeWhenReadingExportsFromSplitFile, ReadRequestSize, StartTime);
 		CompleteRead();
 		check(RequestOffset >= PrecacheStartPos && RequestOffset + RequestSize <= PrecacheEndPos);
 		return true;
 	}
-	LogItem(TEXT("Read Start Cold"), ReadRequestOffset, ReadRequestSize, StartTime);
+	else if (LoadPhase == ELoadPhase::WaitingForSummary)
+	{
+		LogItem(TEXT("Read Start Summary"), ReadRequestOffset - HeaderSizeWhenReadingExportsFromSplitFile, ReadRequestSize, StartTime);
+	}
+	else 
+	{
+		LogItem(TEXT("Read Start Cold"), ReadRequestOffset - HeaderSizeWhenReadingExportsFromSplitFile, ReadRequestSize, StartTime);
+	}
 	return false;
 }
 
-bool FArchiveAsync2::Precache(int64 RequestOffset, int64 RequestSize, bool bUseTimeLimit, bool bUseFullTimeLimit, double TickStartTime, float TimeLimit, bool bTrimPrecache)
+void FArchiveAsync2::FirstExportStarting()
+{
+	ExportReadTime = FPlatformTime::Seconds();
+	LogItem(TEXT("Exports"));
+	LoadPhase = ELoadPhase::ProcessingExports;
+
+#if SPLIT_COOKED_FILES && USE_EVENT_DRIVEN_ASYNC_LOAD
+	FlushCache();
+	if (Handle)
+	{
+		delete Handle;
+		Handle = nullptr;
+	}
+
+	HeaderSizeWhenReadingExportsFromSplitFile = HeaderSize;
+	FileName = FPaths::GetBaseFilename(FileName, false) + TEXT(".uexp");
+
+	Handle = FPlatformFileManager::Get().GetPlatformFile().OpenAsyncRead(*FileName);
+	check(Handle); // this generally cannot fail because it is async
+
+	check(!SizeRequestPtr);
+	SizeRequestPtr = Handle->SizeRequest();
+	if (SizeRequestPtr->PollCompletion())
+	{
+		TotalSize(); // complete the request
+	}
+
+#endif
+}
+
+#if USE_EVENT_DRIVEN_ASYNC_LOAD
+IAsyncReadRequest* FArchiveAsync2::MakeEventDrivenPrecacheRequest(int64 Offset, int64 BytesToRead, FAsyncFileCallBack* CompleteCallback)
+{
+	if (LoadPhase == ELoadPhase::WaitingForFirstExport)
+	{
+		FirstExportStarting();
+	}
+	double StartTime = FPlatformTime::Seconds();
+	check(Offset - HeaderSizeWhenReadingExportsFromSplitFile >= 0);
+	check(Offset + BytesToRead <= TotalSizeOrMaxInt64IfNotReady());
+	IAsyncReadRequest* Precache = Handle->ReadRequest(Offset - HeaderSizeWhenReadingExportsFromSplitFile, BytesToRead, AIOP_Precache, CompleteCallback);
+	LogItem(TEXT("Event Precache"), Offset - HeaderSizeWhenReadingExportsFromSplitFile, BytesToRead, StartTime);
+	return Precache;
+}
+#endif
+
+
+bool FArchiveAsync2::Precache(int64 RequestOffset, int64 RequestSize, bool bUseTimeLimit, bool bUseFullTimeLimit, double TickStartTime, float TimeLimit)
 {
 #if defined(ASYNC_WATCH_FILE)
 	if (FileName.Contains(TEXT(ASYNC_WATCH_FILE)) && RequestOffset == 878129)
@@ -4020,19 +7064,10 @@ bool FArchiveAsync2::Precache(int64 RequestOffset, int64 RequestSize, bool bUseT
 	}
 	if (LoadPhase == ELoadPhase::WaitingForFirstExport)
 	{
-		ExportReadTime = FPlatformTime::Seconds();
-		LogItem(TEXT("Exports"));
-		LoadPhase = ELoadPhase::ProcessingExports;
-#if !USE_FArchiveAsync2_PRECACHE_WHOLE_FILE
-		DispatchPrecache(RequestOffset, true); // start precacheing the exports
-#endif
+		FirstExportStarting();
 	}
 	if (!bUseTimeLimit)
 	{
-		if (bTrimPrecache)
-		{
-			DispatchPrecache(RequestOffset, true); // discard anything before here
-		}
 		return true; // we will stream and do the blocking on the serialize calls
 	}
 	bool bResult = PrecacheInternal(RequestOffset, RequestSize);
@@ -4044,10 +7079,6 @@ bool FArchiveAsync2::Precache(int64 RequestOffset, int64 RequestSize, bool bUseT
 			bResult = true;
 		}
 	}
-	if (bTrimPrecache)
-	{
-		DispatchPrecache(RequestOffset + RequestSize, true); // discard anything before and including this export....we already have the real request in flight
-	}
 	return bResult;
 }
 
@@ -4059,18 +7090,28 @@ bool FArchiveAsync2::Precache(int64 RequestOffset, int64 RequestSize)
 	}
 	if (LoadPhase == ELoadPhase::WaitingForHeader)
 	{
+		//@todoio, it would be nice to check that when we read the header, we don't read any more than we really need...i.e. no "minimum read size"
 		check(RequestOffset == 0 && RequestOffset + RequestSize <= HeaderSize);
 	}
 	return PrecacheInternal(RequestOffset, RequestSize);
 }
 
+bool FArchiveAsync2::PrecacheForEvent(int64 RequestOffset, int64 RequestSize)
+{
+	check(int32(LoadPhase) > int32(ELoadPhase::WaitingForHeader));
+	return PrecacheInternal(RequestOffset, RequestSize, false);
+}
+
 
 void FArchiveAsync2::StartReadingHeader()
 {
 	//LogItem(TEXT("Start Header"));
 	WaitForIntialPhases();
-	check(LoadPhase == ELoadPhase::WaitingForHeader && ReadRequestPtr);
-	WaitRead();
+	if (!ArIsError)
+	{
+		check(LoadPhase == ELoadPhase::WaitingForHeader && ReadRequestPtr);
+		WaitRead();
+	}
 }
 
 void FArchiveAsync2::EndReadingHeader()
@@ -4078,6 +7119,7 @@ void FArchiveAsync2::EndReadingHeader()
 	LogItem(TEXT("End Header"));
 	check(LoadPhase == ELoadPhase::WaitingForHeader);
 	LoadPhase = ELoadPhase::WaitingForFirstExport;
+	FlushPrecacheBlock();
 }
 
 bool FArchiveAsync2::ReadyToStartReadingHeader(bool bUseTimeLimit, bool bUseFullTimeLimit, double TickStartTime, float TimeLimit)
@@ -4116,15 +7158,26 @@ void CallSerializeHook();
 
 void FArchiveAsync2::Serialize(void* Data, int64 Count)
 {
-#if TRACK_SERIALIZE
-	CallSerializeHook();
-#endif
-
 	if (!Count || ArIsError)
 	{
 		return;
 	}
 	check(Count > 0);
+#if DEVIRTUALIZE_FLinkerLoad_Serialize
+	if (ActiveFPLB->StartFastPathLoadBuffer + Count <= ActiveFPLB->EndFastPathLoadBuffer)
+	{
+		// this wasn't one of the cases we devirtualized; we can short circut here to avoid resettting the buffer when we don't need to
+		FMemory::Memcpy(Data, ActiveFPLB->StartFastPathLoadBuffer, Count);
+		ActiveFPLB->StartFastPathLoadBuffer += Count;
+		return;
+	}
+
+	DiscardInlineBufferAndUpdateCurrentPos();
+#endif
+#if TRACK_SERIALIZE
+	CallSerializeHook();
+#endif
+
 #if PLATFORM_DESKTOP
 	// Show a message box indicating, possible, corrupt data (desktop platforms only)
 	if (CurrentPos + Count > TotalSize())
@@ -4188,7 +7241,7 @@ void FArchiveAsync2::Serialize(void* Data, int64 Count)
 	}
 	if (BeforeBlockSize)
 	{
-		UE_LOG(LogAsyncArchive, Warning, TEXT("FArchiveAsync2::DispatchPrecache Backwards streaming in %s  BeforeBlockOffset = %lld   BeforeBlockOffset = %lld"), *FileName, BeforeBlockOffset, BeforeBlockOffset);
+		UE_LOG(LogAsyncArchive, Warning, TEXT("FArchiveAsync2::Serialize Backwards streaming in %s  BeforeBlockOffset = %lld   BeforeBlockOffset = %lld"), *FileName, BeforeBlockOffset, BeforeBlockOffset);
 		LogItem(TEXT("Sync Before Block"), BeforeBlockOffset, BeforeBlockSize);
 		PrecacheInternal(BeforeBlockOffset, BeforeBlockSize);
 		WaitRead();
@@ -4209,9 +7262,21 @@ void FArchiveAsync2::Serialize(void* Data, int64 Count)
 		check(AfterBlockOffset >= PrecacheStartPos && AfterBlockOffset + AfterBlockSize <= PrecacheEndPos);
 		FMemory::Memcpy(((uint8*)Data) + Count - AfterBlockSize, PrecacheBuffer + AfterBlockOffset - PrecacheStartPos, AfterBlockSize);
 	}
+#if DEVIRTUALIZE_FLinkerLoad_Serialize
+	SetPosAndUpdatePrecacheBuffer(CurrentPos + Count);
+#else
 	CurrentPos += Count;
+#endif
 }
 
+#if DEVIRTUALIZE_FLinkerLoad_Serialize
+void FArchiveAsync2::DiscardInlineBufferAndUpdateCurrentPos()
+{
+	CurrentPos += (ActiveFPLB->StartFastPathLoadBuffer - ActiveFPLB->OriginalFastPathLoadBuffer);
+	ActiveFPLB->Reset();
+}
+#endif
+
 #if TRACK_SERIALIZE
 
 #include "StackTracker.h"
