Copying //UE4/Dev-Core to //UE4/Dev-Main (Source: //UE4/Dev-Core @ 3151653)

diff --git a/Engine/Source/Runtime/CoreUObject/Private/Serialization/AsyncLoadingPrivate.h b/Engine/Source/Runtime/CoreUObject/Private/Serialization/AsyncLoadingPrivate.h
--- a/Engine/Source/Runtime/CoreUObject/Private/Serialization/AsyncLoadingPrivate.h
+++ b/Engine/Source/Runtime/CoreUObject/Private/Serialization/AsyncLoadingPrivate.h
@@ -2,6 +2,8 @@
 
 #pragma once
 
+DECLARE_LOG_CATEGORY_EXTERN(LogLoadingDev, Fatal, All);
+
 #if USE_NEW_ASYNC_IO
 
 class COREUOBJECT_API FArchiveAsync2 final: public FArchive
@@ -22,55 +24,73 @@ public:
 		MAX_SUMMARY_SIZE = 4096
 	};
 
-	FArchiveAsync2(const TCHAR* InFileName);
+	FArchiveAsync2(const TCHAR* InFileName
+#if USE_EVENT_DRIVEN_ASYNC_LOAD
+		, TFunction<void()>&& InSummaryReadyCallback
+#endif
+		);
 	virtual ~FArchiveAsync2();
 	virtual bool Close() override;
 	virtual bool SetCompressionMap(TArray<FCompressedChunk>* CompressedChunks, ECompressionFlags CompressionFlags) override;
 	virtual bool Precache(int64 PrecacheOffset, int64 PrecacheSize) override;
 	virtual void Serialize(void* Data, int64 Num) override;
-	virtual int64 Tell() override
+	FORCEINLINE virtual int64 Tell() override
 	{
-		return CurrentPos;
+		return CurrentPos + (ActiveFPLB->StartFastPathLoadBuffer - ActiveFPLB->OriginalFastPathLoadBuffer);
 	}
 	virtual int64 TotalSize() override;
 	virtual void Seek(int64 InPos) override;
 	virtual void FlushCache() override;
 
-	bool Precache(int64 PrecacheOffset, int64 PrecacheSize, bool bUseTimeLimit, bool bUseFullTimeLimit, double TickStartTime, float TimeLimit, bool bTrimPrecache);
+	bool Precache(int64 PrecacheOffset, int64 PrecacheSize, bool bUseTimeLimit, bool bUseFullTimeLimit, double TickStartTime, float TimeLimit);
+	bool PrecacheForEvent(int64 PrecacheOffset, int64 PrecacheSize);
+	void FlushPrecacheBlock();
 	bool ReadyToStartReadingHeader(bool bUseTimeLimit, bool bUseFullTimeLimit, double TickStartTime, float TimeLimit);
 	void StartReadingHeader();
 	void EndReadingHeader();
+	void FirstExportStarting();
+#if USE_EVENT_DRIVEN_ASYNC_LOAD
+	IAsyncReadRequest* MakeEventDrivenPrecacheRequest(int64 Offset, int64 BytesToRead, FAsyncFileCallBack* CompleteCallback);
+#endif
+
+	void LogItem(const TCHAR* Item, int64 Offset = 0, int64 Size = 0, double StartTime = 0.0);
 
 private:
+#if DEVIRTUALIZE_FLinkerLoad_Serialize
+	/**
+	* Updates CurrentPos based on StartFastPathLoadBuffer and sets both StartFastPathLoadBuffer and EndFastPathLoadBuffer to null
+	*/
+	void DiscardInlineBufferAndUpdateCurrentPos();
+
+	void SetPosAndUpdatePrecacheBuffer(int64 Pos);
+
+#endif
 
 	bool WaitRead(float TimeLimit = 0.0f);
 	void CompleteRead();
 	void CancelRead();
 	void CompleteCancel();
-	void CompletePrecache();
 	bool WaitForIntialPhases(float TimeLimit = 0.0f);
-	void DispatchPrecache(int64 PrecacheAt = 0, bool bForceTrimPrecache = false);
 	void ReadCallback(bool bWasCancelled, IAsyncReadRequest*);
-	bool PrecacheInternal(int64 PrecacheOffset, int64 PrecacheSize);
+	bool PrecacheInternal(int64 PrecacheOffset, int64 PrecacheSize, bool bApplyMinReadSize = true);
 	
-	void LogItem(const TCHAR* Item, int64 Offset = 0, int64 Size = 0, double StartTime = 0.0);
-
 	FORCEINLINE int64 TotalSizeOrMaxInt64IfNotReady()
 	{
-		return SizeRequestPtr ? MAX_int64 : FileSize;
+
+		return SizeRequestPtr ? MAX_int64 : (FileSize + HeaderSizeWhenReadingExportsFromSplitFile);
 	}
 
 	IAsyncReadFileHandle* Handle;
 	IAsyncReadRequest* SizeRequestPtr;
+#if !(SPLIT_COOKED_FILES && USE_EVENT_DRIVEN_ASYNC_LOAD)
 	IAsyncReadRequest* SummaryRequestPtr;
 	IAsyncReadRequest* SummaryPrecacheRequestPtr;
+#endif
 	IAsyncReadRequest* ReadRequestPtr;
-	IAsyncReadRequest* PrecacheRequestPtr;
 	IAsyncReadRequest* CanceledReadRequestPtr;
+
 	/** Buffer containing precached data.											*/
 	uint8* PrecacheBuffer;
-	/** Buffer containing precached data.											*/
-	uint8* NextPrecacheBuffer;
 	/** Cached file size															*/
 	int64 FileSize;
 	/** Current position of archive.												*/
@@ -82,9 +102,9 @@ private:
 
 	int64 ReadRequestOffset;
 	int64 ReadRequestSize;
-	int64 PrecacheOffset;
 
 	int64 HeaderSize;
+	int64 HeaderSizeWhenReadingExportsFromSplitFile;
 
 	ELoadPhase LoadPhase;
 	bool bKeepRestOfFilePrecached;
@@ -94,7 +114,11 @@ private:
 	double OpenTime;
 	double SummaryReadTime;
 	double ExportReadTime;
-	double PrecacheEOFTime;
+#if USE_EVENT_DRIVEN_ASYNC_LOAD
+	TFunction<void()> SummaryReadyCallback;
+	FAsyncFileCallBack ReadCallbackFunctionForLinkerLoad;
+#endif
+
 };
 
 #endif
\ No newline at end of file
diff --git a/Engine/Source/Runtime/CoreUObject/Private/Serialization/AsyncLoadingThread.h b/Engine/Source/Runtime/CoreUObject/Private/Serialization/AsyncLoadingThread.h
--- a/Engine/Source/Runtime/CoreUObject/Private/Serialization/AsyncLoadingThread.h
+++ b/Engine/Source/Runtime/CoreUObject/Private/Serialization/AsyncLoadingThread.h
@@ -7,6 +7,126 @@
 
 class IAssetRegistryInterface;
 
+#ifndef USE_EVENT_DRIVEN_ASYNC_LOAD
+#error "USE_EVENT_DRIVEN_ASYNC_LOAD must be defined""
+#endif
+
+#if USE_EVENT_DRIVEN_ASYNC_LOAD
+
+struct FAsyncLoadEvent
+{
+	enum
+	{
+		EventSystemPriority_MAX = MAX_int32
+	};
+
+	int32 UserPriority;
+	int32 PackageSerialNumber;
+	int32 EventSystemPriority;
+	int32 SerialNumber;
+	TFunction<void(FAsyncLoadEventArgs& Args)> Payload;
+
+	FAsyncLoadEvent()
+		: UserPriority(0)
+		, PackageSerialNumber(0)
+		, EventSystemPriority(0)
+		, SerialNumber(0)
+	{
+	}
+	FAsyncLoadEvent(int32 InUserPriority, int32 InPackageSerialNumber, int32 InEventSystemPriority, int32 InSerialNumber, TFunction<void(FAsyncLoadEventArgs& Args)>&& InPayload)
+		: UserPriority(InUserPriority)
+		, PackageSerialNumber(InPackageSerialNumber)
+		, EventSystemPriority(InEventSystemPriority)
+		, SerialNumber(InSerialNumber)
+		, Payload(Forward<TFunction<void(FAsyncLoadEventArgs& Args)>>(InPayload))
+	{
+	}
+
+	FORCEINLINE bool operator<(const FAsyncLoadEvent& Other) const
+	{
+		if (UserPriority != Other.UserPriority)
+		{
+			return UserPriority > Other.UserPriority;
+		}
+		if (EventSystemPriority != Other.EventSystemPriority)
+		{
+			return EventSystemPriority > Other.EventSystemPriority;
+		}
+		if (PackageSerialNumber != Other.PackageSerialNumber)
+		{
+			return PackageSerialNumber > Other.PackageSerialNumber; // roughly DFS
+		}
+		return SerialNumber < Other.SerialNumber;
+	}
+};
+
+struct FAsyncLoadEventQueue
+{
+	//FCriticalSection AsyncLoadEventQueueCritical; //@todoio maybe this doesn't need to be protected by a critical section. 
+	int32 RunningSerialNumber;
+	TArray<FAsyncLoadEvent> EventQueue;
+
+	FAsyncLoadEventQueue()
+		: RunningSerialNumber(0)
+	{
+	}
+
+	FORCEINLINE void AddAsyncEvent(int32 UserPriority, int32 PackageSerialNumber, int32 EventSystemPriority, TFunction<void(FAsyncLoadEventArgs& Args)>&& Payload)
+	{
+		//FScopeLock Lock(&AsyncLoadEventQueueCritical);
+		//@todoio check(FAsyncLoadingThread::IsInAsyncLoadThread());
+		EventQueue.HeapPush(FAsyncLoadEvent(UserPriority, PackageSerialNumber, EventSystemPriority, ++RunningSerialNumber, Forward<TFunction<void(FAsyncLoadEventArgs& Args)>>(Payload)));
+	}
+	bool PopAndExecute(FAsyncLoadEventArgs& Args)
+	{
+		FAsyncLoadEvent Event;
+		bool bResult = false;
+		{
+			//FScopeLock Lock(&AsyncLoadEventQueueCritical);
+			//@todoio check(FAsyncLoadingThread::IsInAsyncLoadThread());
+			if (EventQueue.Num())
+			{
+				EventQueue.HeapPop(Event, false);
+				bResult = true;
+			}
+		}
+		if (bResult)
+		{
+			Event.Payload(Args);
+		}
+		return bResult;
+	}
+};
+
+#endif
+
+/** Package dependency tree used for flushing specific packages */
+struct FFlushTree
+{
+	int32 RequestId;
+	TSet<FName> PackagesToFlush;
+
+	FFlushTree(int32 InRequestId)
+		: RequestId(InRequestId)
+	{
+	}
+
+	bool AddPackage(const FName& Package)
+	{
+		if (!PackagesToFlush.Contains(Package))
+		{
+			PackagesToFlush.Add(Package);
+			return true;
+		}
+		return false;
+	}
+
+	bool Contains(const FName& Package)
+	{
+		return PackagesToFlush.Contains(Package);
+	}
+};
+
 /**
  * Async loading thread. Preloads/serializes packages on async loading thread. Postloads objects on the game thread.
  */
@@ -57,11 +177,13 @@ class FAsyncLoadingThread : public FRunnable
 
 	/** [ASYNC THREAD] Array of packages that are being preloaded */
 	TArray<FAsyncPackage*> AsyncPackages;
-
-	/** Packages that we've already hinted for reading */
-	TSet<FAsyncPackage*> PendingPackageHints;
-
 	TMap<FName, FAsyncPackage*> AsyncPackageNameLookup;
+#if USE_EVENT_DRIVEN_ASYNC_LOAD
+public:
+	TArray<FAsyncPackage*> AsyncPackagesReadyForTick;
+private:
+#endif 
+
 #if THREADSAFE_UOBJECTS
 	/** We only lock AsyncPackages array to make GetAsyncLoadPercentage thread safe, so we only care about locking Add/Remove operations on the async thread */
 	FCriticalSection AsyncPackagesCritical;
@@ -110,13 +232,40 @@ public:
 	/** Returns the async loading thread singleton */
 	static FAsyncLoadingThread& Get();
 
+#if USE_EVENT_DRIVEN_ASYNC_LOAD
+	FAsyncLoadEventQueue EventQueue;
+
+	FORCEINLINE FAsyncPackage* GetPackage(FWeakAsyncPackagePtr Ptr)
+	{
+		if (Ptr.PackageName != NAME_None && Ptr.SerialNumber)
+		{
+			FAsyncPackage* Package = FindAsyncPackage(Ptr.PackageName);
+			if (Package && Package->SerialNumber == Ptr.SerialNumber)
+			{
+				return Package;
+			}
+		}
+		return nullptr;
+	}
+
+	void QueueEvent_CreateLinker(FAsyncPackage* Pkg, int32 EventSystemPriority = 0);
+	void QueueEvent_FinishLinker(FWeakAsyncPackagePtr WeakPtr, int32 EventSystemPriority = 0);
+	void QueueEvent_StartImportPackages(FAsyncPackage* Pkg, int32 EventSystemPriority = 0);
+	void QueueEvent_SetupImports(FAsyncPackage* Pkg, int32 EventSystemPriority = 0);
+	void QueueEvent_SetupExports(FAsyncPackage* Pkg, int32 EventSystemPriority = 0);
+	void QueueEvent_ProcessImportsAndExports(FAsyncPackage* Pkg, int32 EventSystemPriority = 0);
+	void QueueEvent_ExportsDone(FAsyncPackage* Pkg, int32 EventSystemPriority = 0);
+	void QueueEvent_ProcessPostloadWait(FAsyncPackage* Pkg, int32 EventSystemPriority = 0);
+	void QueueEvent_StartPostLoad(FAsyncPackage* Pkg, int32 EventSystemPriority = 0);
+#endif
+
 	/** True if multithreaded async loading should be used. */
 	static FORCEINLINE bool IsMultithreaded()
 	{
 		static struct FAsyncLoadingThreadEnabled
 		{
 			bool Value;
-			FAsyncLoadingThreadEnabled()
+			FORCENOINLINE FAsyncLoadingThreadEnabled()
 			{
 #if THREADSAFE_UOBJECTS
 				if (FPlatformProperties::RequiresCookedData())
@@ -124,8 +273,9 @@ public:
 					check(GConfig);
 					bool bConfigValue = true;
 					GConfig->GetBool(TEXT("/Script/Engine.StreamingSettings"), TEXT("s.AsyncLoadingThreadEnabled"), bConfigValue, GEngineIni);
-					bool bCommandLineNoAsyncThread = false;
-					Value = bConfigValue && FApp::ShouldUseThreadingForPerformance() && !FParse::Param(FCommandLine::Get(), TEXT("NoAsyncLoadingThread"));
+					bool bCommandLineNoAsyncThread = FParse::Param(FCommandLine::Get(), TEXT("NoAsyncLoadingThread"));
+					bool bCommandLineAsyncThread = FParse::Param(FCommandLine::Get(), TEXT("AsyncLoadingThread"));
+					Value = bCommandLineAsyncThread || (bConfigValue && FApp::ShouldUseThreadingForPerformance() && !bCommandLineNoAsyncThread);
 				}
 				else
 #endif
@@ -138,15 +288,35 @@ public:
 	}
 
 	/** Sets the current state of async loading */
-	FORCEINLINE void SetIsInAsyncLoadingTick(bool InTick)
+	void EnterAsyncLoadingTick()
 	{
-		bIsInAsyncLoadingTick = InTick;
+		AsyncLoadingTickCounter++;
+	}
+
+	void LeaveAsyncLoadingTick()
+	{
+		AsyncLoadingTickCounter--;
+		check(AsyncLoadingTickCounter >= 0);
 	}
 
 	/** Gets the current state of async loading */
 	FORCEINLINE bool GetIsInAsyncLoadingTick() const
 	{
-		return bIsInAsyncLoadingTick;
+		return !!AsyncLoadingTickCounter;
+	}
+
+	/** Checks if a package is still being referenced by async loading code. @todo: remove this */
+	bool IsPackageReferenced(FAsyncPackage* Package)
+	{
+		if (PackagesToDelete.Contains(Package))
+		{
+			return true;
+		}
+		if (AsyncPackages.Contains(Package))
+		{
+			return true;
+		}
+		return false;
 	}
 
 	/** Returns true if packages are currently being loaded on the async thread */
@@ -166,11 +336,11 @@ public:
 			// we're on game thread but inside of async loading code (PostLoad mostly)
 			// to make it behave exactly like the non-threaded version
 			bResult = FPlatformTLS::GetCurrentThreadId() == AsyncLoadingThreadID || 
-				(IsInGameThread() && FAsyncLoadingThread::Get().bIsInAsyncLoadingTick);
+				(IsInGameThread() && FAsyncLoadingThread::Get().GetIsInAsyncLoadingTick());
 		}
 		else
 		{
-			bResult = IsInGameThread() && FAsyncLoadingThread::Get().bIsInAsyncLoadingTick;
+			bResult = IsInGameThread() && FAsyncLoadingThread::Get().GetIsInAsyncLoadingTick();
 		}
 		return bResult;
 	}
@@ -213,7 +383,7 @@ public:
 	* @param PackageName - async package name.
 	* @param InsertMode - Insert mode, describing how we insert this package into the request list
 	*/
-	void InsertPackage(FAsyncPackage* Package, EAsyncPackageInsertMode InsertMode = EAsyncPackageInsertMode::InsertBeforeMatchingPriorities);
+	void InsertPackage(FAsyncPackage* Package, bool bReinsert = false, EAsyncPackageInsertMode InsertMode = EAsyncPackageInsertMode::InsertBeforeMatchingPriorities);
 
 	/**
 	* [ASYNC THREAD] Finds an existing async package in the LoadedPackages by its name.
@@ -267,11 +437,17 @@ public:
 	* @param bUseTimeLimit True if time limit should be used [time-slicing].
 	* @param bUseFullTimeLimit True if full time limit should be used [time-slicing].
 	* @param TimeLimit Maximum amount of time that can be spent in this call [time-slicing].
+	* @param FlushTree Package dependency tree to be flushed
 	* @return The current state of async loading
 	*/
-	EAsyncPackageState::Type ProcessAsyncLoading(int32& OutPackagesProcessed, bool bUseTimeLimit = false, bool bUseFullTimeLimit = false, float TimeLimit = 0.0f);
+	EAsyncPackageState::Type ProcessAsyncLoading(int32& OutPackagesProcessed, bool bUseTimeLimit = false, bool bUseFullTimeLimit = false, float TimeLimit = 0.0f, FFlushTree* FlushTree = nullptr);
 
-	void JustInTimeHinting();
+#if USE_EVENT_DRIVEN_ASYNC_LOAD
+	/**
+	* [ASYNC* THREAD] Checks fopr cycles in the event driven loader and does fatal errors in that case
+	*/
+	void CheckForCycles();
+#endif
 
 	/**
 	* [GAME THREAD] Ticks game thread side of async loading.
@@ -279,9 +455,10 @@ public:
 	* @param bUseTimeLimit True if time limit should be used [time-slicing].
 	* @param bUseFullTimeLimit True if full time limit should be used [time-slicing].
 	* @param TimeLimit Maximum amount of time that can be spent in this call [time-slicing].
+	* @param FlushTree Package dependency tree to be flushed
 	* @return The current state of async loading
 	*/
-	EAsyncPackageState::Type TickAsyncLoading(bool bUseTimeLimit, bool bUseFullTimeLimit, float TimeLimit, int32 WaitForRequestID = INDEX_NONE);
+	EAsyncPackageState::Type TickAsyncLoading(bool bUseTimeLimit, bool bUseFullTimeLimit, float TimeLimit, FFlushTree* FlushTree = nullptr);
 
 	/**
 	* [ASYNC THREAD] Main thread loop
@@ -289,8 +466,9 @@ public:
 	* @param bUseTimeLimit True if time limit should be used [time-slicing].
 	* @param bUseFullTimeLimit True if full time limit should be used [time-slicing].
 	* @param TimeLimit Maximum amount of time that can be spent in this call [time-slicing].
+	* @param FlushTree Package dependency tree to be flushed
 	*/
-	EAsyncPackageState::Type TickAsyncThread(bool bUseTimeLimit, bool bUseFullTimeLimit, float TimeLimit);
+	EAsyncPackageState::Type TickAsyncThread(bool bUseTimeLimit, bool bUseFullTimeLimit, float TimeLimit, bool& bDidSomething, FFlushTree* FlushTree = nullptr);
 
 	/** Initializes async loading thread */
 	void InitializeAsyncThread();
@@ -349,21 +527,23 @@ private:
 	* @param bUseTimeLimit True if time limit should be used [time-slicing].
 	* @param bUseFullTimeLimit True if full time limit should be used [time-slicing].
 	* @param TimeLimit Maximum amount of time that can be spent in this call [time-slicing].
-	* @param WaitForRequestID If the package request ID is valid, exits as soon as the request is no longer in the qeueue.
+	* @param FlushTree Package dependency tree to be flushed
 	* @return The current state of async loading
 	*/
-	EAsyncPackageState::Type ProcessLoadedPackages(bool bUseTimeLimit, bool bUseFullTimeLimit, float TimeLimit, int32 WaitForRequestID = INDEX_NONE);
+	EAsyncPackageState::Type ProcessLoadedPackages(bool bUseTimeLimit, bool bUseFullTimeLimit, float TimeLimit, bool& bDidSomething, FFlushTree* FlushTree = nullptr);
 
 	/**
 	* [ASYNC THREAD] Creates async packages from the queued requests
+	* @param FlushTree Package dependency tree to be flushed
 	*/
-	int32 CreateAsyncPackagesFromQueue();
+	int32 CreateAsyncPackagesFromQueue(bool bUseTimeLimit, bool bUseFullTimeLimit, float TimeLimit, FFlushTree* FlushTree = nullptr);
 
 	/**
 	* [ASYNC THREAD] Internal helper function for processing a package load request. If dependency preloading is enabled, 
 	* it will call itself recursively for all the package dependencies
+	* @param FlushTree Package dependency tree to be flushed
 	*/
-	void ProcessAsyncPackageRequest(FAsyncPackageDesc* InRequest, FAsyncPackage* InRootPackage, IAssetRegistryInterface* InAssetRegistry);
+	void ProcessAsyncPackageRequest(FAsyncPackageDesc* InRequest, FAsyncPackage* InRootPackage, IAssetRegistryInterface* InAssetRegistry, FFlushTree* FlushTree);
 
 	/**
 	* [ASYNC THREAD] Internal helper function for updating the priorities of an existing package and all its dependencies
@@ -372,24 +552,18 @@ private:
 
 	/**
 	* [ASYNC THREAD] Finds existing async package and adds the new request's completion callback to it.
+	* @param FlushTree Package dependency tree to be flushed
 	*/
-	FAsyncPackage* FindExistingPackageAndAddCompletionCallback(FAsyncPackageDesc* PackageRequest, TMap<FName, FAsyncPackage*>& PackageList);
+	FAsyncPackage* FindExistingPackageAndAddCompletionCallback(FAsyncPackageDesc* PackageRequest, TMap<FName, FAsyncPackage*>& PackageList, FFlushTree* FlushTree);
 
 	/**
 	* [ASYNC THREAD] Adds a package to a list of packages that have finished loading on the async thread
 	*/
-	FORCEINLINE void AddToLoadedPackages(FAsyncPackage* Package)
-	{
-#if THREADSAFE_UOBJECTS
-		FScopeLock LoadedLock(&LoadedPackagesCritical);
-#endif
-		LoadedPackages.Add(Package);
-		LoadedPackagesNameLookup.Add(Package->GetPackageName(), Package);
-	}
+	void AddToLoadedPackages(FAsyncPackage* Package);
 
 	/** Cancels async loading internally */
 	void CancelAsyncLoadingInternal();
 
-	/** True if async loading is currently being ticked, mostly used by singlethreaded ticking */
-	bool bIsInAsyncLoadingTick;
+	/** Number of times we re-entered the async loading tick, mostly used by singlethreaded ticking. Debug purposes only. */
+	int32 AsyncLoadingTickCounter;
 };
diff --git a/Engine/Source/Runtime/CoreUObject/Private/Serialization/BulkData.cpp b/Engine/Source/Runtime/CoreUObject/Private/Serialization/BulkData.cpp
--- a/Engine/Source/Runtime/CoreUObject/Private/Serialization/BulkData.cpp
+++ b/Engine/Source/Runtime/CoreUObject/Private/Serialization/BulkData.cpp
@@ -656,6 +656,9 @@ void FUntypedBulkData::StartSerializingBulkData(FArchive& Ar, UObject* Owner, in
 	{
 		BulkDataAsync.Reallocate(GetBulkDataSize(), BulkDataAlignment);
 
+#if USE_NEW_ASYNC_IO
+		UE_LOG(LogSerialization, Error, TEXT("Attempt to stream bulk data with EDL enabled. This is not desireable. File %s"), *Filename);
+#endif
 		FArchive* FileReaderAr = IFileManager::Get().CreateFileReader(*Filename, FILEREAD_Silent);
 		checkf(FileReaderAr != NULL, TEXT("Attempted to load bulk data from an invalid filename '%s'."), *Filename);
 
@@ -682,8 +685,30 @@ static FAutoConsoleVariableRef CVarMinimumBulkDataSizeForAsyncLoading(
 	ECVF_Default
 	);
 
+#if !defined(USE_NEW_ASYNC_IO) || !defined(SPLIT_COOKED_FILES) || !defined(USE_EVENT_DRIVEN_ASYNC_LOAD)
+	#error "USE_NEW_ASYNC_IO and SPLIT_COOKED_FILES and USE_EVENT_DRIVEN_ASYNC_LOAD must be defined."
+#endif
+
 bool FUntypedBulkData::ShouldStreamBulkData()
 {
+#if USE_NEW_ASYNC_IO
+	const bool bPayloadInline = !(BulkDataFlags&BULKDATA_PayloadAtEndOfFile);
+	if (bPayloadInline)
+	{
+		return false; // if it is inline, it is already precached, so use it
+	}
+
+#endif
+#if SPLIT_COOKED_FILES && USE_EVENT_DRIVEN_ASYNC_LOAD
+	const bool bSeperateFile = !!(BulkDataFlags&BULKDATA_PayloadInSeperateFile);
+
+	if (!bSeperateFile)
+	{
+		check(!"Bulk data should either be inline or stored in a separate file for the new uobject loader.");
+		return false; // if it is not in a separate file, then we can't easily find the correct offset in the uexp file; we don't want this case anyway!
+	}
+
+#endif
 	const bool bForceStream = !!(BulkDataFlags & BULKDATA_ForceStreamPayload);
 
 	return (FPlatformProperties::RequiresCookedData() && !Filename.IsEmpty() &&
@@ -856,6 +881,17 @@ void FUntypedBulkData::Serialize( FArchive& Ar, UObject* Owner, int32 Idx )
 						if (BulkDataFlags & BULKDATA_PayloadInSeperateFile)
 						{
 							// open seperate bulk data file
+#if USE_NEW_ASYNC_IO
+							UE_LOG(LogSerialization, Error, TEXT("Attempt to sync load bulk data with EDL enabled (separate file). This is not desireable. File %s"), *Filename);
+#endif
+#if SPLIT_COOKED_FILES && USE_EVENT_DRIVEN_ASYNC_LOAD
+							if (Filename.EndsWith(TEXT(".uasset")) || Filename.EndsWith(TEXT(".umap")))
+							{
+								BulkDataOffsetInFile -= IFileManager::Get().FileSize(*Filename);
+								check(BulkDataOffsetInFile >= 0);
+								Filename = FPaths::GetBaseFilename(Filename, false) + TEXT(".uexp");
+							}
+#endif
 							FArchive* TargetArchive = IFileManager::Get().CreateFileReader(*Filename);
 							// seek to the location in the file where the payload is stored
 							TargetArchive->Seek(BulkDataOffsetInFile);
@@ -866,6 +902,9 @@ void FUntypedBulkData::Serialize( FArchive& Ar, UObject* Owner, int32 Idx )
 						}
 						else
 						{
+#if USE_NEW_ASYNC_IO
+							UE_LOG(LogSerialization, Error, TEXT("Attempt to sync load bulk data with EDL enabled. This is not desireable. File %s"), *Filename);
+#endif
 							// store the current file offset
 							int64 CurOffset = Ar.Tell();
 							// seek to the location in the file where the payload is stored
@@ -928,6 +967,11 @@ void FUntypedBulkData::Serialize( FArchive& Ar, UObject* Owner, int32 Idx )
 			// determine whether we are going to store the payload inline or not.
 			bool bStoreInline = !!(BulkDataFlags&BULKDATA_ForceInlinePayload) || !LinkerSave;
 
+			if (IsEventDrivenLoaderEnabledInCookedBuilds() && Ar.IsCooking() && !bStoreInline && !(BulkDataFlags&BULKDATA_Force_NOT_InlinePayload))
+			{
+				bStoreInline = true;
+			}
+
 			if (!bStoreInline)
 			{
 				// set the flag indicating where the payload is stored
@@ -1328,6 +1372,17 @@ void FUntypedBulkData::LoadDataIntoMemory( void* Dest )
 		// load from the specied filename when the linker has been cleared
 		checkf( Filename != TEXT(""), TEXT( "Attempted to load bulk data without a proper filename." ) );
 	
+#if USE_NEW_ASYNC_IO
+		UE_CLOG(!(IsInGameThread() || IsInAsyncLoadingThread()), LogSerialization, Error, TEXT("Attempt to sync load bulk data with EDL enabled (LoadDataIntoMemory). This is not desireable. File %s"), *Filename);
+#endif
+#if SPLIT_COOKED_FILES && USE_EVENT_DRIVEN_ASYNC_LOAD
+		if (Filename.EndsWith(TEXT(".uasset")) || Filename.EndsWith(TEXT(".umap")))
+		{
+			BulkDataOffsetInFile -= IFileManager::Get().FileSize(*Filename);
+			check(BulkDataOffsetInFile >= 0);
+			Filename = FPaths::GetBaseFilename(Filename, false) + TEXT(".uexp");
+		}
+#endif
 		FArchive* Ar = IFileManager::Get().CreateFileReader(*Filename, FILEREAD_Silent);
 		checkf( Ar != NULL, TEXT( "Attempted to load bulk data from an invalid filename '%s'." ), *Filename );
 	
diff --git a/Engine/Source/Runtime/CoreUObject/Private/Serialization/ObjectReader.cpp b/Engine/Source/Runtime/CoreUObject/Private/Serialization/ObjectReader.cpp
--- a/Engine/Source/Runtime/CoreUObject/Private/Serialization/ObjectReader.cpp
+++ b/Engine/Source/Runtime/CoreUObject/Private/Serialization/ObjectReader.cpp
@@ -58,6 +58,12 @@ FArchive& FObjectReader::operator<<(FStringAssetReference& Value)
 	return *this;
 }
 
+FArchive& FObjectReader::operator<< (struct FWeakObjectPtr& Value)
+{
+	Value.Serialize(*this);
+	return *this;
+}
+
 FString FObjectReader::GetArchiveName() const
 {
 	return TEXT("FObjectReader");
diff --git a/Engine/Source/Runtime/CoreUObject/Private/Serialization/ObjectWriter.cpp b/Engine/Source/Runtime/CoreUObject/Private/Serialization/ObjectWriter.cpp
--- a/Engine/Source/Runtime/CoreUObject/Private/Serialization/ObjectWriter.cpp
+++ b/Engine/Source/Runtime/CoreUObject/Private/Serialization/ObjectWriter.cpp
@@ -49,6 +49,12 @@ FArchive& FObjectWriter::operator<<(FStringAssetReference& Value)
 	return *this;
 }
 
+FArchive& FObjectWriter::operator<< (struct FWeakObjectPtr& Value)
+{
+	Value.Serialize(*this);
+	return *this;
+}
+
 FString FObjectWriter::GetArchiveName() const
 {
 	return TEXT("FObjectWriter");
diff --git a/Engine/Source/Runtime/CoreUObject/Private/UObject/Class.cpp b/Engine/Source/Runtime/CoreUObject/Private/UObject/Class.cpp
--- a/Engine/Source/Runtime/CoreUObject/Private/UObject/Class.cpp
+++ b/Engine/Source/Runtime/CoreUObject/Private/UObject/Class.cpp
@@ -26,6 +26,10 @@ DEFINE_LOG_CATEGORY(LogClass);
 	#endif
 #endif
 
+#if !defined(USE_EVENT_DRIVEN_ASYNC_LOAD)
+#error "USE_EVENT_DRIVEN_ASYNC_LOAD must be defined"
+#endif
+
 //////////////////////////////////////////////////////////////////////////
 
 FThreadSafeBool& InternalSafeGetTokenStreamDirtyFlag()
@@ -498,6 +502,18 @@ void UStruct::StaticLink(bool bRelinkExistingProperties)
 	Link(ArDummy, bRelinkExistingProperties);
 }
 
+void UStruct::GetPreloadDependencies(TArray<UObject*>& OutDeps)
+{
+	Super::GetPreloadDependencies(OutDeps);
+	UStruct* InheritanceSuper = GetInheritanceSuper();
+	OutDeps.Add(InheritanceSuper);
+
+	for (UField* Field = Children; Field; Field = Field->Next)
+	{
+		OutDeps.Add(Field);
+	}
+}
+
 void UStruct::Link(FArchive& Ar, bool bRelinkExistingProperties)
 {
 	if (bRelinkExistingProperties)
@@ -672,11 +688,12 @@ void UStruct::Link(FArchive& Ar, bool bRelinkExistingProperties)
 	UProperty** RefLinkPtr = (UProperty**)&RefLink;
 	UProperty** PostConstructLinkPtr = &PostConstructLink;
 
+	TArray<const UStructProperty*> EncounteredStructProps;
 	for (TFieldIterator<UProperty> It(this); It; ++It)
 	{
 		UProperty* Property = *It;
 
-		if (Property->ContainsObjectReference() || Property->ContainsWeakObjectReference())
+		if (Property->ContainsObjectReference(EncounteredStructProps) || Property->ContainsWeakObjectReference())
 		{
 			*RefLinkPtr = Property;
 			RefLinkPtr = &(*RefLinkPtr)->NextRef;
@@ -2443,6 +2460,9 @@ UObject* UClass::CreateDefaultObject()
 		{
 			UObjectForceRegistration(ParentClass);
 			ParentDefaultObject = ParentClass->GetDefaultObject(); // Force the default object to be constructed if it isn't already
+#if USE_EVENT_DRIVEN_ASYNC_LOAD
+			check(ParentDefaultObject && !ParentDefaultObject->HasAnyFlags(RF_NeedLoad));
+#endif
 		}
 
 		if ( (ParentDefaultObject != NULL) || (this == UObject::StaticClass()) )
@@ -2864,7 +2884,7 @@ void UClass::Link(FArchive& Ar, bool bRelinkExistingProperties)
 	}
 }
 
-#if UCLASS_FAST_ISA_IMPL & 2
+#if UCLASS_FAST_ISA_IMPL == UCLASS_ISA_INDEXTREE
 
 	struct FClassParentPair
 	{
@@ -3084,25 +3104,63 @@ void UClass::Link(FArchive& Ar, bool bRelinkExistingProperties)
 		FFastIndexingClassTree::Unregister((UClass*)this);
 	}
 
+#elif UCLASS_FAST_ISA_IMPL == UCLASS_ISA_CLASSARRAY
+
+	FClassBaseChain::FClassBaseChain()
+		: ClassBaseChainArray(nullptr)
+		, NumClassBasesInChainMinusOne(-1)
+	{
+	}
+
+	FClassBaseChain::~FClassBaseChain()
+	{
+		delete [] ClassBaseChainArray;
+	}
+
+	void FClassBaseChain::ReinitializeBaseChainArray()
+	{
+		delete [] ClassBaseChainArray;
+
+		int32 Depth = 0;
+		for (UClass* Ptr = static_cast<UClass*>(this); Ptr; Ptr = Ptr->GetSuperClass())
+		{
+			++Depth;
+		}
+
+		FClassBaseChain** Bases = new FClassBaseChain*[Depth];
+		{
+			FClassBaseChain** Base = Bases + Depth;
+			for (UClass* Ptr = static_cast<UClass*>(this); Ptr; Ptr = Ptr->GetSuperClass())
+			{
+				*--Base = Ptr;
+			}
+		}
+
+		ClassBaseChainArray = Bases;
+		NumClassBasesInChainMinusOne = Depth - 1;
+	}
+
 #endif
 
 void UClass::SetSuperStruct(UStruct* NewSuperStruct)
 {
 	UnhashObject(this);
-#if UCLASS_FAST_ISA_IMPL & 2
+#if UCLASS_FAST_ISA_IMPL == UCLASS_ISA_INDEXTREE
 	FFastIndexingClassTree::Unregister(this);
 #endif
 	ClearFunctionMapsCaches();
 	Super::SetSuperStruct(NewSuperStruct);
-#if UCLASS_FAST_ISA_IMPL & 2
+#if UCLASS_FAST_ISA_IMPL == UCLASS_ISA_INDEXTREE
 	FFastIndexingClassTree::Register(this);
+#elif UCLASS_FAST_ISA_IMPL == UCLASS_ISA_CLASSARRAY
+	this->ReinitializeBaseChainArray();
 #endif
 	HashObject(this);
 }
 
 void UClass::SerializeSuperStruct(FArchive& Ar)
 {
-#if UCLASS_FAST_ISA_IMPL & 2
+#if UCLASS_FAST_ISA_IMPL == UCLASS_ISA_INDEXTREE
 	bool bIsLoading = Ar.IsLoading();
 	if (bIsLoading)
 	{
@@ -3110,11 +3168,13 @@ void UClass::SerializeSuperStruct(FArchive& Ar)
 	}
 #endif
 	Super::SerializeSuperStruct(Ar);
-#if UCLASS_FAST_ISA_IMPL & 2
+#if UCLASS_FAST_ISA_IMPL == UCLASS_ISA_INDEXTREE
 	if (bIsLoading)
 	{
 		FFastIndexingClassTree::Register(this);
 	}
+#elif UCLASS_FAST_ISA_IMPL == UCLASS_ISA_CLASSARRAY
+	this->ReinitializeBaseChainArray();
 #endif
 }
 
@@ -3126,11 +3186,11 @@ void UClass::Serialize( FArchive& Ar )
 		UnhashObject(this);
 	}
 
-#if UCLASS_FAST_ISA_IMPL & 2
+#if UCLASS_FAST_ISA_IMPL == UCLASS_ISA_INDEXTREE || UCLASS_FAST_ISA_IMPL == UCLASS_ISA_CLASSARRAY
 	UClass* SuperClassBefore = GetSuperClass();
 #endif
 	Super::Serialize( Ar );
-#if UCLASS_FAST_ISA_IMPL & 2
+#if UCLASS_FAST_ISA_IMPL == UCLASS_ISA_INDEXTREE || UCLASS_FAST_ISA_IMPL == UCLASS_ISA_CLASSARRAY
 	// Handle that fact that FArchive takes UObject*s by reference, and archives can just blat
 	// over our SuperStruct with impunity.
 	if (SuperClassBefore)
@@ -3138,8 +3198,12 @@ void UClass::Serialize( FArchive& Ar )
 		UClass* SuperClassAfter = GetSuperClass();
 		if (SuperClassBefore != SuperClassAfter)
 		{
-			FFastIndexingClassTree::Unregister(this);
-			FFastIndexingClassTree::Register(this);
+			#if UCLASS_FAST_ISA_IMPL == UCLASS_ISA_INDEXTREE
+				FFastIndexingClassTree::Unregister(this);
+				FFastIndexingClassTree::Register(this);
+			#elif UCLASS_FAST_ISA_IMPL == UCLASS_ISA_CLASSARRAY
+				this->ReinitializeBaseChainArray();
+			#endif
 		}
 	}
 #endif
@@ -3364,10 +3428,17 @@ void UClass::Serialize( FArchive& Ar )
 	{
 		if (ClassDefaultObject == NULL)
 		{
-			UE_LOG(LogClass, Error, TEXT("CDO for class %s did not load!"), *GetPathName() );
+
+#if USE_EVENT_DRIVEN_ASYNC_LOAD
+			ClassDefaultObject = GetDefaultObject();
+			// we do this later anyway, once we find it and set it in the export table. 
+			// ClassDefaultObject->SetFlags(RF_NeedLoad | RF_NeedPostLoad | RF_NeedPostLoadSubobjects);
+#else
+			UE_LOG(LogClass, Error, TEXT("CDO for class %s did not load!"), *GetPathName());
 			ensure(ClassDefaultObject != NULL);
 			ClassDefaultObject = GetDefaultObject();
 			Ar.ForceBlueprintFinalization();
+#endif
 		}
 	}
 }
diff --git a/Engine/Source/Runtime/CoreUObject/Private/UObject/CoreNative.cpp b/Engine/Source/Runtime/CoreUObject/Private/UObject/CoreNative.cpp
--- a/Engine/Source/Runtime/CoreUObject/Private/UObject/CoreNative.cpp
+++ b/Engine/Source/Runtime/CoreUObject/Private/UObject/CoreNative.cpp
@@ -268,9 +268,15 @@ UObject* FObjectInstancingGraph::InstancePropertyValue( class UObject* Component
 	return NewValue;
 }
 
+#if !defined(USE_EVENT_DRIVEN_ASYNC_LOAD)
+#error "USE_EVENT_DRIVEN_ASYNC_LOAD must be defined"
+#endif
 
 void FObjectInstancingGraph::AddNewObject(class UObject* ObjectInstance, UObject* InArchetype /*= nullptr*/)
 {
+#if USE_EVENT_DRIVEN_ASYNC_LOAD
+	check(!InArchetype || !InArchetype->HasAnyFlags(RF_NeedLoad));
+#endif
 	if (HasDestinationRoot())
 	{
 		AddNewInstance(ObjectInstance, InArchetype);
diff --git a/Engine/Source/Runtime/CoreUObject/Private/UObject/CoreNet.cpp b/Engine/Source/Runtime/CoreUObject/Private/UObject/CoreNet.cpp
--- a/Engine/Source/Runtime/CoreUObject/Private/UObject/CoreNet.cpp
+++ b/Engine/Source/Runtime/CoreUObject/Private/UObject/CoreNet.cpp
@@ -363,6 +363,11 @@ FArchive& FNetBitWriter::operator<<(FStringAssetReference& Value)
 	return *this;
 }
 
+FArchive& FNetBitWriter::operator<<(struct FWeakObjectPtr& WeakObjectPtr)
+{
+	WeakObjectPtr.Serialize(*this);
+	return *this;
+}
 
 // ----------------------------------------------------------------
 //	FNetBitReader
@@ -399,6 +404,12 @@ FArchive& FNetBitReader::operator<<(FStringAssetReference& Value)
 	return *this;
 }
 
+FArchive& FNetBitReader::operator<<(struct FWeakObjectPtr& WeakObjectPtr)
+{
+	WeakObjectPtr.Serialize(*this);
+	return *this;
+}
+
 static const TCHAR* GLastRPCFailedReason = NULL;
 
 void RPC_ResetLastFailedReason()
diff --git a/Engine/Source/Runtime/CoreUObject/Private/UObject/GarbageCollection.cpp b/Engine/Source/Runtime/CoreUObject/Private/UObject/GarbageCollection.cpp
--- a/Engine/Source/Runtime/CoreUObject/Private/UObject/GarbageCollection.cpp
+++ b/Engine/Source/Runtime/CoreUObject/Private/UObject/GarbageCollection.cpp
@@ -32,7 +32,9 @@ int32		GObjectCountDuringLastMarkPhase			= 0;
 /** Count of objects purged since last mark phase																	*/
 int32		GPurgedObjectCountSinceLastMarkPhase	= 0;
 /** Whether incremental object purge is in progress										*/
-static bool GObjIncrementalPurgeIsInProgress = false;
+bool GObjIncrementalPurgeIsInProgress = false;
+/** Whether GC is currently routing BeginDestroy to objects										*/
+bool GObjUnhashUnreachableIsInProgress = false;
 /** Whether FinishDestroy has already been routed to all unreachable objects. */
 static bool GObjFinishDestroyHasBeenRoutedToAllObjects	= false;
 /** 
@@ -518,25 +520,32 @@ public:
 	FORCEINLINE void HandleTokenStreamObjectReference(TArray<UObject*>& ObjectsToSerialize, UObject* ReferencingObject, UObject*& Object, const int32 TokenIndex, bool bAllowReferenceElimination)
 	{
 #if !(UE_BUILD_TEST || UE_BUILD_SHIPPING)
-		if (Object && !Object->IsValidLowLevelFast())
+		if (Object)
 		{
-			FString TokenDebugInfo;
-			if (UClass *Class = (ReferencingObject ? ReferencingObject->GetClass() : nullptr))
-			{
-				auto& TokenInfo = Class->DebugTokenMap.GetTokenInfo(TokenIndex);
-				TokenDebugInfo = FString::Printf(TEXT("ReferencingObjectClass: %s, Property Name: %s, Offset: %d"),
-					*Class->GetFullName(), *TokenInfo.Name.GetPlainNameString(), TokenInfo.Offset);
-			}
-			else
+			if (
+#if DO_POINTER_CHECKS_ON_GC
+				!IsPossiblyAllocatedUObjectPointer(Object) ||
+#endif
+				!Object->IsValidLowLevelFast())
 			{
-				// This means this objects is most likely being referenced by AddReferencedObjects
-				TokenDebugInfo = TEXT("Native Reference");
-			}
+				FString TokenDebugInfo;
+				if (UClass *Class = (ReferencingObject ? ReferencingObject->GetClass() : nullptr))
+				{
+					auto& TokenInfo = Class->DebugTokenMap.GetTokenInfo(TokenIndex);
+					TokenDebugInfo = FString::Printf(TEXT("ReferencingObjectClass: %s, Property Name: %s, Offset: %d"),
+						*Class->GetFullName(), *TokenInfo.Name.GetPlainNameString(), TokenInfo.Offset);
+				}
+				else
+				{
+					// This means this objects is most likely being referenced by AddReferencedObjects
+					TokenDebugInfo = TEXT("Native Reference");
+				}
 
-			UE_LOG(LogGarbage, Fatal, TEXT("Invalid object in GC: 0x%016llx, ReferencingObject: %s, %s, TokenIndex: %d"),
-				(int64)(PTRINT)Object,
-				ReferencingObject ? *ReferencingObject->GetFullName() : TEXT("NULL"),
-				*TokenDebugInfo, TokenIndex);
+				UE_LOG(LogGarbage, Fatal, TEXT("Invalid object in GC: 0x%016llx, ReferencingObject: %s, %s, TokenIndex: %d"),
+					(int64)(PTRINT)Object,
+					ReferencingObject ? *ReferencingObject->GetFullName() : TEXT("NULL"),
+					*TokenDebugInfo, TokenIndex);
+			}
 		}
 #endif
 		HandleObjectReference(ObjectsToSerialize, ReferencingObject, Object, bAllowReferenceElimination);
@@ -1130,7 +1139,9 @@ void CollectGarbageInternal(EObjectFlags KeepFlags, bool bPerformFullPurge)
 	// Flush streaming before GC if requested
 	if (GFlushStreamingOnGC)
 	{
+		GGarbageCollectionGuardCritical.GCUnlock();
 		FlushAsyncLoading();
+		GGarbageCollectionGuardCritical.GCLock();
 	}
 
 	// Route callbacks so we can ensure that we are e.g. not in the middle of loading something by flushing
@@ -1138,172 +1149,178 @@ void CollectGarbageInternal(EObjectFlags KeepFlags, bool bPerformFullPurge)
 	FCoreUObjectDelegates::PreGarbageCollect.Broadcast();
 	GLastGCFrame = GFrameCounter;
 
-	// Set 'I'm garbage collecting' flag - might be checked inside various functions.
-	FGCScopeLock GCLock;
+	{
+		// Set 'I'm garbage collecting' flag - might be checked inside various functions.
+		// This has to be unlocked before we call post GC callbacks
+		FGCScopeLock GCLock;
 
-	UE_LOG(LogGarbage, Log, TEXT("Collecting garbage") );
+		UE_LOG(LogGarbage, Log, TEXT("Collecting garbage"));
 
-	// Make sure previous incremental purge has finished or we do a full purge pass in case we haven't kicked one
-	// off yet since the last call to garbage collection.
-	if( GObjIncrementalPurgeIsInProgress || GObjPurgeIsRequired )
-	{
-		IncrementalPurgeGarbage( false );
-		FMemory::Trim();
-	}
-	check( !GObjIncrementalPurgeIsInProgress );
-	check( !GObjPurgeIsRequired );
+		// Make sure previous incremental purge has finished or we do a full purge pass in case we haven't kicked one
+		// off yet since the last call to garbage collection.
+		if (GObjIncrementalPurgeIsInProgress || GObjPurgeIsRequired)
+		{
+			IncrementalPurgeGarbage(false);
+			FMemory::Trim();
+		}
+		check(!GObjIncrementalPurgeIsInProgress);
+		check(!GObjPurgeIsRequired);
 
 #if VERIFY_DISREGARD_GC_ASSUMPTIONS
-	FUObjectArray& UObjectArray = GUObjectArray;
-	// Only verify assumptions if option is enabled. This avoids false positives in the Editor or commandlets.
-	if ((UObjectArray.DisregardForGCEnabled() || GUObjectClusters.Num()) && GShouldVerifyGCAssumptions)
-	{
-		DECLARE_SCOPE_CYCLE_COUNTER( TEXT( "CollectGarbageInternal.VerifyGCAssumptions" ), STAT_CollectGarbageInternal_VerifyGCAssumptions, STATGROUP_GC );
-		bool bShouldAssert = false;
-
-		// Verify that objects marked to be disregarded for GC are not referencing objects that are not part of the root set.
-		for (FRawObjectIterator It(false); It; ++It)
+		FUObjectArray& UObjectArray = GUObjectArray;
+		// Only verify assumptions if option is enabled. This avoids false positives in the Editor or commandlets.
+		if ((UObjectArray.DisregardForGCEnabled() || GUObjectClusters.Num()) && GShouldVerifyGCAssumptions)
 		{
-			FUObjectItem* ObjectItem = *It;
-			UObject* Object = (UObject*)ObjectItem->Object;
-			// Don't require UGCObjectReferencer's references to adhere to the assumptions.
-			// Although we want the referencer itself to sit in the disregard for gc set, most of the objects
-			// it's referencing will not be in the root set.
-			if (UObjectArray.IsDisregardForGC(Object) && !Object->IsA(UGCObjectReferencer::StaticClass()))
+			DECLARE_SCOPE_CYCLE_COUNTER(TEXT("CollectGarbageInternal.VerifyGCAssumptions"), STAT_CollectGarbageInternal_VerifyGCAssumptions, STATGROUP_GC);
+			bool bShouldAssert = false;
+
+			// Verify that objects marked to be disregarded for GC are not referencing objects that are not part of the root set.
+			for (FRawObjectIterator It(false); It; ++It)
 			{
-				// Serialize object with reference collector.
-				TArray<UObject*> CollectedReferences;
-				FReferenceFinder ObjectReferenceCollector( CollectedReferences );
-				ObjectReferenceCollector.FindReferences( Object );
+				FUObjectItem* ObjectItem = *It;
+				UObject* Object = (UObject*)ObjectItem->Object;
+				// Don't require UGCObjectReferencer's references to adhere to the assumptions.
+				// Although we want the referencer itself to sit in the disregard for gc set, most of the objects
+				// it's referencing will not be in the root set.
+				if (UObjectArray.IsDisregardForGC(Object) && !Object->IsA(UGCObjectReferencer::StaticClass()))
+				{
+					// Serialize object with reference collector.
+					TArray<UObject*> CollectedReferences;
+					FReferenceFinder ObjectReferenceCollector(CollectedReferences);
+					ObjectReferenceCollector.FindReferences(Object);
 
-				// Iterate over referenced objects, finding bad ones.
-				for( int32 ReferenceIndex=0; ReferenceIndex<CollectedReferences.Num(); ReferenceIndex++ )
+					// Iterate over referenced objects, finding bad ones.
+					for (int32 ReferenceIndex = 0; ReferenceIndex < CollectedReferences.Num(); ReferenceIndex++)
+					{
+						UObject* ReferencedObject = CollectedReferences[ReferenceIndex];
+						if (ReferencedObject &&
+							!(ReferencedObject->IsRooted() ||
+								UObjectArray.IsDisregardForGC(ReferencedObject) ||
+								UObjectArray.ObjectToObjectItem(ReferencedObject)->GetOwnerIndex() ||
+								UObjectArray.ObjectToObjectItem(ReferencedObject)->HasAnyFlags(EInternalObjectFlags::ClusterRoot)))
+						{
+							UE_LOG(LogGarbage, Warning, TEXT("Disregard for GC object %s referencing %s which is not part of root set"),
+								*Object->GetFullName(),
+								*ReferencedObject->GetFullName());
+							bShouldAssert = true;
+						}
+					}
+				}
+				else if (ObjectItem->HasAnyFlags(EInternalObjectFlags::ClusterRoot))
 				{
-					UObject* ReferencedObject = CollectedReferences[ReferenceIndex];
-					if (ReferencedObject && 
-						!(ReferencedObject->IsRooted() || 
-						  UObjectArray.IsDisregardForGC(ReferencedObject) || 
-							UObjectArray.ObjectToObjectItem(ReferencedObject)->GetOwnerIndex() ||
-							UObjectArray.ObjectToObjectItem(ReferencedObject)->HasAnyFlags(EInternalObjectFlags::ClusterRoot)))
+					if (!VerifyClusterAssumptions(Object))
 					{
-						UE_LOG(LogGarbage, Warning, TEXT("Disregard for GC object %s referencing %s which is not part of root set"),
-							*Object->GetFullName(),
-							*ReferencedObject->GetFullName());						
 						bShouldAssert = true;
 					}
 				}
 			}
-			else if (ObjectItem->HasAnyFlags(EInternalObjectFlags::ClusterRoot))
+			// Assert if we encountered any objects breaking implicit assumptions.
+			if (bShouldAssert)
 			{
-				if (!VerifyClusterAssumptions(Object))
-				{
-					bShouldAssert = true;
-				}
+				UE_LOG(LogGarbage, Fatal, TEXT("Encountered object(s) breaking Disregard for GC assumption. Please check log for details."));
 			}
 		}
-		// Assert if we encountered any objects breaking implicit assumptions.
-		if (bShouldAssert)
-		{
-			UE_LOG(LogGarbage, Fatal, TEXT("Encountered object(s) breaking Disregard for GC assumption. Please check log for details."));
-		}
-	}
 #endif
 
-	// Fall back to single threaded GC if processor count is 1 or parallel GC is disabled
-	// or detailed per class gc stats are enabled (not thread safe)
-	// Temporarily forcing single-threaded GC in the editor until Modify() can be safely removed from HandleObjectReference.
-	const bool bForceSingleThreadedGC = !FApp::ShouldUseThreadingForPerformance() || !FPlatformProcess::SupportsMultithreading() ||
+		// Fall back to single threaded GC if processor count is 1 or parallel GC is disabled
+		// or detailed per class gc stats are enabled (not thread safe)
+		// Temporarily forcing single-threaded GC in the editor until Modify() can be safely removed from HandleObjectReference.
+		const bool bForceSingleThreadedGC = !FApp::ShouldUseThreadingForPerformance() || !FPlatformProcess::SupportsMultithreading() ||
 #if PLATFORM_SUPPORTS_MULTITHREADED_GC
-		( FPlatformMisc::NumberOfCores() < 2 || GAllowParallelGC == 0 || PERF_DETAILED_PER_CLASS_GC_STATS );
+		(FPlatformMisc::NumberOfCores() < 2 || GAllowParallelGC == 0 || PERF_DETAILED_PER_CLASS_GC_STATS);
 #else	//PLATFORM_SUPPORTS_MULTITHREADED_GC
-		true;
+			true;
 #endif	//PLATFORM_SUPPORTS_MULTITHREADED_GC
 
-	// Perform reachability analysis.
-	{
-		const double StartTime = FPlatformTime::Seconds();
-		FRealtimeGC TagUsedRealtimeGC;
-		TagUsedRealtimeGC.PerformReachabilityAnalysis( KeepFlags, bForceSingleThreadedGC );
-		UE_LOG(LogGarbage, Log, TEXT("%f ms for GC"), (FPlatformTime::Seconds() - StartTime) * 1000 );
-	}
+		// Perform reachability analysis.
+		{
+			const double StartTime = FPlatformTime::Seconds();
+			FRealtimeGC TagUsedRealtimeGC;
+			TagUsedRealtimeGC.PerformReachabilityAnalysis(KeepFlags, bForceSingleThreadedGC);
+			UE_LOG(LogGarbage, Log, TEXT("%f ms for GC"), (FPlatformTime::Seconds() - StartTime) * 1000);
+		}
 
 #if WITH_EDITOR
-	if ( GIsEditor && EditorPostReachabilityAnalysisCallback )
-	{
-		EditorPostReachabilityAnalysisCallback();
-	}
+		if (GIsEditor && EditorPostReachabilityAnalysisCallback)
+		{
+			EditorPostReachabilityAnalysisCallback();
+		}
 #endif // WITH_EDITOR
 
-	{
-		DECLARE_SCOPE_CYCLE_COUNTER( TEXT( "CollectGarbageInternal.UnhashUnreachable" ), STAT_CollectGarbageInternal_UnhashUnreachable, STATGROUP_GC );
-
-		// Unhash all unreachable objects.
-		const double StartTime = FPlatformTime::Seconds();
-		int32 ClustersRemoved = 0;
-		for ( FRawObjectIterator It(true); It; ++It )
 		{
-			//@todo UE4 - A prefetch was removed here. Re-add it. It wasn't right anyway, since it was ten items ahead and the consoles on have 8 prefetch slots
+			DECLARE_SCOPE_CYCLE_COUNTER(TEXT("CollectGarbageInternal.UnhashUnreachable"), STAT_CollectGarbageInternal_UnhashUnreachable, STATGROUP_GC);
 
-			FUObjectItem* ObjectItem = *It;
-			checkSlow(ObjectItem);
-			if (ObjectItem->IsUnreachable())
+			TGuardValue<bool> GuardObjUnhashUnreachableIsInProgress(GObjUnhashUnreachableIsInProgress, true);
+
+			// Unhash all unreachable objects.
+			const double StartTime = FPlatformTime::Seconds();
+			int32 ClustersRemoved = 0;
+			for (FRawObjectIterator It(true); It; ++It)
 			{
-				if ((ObjectItem->GetFlags() & EInternalObjectFlags::ClusterRoot) == EInternalObjectFlags::ClusterRoot)
+				//@todo UE4 - A prefetch was removed here. Re-add it. It wasn't right anyway, since it was ten items ahead and the consoles on have 8 prefetch slots
+
+				FUObjectItem* ObjectItem = *It;
+				checkSlow(ObjectItem);
+				if (ObjectItem->IsUnreachable())
 				{
-					// Nuke the entire cluster
-					ObjectItem->ClearFlags(EInternalObjectFlags::ClusterRoot|EInternalObjectFlags::NoStrongReference);					
-					const int32 ClusterRootIndex = It.GetIndex();
-					FUObjectCluster* Cluster = GUObjectClusters.FindChecked(ClusterRootIndex);
-					checkSlow(Cluster);
-					for (int32 ClusterObjectIndex : Cluster->Objects)
+					if ((ObjectItem->GetFlags() & EInternalObjectFlags::ClusterRoot) == EInternalObjectFlags::ClusterRoot)
 					{
-						FUObjectItem* ClusterObjectItem = GUObjectArray.IndexToObjectUnsafeForGC(ClusterObjectIndex);
-						ClusterObjectItem->ClearFlags(EInternalObjectFlags::NoStrongReference);
-						ClusterObjectItem->SetOwnerIndex(0);
-
-						if (!ClusterObjectItem->HasAnyFlags(EInternalObjectFlags::ReachableInCluster))
+						// Nuke the entire cluster
+						ObjectItem->ClearFlags(EInternalObjectFlags::ClusterRoot | EInternalObjectFlags::NoStrongReference);
+						const int32 ClusterRootIndex = It.GetIndex();
+						FUObjectCluster* Cluster = GUObjectClusters.FindChecked(ClusterRootIndex);
+						checkSlow(Cluster);
+						for (int32 ClusterObjectIndex : Cluster->Objects)
 						{
-							ClusterObjectItem->SetFlags(EInternalObjectFlags::Unreachable);
-							if (ClusterObjectIndex < ClusterRootIndex)
+							FUObjectItem* ClusterObjectItem = GUObjectArray.IndexToObjectUnsafeForGC(ClusterObjectIndex);
+							ClusterObjectItem->ClearFlags(EInternalObjectFlags::NoStrongReference);
+							ClusterObjectItem->SetOwnerIndex(0);
+
+							if (!ClusterObjectItem->HasAnyFlags(EInternalObjectFlags::ReachableInCluster))
 							{
-								UObject* ClusterObject = (UObject*)ClusterObjectItem->Object;
-								ClusterObject->ConditionalBeginDestroy();
+								ClusterObjectItem->SetFlags(EInternalObjectFlags::Unreachable);
+								if (ClusterObjectIndex < ClusterRootIndex)
+								{
+									UObject* ClusterObject = (UObject*)ClusterObjectItem->Object;
+									ClusterObject->ConditionalBeginDestroy();
+								}
 							}
 						}
+						delete Cluster;
+						GUObjectClusters.Remove(ClusterRootIndex);
+						ClustersRemoved++;
 					}
-					delete Cluster;
-					GUObjectClusters.Remove(ClusterRootIndex);
-					ClustersRemoved++;
-				}
 
-				// Begin the object's asynchronous destruction.
-				UObject* Object = (UObject*)ObjectItem->Object;
-				Object->ConditionalBeginDestroy();
-			}
-			else if (ObjectItem->IsNoStrongReference())
-			{
-				ObjectItem->ClearNoStrongReference();
-				ObjectItem->SetPendingKill();
+					// Begin the object's asynchronous destruction.
+					UObject* Object = (UObject*)ObjectItem->Object;
+					Object->ConditionalBeginDestroy();
+				}
+				else if (ObjectItem->IsNoStrongReference())
+				{
+					ObjectItem->ClearNoStrongReference();
+					ObjectItem->SetPendingKill();
+				}
 			}
+			UE_LOG(LogGarbage, Log, TEXT("%f ms for unhashing unreachable objects. Clusters removed: %d."), (FPlatformTime::Seconds() - StartTime) * 1000, ClustersRemoved);
 		}
-		UE_LOG(LogGarbage, Log, TEXT("%f ms for unhashing unreachable objects. Clusters removed: %d."), (FPlatformTime::Seconds() - StartTime) * 1000, ClustersRemoved);
-	}
 
-	// Set flag to indicate that we are relying on a purge to be performed.
-	GObjPurgeIsRequired = true;
-	// Reset purged count.
-	GPurgedObjectCountSinceLastMarkPhase = 0;
+		// Set flag to indicate that we are relying on a purge to be performed.
+		GObjPurgeIsRequired = true;
+		// Reset purged count.
+		GPurgedObjectCountSinceLastMarkPhase = 0;
 
-	// Perform a full purge by not using a time limit for the incremental purge. The Editor always does a full purge.
-	if( bPerformFullPurge || GIsEditor )
-	{
-		IncrementalPurgeGarbage( false );	
+		// Perform a full purge by not using a time limit for the incremental purge. The Editor always does a full purge.
+		if (bPerformFullPurge || GIsEditor)
+		{
+			IncrementalPurgeGarbage(false);
+		}
 	}
-	FMemory::Trim();
 
 	// Destroy all pending delete linkers
 	DeleteLoaders();
 
+	FMemory::Trim();
+
 	// Route callbacks to verify GC assumptions
 	FCoreUObjectDelegates::PostGarbageCollect.Broadcast();
 
@@ -1419,7 +1436,7 @@ void UObject::AddReferencedObjects(UObject* This, FReferenceCollector& Collector
  *
  * @return true if property (or sub- properties) contain a UObject reference, false otherwise
  */
-bool UProperty::ContainsObjectReference() const
+bool UProperty::ContainsObjectReference(TArray<const UStructProperty*>& EncounteredStructProps) const
 {
 	return false;
 }
@@ -1430,10 +1447,10 @@ bool UProperty::ContainsObjectReference() const
  *
  * @return true if property (or sub- properties) contain a UObject reference, false otherwise
  */
-bool UArrayProperty::ContainsObjectReference() const
+bool UArrayProperty::ContainsObjectReference(TArray<const UStructProperty*>& EncounteredStructProps) const
 {
 	check(Inner);
-	return Inner->ContainsObjectReference();
+	return Inner->ContainsObjectReference(EncounteredStructProps);
 }
 
 /**
@@ -1442,11 +1459,11 @@ bool UArrayProperty::ContainsObjectReference() const
  *
  * @return true if property (or sub- properties) contain a UObject reference, false otherwise
  */
-bool UMapProperty::ContainsObjectReference() const
+bool UMapProperty::ContainsObjectReference(TArray<const UStructProperty*>& EncounteredStructProps) const
 {
 	check(KeyProp);
 	check(ValueProp);
-	return KeyProp->ContainsObjectReference() || ValueProp->ContainsObjectReference();
+	return KeyProp->ContainsObjectReference(EncounteredStructProps) || ValueProp->ContainsObjectReference(EncounteredStructProps);
 }
 
 /**
@@ -1455,10 +1472,10 @@ bool UMapProperty::ContainsObjectReference() const
 *
 * @return true if property (or sub- properties) contain a UObject reference, false otherwise
 */
-bool USetProperty::ContainsObjectReference() const
+bool USetProperty::ContainsObjectReference(TArray<const UStructProperty*>& EncounteredStructProps) const
 {
 	check(ElementProp);
-	return ElementProp->ContainsObjectReference();
+	return ElementProp->ContainsObjectReference(EncounteredStructProps);
 }
 
 /**
@@ -1467,10 +1484,8 @@ bool USetProperty::ContainsObjectReference() const
  *
  * @return true if property (or sub- properties) contain a UObject reference, false otherwise
  */
-bool UStructProperty::ContainsObjectReference() const
+bool UStructProperty::ContainsObjectReference(TArray<const UStructProperty*>& EncounteredStructProps) const
 {
-	// prevent recursion in the case of structs containing dynamic arrays of themselves
-	static TArray<const UStructProperty*> EncounteredStructProps;
 	if (EncounteredStructProps.Contains(this))
 	{
 		return false;
@@ -1487,7 +1502,7 @@ bool UStructProperty::ContainsObjectReference() const
 			UProperty* Property = Struct->PropertyLink;
 			while( Property )
 			{
-				if( Property->ContainsObjectReference() )
+				if (Property->ContainsObjectReference(EncounteredStructProps))
 				{
 					EncounteredStructProps.RemoveSingleSwap(this);
 					return true;
@@ -1621,7 +1636,7 @@ private:
  * Emits tokens used by realtime garbage collection code to passed in ReferenceTokenStream. The offset emitted is relative
  * to the passed in BaseOffset which is used by e.g. arrays of structs.
  */
-void UProperty::EmitReferenceInfo(UClass& OwnerClass, int32 BaseOffset)
+void UProperty::EmitReferenceInfo(UClass& OwnerClass, int32 BaseOffset, TArray<const UStructProperty*>& EncounteredStructProps)
 {
 }
 
@@ -1629,7 +1644,7 @@ void UProperty::EmitReferenceInfo(UClass& OwnerClass, int32 BaseOffset)
  * Emits tokens used by realtime garbage collection code to passed in OwnerClass' ReferenceTokenStream. The offset emitted is relative
  * to the passed in BaseOffset which is used by e.g. arrays of structs.
  */
-void UObjectProperty::EmitReferenceInfo(UClass& OwnerClass, int32 BaseOffset)
+void UObjectProperty::EmitReferenceInfo(UClass& OwnerClass, int32 BaseOffset, TArray<const UStructProperty*>& EncounteredStructProps)
 {
 	FGCReferenceFixedArrayTokenHelper FixedArrayHelper(OwnerClass, BaseOffset + GetOffset_ForGC(), ArrayDim, sizeof(UObject*), *this);
 	OwnerClass.EmitObjectReference(BaseOffset + GetOffset_ForGC(), GetFName(), GCRT_Object);
@@ -1639,9 +1654,9 @@ void UObjectProperty::EmitReferenceInfo(UClass& OwnerClass, int32 BaseOffset)
  * Emits tokens used by realtime garbage collection code to passed in OwnerClass' ReferenceTokenStream. The offset emitted is relative
  * to the passed in BaseOffset which is used by e.g. arrays of structs.
  */
-void UArrayProperty::EmitReferenceInfo(UClass& OwnerClass, int32 BaseOffset)
+void UArrayProperty::EmitReferenceInfo(UClass& OwnerClass, int32 BaseOffset, TArray<const UStructProperty*>& EncounteredStructProps)
 {
-	if( Inner->ContainsObjectReference() )
+	if (Inner->ContainsObjectReference(EncounteredStructProps))
 	{
 		if( Inner->IsA(UStructProperty::StaticClass()) )
 		{
@@ -1649,7 +1664,7 @@ void UArrayProperty::EmitReferenceInfo(UClass& OwnerClass, int32 BaseOffset)
 
 			OwnerClass.ReferenceTokenStream.EmitStride(Inner->ElementSize);
 			const uint32 SkipIndexIndex = OwnerClass.ReferenceTokenStream.EmitSkipIndexPlaceholder();
-			Inner->EmitReferenceInfo(OwnerClass, 0);
+			Inner->EmitReferenceInfo(OwnerClass, 0, EncounteredStructProps);
 			const uint32 SkipIndex = OwnerClass.ReferenceTokenStream.EmitReturn();
 			OwnerClass.ReferenceTokenStream.UpdateSkipIndexPlaceholder(SkipIndexIndex, SkipIndex);
 		}
@@ -1681,9 +1696,9 @@ void UArrayProperty::EmitReferenceInfo(UClass& OwnerClass, int32 BaseOffset)
  * Emits tokens used by realtime garbage collection code to passed in OwnerClass' ReferenceTokenStream. The offset emitted is relative
  * to the passed in BaseOffset which is used by e.g. arrays of structs.
  */
-void UMapProperty::EmitReferenceInfo(UClass& OwnerClass, int32 BaseOffset)
+void UMapProperty::EmitReferenceInfo(UClass& OwnerClass, int32 BaseOffset, TArray<const UStructProperty*>& EncounteredStructProps)
 {
-	if (ContainsObjectReference())
+	if (ContainsObjectReference(EncounteredStructProps))
 	{
 		OwnerClass.EmitObjectReference(BaseOffset + GetOffset_ForGC(), GetFName(), GCRT_AddTMapReferencedObjects);
 		OwnerClass.ReferenceTokenStream.EmitPointer((const void*)this);
@@ -1694,9 +1709,9 @@ void UMapProperty::EmitReferenceInfo(UClass& OwnerClass, int32 BaseOffset)
 * Emits tokens used by realtime garbage collection code to passed in OwnerClass' ReferenceTokenStream. The offset emitted is relative
 * to the passed in BaseOffset which is used by e.g. arrays of structs.
 */
-void USetProperty::EmitReferenceInfo(UClass& OwnerClass, int32 BaseOffset)
+void USetProperty::EmitReferenceInfo(UClass& OwnerClass, int32 BaseOffset, TArray<const UStructProperty*>& EncounteredStructProps)
 {
-	if (ContainsObjectReference())
+	if (ContainsObjectReference(EncounteredStructProps))
 	{
 		OwnerClass.EmitObjectReference(BaseOffset + GetOffset_ForGC(), GetFName(), GCRT_AddTSetReferencedObjects);
 		OwnerClass.ReferenceTokenStream.EmitPointer((const void*)this);
@@ -1708,7 +1723,7 @@ void USetProperty::EmitReferenceInfo(UClass& OwnerClass, int32 BaseOffset)
  * Emits tokens used by realtime garbage collection code to passed in ReferenceTokenStream. The offset emitted is relative
  * to the passed in BaseOffset which is used by e.g. arrays of structs.
  */
-void UStructProperty::EmitReferenceInfo(UClass& OwnerClass, int32 BaseOffset)
+void UStructProperty::EmitReferenceInfo(UClass& OwnerClass, int32 BaseOffset, TArray<const UStructProperty*>& EncounteredStructProps)
 {
 	if (Struct->StructFlags & STRUCT_AddStructReferencedObjects)
 	{
@@ -1723,14 +1738,14 @@ void UStructProperty::EmitReferenceInfo(UClass& OwnerClass, int32 BaseOffset)
 		return;
 	}
 	check(Struct);
-	if( ContainsObjectReference() )
+	if (ContainsObjectReference(EncounteredStructProps))
 	{
 		FGCReferenceFixedArrayTokenHelper FixedArrayHelper(OwnerClass, BaseOffset + GetOffset_ForGC(), ArrayDim, ElementSize, *this);
 
 		UProperty* Property = Struct->PropertyLink;
 		while( Property )
 		{
-			Property->EmitReferenceInfo(OwnerClass, BaseOffset + GetOffset_ForGC());
+			Property->EmitReferenceInfo(OwnerClass, BaseOffset + GetOffset_ForGC(), EncounteredStructProps);
 			Property = Property->PropertyLinkNext;
 		}
 	}
@@ -1740,7 +1755,7 @@ void UStructProperty::EmitReferenceInfo(UClass& OwnerClass, int32 BaseOffset)
  * Emits tokens used by realtime garbage collection code to passed in ReferenceTokenStream. The offset emitted is relative
  * to the passed in BaseOffset which is used by e.g. arrays of structs.
  */
-void UInterfaceProperty::EmitReferenceInfo(UClass& OwnerClass, int32 BaseOffset)
+void UInterfaceProperty::EmitReferenceInfo(UClass& OwnerClass, int32 BaseOffset, TArray<const UStructProperty*>& EncounteredStructProps)
 {
 	FGCReferenceFixedArrayTokenHelper FixedArrayHelper(OwnerClass, BaseOffset + GetOffset_ForGC(), ArrayDim, sizeof(FScriptInterface), *this);
 
@@ -1815,12 +1830,13 @@ void UClass::AssembleReferenceTokenStream(bool bForce)
 #endif
 			ClassFlags &= ~CLASS_TokenStreamAssembled;
 		}
+		TArray<const UStructProperty*> EncounteredStructProps;
 
 		// Iterate over properties defined in this class
 		for( TFieldIterator<UProperty> It(this,EFieldIteratorFlags::ExcludeSuper); It; ++It)
 		{
 			UProperty* Property = *It;
-			Property->EmitReferenceInfo(*this, 0);
+			Property->EmitReferenceInfo(*this, 0, EncounteredStructProps);
 		}
 
 		if (GetSuperClass())
diff --git a/Engine/Source/Runtime/CoreUObject/Private/UObject/LinkerLoad.cpp b/Engine/Source/Runtime/CoreUObject/Private/UObject/LinkerLoad.cpp
--- a/Engine/Source/Runtime/CoreUObject/Private/UObject/LinkerLoad.cpp
+++ b/Engine/Source/Runtime/CoreUObject/Private/UObject/LinkerLoad.cpp
@@ -43,8 +43,10 @@ static FAutoConsoleVariableRef CVarLinkerAllowDynamicClasses(
 	);
 #endif
 
+#if !USE_NEW_ASYNC_IO
 /** Map that keeps track of any precached full package reads															*/
 TMap<FString, FLinkerLoad::FPackagePrecacheInfo> FLinkerLoad::PackagePrecacheMap;
+#endif
 
 UClass* FLinkerLoad::UTexture2DStaticClass = NULL;
 
@@ -230,44 +232,26 @@ void FLinkerLoad::CreateActiveRedirectsMap(const FString& GEngineIniName)
 	}
 }
 
-/** Helper struct to keep track of the first time CreateImport() is called in the current callstack. */
-struct FScopedCreateImportCounter
-{
-	/**
-	 *	Constructor. Called upon CreateImport() entry.
-	 *	@param Linker	- Current Linker
-	 *	@param Index	- Index of the current Import
-	 */
-	FScopedCreateImportCounter( FLinkerLoad* Linker, int32 Index )
-	{
-		// First time CreateImport() is called for this callstack?
-		if ( Counter++ == 0 )
-		{
-			FUObjectThreadContext& ThreadContext = FUObjectThreadContext::Get();
-			// Remember the current linker and index.
-			ThreadContext.SerializedImportLinker = Linker;
-			ThreadContext.SerializedImportIndex = Index;
-		}
-	}
 
-	/** Destructor. Called upon CreateImport() exit. */
-	~FScopedCreateImportCounter()
-	{
-		// Last time CreateImport() exits for this callstack?
-		if ( --Counter == 0 )
-		{
-			FUObjectThreadContext& ThreadContext = FUObjectThreadContext::Get();
-			ThreadContext.SerializedImportLinker = nullptr;
-			ThreadContext.SerializedImportIndex = INDEX_NONE;
-		}
-	}
+FScopedCreateImportCounter::FScopedCreateImportCounter(FLinkerLoad* Linker, int32 Index)
+{
+	FUObjectThreadContext& ThreadContext = FUObjectThreadContext::Get();
+	// Remember the old linker and index
+	PreviousLinker = ThreadContext.SerializedImportLinker;
+	PreviousIndex = ThreadContext.SerializedImportIndex;
+	// Remember the current linker and index.
+	ThreadContext.SerializedImportLinker = Linker;
+	ThreadContext.SerializedImportIndex = Index;
+}
 
-	/** Number of times CreateImport() has been called in the current callstack. */
-	static int32	Counter;
-};
+FScopedCreateImportCounter::~FScopedCreateImportCounter()
+{
+	FUObjectThreadContext& ThreadContext = FUObjectThreadContext::Get();
+	// Restore old values
+	ThreadContext.SerializedImportLinker = PreviousLinker;
+	ThreadContext.SerializedImportIndex = PreviousIndex;
+}
 
-/** Number of times CreateImport() has been called in the current callstack. */
-int32 FScopedCreateImportCounter::Counter = 0;
 
 /** Helper struct to keep track of the CreateExport() entry/exit. */
 struct FScopedCreateExportCounter
@@ -277,9 +261,13 @@ struct FScopedCreateExportCounter
 	 *	@param Linker	- Current Linker
 	 *	@param Index	- Index of the current Import
 	 */
-	FScopedCreateExportCounter( FLinkerLoad* Linker, int32 Index )
+	FScopedCreateExportCounter(FLinkerLoad* Linker, int32 Index)
 	{
 		FUObjectThreadContext& ThreadContext = FUObjectThreadContext::Get();
+		// Remember the old linker and index
+		PreviousLinker = ThreadContext.SerializedExportLinker;
+		PreviousIndex = ThreadContext.SerializedExportIndex;
+		// Remember the current linker and index.
 		ThreadContext.SerializedExportLinker = Linker;
 		ThreadContext.SerializedExportIndex = Index;
 	}
@@ -288,9 +276,15 @@ struct FScopedCreateExportCounter
 	~FScopedCreateExportCounter()
 	{
 		FUObjectThreadContext& ThreadContext = FUObjectThreadContext::Get();
-		ThreadContext.SerializedExportLinker = nullptr;
-		ThreadContext.SerializedExportIndex = INDEX_NONE;
+		// Restore old values
+		ThreadContext.SerializedExportLinker = PreviousLinker;
+		ThreadContext.SerializedExportIndex = PreviousIndex;
 	}
+
+	/** Poreviously stored linker */
+	FLinkerLoad* PreviousLinker;
+	/** Previously stored index */
+	int32 PreviousIndex;
 };
 
 /**
@@ -370,11 +364,8 @@ static FORCEINLINE bool IsCoreUObjectPackage(const FName& PackageName)
 	FLinkerLoad.
 ----------------------------------------------------------------------------*/
 
-/**
- * Fills in the passed in TArray with the packages that are in its PrecacheMap
- *
- * @param TArray<FString> to be populated
- */
+#if !USE_NEW_ASYNC_IO
+
 void FLinkerLoad::GetListOfPackagesInPackagePrecacheMap( TArray<FString>& ListOfPackages )
 {
 	for ( TMap<FString, FLinkerLoad::FPackagePrecacheInfo>::TIterator It(PackagePrecacheMap); It; ++It )
@@ -383,6 +374,8 @@ void FLinkerLoad::GetListOfPackagesInPackagePrecacheMap( TArray<FString>& ListOf
 	}
 }
 
+#endif
+
 void FLinkerLoad::StaticInit(UClass* InUTexture2DStaticClass)
 {
 	UTexture2DStaticClass = InUTexture2DStaticClass;
@@ -413,7 +406,11 @@ FLinkerLoad* FLinkerLoad::CreateLinker(UPackage* Parent, const TCHAR* Filename,
 	LoadFlags &= ~LOAD_DeferDependencyLoads;
 #endif // USE_CIRCULAR_DEPENDENCY_LOAD_DEFERRING
 
-	FLinkerLoad* Linker = CreateLinkerAsync(Parent, Filename, LoadFlags);
+	FLinkerLoad* Linker = CreateLinkerAsync(Parent, Filename, LoadFlags
+#if USE_EVENT_DRIVEN_ASYNC_LOAD
+		, TFunction<void()>([](){})
+#endif
+		);
 	{
 #if USE_CIRCULAR_DEPENDENCY_LOAD_DEFERRING
 		// the linker could already have the DeferDependencyLoads flag present 
@@ -554,7 +551,11 @@ FName FLinkerLoad::FindSubobjectRedirectName(const FName& Name)
  *
  * @return	new FLinkerLoad object for Parent/ Filename
  */
-FLinkerLoad* FLinkerLoad::CreateLinkerAsync( UPackage* Parent, const TCHAR* Filename, uint32 LoadFlags )
+FLinkerLoad* FLinkerLoad::CreateLinkerAsync( UPackage* Parent, const TCHAR* Filename, uint32 LoadFlags 
+#if USE_EVENT_DRIVEN_ASYNC_LOAD
+	, TFunction<void()>&& InSummaryReadyCallback
+#endif
+	)
 {
 	check(Parent);
 
@@ -562,18 +563,31 @@ FLinkerLoad* FLinkerLoad::CreateLinkerAsync( UPackage* Parent, const TCHAR* File
 	FLinkerLoad* Linker = FindExistingLinkerForPackage(Parent);
 	if (Linker)
 	{
+#if USE_EVENT_DRIVEN_ASYNC_LOAD
+		UE_LOG(LogStreaming, Fatal, TEXT("FLinkerLoad::CreateLinkerAsync: Found existing linker for '%s'"), *Parent->GetName());
+#else
 		UE_LOG(LogStreaming, Log, TEXT("FLinkerLoad::CreateLinkerAsync: Found existing linker for '%s'"), *Parent->GetName());
+#endif
 	}
 
 	// Create a new linker if there isn't an existing one.
 	if( Linker == NULL )
 	{
-		if( GUseSeekFreeLoading )
+#if USE_NEW_ASYNC_IO
+		if (FApp::IsGame() && !GIsEditor)
 		{
-			LoadFlags |= LOAD_SeekFree;
+			LoadFlags |= LOAD_Async;
 		}
+#endif
 		Linker = new FLinkerLoad(Parent, Filename, LoadFlags );
 		Parent->LinkerLoad = Linker;
+#if USE_EVENT_DRIVEN_ASYNC_LOAD
+		if (Linker)
+		{
+			Linker->CreateLoader(Forward<TFunction<void()>>(InSummaryReadyCallback));
+		}
+#endif
+
 	}
 	
 	check(Parent->LinkerLoad == Linker);
@@ -606,6 +620,10 @@ FLinkerLoad::ELinkerStatus FLinkerLoad::Tick( float InTimeLimit, bool bInUseTime
 
 		do
 		{
+#if USE_EVENT_DRIVEN_ASYNC_LOAD
+			check(Loader);
+			if (true)
+#else
 			// Create loader, aka FArchive used for serialization and also precache the package file summary.
 			// false is returned until any precaching is complete.
 			if( true )
@@ -616,6 +634,7 @@ FLinkerLoad::ELinkerStatus FLinkerLoad::Tick( float InTimeLimit, bool bInUseTime
 
 			// Serialize the package file summary and presize the various arrays (name, import & export map)
 			if( Status == LINKER_Loaded )
+#endif
 			{
 				SCOPED_LOADTIMER(LinkerLoad_SerializePackageFileSummary);
 				Status = SerializePackageFileSummary();
@@ -649,15 +668,6 @@ FLinkerLoad::ELinkerStatus FLinkerLoad::Tick( float InTimeLimit, bool bInUseTime
 				Status = SerializeExportMap();
 			}
 
-			// Start pre-allocation of texture memory.
-			if( Status == LINKER_Loaded )
-			{
-				SCOPED_LOADTIMER(LinkerLoad_StartTextureAllocation);
-#if WITH_ENGINE
-				Status = StartTextureAllocation();
-#endif		// WITH_ENGINE
-			}
-
 			// Fix up import map for backward compatible serialization.
 			if( Status == LINKER_Loaded )
 			{	
@@ -699,6 +709,12 @@ FLinkerLoad::ELinkerStatus FLinkerLoad::Tick( float InTimeLimit, bool bInUseTime
 				Status = FindExistingExports();
 			}
 
+			if (Status == LINKER_Loaded)
+			{
+				SCOPED_LOADTIMER(LinkerLoad_SerializePreloadDependencies);
+				Status = SerializePreloadDependencies();
+			}
+
 			// Finalize creation process.
 			if( Status == LINKER_Loaded )
 			{
@@ -738,13 +754,20 @@ FLinkerLoad::FLinkerLoad(UPackage* InParent, const TCHAR* InFilename, uint32 InL
 , LoadFlags(InLoadFlags)
 , bHaveImportsBeenVerified(false)
 , bDynamicClassLinker(false)
+#if USE_EVENT_DRIVEN_ASYNC_LOAD
+, TemplateForGetArchetypeFromLoader(nullptr)
+, bForceSimpleIndexToObject(false)
+, bLockoutLegacyOperations(false)
+#endif
+#if USE_NEW_ASYNC_IO
+, bLoaderIsFArchiveAsync2(false)
+#endif
 , Loader(nullptr)
 , AsyncRoot(nullptr)
 , NameMapIndex(0)
 , GatherableTextDataMapIndex(0)
 , ImportMapIndex(0)
 , ExportMapIndex(0)
-, FirstNotLoadedExportMapIndex(0)
 , DependsMapIndex(0)
 , ExportHashIndex(0)
 , bHasSerializedPackageFileSummary(false)
@@ -772,7 +795,7 @@ FLinkerLoad::FLinkerLoad(UPackage* InParent, const TCHAR* InFilename, uint32 InL
 	FMemory::Memset(ExportHash, INDEX_NONE, sizeof(ExportHash));
 	INC_DWORD_STAT(STAT_LinkerCount);
 	INC_DWORD_STAT(STAT_LiveLinkerCount);
-#if UE_BUILD_DEBUG
+#if !UE_BUILD_SHIPPING
 	FLinkerManager::Get().GetLiveLinkers().Add(this);
 #endif
 
@@ -781,7 +804,7 @@ FLinkerLoad::FLinkerLoad(UPackage* InParent, const TCHAR* InFilename, uint32 InL
 
 FLinkerLoad::~FLinkerLoad()
 {
-#if UE_BUILD_DEBUG
+#if !UE_BUILD_SHIPPING
 	FLinkerManager::Get().GetLiveLinkers().Remove(this);
 #endif
 
@@ -796,6 +819,7 @@ FLinkerLoad::~FLinkerLoad()
 	// Make sure this is deleted if it's still allocated
 	delete LoadProgressScope;
 #endif
+	check(!Loader);
 }
 
 /**
@@ -832,7 +856,11 @@ bool FLinkerLoad::IsTimeLimitExceeded( const TCHAR* CurrentTask, int32 Granulari
 /**
  * Creates loader used to serialize content.
  */
-FLinkerLoad::ELinkerStatus FLinkerLoad::CreateLoader()
+FLinkerLoad::ELinkerStatus FLinkerLoad::CreateLoader(
+#if USE_EVENT_DRIVEN_ASYNC_LOAD
+	TFunction<void()>&& InSummaryReadyCallback
+#endif
+	)
 {
 	//DECLARE_SCOPE_CYCLE_COUNTER( TEXT( "FLinkerLoad::CreateLoader" ), STAT_LinkerLoad_CreateLoader, STATGROUP_LinkerLoad );
 
@@ -850,9 +878,7 @@ FLinkerLoad::ELinkerStatus FLinkerLoad::CreateLoader()
 
 	if( !Loader && !bDynamicClassLinker )
 	{
-#if !USE_NEW_ASYNC_IO
-		bool bIsSeekFree = LoadFlags & LOAD_SeekFree; // delete all of the seek free stuff
-#endif
+		const bool bIsAsyncLoad = LoadFlags & LOAD_Async;
 
 #if WITH_EDITOR
 		FFormatNamedArguments FeedbackArgs;
@@ -872,37 +898,24 @@ FLinkerLoad::ELinkerStatus FLinkerLoad::CreateLoader()
 			// In this case we can skip serializing PackageFileSummary and fill all the required info here
 			CreateDynamicTypeLoader();
 		}
-		// NOTE: Precached memory read gets highest priority, then memory reader, then seek free, then normal
-		// check to see if there is was an async preload request for this file
-		else if (FPackagePrecacheInfo* PrecacheInfo = PackagePrecacheMap.Find(*Filename))
+#if USE_NEW_ASYNC_IO
+		else if (!bIsAsyncLoad)		
 		{
-			// if so, serialize from memory (note this will have uncompressed a fully compressed package)
-			// block until the async read is complete
-			if( PrecacheInfo->SynchronizationObject->GetValue() != 0 )
+			check(!FPlatformProperties::RequiresCookedData() && !FSHA1::GetFileSHAHash(*Filename, NULL));
+			Loader = IFileManager::Get().CreateFileReader(*Filename, 0);
+			if (!Loader)
 			{
-				double StartTime = FPlatformTime::Seconds();
-			
-				FPlatformProcess::ConditionalSleep( [&]()
-				{
-					SHUTDOWN_IF_EXIT_REQUESTED;
-					return PrecacheInfo->SynchronizationObject->GetValue() == 0;
-				} 
-				);
-				float WaitTime = FPlatformTime::Seconds() - StartTime;
-				UE_LOG(LogInit, Log, TEXT("Waited %.3f sec for async package '%s' to complete caching."), WaitTime, *Filename);
+				UE_LOG(LogLinker, Warning, TEXT("Error opening file '%s'."), *Filename );
+				return LINKER_Failed;
 			}
-
-			// create a buffer reader using the read in data
-			// assume that all precached startup packages have SHA entries
-			Loader = new FBufferReaderWithSHA(PrecacheInfo->PackageData, PrecacheInfo->PackageDataSize, true, *Filename, true);
-
-			// remove the precache info from the map
-			PackagePrecacheMap.Remove(*Filename);
 		}
-#if USE_NEW_ASYNC_IO
 		else
 		{
-			Loader = NewFArchiveAsync2(*Filename);
+			Loader = new FArchiveAsync2(*Filename
+#if USE_EVENT_DRIVEN_ASYNC_LOAD
+				, Forward<TFunction<void()>>(InSummaryReadyCallback)
+#endif				
+				);
 
 			if (!Loader)
 			{
@@ -917,6 +930,7 @@ FLinkerLoad::ELinkerStatus FLinkerLoad::CreateLoader()
 				UE_LOG(LogLinker, Warning, TEXT("Error opening file '%s'."), *Filename);
 				return LINKER_Failed;
 			}
+			ActiveFPLB = Loader->ActiveFPLB; // make sure my fast past loading is using the FAA2 fast path buffer
 
 			bool bHasHashEntry = FSHA1::GetFileSHAHash(*Filename, NULL);
 			if ((LoadFlags & LOAD_MemoryReader) || bHasHashEntry)
@@ -945,7 +959,34 @@ FLinkerLoad::ELinkerStatus FLinkerLoad::CreateLoader()
 			}
 		}
 #else
-		else if ((LoadFlags & LOAD_MemoryReader) || !bIsSeekFree)
+		// NOTE: Precached memory read gets highest priority, then memory reader, then seek free, then normal
+		// check to see if there is was an async preload request for this file
+		else if (FPackagePrecacheInfo* PrecacheInfo = PackagePrecacheMap.Find(*Filename))
+		{
+			// if so, serialize from memory (note this will have uncompressed a fully compressed package)
+			// block until the async read is complete
+			if (PrecacheInfo->SynchronizationObject->GetValue() != 0)
+			{
+				double StartTime = FPlatformTime::Seconds();
+
+				FPlatformProcess::ConditionalSleep([&]()
+				{
+					SHUTDOWN_IF_EXIT_REQUESTED;
+					return PrecacheInfo->SynchronizationObject->GetValue() == 0;
+				}
+				);
+				float WaitTime = FPlatformTime::Seconds() - StartTime;
+				UE_LOG(LogInit, Log, TEXT("Waited %.3f sec for async package '%s' to complete caching."), WaitTime, *Filename);
+			}
+
+			// create a buffer reader using the read in data
+			// assume that all precached startup packages have SHA entries
+			Loader = new FBufferReaderWithSHA(PrecacheInfo->PackageData, PrecacheInfo->PackageDataSize, true, *Filename, true);
+
+			// remove the precache info from the map
+			PackagePrecacheMap.Remove(*Filename);
+		}
+		else if ((LoadFlags & LOAD_MemoryReader) || !bIsAsyncLoad)
 		{
 			// Create file reader used for serialization.
 			FArchive* FileReader = IFileManager::Get().CreateFileReader( *Filename, 0 );
@@ -983,7 +1024,7 @@ FLinkerLoad::ELinkerStatus FLinkerLoad::CreateLoader()
 				Loader = FileReader;
 			}
 		}
-		else if (bIsSeekFree)
+		else if (bIsAsyncLoad)
 		{
 			// Use the async archive as it supports proper Precache and package compression.
 			Loader = new FArchiveAsync( *Filename );
@@ -1017,7 +1058,13 @@ FLinkerLoad::ELinkerStatus FLinkerLoad::CreateLoader()
 		// Reset all custom versions
 		ResetCustomVersions();
 	}
-
+#if USE_EVENT_DRIVEN_ASYNC_LOAD
+	else
+	{
+		check(0);
+	}
+	return LINKER_TimedOut;
+#else
 	bool bExecuteNextStep = true;
 	if( bHasSerializedPackageFileSummary == false )
 	{
@@ -1051,6 +1098,7 @@ FLinkerLoad::ELinkerStatus FLinkerLoad::CreateLoader()
 	}
 
 	return (bExecuteNextStep && !IsTimeLimitExceeded( TEXT("creating loader") )) ? LINKER_Loaded : LINKER_TimedOut;
+#endif
 }
 
 /**
@@ -1062,6 +1110,11 @@ FLinkerLoad::ELinkerStatus FLinkerLoad::SerializePackageFileSummary()
 
 	if (bHasSerializedPackageFileSummary == false)
 	{
+		if (Loader->IsError())
+		{
+			UE_LOG(LogLinker, Warning, TEXT("The file '%s' contains unrecognizable data, check that it is of the expected type."), *Filename);
+			return LINKER_Failed;
+		}
 #if USE_NEW_ASYNC_IO
 		if (bLoaderIsFArchiveAsync2)
 		{
@@ -1152,6 +1205,14 @@ FLinkerLoad::ELinkerStatus FLinkerLoad::SerializePackageFileSummary()
 			}
 		}
 
+		if (FPlatformProperties::RequiresCookedData() && 
+			Summary.PreloadDependencyCount > 0 && Summary.PreloadDependencyOffset > 0 &&
+			!IsEventDrivenLoaderEnabledInCookedBuilds())
+		{
+			UE_LOG(LogLinker, Fatal, TEXT("Package %s contains preload dependency data but the current build does not support it. Make sure Event Driven Loader is enabled and rebuild the game executable."),
+				*GetArchiveName())
+		}
+
 #if PLATFORM_WINDOWS
 		if (!FPlatformProperties::RequiresCookedData())
 		{
@@ -1335,6 +1396,9 @@ FLinkerLoad::ELinkerStatus FLinkerLoad::SerializeNameMap()
 			if (bLoaderIsFArchiveAsync2)
 			{
 				bFinishedPrecaching = GetFArchiveAsync2Loader()->ReadyToStartReadingHeader(bUseTimeLimit, bUseFullTimeLimit, TickStartTime, TimeLimit);
+#if USE_EVENT_DRIVEN_ASYNC_LOAD
+				check(bFinishedPrecaching);
+#endif
 			}
 			else
 #endif
@@ -1699,77 +1763,6 @@ FLinkerLoad::ELinkerStatus FLinkerLoad::RemapImports()
 	return LINKER_Loaded;
 }
 
-#if WITH_ENGINE
-/**
- * Kicks off async memory allocations for all textures that will be loaded from this package.
- */
-FLinkerLoad::ELinkerStatus FLinkerLoad::StartTextureAllocation()
-{
-	DECLARE_SCOPE_CYCLE_COUNTER( TEXT( "FLinkerLoad::StartTextureAllocation" ), STAT_LinkerLoad_StartTextureAllocation, STATGROUP_LinkerLoad );
-
-	double StartTime = FPlatformTime::Seconds();
-	int32 NumAllocationsStarted = 0;
-	int32 NumAllocationsConsidered = 0;
-
-	// Only kick off async allocation if the loader is async.
-	bool bIsDone = true;
-	if ( bUseTimeLimit && !Summary.TextureAllocations.HaveAllAllocationsBeenConsidered() )
-	{
-		bool bContinue = true;
-		for ( int32 TypeIndex=Summary.TextureAllocations.NumTextureTypesConsidered;
-			  TypeIndex < Summary.TextureAllocations.TextureTypes.Num() && bContinue;
-			  ++TypeIndex )
-		{
-			FTextureAllocations::FTextureType& TextureType = Summary.TextureAllocations.TextureTypes[ TypeIndex ];
-			for ( int32 ResourceIndex=TextureType.NumExportIndicesProcessed; ResourceIndex < TextureType.ExportIndices.Num() && bContinue; ++ResourceIndex )
-			{
-				int32 ExportIndex = TextureType.ExportIndices[ ResourceIndex ];
-				if ( WillTextureBeLoaded( UTexture2DStaticClass, ExportIndex ) )
-				{
-					FTexture2DResourceMem* ResourceMem = CreateResourceMem(
-						TextureType.SizeX,
-						TextureType.SizeY,
-						TextureType.NumMips,
-						TextureType.Format,
-						TextureType.TexCreateFlags,
-						&Summary.TextureAllocations.PendingAllocationCount );
-
-					if ( ResourceMem )
-					{
-						TextureType.Allocations.Add( ResourceMem );
-						Summary.TextureAllocations.PendingAllocationSize += ResourceMem->GetResourceBulkDataSize();
-						Summary.TextureAllocations.PendingAllocationCount.Increment();
-						NumAllocationsStarted++;
-					}
-				}
-
-				TextureType.NumExportIndicesProcessed++;
-				NumAllocationsConsidered++;
-
-				bContinue = !IsTimeLimitExceeded( TEXT("allocating texture memory") );
-			}
-
-			// Have we processed all potential allocations for this texture type yet?
-			if ( TextureType.HaveAllAllocationsBeenConsidered() )
-			{
-				Summary.TextureAllocations.NumTextureTypesConsidered++;
-			}
-		}
-		bIsDone = Summary.TextureAllocations.HaveAllAllocationsBeenConsidered();
-	}
-
-	double Duration = FPlatformTime::Seconds() - StartTime;
-
-	// For profiling:
-// 	if ( NumAllocationsStarted )
-// 	{
-// 		UE_LOG(LogLinker, Log,  TEXT("StartTextureAllocation duration: %.3f ms (%d textures allocated, %d textures considered)"), Duration*1000.0, NumAllocationsStarted, NumAllocationsConsidered );
-// 	}
-
-	return (bIsDone && !IsTimeLimitExceeded( TEXT("kicking off texture allocations") )) ? LINKER_Loaded : LINKER_TimedOut;
-}
-#endif // WITH_ENGINE
-
 /**
  * Serializes the depends map.
  */
@@ -1805,6 +1798,33 @@ FLinkerLoad::ELinkerStatus FLinkerLoad::SerializeDependsMap()
 	return ((DependsMapIndex == Summary.ExportCount) && !IsTimeLimitExceeded( TEXT("serializing depends map") )) ? LINKER_Loaded : LINKER_TimedOut;
 }
 
+/**
+* Serializes the depends map.
+*/
+FLinkerLoad::ELinkerStatus FLinkerLoad::SerializePreloadDependencies()
+{
+	DECLARE_SCOPE_CYCLE_COUNTER(TEXT("FLinkerLoad::SerializePreloadDependencies"), STAT_LinkerLoad_SerializePreloadDependencies, STATGROUP_LinkerLoad);
+
+	// Skip serializing depends map if this is the editor or the data is missing
+	if (Summary.PreloadDependencyCount < 1 || Summary.PreloadDependencyOffset <= 0)
+	{
+		return LINKER_Loaded;
+	}
+
+	Seek(Summary.PreloadDependencyOffset);
+
+	PreloadDependencies.Reserve(Summary.PreloadDependencyCount);
+	//@todoio check endiness and fastpath this as a single serialize
+	for (int32 Index = 0; Index < Summary.PreloadDependencyCount; Index++)
+	{
+		FPackageIndex Idx;
+		*this << Idx;
+		PreloadDependencies.Add(Idx);
+	}
+	// Return whether we finished this step and it's safe to start with the next.
+	return !IsTimeLimitExceeded(TEXT("serialize preload dependencies")) ? LINKER_Loaded : LINKER_TimedOut;
+}
+
 /**
  * Serializes thumbnails
  */
@@ -1905,12 +1925,6 @@ FLinkerLoad::ELinkerStatus FLinkerLoad::CreateExportHash()
 	// Zero initialize hash on first iteration.
 	if( ExportHashIndex == 0 )
 	{
-#if USE_NEW_ASYNC_IO
-		if (bLoaderIsFArchiveAsync2)
-		{
-			GetFArchiveAsync2Loader()->EndReadingHeader();
-		}
-#endif
 		for( int32 i=0; i<ARRAY_COUNT(ExportHash); i++ )
 		{
 			ExportHash[i] = INDEX_NONE;
@@ -1998,6 +2012,28 @@ FLinkerLoad::ELinkerStatus FLinkerLoad::FinalizeCreation()
 			Verify();
 		}
 
+#if USE_EVENT_DRIVEN_ASYNC_LOAD
+		if (AsyncRoot)
+		{
+			for (int32 ImportIndex = 0; ImportIndex < ImportMap.Num(); ++ImportIndex)
+			{
+				FPackageIndex Index = FPackageIndex::FromImport(ImportIndex);
+				AsyncRoot->ObjectNameToImportOrExport.Add(Imp(Index).ObjectName, Index);
+			}
+			for (int32 ExportIndex = 0; ExportIndex < ExportMap.Num(); ++ExportIndex)
+			{
+				FPackageIndex Index = FPackageIndex::FromExport(ExportIndex);
+				AsyncRoot->ObjectNameToImportOrExport.Add(Exp(Index).ObjectName, Index);
+			}
+		}
+#endif
+#if USE_NEW_ASYNC_IO
+		if (bLoaderIsFArchiveAsync2)
+		{
+			GetFArchiveAsync2Loader()->EndReadingHeader();
+		}
+#endif
+
 		// Avoid duplicate work in the case of async linker creation.
 		bHasFinishedInitialization = true;
 
@@ -2144,6 +2180,7 @@ FString FLinkerLoad::GetArchiveName() const
  * @param Dependencies Array of all dependencies needed
  * @param bSkipLoadedObjects Whether to skip already loaded objects when gathering dependencies
  */
+#if WITH_EDITORONLY_DATA
 void FLinkerLoad::GatherExportDependencies(int32 ExportIndex, TSet<FDependencyRef>& Dependencies, bool bSkipLoadedObjects)
 {
 	// make sure we have dependencies
@@ -2296,9 +2333,13 @@ void FLinkerLoad::GatherImportDependencies(int32 ImportIndex, TSet<FDependencyRe
 		NewRef.Linker->GatherExportDependencies(NewRef.ExportIndex, Dependencies, bSkipLoadedObjects);
 	}
 }
+#endif
 
 FLinkerLoad::EVerifyResult FLinkerLoad::VerifyImport(int32 ImportIndex)
 {
+#if USE_EVENT_DRIVEN_ASYNC_LOAD
+	check(GIsInitialLoad);
+#endif
 	FObjectImport& Import = ImportMap[ImportIndex];
 
 	// keep a string of modifiers to add to the Editor Warning dialog
@@ -2484,6 +2525,10 @@ UPackage* LoadPackageInternal(UPackage* InOuter, const TCHAR* InLongPackageName,
  */
 bool FLinkerLoad::VerifyImportInner(const int32 ImportIndex, FString& WarningSuffix)
 {
+#if USE_EVENT_DRIVEN_ASYNC_LOAD
+	check(GIsInitialLoad);
+#endif
+
 	check(IsLoading());
 
 	FObjectImport& Import = ImportMap[ImportIndex];
@@ -2843,7 +2888,7 @@ bool FLinkerLoad::VerifyImportInner(const int32 ImportIndex, FString& WarningSuf
 					FindObject = StaticFindObject( FindClass, ANY_PACKAGE, *Import.ObjectName.ToString(), true );
 					if (FindObject && FindOuter->GetOutermost() != FindObject->GetOutermost())
 					{
-						// Limit the results to the same package.
+						// Limit the results to the same package.I
 						FindObject = NULL;
 					}
 				}
@@ -2967,7 +3012,7 @@ void FLinkerLoad::LoadAllObjects( bool bForcePreload )
 	check((LoadFlags & LOAD_DeferDependencyLoads) == 0);
 #endif // USE_DEFERRED_DEPENDENCY_CHECK_VERIFICATION_TESTS
 
-	if ( (LoadFlags&LOAD_SeekFree) != 0 )
+	if ((LoadFlags & LOAD_Async) != 0)
 	{
 		bForcePreload = true;
 	}
@@ -3223,44 +3268,49 @@ UObject* FLinkerLoad::Create( UClass* ObjectClass, FName ObjectName, UObject* Ou
  *					If Object is a UClass and the class default object has already been created, calls
  *					Preload for the class default object as well.
  */
+
 void FLinkerLoad::Preload( UObject* Object )
 {
+
 	//check(IsValidLowLevel());
 	check(Object);
 
-#if USE_CIRCULAR_DEPENDENCY_LOAD_DEFERRING
-	bool const bIsNonNativeObject = !Object->GetOutermost()->HasAnyPackageFlags(PKG_CompiledIn);
-	// we can determine that this is a blueprint class/struct by checking if it 
-	// is a class/struct object AND if it is not native (blueprint 
-	// structs/classes are the only asset package structs/classes we have)
-	bool const bIsBlueprintClass  = (Cast<UClass>(Object) != nullptr) && bIsNonNativeObject;
-	bool const bIsBlueprintStruct = (Cast<UScriptStruct>(Object) != nullptr) && bIsNonNativeObject;
-	// to avoid cyclic dependency issues, we want to defer all external loads 
-	// that MAY rely on this class/struct (meaning all other blueprint packages)  
-	bool const bDeferDependencyLoads = (bIsBlueprintClass || bIsBlueprintStruct) && FBlueprintSupport::UseDeferredDependencyLoading();
-
-#if USE_DEFERRED_DEPENDENCY_CHECK_VERIFICATION_TESTS
-	// we should NEVER be pre-loading another blueprint class when the 
-	// DeferDependencyLoads flag is set (some other blueprint class/struct is  
-	// already being loaded further up the load chain, and this could introduce  
-	// a circular load)
-	//
-	// NOTE: we do allow Preload() calls for structs (because we need a struct 
-	//       loaded to determine its size), but structs will be prevented from 
-	//       further loading any of its BP class dependencies (we pass along the 
-	//       LOAD_DeferDependencyLoads flag)
-	check(!bIsBlueprintClass || !Object->HasAnyFlags(RF_NeedLoad) || !(LoadFlags & LOAD_DeferDependencyLoads));
-	// right now there are no known scenarios where someone requests a Preloa() 
-	// on a temporary ULinkerPlaceholderExportObject
-	check(!Object->IsA<ULinkerPlaceholderExportObject>());
-#endif // USE_DEFERRED_DEPENDENCY_CHECK_VERIFICATION_TESTS
-#endif // USE_CIRCULAR_DEPENDENCY_LOAD_DEFERRING
-	
 	// Preload the object if necessary.
 	if (Object->HasAnyFlags(RF_NeedLoad))
 	{
 		if (Object->GetLinker() == this)
 		{
+#if USE_EVENT_DRIVEN_ASYNC_LOAD
+			check(!bLockoutLegacyOperations);
+#endif
+#if USE_CIRCULAR_DEPENDENCY_LOAD_DEFERRING
+			bool const bIsNonNativeObject = !Object->GetOutermost()->HasAnyPackageFlags(PKG_CompiledIn);
+			// we can determine that this is a blueprint class/struct by checking if it 
+			// is a class/struct object AND if it is not native (blueprint 
+			// structs/classes are the only asset package structs/classes we have)
+			bool const bIsBlueprintClass = (Cast<UClass>(Object) != nullptr) && bIsNonNativeObject;
+			bool const bIsBlueprintStruct = (Cast<UScriptStruct>(Object) != nullptr) && bIsNonNativeObject;
+			// to avoid cyclic dependency issues, we want to defer all external loads 
+			// that MAY rely on this class/struct (meaning all other blueprint packages)  
+			bool const bDeferDependencyLoads = (bIsBlueprintClass || bIsBlueprintStruct) && FBlueprintSupport::UseDeferredDependencyLoading();
+
+#if USE_DEFERRED_DEPENDENCY_CHECK_VERIFICATION_TESTS
+			// we should NEVER be pre-loading another blueprint class when the 
+			// DeferDependencyLoads flag is set (some other blueprint class/struct is  
+			// already being loaded further up the load chain, and this could introduce  
+			// a circular load)
+			//
+			// NOTE: we do allow Preload() calls for structs (because we need a struct 
+			//       loaded to determine its size), but structs will be prevented from 
+			//       further loading any of its BP class dependencies (we pass along the 
+			//       LOAD_DeferDependencyLoads flag)
+			check(!bIsBlueprintClass || !Object->HasAnyFlags(RF_NeedLoad) || !(LoadFlags & LOAD_DeferDependencyLoads));
+			// right now there are no known scenarios where someone requests a Preloa() 
+			// on a temporary ULinkerPlaceholderExportObject
+			check(!Object->IsA<ULinkerPlaceholderExportObject>());
+#endif // USE_DEFERRED_DEPENDENCY_CHECK_VERIFICATION_TESTS
+#endif // USE_CIRCULAR_DEPENDENCY_LOAD_DEFERRING
+
 #if USE_CIRCULAR_DEPENDENCY_LOAD_DEFERRING
 			// Because of delta serialization, we require that a parent's CDO be 
 			// fully serialized before its children's CDOs are created. However, 
@@ -3282,7 +3332,7 @@ void FLinkerLoad::Preload( UObject* Object )
 			// if this is an inherited sub-object on a CDO, and that CDO has had
 			// its initialization deferred (for reasons explained above), then 
 			// we shouldn't serialize in data for this quite yet... not until 
-			// its owner has had a chance to initialize itself (because, as part
+			// its owner has had a chaTnce to initialize itself (because, as part
 			// of CDO initialization, inherited sub-objects get filled in with 
 			// values inherited from the super)
 			else if (Object->HasAnyFlags(RF_DefaultSubObject|RF_InheritableComponentTemplate) && FDeferredObjInitializerTracker::DeferSubObjectPreload(Object))
@@ -3338,17 +3388,8 @@ void FLinkerLoad::Preload( UObject* Object )
 #if USE_NEW_ASYNC_IO
 					if (FAA2)
 					{
-						for (int32 TrimIndex = FirstNotLoadedExportMapIndex; TrimIndex < ExportIndex; TrimIndex++)
-						{
-							FObjectExport& TrimExport = ExportMap[TrimIndex];
-							if (!TrimExport.Object || TrimExport.Object->HasAnyFlags(RF_NeedLoad))
-							{
-								break;
-							}
-							FirstNotLoadedExportMapIndex = TrimIndex + 1;
-						}
-						bool bReady = FAA2->Precache(Export.SerialOffset, Export.SerialSize, bUseTimeLimit, bUseFullTimeLimit, TickStartTime, TimeLimit, FirstNotLoadedExportMapIndex == ExportIndex);
-						check(bReady || !FPlatformProperties::RequiresCookedData()); // if we are cooked, then this should be ready to go
+						bool bReady = FAA2->Precache(Export.SerialOffset, Export.SerialSize, bUseTimeLimit, bUseFullTimeLimit, TickStartTime, TimeLimit);
+						UE_CLOG(!(bReady || !bUseTimeLimit || !FPlatformProperties::RequiresCookedData()), LogLinker, Warning, TEXT("Hitch on async loading of %s; this export was not properly precached."), *Object->GetFullName());
 					}
 					else
 #endif
@@ -3475,28 +3516,6 @@ void FLinkerLoad::Preload( UObject* Object )
 				}
 
 				Loader->Seek( SavedPos );
-#if USE_NEW_ASYNC_IO
-				if (FPlatformProperties::RequiresCookedData() && FAA2)
-				{
-					if (Export.Object && !Export.Object->HasAnyFlags(RF_NeedLoad)) // if we somehow failed here, then we know there is at least one thing pending and we will forget about it
-					{
-						for (int32 TrimIndex = FirstNotLoadedExportMapIndex; TrimIndex < ExportMap.Num(); TrimIndex++)
-						{
-							FObjectExport& TrimExport = ExportMap[TrimIndex];
-							if (!TrimExport.Object || TrimExport.Object->HasAnyFlags(RF_NeedLoad))
-							{
-								break;
-							}
-							FirstNotLoadedExportMapIndex = TrimIndex + 1;
-						}
-						if (FirstNotLoadedExportMapIndex == ExportMap.Num())
-						{
-							// EOF, flush archive
-							FAA2->FlushCache();
-						}
-					}
-				}
-#endif
 
 				// if this is a UClass object and it already has a class default object
 				if ( Cls != NULL && Cls->GetDefaultsCount() )
@@ -3625,6 +3644,10 @@ UObject* FLinkerLoad::CreateExport( int32 Index )
 	// Check whether we already loaded the object and if not whether the context flags allow loading it.
 	if( !Export.Object && !FilterExport(Export) ) // for some acceptable position, it was not "not for" 
 	{
+#if USE_EVENT_DRIVEN_ASYNC_LOAD
+		check(!bLockoutLegacyOperations);
+#endif
+
 		check(Export.ObjectName!=NAME_None || !(Export.ObjectFlags&RF_Public));
 		check(IsLoading());
 
@@ -3953,8 +3976,6 @@ UObject* FLinkerLoad::CreateExport( int32 Index )
 		}
 #endif // WITH_EDITOR
 
-		// Create the export object, marking it with the appropriate flags to
-		// indicate that the object's data still needs to be loaded.
 		if (ActualObjectWithTheName && !ActualObjectWithTheName->GetClass()->IsChildOf(LoadClass))
 		{
 			UE_LOG(LogLinker, Error, TEXT("Failed import: class '%s' name '%s' outer '%s'. There is another object (of '%s' class) at the path."),
@@ -3962,6 +3983,8 @@ UObject* FLinkerLoad::CreateExport( int32 Index )
 			return NULL;
 		}
 
+		// Create the export object, marking it with the appropriate flags to
+		// indicate that the object's data still needs to be loaded.
 		EObjectFlags ObjectLoadFlags = Export.ObjectFlags;
 		// if we are loading objects just to verify an object reference during script compilation,
 		if (!GVerifyObjectReferencesOnly
@@ -4165,6 +4188,9 @@ bool FLinkerLoad::IsImportNative(const int32 Index) const
 // Return the loaded object corresponding to an import index; any errors are fatal.
 UObject* FLinkerLoad::CreateImport( int32 Index )
 {
+#if USE_EVENT_DRIVEN_ASYNC_LOAD
+	check(!bLockoutLegacyOperations);
+#endif
 	FScopedCreateImportCounter ScopedCounter( this, Index );
 	FObjectImport& Import = ImportMap[ Index ];
 	
@@ -4232,8 +4258,9 @@ UObject* FLinkerLoad::CreateImport( int32 Index )
 						}
 						if (!FindOuter)
 						{
+							// This can happen when deleting native properties or restructing blueprints. If there is an actual problem it will be caught when trying to resolve the outer itself
 							FString OuterName = Import.OuterIndex.IsNull() ? LinkerRoot->GetFullName() : GetFullImpExpName(Import.OuterIndex);
-							UE_LOG(LogLinker, Warning, TEXT("CreateImport: Failed to load Outer for resource '%s': %s"), *Import.ObjectName.ToString(), *OuterName);
+							UE_LOG(LogLinker, Verbose, TEXT("CreateImport: Failed to load Outer for resource '%s': %s"), *Import.ObjectName.ToString(), *OuterName);
 							return NULL;
 						}
 	
@@ -4296,6 +4323,8 @@ UObject* FLinkerLoad::CreateImport( int32 Index )
 	return Import.XObject;
 }
 
+
+
 // Map an import/export index to an object; all errors here are fatal.
 UObject* FLinkerLoad::IndexToObject( FPackageIndex Index )
 {
@@ -4510,46 +4539,19 @@ void FLinkerLoad::DetachAllBulkData(bool bEnsureAllBulkDataIsLoaded)
 
 #endif // WITH_EDITOR
 
-/**
- * Hint the archive that the region starting at passed in offset and spanning the passed in size
- * is going to be read soon and should be precached.
- *
- * The function returns whether the precache operation has completed or not which is an important
- * hint for code knowing that it deals with potential async I/O. The archive is free to either not 
- * implement this function or only partially precache so it is required that given sufficient time
- * the function will return true. Archives not based on async I/O should always return true.
- *
- * This function will not change the current archive position.
- *
- * @param	PrecacheOffset	Offset at which to begin precaching.
- * @param	PrecacheSize	Number of bytes to precache
- * @return	false if precache operation is still pending, true otherwise
- */
-bool FLinkerLoad::Precache( int64 PrecacheOffset, int64 PrecacheSize )
-{
-	return bDynamicClassLinker || Loader->Precache(PrecacheOffset, PrecacheSize);
-}
-
-void FLinkerLoad::Seek( int64 InPos )
-{
-	Loader->Seek( InPos );
-}
-
-int64 FLinkerLoad::Tell()
-{
-	return Loader->Tell();
-}
-
-int64 FLinkerLoad::TotalSize()
-{
-	return Loader->TotalSize();
-}
-
 FArchive& FLinkerLoad::operator<<( UObject*& Object )
 {
 	FPackageIndex Index;
 	FArchive& Ar = *this;
 	Ar << Index;
+#if USE_EVENT_DRIVEN_ASYNC_LOAD
+	if (bForceSimpleIndexToObject)
+	{
+		check(Ar.IsLoading() && AsyncRoot);
+		Object = AsyncRoot->EventDrivenIndexToObject(Index, false);
+		return *this;
+	}
+#endif
 
 	UObject* Temporary = NULL;
 	Temporary = IndexToObject( Index );
@@ -4584,79 +4586,16 @@ FArchive& FLinkerLoad::operator<<( UObject*& Object )
 	}
 #endif
 
-	FMemory::Memcpy(&Object, &Temporary, sizeof(UObject*));
-	return *this;
-}
-
-FArchive& FLinkerLoad::operator<<( FLazyObjectPtr& LazyObjectPtr)
-{
-	FArchive& Ar = *this;
-	FUniqueObjectGuid ID;
-	Ar << ID;
-	LazyObjectPtr = ID;
-	return Ar;
-}
-
-FArchive& FLinkerLoad::operator<<( FAssetPtr& AssetPtr)
-{
-	FArchive& Ar = *this;
-	FStringAssetReference ID;
-	ID.Serialize(Ar);
-	AssetPtr = ID;
-	return Ar;
-}
-
-FArchive& FLinkerLoad::operator<<( FName& Name )
-{
-	NAME_INDEX NameIndex;
-	FArchive& Ar = *this;
-	Ar << NameIndex;
-
-	if( !NameMap.IsValidIndex(NameIndex) )
-	{
-		UE_LOG(LogLinker, Error, TEXT("Bad name index %i/%i"), NameIndex, NameMap.Num() );
-		ArIsError = true;
-		ArIsCriticalError = true;
-		Name = NAME_None;
-		return *this;
-	}
-
-	// if the name wasn't loaded (because it wasn't valid in this context)
-	const FName& MappedName = NameMap[NameIndex];
-	if (MappedName.IsNone())
-	{
-		int32 TempNumber;
-		Ar << TempNumber;
-		Name = NAME_None;
-	}
-	else
-	{
-		int32 Number;
-		Ar << Number;
-		// simply create the name from the NameMap's name and the serialized instance number
-#ifndef __clang__
-		Name = FName(MappedName, Number);
-#else
-		// @todo-mobile: IOS crashes on the assignment; need to do a memcpy manually...
-		FName TempName = FName(MappedName, Number);
-		FMemory::Memcpy(&Name, &TempName, sizeof(FName));
-#endif
-	}
-
+	Object = Temporary;
 	return *this;
 }
 
-void FLinkerLoad::Serialize( void* V, int64 Length )
+void FLinkerLoad::BadNameIndexError(NAME_INDEX NameIndex)
 {
-	checkSlow(FPlatformTLS::GetCurrentThreadId() == OwnerThread);
-	Loader->Serialize( V, Length );
+	UE_LOG(LogLinker, Error, TEXT("Bad name index %i/%i"), NameIndex, NameMap.Num());
 }
 
-/**
-* Kick off an async load of a package file into memory
-* 
-* @param PackageName Name of package to read in. Must be the same name as passed into LoadPackage
-*/
+#if !USE_NEW_ASYNC_IO
 void FLinkerLoad::AsyncPreloadPackage(const TCHAR* PackageName)
 {
 	// get package filename
@@ -4725,6 +4664,7 @@ void FLinkerLoad::AsyncPreloadPackage(const TCHAR* PackageName)
 
 	check(RequestId);
 }
+#endif
 
 /**
  * Called when an object begins serializing property data using script serialization.
@@ -4768,6 +4708,15 @@ bool FLinkerLoad::FindImportClassAndPackage( FName ClassName, FPackageIndex &Cla
 	return false;
 }
 
+#if USE_EVENT_DRIVEN_ASYNC_LOAD
+UObject* FLinkerLoad::GetArchetypeFromLoader(const UObject* Obj)
+{
+	check(!TemplateForGetArchetypeFromLoader || FUObjectThreadContext::Get().SerializedObject == Obj);
+	return TemplateForGetArchetypeFromLoader;
+}
+#endif
+
+
 /**
 * Attempts to find the index for the given class object in the import list and adds it + its package if it does not exist
 */
diff --git a/Engine/Source/Runtime/CoreUObject/Private/UObject/LinkerManager.cpp b/Engine/Source/Runtime/CoreUObject/Private/UObject/LinkerManager.cpp
--- a/Engine/Source/Runtime/CoreUObject/Private/UObject/LinkerManager.cpp
+++ b/Engine/Source/Runtime/CoreUObject/Private/UObject/LinkerManager.cpp
@@ -18,9 +18,21 @@ bool FLinkerManager::Exec(class UWorld* InWorld, const TCHAR* Cmd, FOutputDevice
 	if (FParse::Command(&Cmd, TEXT("LinkerLoadList")))
 	{
 		UE_LOG(LogLinker, Display, TEXT("ObjectLoaders: %d"), ObjectLoaders.Num());
+#if !UE_BUILD_SHIPPING
+		for (auto Linker : ObjectLoaders)
+		{
+			UE_LOG(LogLinker, Display, TEXT("%s"), *Linker->Filename);
+		}
+#endif
 		UE_LOG(LogLinker, Display, TEXT("LoadersWithNewImports: %d"), LoadersWithNewImports.Num());
+#if !UE_BUILD_SHIPPING
+		for (auto Linker : LoadersWithNewImports)
+		{
+			UE_LOG(LogLinker, Display, TEXT("%s"), *Linker->Filename);
+		}
+#endif
 
-#if UE_BUILD_DEBUG
+#if !UE_BUILD_SHIPPING
 		UE_LOG(LogLinker, Display, TEXT("LiveLinkers: %d"), LiveLinkers.Num());
 		for (auto Linker : LiveLinkers)
 		{
