Copying //UE4/Dev-Core to //UE4/Dev-Main (Source: //UE4/Dev-Core @ 3151653)

diff --git a/Engine/Source/Runtime/CoreUObject/Private/UObject/LinkerManager.h b/Engine/Source/Runtime/CoreUObject/Private/UObject/LinkerManager.h
--- a/Engine/Source/Runtime/CoreUObject/Private/UObject/LinkerManager.h
+++ b/Engine/Source/Runtime/CoreUObject/Private/UObject/LinkerManager.h
@@ -95,7 +95,7 @@ public:
 		LoadersWithNewImports.Empty();
 	}
 
-#if UE_BUILD_DEBUG
+#if !UE_BUILD_SHIPPING
 	FORCEINLINE TArray<FLinkerLoad*>& GetLiveLinkers()
 	{
 		return LiveLinkers;
@@ -133,7 +133,7 @@ private:
 #if THREADSAFE_UOBJECTS
 	FCriticalSection LoadersWithNewImportsCritical;
 #endif
-#if UE_BUILD_DEBUG
+#if !UE_BUILD_SHIPPING
 	/** List of all the existing linker loaders **/
 	TArray<FLinkerLoad*> LiveLinkers;
 #endif
diff --git a/Engine/Source/Runtime/CoreUObject/Private/UObject/LinkerSave.cpp b/Engine/Source/Runtime/CoreUObject/Private/UObject/LinkerSave.cpp
--- a/Engine/Source/Runtime/CoreUObject/Private/UObject/LinkerSave.cpp
+++ b/Engine/Source/Runtime/CoreUObject/Private/UObject/LinkerSave.cpp
@@ -172,6 +172,40 @@ FPackageIndex FLinkerSave::MapObject( const UObject* Object ) const
 		const FPackageIndex *Found = ObjectIndicesMap.Find(Object);
 		if (Found)
 		{
+			if (IsEventDrivenLoaderEnabledInCookedBuilds() &&
+				IsCooking() && CurrentlySavingExport.IsExport() &&
+				!Object->GetOutermost()->HasAllPackagesFlags(PKG_CompiledIn) && 
+				*Found != CurrentlySavingExport) // would be weird, but I can't be a dependency on myself
+			{
+				const FObjectExport& SavingExport = Exp(CurrentlySavingExport);
+				bool bFoundDep = false;
+				if (SavingExport.FirstExportDependency >= 0)
+				{
+					int32 NumDeps = SavingExport.CreateBeforeCreateDependencies + SavingExport.CreateBeforeSerializationDependencies + SavingExport.SerializationBeforeCreateDependencies + SavingExport.SerializationBeforeSerializationDependencies;
+					for (int32 DepIndex = SavingExport.FirstExportDependency; DepIndex < SavingExport.FirstExportDependency + NumDeps; DepIndex++)
+					{
+						if (DepListForErrorChecking[DepIndex] == *Found)
+						{
+							bFoundDep = true;
+						}
+					}
+				}
+				if (!bFoundDep)
+				{
+					if (SavingExport.Object && SavingExport.Object->IsA(UClass::StaticClass()) && CastChecked<UClass>(SavingExport.Object)->GetDefaultObject() == Object)
+					{
+						bFoundDep = true; // the class is saving a ref to the CDO...which doesn't really work or do anything useful, but it isn't an error
+					}
+				}
+				if (!bFoundDep)
+				{
+					UE_LOG(LogLinker, Fatal, TEXT("Attempt to map an object during save that was not listed as a dependency. Saving Export %d %s in %s. Missing Dep on %s %s."),
+						CurrentlySavingExport.ForDebugging(), *SavingExport.ObjectName.ToString(), *GetArchiveName(),
+						Found->IsExport() ? TEXT("Export") : TEXT("Import"), *ImpExp(*Found).ObjectName.ToString()
+						);
+				}
+			}
+
 			return *Found;
 		}
 	}
diff --git a/Engine/Source/Runtime/CoreUObject/Private/UObject/Obj.cpp b/Engine/Source/Runtime/CoreUObject/Private/UObject/Obj.cpp
--- a/Engine/Source/Runtime/CoreUObject/Private/UObject/Obj.cpp
+++ b/Engine/Source/Runtime/CoreUObject/Private/UObject/Obj.cpp
@@ -788,11 +788,20 @@ bool UObject::ConditionalFinishDestroy()
 	}
 }
 
+#if !defined(USE_EVENT_DRIVEN_ASYNC_LOAD)
+#error "USE_EVENT_DRIVEN_ASYNC_LOAD must be defined"
+#endif
 
 void UObject::ConditionalPostLoad()
 {
-	if( HasAnyFlags(RF_NeedPostLoad) )
+#if USE_EVENT_DRIVEN_ASYNC_LOAD
+	check(!HasAnyFlags(RF_NeedLoad)); //@todoio Added this as "nicks rule"
+#endif
+									  // PostLoad only if the object needs it and has already been serialized
+	//@todoio note this logic should be unchanged compared to main
+	if (HasAnyFlags(RF_NeedPostLoad))
 	{
+
 		check(IsInGameThread() || HasAnyFlags(RF_ClassDefaultObject|RF_ArchetypeObject) || IsPostLoadThreadSafe() || IsA(UClass::StaticClass()))
 
 #if !(UE_BUILD_SHIPPING || UE_BUILD_TEST)
@@ -833,15 +842,25 @@ void UObject::ConditionalPostLoad()
 	}
 }
 
+#if !defined(USE_EVENT_DRIVEN_ASYNC_LOAD)
+#error "USE_EVENT_DRIVEN_ASYNC_LOAD must be defined."
+#endif
+
 
 void UObject::PostLoadSubobjects( FObjectInstancingGraph* OuterInstanceGraph/*=NULL*/ )
 {
+#if USE_EVENT_DRIVEN_ASYNC_LOAD
+	check(!HasAnyFlags(RF_NeedLoad));
+#endif
 	if( GetClass()->HasAnyClassFlags(CLASS_HasInstancedReference) )
 	{
 		UObject* ObjOuter = GetOuter();
 		// make sure our Outer has already called ConditionalPostLoadSubobjects
 		if (ObjOuter != NULL && ObjOuter->HasAnyFlags(RF_NeedPostLoadSubobjects) )
 		{
+#if USE_EVENT_DRIVEN_ASYNC_LOAD
+			check(!ObjOuter->HasAnyFlags(RF_NeedLoad));
+#endif
 			if (ObjOuter->HasAnyFlags(RF_NeedPostLoad) )
 			{
 				ObjOuter->ConditionalPostLoad();
@@ -959,6 +978,20 @@ bool UObject::IsSelected() const
 	return !IsPendingKill() && GSelectedAnnotation.Get(this);
 }
 
+void UObject::GetPreloadDependencies(TArray<UObject*>& OutDeps)
+{
+	UClass *ObjClass = GetClass();
+	if (!ObjClass->HasAnyClassFlags(CLASS_Intrinsic | CLASS_Native))
+	{
+		OutDeps.Add(ObjClass);
+
+		if (!HasAnyFlags(RF_ClassDefaultObject) && ObjClass->GetDefaultsCount() > 0)
+		{
+			OutDeps.Add(ObjClass->GetDefaultObject());
+		}
+	}
+}
+
 void UObject::Serialize( FArchive& Ar )
 {
 	// These three items are very special items from a serialization standpoint. They aren't actually serialized.
@@ -1077,7 +1110,12 @@ void UObject::SerializeScriptProperties( FArchive& Ar ) const
 
 	if( (Ar.IsLoading() || Ar.IsSaving()) && !Ar.WantBinaryPropertySerialization() )
 	{
-		UObject* DiffObject = GetArchetype();
+		//@todoio GetArchetype is pathological for blueprint classes and the event driven loader; the EDL already knows what the archetype is; just calling this->GetArchetype() tries to load some other stuff.
+		UObject* DiffObject = Ar.GetArchetypeFromLoader(this);
+		if (!DiffObject)
+		{
+			DiffObject = GetArchetype();
+		}
 #if WITH_EDITOR
 		static const FBoolConfigValueHelper BreakSerializationRecursion(TEXT("StructSerialization"), TEXT("BreakSerializationRecursion"));
 		const bool bBreakSerializationRecursion = BreakSerializationRecursion && Ar.IsLoading() && Ar.GetLinker();
@@ -1093,7 +1131,12 @@ void UObject::SerializeScriptProperties( FArchive& Ar ) const
 	}
 	else if ( Ar.GetPortFlags() != 0 && !Ar.ArUseCustomPropertyList )
 	{
-		UObject* DiffObject = GetArchetype();
+		//@todoio GetArchetype is pathological for blueprint classes and the event driven loader; the EDL already knows what the archetype is; just calling this->GetArchetype() tries to load some other stuff.
+		UObject* DiffObject = Ar.GetArchetypeFromLoader(this);
+		if (!DiffObject)
+		{
+			DiffObject = GetArchetype();
+		}
 		ObjClass->SerializeBinEx( Ar, const_cast<UObject *>(this), DiffObject, DiffObject ? DiffObject->GetClass() : NULL );
 	}
 	else
diff --git a/Engine/Source/Runtime/CoreUObject/Private/UObject/ObjectResource.cpp b/Engine/Source/Runtime/CoreUObject/Private/UObject/ObjectResource.cpp
--- a/Engine/Source/Runtime/CoreUObject/Private/UObject/ObjectResource.cpp
+++ b/Engine/Source/Runtime/CoreUObject/Private/UObject/ObjectResource.cpp
@@ -49,6 +49,12 @@ FObjectExport::FObjectExport()
 , bWasFiltered(false)
 , PackageGuid(FGuid(0, 0, 0, 0))
 , PackageFlags(0)
+, FirstExportDependency(-1)
+, SerializationBeforeSerializationDependencies(0)
+, CreateBeforeSerializationDependencies(0)
+, SerializationBeforeCreateDependencies(0)
+, CreateBeforeCreateDependencies(0)
+
 {}
 
 FObjectExport::FObjectExport( UObject* InObject )
@@ -70,6 +76,11 @@ FObjectExport::FObjectExport( UObject* InObject )
 , bWasFiltered(false)
 , PackageGuid(FGuid(0, 0, 0, 0))
 , PackageFlags(0)
+, FirstExportDependency(-1)
+, SerializationBeforeSerializationDependencies(0)
+, CreateBeforeSerializationDependencies(0)
+, SerializationBeforeCreateDependencies(0)
+, CreateBeforeCreateDependencies(0)
 {
 	if(Object)		
 	{
@@ -84,6 +95,11 @@ FArchive& operator<<( FArchive& Ar, FObjectExport& E )
 {
 	Ar << E.ClassIndex;
 	Ar << E.SuperIndex;
+	if (Ar.UE4Ver() >= VER_UE4_TemplateIndex_IN_COOKED_EXPORTS)
+	{
+		Ar << E.TemplateIndex;
+	}
+
 	Ar << E.OuterIndex;
 	Ar << E.ObjectName;
 
@@ -114,6 +130,14 @@ FArchive& operator<<( FArchive& Ar, FObjectExport& E )
 		Ar << E.bIsAsset;
 	}
 
+	if (Ar.UE4Ver() >= VER_UE4_PRELOAD_DEPENDENCIES_IN_COOKED_EXPORTS)
+	{
+		Ar << E.FirstExportDependency;
+		Ar << E.SerializationBeforeSerializationDependencies;
+		Ar << E.CreateBeforeSerializationDependencies;
+		Ar << E.SerializationBeforeCreateDependencies;
+		Ar << E.CreateBeforeCreateDependencies;
+	}	
 	return Ar;
 }
 
@@ -123,6 +147,11 @@ FArchive& operator<<( FArchive& Ar, FObjectExport& E )
 
 FObjectImport::FObjectImport()
 :	FObjectResource	()
+#if USE_EVENT_DRIVEN_ASYNC_LOAD
+,	bImportPackageHandled(false)
+, bImportSearchedFor(false)
+#endif
+
 {
 }
 
@@ -133,6 +162,10 @@ FObjectImport::FObjectImport( UObject* InObject )
 ,	XObject			( InObject																)
 ,	SourceLinker	( NULL																	)
 ,	SourceIndex		( INDEX_NONE															)
+#if USE_EVENT_DRIVEN_ASYNC_LOAD
+, bImportPackageHandled(false)
+, bImportSearchedFor(false)
+#endif
 {
 }
 
@@ -143,6 +176,10 @@ FObjectImport::FObjectImport(UObject* InObject, UClass* InClass)
 ,	XObject			( InObject																)
 ,	SourceLinker	( NULL																	)
 ,	SourceIndex		( INDEX_NONE															)
+#if USE_EVENT_DRIVEN_ASYNC_LOAD
+, bImportPackageHandled(false)
+, bImportSearchedFor(false)
+#endif
 {
 }
 
diff --git a/Engine/Source/Runtime/CoreUObject/Private/UObject/PackageFileSummary.cpp b/Engine/Source/Runtime/CoreUObject/Private/UObject/PackageFileSummary.cpp
--- a/Engine/Source/Runtime/CoreUObject/Private/UObject/PackageFileSummary.cpp
+++ b/Engine/Source/Runtime/CoreUObject/Private/UObject/PackageFileSummary.cpp
@@ -72,14 +72,15 @@ FArchive& operator<<( FArchive& Ar, FPackageFileSummary& Sum )
 		 *
 		 * Lower 16 bits stores the UE3 engine version
 		 * Upper 16 bits stores the UE4/licensee version
-		 * For newer packages this is -6
+		 * For newer packages this is -7
 		 *		-2 indicates presence of enum-based custom versions
 		 *		-3 indicates guid-based custom versions
 		 *		-4 indicates removal of the UE3 version. Packages saved with this ID cannot be loaded in older engine versions 
 		 *		-5 indicates the replacement of writing out the "UE3 version" so older versions of engine can gracefully fail to open newer packages
 		 *		-6 indicates optimizations to how custom versions are being serialized
+		 *		-7 indicates the texture allocation info has been removed from the summary
 		 */
-		const int32 CurrentLegacyFileVersion = -6;
+		const int32 CurrentLegacyFileVersion = -7;
 		int32 LegacyFileVersion = CurrentLegacyFileVersion;
 		Ar << LegacyFileVersion;
 
@@ -249,8 +250,12 @@ FArchive& operator<<( FArchive& Ar, FPackageFileSummary& Sum )
 		Ar << Sum.AdditionalPackagesToCook;
 
 #if WITH_ENGINE
-		//@todo legacy
-		Ar << Sum.TextureAllocations;
+		if (LegacyFileVersion > -7)
+		{
+			int32 NumTextureAllocations = 0;
+			Ar << NumTextureAllocations;
+			check(NumTextureAllocations == 0);
+		}
 #else
 		check(!"this can't serialize successfully");
 #endif		// WITH_ENGINE
@@ -282,6 +287,15 @@ FArchive& operator<<( FArchive& Ar, FPackageFileSummary& Sum )
 				}
 			}
 		}
+		if (Ar.IsSaving() || Sum.FileVersionUE4 >= VER_UE4_PRELOAD_DEPENDENCIES_IN_COOKED_EXPORTS)
+		{
+			Ar << Sum.PreloadDependencyCount << Sum.PreloadDependencyOffset;
+		}
+		else
+		{
+			Sum.PreloadDependencyCount = -1;
+			Sum.PreloadDependencyOffset = 0;
+		}
 	}
 
 	return Ar;
diff --git a/Engine/Source/Runtime/CoreUObject/Private/UObject/Property.cpp b/Engine/Source/Runtime/CoreUObject/Private/UObject/Property.cpp
--- a/Engine/Source/Runtime/CoreUObject/Private/UObject/Property.cpp
+++ b/Engine/Source/Runtime/CoreUObject/Private/UObject/Property.cpp
@@ -358,6 +358,26 @@ void UProperty::Serialize( FArchive& Ar )
 	}
 }
 
+void UProperty::CopySingleValueToScriptVM( void* Dest, void const* Src ) const
+{
+	CopySingleValue(Dest, Src);
+}
+
+void UProperty::CopyCompleteValueToScriptVM( void* Dest, void const* Src ) const
+{
+	CopyCompleteValue(Dest, Src);
+}
+
+void UProperty::CopySingleValueFromScriptVM( void* Dest, void const* Src ) const
+{
+	CopySingleValue(Dest, Src);
+}
+
+void UProperty::CopyCompleteValueFromScriptVM( void* Dest, void const* Src ) const
+{
+	CopyCompleteValue(Dest, Src);
+}
+
 void UProperty::ClearValueInternal( void* Data ) const
 {
 	checkf(0, TEXT("%s failed to handle ClearValueInternal, but it was not CPF_NoDestructor | CPF_ZeroConstructor"), *GetFullName());
@@ -415,6 +435,12 @@ FString UProperty::GetCPPMacroType( FString& ExtendedTypeText ) const
 	return TEXT("PROPERTY");
 }
 
+bool UProperty::PassCPPArgsByRef() const
+{
+	return false;
+}
+
+
 void UProperty::ExportCppDeclaration(FOutputDevice& Out, EExportedDeclaration::Type DeclarationType, const TCHAR* ArrayDimOverride, uint32 AdditionalExportCPPFlags
 	, bool bSkipParameterName, const FString* ActualCppType, const FString* ActualExtendedType, const FString* ActualParameterName) const
 {
@@ -662,6 +688,16 @@ FName UProperty::GetID() const
 	return GetClass()->GetFName();
 }
 
+void UProperty::InstanceSubobjects( void* Data, void const* DefaultData, UObject* Owner, struct FObjectInstancingGraph* InstanceGraph )
+{
+}
+
+int32 UProperty::GetMinAlignment() const
+{
+	return 1;
+}
+
+
 //
 // Link property loaded from file.
 //
@@ -670,6 +706,12 @@ void UProperty::LinkInternal(FArchive& Ar)
 	check(0); // Link shouldn't call super...and we should never link an abstract property, like this base class
 }
 
+bool UProperty::ConvertFromType(const FPropertyTag& Tag, FArchive& Ar, uint8* Data, UStruct* DefaultsStruct, bool& bOutAdvanceProperty)
+{
+	return false;
+}
+
+
 int32 UProperty::SetupOffset()
 {
 	Offset_Internal = Align((GetOuter()->GetClass()->ClassCastFlags & CASTCLASS_UStruct) ? ((UStruct*)GetOuter())->GetPropertiesSize() : 0, GetMinAlignment());
@@ -844,7 +886,7 @@ const TCHAR* UProperty::ImportSingleProperty( const TCHAR* Str, void* DestData,
 
 		if (Property == NULL)
 		{
-			Warn->Logf(ELogVerbosity::Warning, TEXT("Unknown property in %s: %s "), *ObjectStruct->GetName(), Start);
+			UE_SUPPRESS(LogExec, Verbose, Warn->Logf(TEXT("Unknown property in %s: %s "), *ObjectStruct->GetName(), Start));
 			return Str;
 		}
 
@@ -1194,123 +1236,22 @@ uint32 UProperty::GetValueTypeHash(const void* Src) const
 	return GetValueTypeHashInternal(Src);
 }
 
-
-IMPLEMENT_CORE_INTRINSIC_CLASS(UProperty, UField,
-	{
-	}
-);
-
-uint8 UNumericProperty::ReadEnumAsUint8(FArchive& Ar, UStruct* DefaultsStruct, const FPropertyTag& Tag)
+void UProperty::CopyValuesInternal( void* Dest, void const* Src, int32 Count  ) const
 {
-	//@warning: mirrors loading code in UByteProperty::SerializeItem()
-	FName EnumName;
-	Ar << EnumName;
-
-	UEnum* Enum = FindField<UEnum>(dynamic_cast<UClass*>(DefaultsStruct) ? static_cast<UClass*>(DefaultsStruct) : DefaultsStruct->GetTypedOuter<UClass>(), Tag.EnumName);
-	if (!Enum)
-	{
-		Enum = FindObject<UEnum>(ANY_PACKAGE, *Tag.EnumName.ToString(), true);
-	}
-
-	if (!Enum)
-	{
-		UE_LOG(LogClass, Warning, TEXT("Failed to find enum '%s' when converting property '%s' during property loading - setting to 0"), *Tag.EnumName.ToString(), *Tag.Name.ToString());
-		return 0;
-	}
-
-	Ar.Preload(Enum);
-
-	uint8 Result = Enum->GetValueByName(EnumName);
-	if (!Enum->IsValidEnumValue(Result))
-	{
-		UE_LOG(
-			LogClass,
-			Warning,
-			TEXT("Failed to find valid enum value '%s' for enum type '%s' when converting property '%s' during property loading - setting to '%s'"),
-			*EnumName.ToString(),
-			*Enum->GetName(),
-			*Tag.Name.ToString(),
-			*Enum->GetNameByValue(Enum->GetMaxEnumValue()).ToString()
-			);
-
-		return Enum->GetMaxEnumValue();
-	}
-
-	return Result;
-};
+	check(0); // if you are not memcpyable, then you need to deal with the virtual call
+}
 
-const TCHAR* UNumericProperty::ImportText_Internal( const TCHAR* Buffer, void* Data, int32 PortFlags, UObject* Parent, FOutputDevice* ErrorText ) const
+uint32 UProperty::GetValueTypeHashInternal(const void* Src) const
 {
-	if ( Buffer != NULL )
-	{
-		const TCHAR* Start = Buffer;
-		if (IsInteger())
-		{
-			if (FChar::IsAlpha(*Buffer))
-			{
-				int32 EnumValue = UEnum::ParseEnum(Buffer);
-				if (EnumValue != INDEX_NONE)
-				{
-					SetIntPropertyValue(Data, int64(EnumValue));
-					return Buffer;
-				}
-				else
-				{
-					return NULL;
-				}
-			}
-			else
-			{
-				if ( !FCString::Strnicmp(Start,TEXT("0x"),2) )
-				{
-					Buffer+=2;
-					while ( Buffer && (FParse::HexDigit(*Buffer) != 0 || *Buffer == TCHAR('0')) )
-					{
-						Buffer++;
-					}
-				}
-				else
-				{
-					while ( Buffer && (*Buffer == TCHAR('-') || *Buffer == TCHAR('+')) )
-					{
-						Buffer++;
-					}
+	check(false); // you need to deal with the virtual call
+	return 0;
+}
 
-					while ( Buffer &&  FChar::IsDigit(*Buffer) )
-					{
-						Buffer++;
-					}
-				}
 
-				if (Start == Buffer)
-				{
-					// import failure
-					return NULL;
-				}
-			}
-		}
-		else
-		{
-			check(IsFloatingPoint());
-			// floating point
-			while( *Buffer == TCHAR('+') || *Buffer == TCHAR('-') || *Buffer == TCHAR('.') || (*Buffer >= TCHAR('0') && *Buffer <= TCHAR('9')) )
-			{
-				Buffer++;
-			}
-			if ( *Buffer == TCHAR('f') || *Buffer == TCHAR('F') )
-			{
-				Buffer++;
-			}
-		}
-		SetNumericPropertyValueFromString(Data, Start);
+IMPLEMENT_CORE_INTRINSIC_CLASS(UProperty, UField,
+	{
 	}
-	return Buffer;
-}
-
-void UNumericProperty::ExportTextItem( FString& ValueStr, const void* PropertyValue, const void* DefaultValue, UObject* Parent, int32 PortFlags, UObject* ExportRootScope ) const
-{
-	ValueStr += GetNumericPropertyValueToString(PropertyValue);
-}
+);
 
 void UFloatProperty::ExportTextItem(FString& ValueStr, const void* PropertyValue, const void* DefaultValue, UObject* Parent, int32 PortFlags, UObject* ExportRootScope) const
 {
@@ -1323,11 +1264,6 @@ void UFloatProperty::ExportTextItem(FString& ValueStr, const void* PropertyValue
 }
 
 
-IMPLEMENT_CORE_INTRINSIC_CLASS(UNumericProperty, UProperty,
-{
-}
-);
-
 UProperty* UStruct::FindPropertyByName(FName InName) const
 {
 	for (UProperty* Property = PropertyLink; Property != NULL; Property = Property->PropertyLinkNext)
diff --git a/Engine/Source/Runtime/CoreUObject/Private/UObject/PropertyArray.cpp b/Engine/Source/Runtime/CoreUObject/Private/UObject/PropertyArray.cpp
--- a/Engine/Source/Runtime/CoreUObject/Private/UObject/PropertyArray.cpp
+++ b/Engine/Source/Runtime/CoreUObject/Private/UObject/PropertyArray.cpp
@@ -8,6 +8,11 @@
 	UArrayProperty.
 -----------------------------------------------------------------------------*/
 
+void UArrayProperty::GetPreloadDependencies(TArray<UObject*>& OutDeps)
+{
+	Super::GetPreloadDependencies(OutDeps);
+	OutDeps.Add(Inner);
+}
 
 void UArrayProperty::LinkInternal(FArchive& Ar)
 {
diff --git a/Engine/Source/Runtime/CoreUObject/Private/UObject/PropertyAssetObject.cpp b/Engine/Source/Runtime/CoreUObject/Private/UObject/PropertyAssetObject.cpp
--- a/Engine/Source/Runtime/CoreUObject/Private/UObject/PropertyAssetObject.cpp
+++ b/Engine/Source/Runtime/CoreUObject/Private/UObject/PropertyAssetObject.cpp
@@ -165,6 +165,46 @@ bool UAssetObjectProperty::ConvertFromType(const FPropertyTag& Tag, FArchive& Ar
 	return false;
 }
 
+UObject* UAssetObjectProperty::GetObjectPropertyValue(const void* PropertyValueAddress) const
+{
+	return GetPropertyValue(PropertyValueAddress).Get();
+}
+
+void UAssetObjectProperty::SetObjectPropertyValue(void* PropertyValueAddress, UObject* Value) const
+{
+	SetPropertyValue(PropertyValueAddress, TCppType(Value));
+}
+
+bool UAssetObjectProperty::AllowCrossLevel() const
+{
+	return true;
+}
+
+uint32 UAssetObjectProperty::GetValueTypeHashInternal(const void* Src) const
+{
+	return GetTypeHash(GetPropertyValue(Src));
+}
+
+void UAssetObjectProperty::CopySingleValueToScriptVM(void* Dest, void const* Src) const
+{
+	CopySingleValue(Dest, Src);
+}
+
+void UAssetObjectProperty::CopyCompleteValueToScriptVM(void* Dest, void const* Src) const
+{
+	CopyCompleteValue(Dest, Src);
+}
+
+void UAssetObjectProperty::CopySingleValueFromScriptVM(void* Dest, void const* Src) const
+{
+	CopySingleValue(Dest, Src);
+}
+
+void UAssetObjectProperty::CopyCompleteValueFromScriptVM(void* Dest, void const* Src) const
+{
+	CopyCompleteValue(Dest, Src);
+}
+
 IMPLEMENT_CORE_INTRINSIC_CLASS(UAssetObjectProperty, UObjectPropertyBase,
 	{
 	}
diff --git a/Engine/Source/Runtime/CoreUObject/Private/UObject/PropertyBaseObject.cpp b/Engine/Source/Runtime/CoreUObject/Private/UObject/PropertyBaseObject.cpp
--- a/Engine/Source/Runtime/CoreUObject/Private/UObject/PropertyBaseObject.cpp
+++ b/Engine/Source/Runtime/CoreUObject/Private/UObject/PropertyBaseObject.cpp
@@ -456,6 +456,22 @@ FName UObjectPropertyBase::GetID() const
 	return NAME_ObjectProperty;
 }
 
+UObject* UObjectPropertyBase::GetObjectPropertyValue(const void* PropertyValueAddress) const
+{
+	check(0);
+	return NULL;
+}
+
+void UObjectPropertyBase::SetObjectPropertyValue(void* PropertyValueAddress, UObject* Value) const
+{
+	check(0);
+}
+
+bool UObjectPropertyBase::AllowCrossLevel() const
+{
+	return false;
+}
+
 void UObjectPropertyBase::CheckValidObject(void* Value) const
 {
 	UObject *Object = GetObjectPropertyValue(Value);
@@ -507,6 +523,33 @@ bool UObjectPropertyBase::SameType(const UProperty* Other) const
 	return Super::SameType(Other) && (PropertyClass == ((UObjectPropertyBase*)Other)->PropertyClass);
 }
 
+void UObjectPropertyBase::CopySingleValueToScriptVM( void* Dest, void const* Src ) const
+{
+	*(UObject**)Dest = GetObjectPropertyValue(Src);
+}
+
+void UObjectPropertyBase::CopyCompleteValueToScriptVM( void* Dest, void const* Src ) const
+{
+	for (int32 Index = 0; Index < ArrayDim; Index++)
+	{
+		((UObject**)Dest)[Index] = GetObjectPropertyValue(((uint8*)Src) + Index * ElementSize);
+	}
+}
+
+void UObjectPropertyBase::CopySingleValueFromScriptVM( void* Dest, void const* Src ) const
+{
+	SetObjectPropertyValue(Dest, *(UObject**)Src);
+}
+
+void UObjectPropertyBase::CopyCompleteValueFromScriptVM( void* Dest, void const* Src ) const
+{
+	checkSlow(ElementSize == sizeof(UObject*)); // the idea that script pointers are the same size as weak pointers is maybe required, maybe not
+	for (int32 Index = 0; Index < ArrayDim; Index++)
+	{
+		SetObjectPropertyValue(((uint8*)Dest) + Index * ElementSize, ((UObject**)Src)[Index]);
+	}
+}
+
 IMPLEMENT_CORE_INTRINSIC_CLASS(UObjectPropertyBase, UProperty,
 	{
 		Class->EmitObjectReference(STRUCT_OFFSET(UObjectProperty, PropertyClass), TEXT("PropertyClass"));
diff --git a/Engine/Source/Runtime/CoreUObject/Private/UObject/PropertyBool.cpp b/Engine/Source/Runtime/CoreUObject/Private/UObject/PropertyBool.cpp
--- a/Engine/Source/Runtime/CoreUObject/Private/UObject/PropertyBool.cpp
+++ b/Engine/Source/Runtime/CoreUObject/Private/UObject/PropertyBool.cpp
@@ -19,7 +19,7 @@ UBoolProperty::UBoolProperty( const FObjectInitializer& ObjectInitializer )
 }
 
 UBoolProperty::UBoolProperty(ECppProperty, int32 InOffset, uint64 InFlags, uint32 InBitMask, uint32 InElementSize, bool bIsNativeBool)
-	: UProperty(FObjectInitializer::Get(), EC_CppProperty, InOffset, InFlags)
+	: UProperty(FObjectInitializer::Get(), EC_CppProperty, InOffset, InFlags | CPF_HasGetValueTypeHash)
 	, FieldSize(0)
 	, ByteOffset(0)
 	, ByteMask(1)
@@ -29,7 +29,7 @@ UBoolProperty::UBoolProperty(ECppProperty, int32 InOffset, uint64 InFlags, uint3
 }
 
 UBoolProperty::UBoolProperty( const FObjectInitializer& ObjectInitializer, ECppProperty, int32 InOffset, uint64 InFlags, uint32 InBitMask, uint32 InElementSize, bool bIsNativeBool )
-: UProperty( ObjectInitializer, EC_CppProperty, InOffset, InFlags )
+: UProperty( ObjectInitializer, EC_CppProperty, InOffset, InFlags | CPF_HasGetValueTypeHash)
 , FieldSize(0)
 , ByteOffset(0)
 , ByteMask(1)
@@ -355,6 +355,11 @@ void UBoolProperty::InitializeValueInternal( void* Data ) const
 	*ByteValue &= ~FieldMask;
 }
 
+uint32 UBoolProperty::GetValueTypeHashInternal(const void* Src) const
+{
+	return GetTypeHash(*(const bool*)Src);
+}
+
 IMPLEMENT_CORE_INTRINSIC_CLASS(UBoolProperty, UProperty,
 	{
 	}
diff --git a/Engine/Source/Runtime/CoreUObject/Private/UObject/PropertyByte.cpp b/Engine/Source/Runtime/CoreUObject/Private/UObject/PropertyByte.cpp
--- a/Engine/Source/Runtime/CoreUObject/Private/UObject/PropertyByte.cpp
+++ b/Engine/Source/Runtime/CoreUObject/Private/UObject/PropertyByte.cpp
@@ -6,6 +6,13 @@
 /*-----------------------------------------------------------------------------
 	UByteProperty.
 -----------------------------------------------------------------------------*/
+
+void UByteProperty::GetPreloadDependencies(TArray<UObject*>& OutDeps)
+{
+	Super::GetPreloadDependencies(OutDeps);
+	OutDeps.Add(Enum);
+}
+
 void UByteProperty::SerializeItem( FArchive& Ar, void* Value, void const* Defaults ) const
 {
 	if(Enum && Ar.UseToResolveEnumerators())
@@ -361,6 +368,11 @@ const TCHAR* UByteProperty::ImportText_Internal( const TCHAR* InBuffer, void* Da
 	return Super::ImportText_Internal( InBuffer, Data, PortFlags, Parent, ErrorText );
 }
 
+UEnum* UByteProperty::GetIntPropertyEnum() const
+{
+	return Enum;
+}
+
 IMPLEMENT_CORE_INTRINSIC_CLASS(UByteProperty, UNumericProperty,
 	{
 		Class->EmitObjectReference(STRUCT_OFFSET(UByteProperty, Enum), TEXT("Enum"));
diff --git a/Engine/Source/Runtime/CoreUObject/Private/UObject/PropertyInterface.cpp b/Engine/Source/Runtime/CoreUObject/Private/UObject/PropertyInterface.cpp
--- a/Engine/Source/Runtime/CoreUObject/Private/UObject/PropertyInterface.cpp
+++ b/Engine/Source/Runtime/CoreUObject/Private/UObject/PropertyInterface.cpp
@@ -215,7 +215,7 @@ const TCHAR* UInterfaceProperty::ImportText_Internal( const TCHAR* InBuffer, voi
 	return Buffer;
 }
 
-bool UInterfaceProperty::ContainsObjectReference() const
+bool UInterfaceProperty::ContainsObjectReference(TArray<const UStructProperty*>& EncounteredStructProps) const
 {
 	return true; 
 }
diff --git a/Engine/Source/Runtime/CoreUObject/Private/UObject/PropertyLazyObjectPtr.cpp b/Engine/Source/Runtime/CoreUObject/Private/UObject/PropertyLazyObjectPtr.cpp
--- a/Engine/Source/Runtime/CoreUObject/Private/UObject/PropertyLazyObjectPtr.cpp
+++ b/Engine/Source/Runtime/CoreUObject/Private/UObject/PropertyLazyObjectPtr.cpp
@@ -75,6 +75,26 @@ bool ULazyObjectProperty::Identical( const void* A, const void* B, uint32 PortFl
 	return bResult;
 }
 
+UObject* ULazyObjectProperty::GetObjectPropertyValue(const void* PropertyValueAddress) const
+{
+	return GetPropertyValue(PropertyValueAddress).Get();
+}
+
+void ULazyObjectProperty::SetObjectPropertyValue(void* PropertyValueAddress, UObject* Value) const
+{
+	SetPropertyValue(PropertyValueAddress, TCppType(Value));
+}
+
+bool ULazyObjectProperty::AllowCrossLevel() const
+{
+	return true;
+}
+
+uint32 ULazyObjectProperty::GetValueTypeHashInternal(const void* Src) const
+{
+	return GetTypeHash(GetPropertyValue(Src));
+}
+
 IMPLEMENT_CORE_INTRINSIC_CLASS(ULazyObjectProperty, UObjectPropertyBase,
 	{
 	}
diff --git a/Engine/Source/Runtime/CoreUObject/Private/UObject/PropertyMap.cpp b/Engine/Source/Runtime/CoreUObject/Private/UObject/PropertyMap.cpp
--- a/Engine/Source/Runtime/CoreUObject/Private/UObject/PropertyMap.cpp
+++ b/Engine/Source/Runtime/CoreUObject/Private/UObject/PropertyMap.cpp
@@ -226,6 +226,13 @@ bool UMapProperty::Identical(const void* A, const void* B, uint32 PortFlags) con
 	return UE4MapProperty_Private::IsPermutation(MapHelperA, MapHelperB, PortFlags);
 }
 
+void UMapProperty::GetPreloadDependencies(TArray<UObject*>& OutDeps)
+{
+	Super::GetPreloadDependencies(OutDeps);
+	OutDeps.Add(KeyProp);
+	OutDeps.Add(ValueProp);
+}
+
 void UMapProperty::SerializeItem(FArchive& Ar, void* Value, const void* Defaults) const
 {
 	checkSlow(KeyProp);
diff --git a/Engine/Source/Runtime/CoreUObject/Private/UObject/PropertyName.cpp b/Engine/Source/Runtime/CoreUObject/Private/UObject/PropertyName.cpp
--- a/Engine/Source/Runtime/CoreUObject/Private/UObject/PropertyName.cpp
+++ b/Engine/Source/Runtime/CoreUObject/Private/UObject/PropertyName.cpp
@@ -72,6 +72,16 @@ bool UNameProperty::ConvertFromType(const FPropertyTag& Tag, FArchive& Ar, uint8
 	return bOutAdvanceProperty;
 }
 
+FString UNameProperty::GetCPPTypeForwardDeclaration() const
+{
+	return FString();
+}
+
+uint32 UNameProperty::GetValueTypeHashInternal(const void* Src) const
+{
+	return GetTypeHash(*(const FName*)Src);
+}
+
 IMPLEMENT_CORE_INTRINSIC_CLASS(UNameProperty, UProperty,
 	{
 	}
diff --git a/Engine/Source/Runtime/CoreUObject/Private/UObject/PropertyNumeric.cpp b/Engine/Source/Runtime/CoreUObject/Private/UObject/PropertyNumeric.cpp
new file mode 100644
--- /dev/null
+++ b/Engine/Source/Runtime/CoreUObject/Private/UObject/PropertyNumeric.cpp
@@ -0,0 +1,222 @@
+// Copyright 1998-2016 Epic Games, Inc. All Rights Reserved.
+
+#include "CoreUObjectPrivate.h"
+#include "PropertyTag.h"
+
+uint8 UNumericProperty::ReadEnumAsUint8(FArchive& Ar, UStruct* DefaultsStruct, const FPropertyTag& Tag)
+{
+	//@warning: mirrors loading code in UByteProperty::SerializeItem()
+	FName EnumName;
+	Ar << EnumName;
+
+	UEnum* Enum = FindField<UEnum>(dynamic_cast<UClass*>(DefaultsStruct) ? static_cast<UClass*>(DefaultsStruct) : DefaultsStruct->GetTypedOuter<UClass>(), Tag.EnumName);
+	if (!Enum)
+	{
+		Enum = FindObject<UEnum>(ANY_PACKAGE, *Tag.EnumName.ToString(), true);
+	}
+
+	if (!Enum)
+	{
+		UE_LOG(LogClass, Warning, TEXT("Failed to find enum '%s' when converting property '%s' during property loading - setting to 0"), *Tag.EnumName.ToString(), *Tag.Name.ToString());
+		return 0;
+	}
+
+	Ar.Preload(Enum);
+
+	uint8 Result = Enum->GetValueByName(EnumName);
+	if (!Enum->IsValidEnumValue(Result))
+	{
+		UE_LOG(
+			LogClass,
+			Warning,
+			TEXT("Failed to find valid enum value '%s' for enum type '%s' when converting property '%s' during property loading - setting to '%s'"),
+			*EnumName.ToString(),
+			*Enum->GetName(),
+			*Tag.Name.ToString(),
+			*Enum->GetNameByValue(Enum->GetMaxEnumValue()).ToString()
+			);
+
+		return Enum->GetMaxEnumValue();
+	}
+
+	return Result;
+};
+
+const TCHAR* UNumericProperty::ImportText_Internal( const TCHAR* Buffer, void* Data, int32 PortFlags, UObject* Parent, FOutputDevice* ErrorText ) const
+{
+	if ( Buffer != NULL )
+	{
+		const TCHAR* Start = Buffer;
+		if (IsInteger())
+		{
+			if (FChar::IsAlpha(*Buffer))
+			{
+				int32 EnumValue = UEnum::ParseEnum(Buffer);
+				if (EnumValue != INDEX_NONE)
+				{
+					SetIntPropertyValue(Data, int64(EnumValue));
+					return Buffer;
+				}
+				else
+				{
+					return NULL;
+				}
+			}
+			else
+			{
+				if ( !FCString::Strnicmp(Start,TEXT("0x"),2) )
+				{
+					Buffer+=2;
+					while ( Buffer && (FParse::HexDigit(*Buffer) != 0 || *Buffer == TCHAR('0')) )
+					{
+						Buffer++;
+					}
+				}
+				else
+				{
+					while ( Buffer && (*Buffer == TCHAR('-') || *Buffer == TCHAR('+')) )
+					{
+						Buffer++;
+					}
+
+					while ( Buffer &&  FChar::IsDigit(*Buffer) )
+					{
+						Buffer++;
+					}
+				}
+
+				if (Start == Buffer)
+				{
+					// import failure
+					return NULL;
+				}
+			}
+		}
+		else
+		{
+			check(IsFloatingPoint());
+			// floating point
+			while( *Buffer == TCHAR('+') || *Buffer == TCHAR('-') || *Buffer == TCHAR('.') || (*Buffer >= TCHAR('0') && *Buffer <= TCHAR('9')) )
+			{
+				Buffer++;
+			}
+			if ( *Buffer == TCHAR('f') || *Buffer == TCHAR('F') )
+			{
+				Buffer++;
+			}
+		}
+		SetNumericPropertyValueFromString(Data, Start);
+	}
+	return Buffer;
+}
+
+void UNumericProperty::ExportTextItem( FString& ValueStr, const void* PropertyValue, const void* DefaultValue, UObject* Parent, int32 PortFlags, UObject* ExportRootScope ) const
+{
+	ValueStr += GetNumericPropertyValueToString(PropertyValue);
+}
+
+bool UNumericProperty::IsFloatingPoint() const
+{
+	return false;
+}
+
+bool UNumericProperty::IsInteger() const
+{
+	return true;
+}
+
+UEnum* UNumericProperty::GetIntPropertyEnum() const
+{
+	return nullptr;
+}
+
+/** 
+	* Set the value of an unsigned integral property type
+	* @param Data - pointer to property data to set
+	* @param Value - Value to set data to
+**/
+void UNumericProperty::SetIntPropertyValue(void* Data, uint64 Value) const
+{
+	check(0);
+}
+
+/** 
+	* Set the value of a signed integral property type
+	* @param Data - pointer to property data to set
+	* @param Value - Value to set data to
+**/
+void UNumericProperty::SetIntPropertyValue(void* Data, int64 Value) const
+{
+	check(0);
+}
+
+/** 
+	* Set the value of a floating point property type
+	* @param Data - pointer to property data to set
+	* @param Value - Value to set data to
+**/
+void UNumericProperty::SetFloatingPointPropertyValue(void* Data, double Value) const
+{
+	check(0);
+}
+
+/** 
+	* Set the value of any numeric type from a string point
+	* @param Data - pointer to property data to set
+	* @param Value - Value (as a string) to set 
+	* CAUTION: This routine does not do enum name conversion
+**/
+void UNumericProperty::SetNumericPropertyValueFromString(void* Data, TCHAR const* Value) const
+{
+	check(0);
+}
+
+/** 
+	* Gets the value of a signed integral property type
+	* @param Data - pointer to property data to get
+	* @return Data as a signed int
+**/
+int64 UNumericProperty::GetSignedIntPropertyValue(void const* Data) const
+{
+	check(0);
+	return 0;
+}
+
+/** 
+	* Gets the value of an unsigned integral property type
+	* @param Data - pointer to property data to get
+	* @return Data as an unsigned int
+**/
+uint64 UNumericProperty::GetUnsignedIntPropertyValue(void const* Data) const
+{
+	check(0);
+	return 0;
+}
+
+/** 
+	* Gets the value of an floating point property type
+	* @param Data - pointer to property data to get
+	* @return Data as a double
+**/
+double UNumericProperty::GetFloatingPointPropertyValue(void const* Data) const
+{
+	check(0);
+	return 0.0;
+}
+
+/** 
+	* Get the value of any numeric type and return it as a string
+	* @param Data - pointer to property data to get
+	* @return Data as a string
+	* CAUTION: This routine does not do enum name conversion
+**/
+FString UNumericProperty::GetNumericPropertyValueToString(void const* Data) const
+{
+	check(0);
+	return FString();
+}
+
+IMPLEMENT_CORE_INTRINSIC_CLASS(UNumericProperty, UProperty,
+{
+}
+);
diff --git a/Engine/Source/Runtime/CoreUObject/Private/UObject/PropertyObject.cpp b/Engine/Source/Runtime/CoreUObject/Private/UObject/PropertyObject.cpp
--- a/Engine/Source/Runtime/CoreUObject/Private/UObject/PropertyObject.cpp
+++ b/Engine/Source/Runtime/CoreUObject/Private/UObject/PropertyObject.cpp
@@ -111,6 +111,16 @@ const TCHAR* UObjectProperty::ImportText_Internal(const TCHAR* Buffer, void* Dat
 	return Result;
 }
 
+uint32 UObjectProperty::GetValueTypeHashInternal(const void* Src) const
+{
+	return GetTypeHash(GetPropertyValue(Src));
+}
+
+UObject* UObjectProperty::GetObjectPropertyValue(const void* PropertyValueAddress) const
+{
+	return GetPropertyValue(PropertyValueAddress);
+}
+
 void UObjectProperty::SetObjectPropertyValue(void* PropertyValueAddress, UObject* Value) const
 {
 	SetPropertyValue(PropertyValueAddress, Value);
diff --git a/Engine/Source/Runtime/CoreUObject/Private/UObject/PropertySet.cpp b/Engine/Source/Runtime/CoreUObject/Private/UObject/PropertySet.cpp
--- a/Engine/Source/Runtime/CoreUObject/Private/UObject/PropertySet.cpp
+++ b/Engine/Source/Runtime/CoreUObject/Private/UObject/PropertySet.cpp
@@ -202,6 +202,12 @@ bool USetProperty::Identical(const void* A, const void* B, uint32 PortFlags) con
 	return UE4SetProperty_Private::IsPermutation(SetHelperA, SetHelperB, PortFlags);
 }
 
+void USetProperty::GetPreloadDependencies(TArray<UObject*>& OutDeps)
+{
+	Super::GetPreloadDependencies(OutDeps);
+	OutDeps.Add(ElementProp);
+}
+
 void USetProperty::SerializeItem(FArchive& Ar, void* Value, const void* Defaults) const
 {
 	checkSlow(ElementProp);
diff --git a/Engine/Source/Runtime/CoreUObject/Private/UObject/PropertyStr.cpp b/Engine/Source/Runtime/CoreUObject/Private/UObject/PropertyStr.cpp
--- a/Engine/Source/Runtime/CoreUObject/Private/UObject/PropertyStr.cpp
+++ b/Engine/Source/Runtime/CoreUObject/Private/UObject/PropertyStr.cpp
@@ -26,6 +26,11 @@ bool UStrProperty::ConvertFromType(const FPropertyTag& Tag, FArchive& Ar, uint8*
 	return bOutAdvanceProperty;
 }
 
+FString UStrProperty::GetCPPTypeForwardDeclaration() const
+{
+	return FString();
+}
+
 // Necessary to fix Compiler Error C2026 and C1091
 FString UStrProperty::ExportCppHardcodedText(const FString& InSource, const FString& Indent)
 {
@@ -122,6 +127,11 @@ const TCHAR* UStrProperty::ImportText_Internal( const TCHAR* Buffer, void* Data,
 	return Buffer;
 }
 
+uint32 UStrProperty::GetValueTypeHashInternal(const void* Src) const
+{
+	return GetTypeHash(*(const FString*)Src);
+}
+
 IMPLEMENT_CORE_INTRINSIC_CLASS(UStrProperty, UProperty,
 	{
 	}
diff --git a/Engine/Source/Runtime/CoreUObject/Private/UObject/PropertyStruct.cpp b/Engine/Source/Runtime/CoreUObject/Private/UObject/PropertyStruct.cpp
--- a/Engine/Source/Runtime/CoreUObject/Private/UObject/PropertyStruct.cpp
+++ b/Engine/Source/Runtime/CoreUObject/Private/UObject/PropertyStruct.cpp
@@ -136,6 +136,12 @@ bool UStructProperty::NetSerializeItem( FArchive& Ar, UPackageMap* Map, void* Da
 	return 1;
 }
 
+void UStructProperty::GetPreloadDependencies(TArray<UObject*>& OutDeps)
+{
+	Super::GetPreloadDependencies(OutDeps);
+	OutDeps.Add(Struct);
+}
+
 void UStructProperty::Serialize( FArchive& Ar )
 {
 	Super::Serialize( Ar );
diff --git a/Engine/Source/Runtime/CoreUObject/Private/UObject/PropertyWeakObjectPtr.cpp b/Engine/Source/Runtime/CoreUObject/Private/UObject/PropertyWeakObjectPtr.cpp
--- a/Engine/Source/Runtime/CoreUObject/Private/UObject/PropertyWeakObjectPtr.cpp
+++ b/Engine/Source/Runtime/CoreUObject/Private/UObject/PropertyWeakObjectPtr.cpp
@@ -41,6 +41,16 @@ void UWeakObjectProperty::SerializeItem( FArchive& Ar, void* Value, void const*
 	}
 }
 
+UObject* UWeakObjectProperty::GetObjectPropertyValue(const void* PropertyValueAddress) const
+{
+	return GetPropertyValue(PropertyValueAddress).Get();
+}
+
+void UWeakObjectProperty::SetObjectPropertyValue(void* PropertyValueAddress, UObject* Value) const
+{
+	SetPropertyValue(PropertyValueAddress, TCppType(Value));
+}
+
 IMPLEMENT_CORE_INTRINSIC_CLASS(UWeakObjectProperty, UObjectPropertyBase,
 	{
 	}
diff --git a/Engine/Source/Runtime/CoreUObject/Private/UObject/SavePackage.cpp b/Engine/Source/Runtime/CoreUObject/Private/UObject/SavePackage.cpp
--- a/Engine/Source/Runtime/CoreUObject/Private/UObject/SavePackage.cpp
+++ b/Engine/Source/Runtime/CoreUObject/Private/UObject/SavePackage.cpp
@@ -16,7 +16,6 @@ DEFINE_LOG_CATEGORY_STATIC(LogSavePackage, Log, All);
 static const int32 MAX_MERGED_COMPRESSION_CHUNKSIZE = 1024 * 1024;
 static const FName WorldClassName = FName("World");
 
-
 #define VALIDATE_INITIALIZECORECLASSES 0
 #define EXPORT_SORTING_DETAILED_LOGGING 0
 
@@ -354,7 +353,7 @@ void AsyncWriteFile(FLargeMemoryPtr Data, const int64 DataSize, const TCHAR* Fil
 		{
 			check(DataSize);
 			FString TempFilename; 
-			
+
 			if ( bUseTempFilename )
 			{
 				TempFilename = FPaths::GetBaseFilename(Filename, false);
@@ -420,6 +419,157 @@ void AsyncWriteFile(FLargeMemoryPtr Data, const int64 DataSize, const TCHAR* Fil
 	(new FAutoDeleteAsyncTask<FAsyncWriteWorker>(Filename, MoveTemp(Data), DataSize, TimeStamp, bUseTempFilename))->StartBackgroundTask();
 }
 
+#if SPLIT_COOKED_FILES
+void AsyncWriteFileWithSplitExports(FLargeMemoryPtr Data, const int64 DataSize, const int64 HeaderSize, const TCHAR* Filename, const FDateTime& TimeStamp, bool bUseTempFilename = true)
+{
+	class FAsyncWriteWorkerWithSplitExports : public FNonAbandonableTask
+	{
+	public:
+		/** Filename To write to**/
+		FString Filename;
+		/** Should we write to a temp file then move it */
+		bool bUseTempFilename;
+		/** Data for the file. Will be freed after write **/
+		FLargeMemoryPtr Data;
+		/** Size of data */
+		const int64 DataSize;
+		/** Size of data */
+		const int64 HeaderSize;
+		/** Timestamp to give the file. MinValue if shouldn't be modified */
+		FDateTime FinalTimeStamp;
+
+		/** Constructor
+		*/
+		FAsyncWriteWorkerWithSplitExports(const TCHAR* InFilename, FLargeMemoryPtr InData, const int64 InDataSize, const int64 InHeaderSize, const FDateTime& InTimeStamp, bool inbUseTempFilename)
+			: Filename(InFilename)
+			, bUseTempFilename(inbUseTempFilename)
+			, Data(MoveTemp(InData))
+			, DataSize(InDataSize)
+			, HeaderSize(InHeaderSize)
+			, FinalTimeStamp(InTimeStamp)
+		{
+		}
+
+		/** Write the file  */
+		void DoWork()
+		{
+			check(DataSize);
+			FString BaseFilename;
+			FString TempFilename;
+			FString FilenameExports;
+			FString TempFilenameExports;
+
+			BaseFilename = FPaths::GetBaseFilename(Filename, false);
+			FilenameExports = BaseFilename;
+			FilenameExports += TEXT(".uexp");
+			if (bUseTempFilename)
+			{
+				TempFilename = BaseFilename;
+				TempFilename += TEXT(".t");
+				TempFilenameExports = BaseFilename;
+				TempFilenameExports += TEXT(".e");
+			}
+			else
+			{
+				TempFilename = Filename;
+				TempFilenameExports = FilenameExports;
+			}
+
+			{
+				FArchive* Ar = IFileManager::Get().CreateFileWriter(*TempFilename);
+				if (Ar)
+				{
+					Ar->Serialize(Data.Get(), HeaderSize);
+					delete Ar;
+					Ar = nullptr;
+
+					if (IFileManager::Get().FileSize(*TempFilename) == HeaderSize)
+					{
+						if (bUseTempFilename)
+						{
+							if (!IFileManager::Get().Move(*Filename, *TempFilename, true, true, false, false))
+							{
+								UE_LOG(LogSavePackage, Fatal, TEXT("Could not move to %s."), *Filename);
+							}
+
+							// if everything worked, this is not necessary, but we will make every effort to avoid leaving junk in the cache
+							if (FPaths::FileExists(TempFilename))
+							{
+								IFileManager::Get().Delete(*TempFilename);
+							}
+						}
+
+						if (FinalTimeStamp != FDateTime::MinValue())
+						{
+							IFileManager::Get().SetTimeStamp(*Filename, FinalTimeStamp);
+						}
+					}
+					else
+					{
+						UE_LOG(LogSavePackage, Fatal, TEXT("Could not save to %s!"), *TempFilename);
+					}
+				}
+				else
+				{
+					UE_LOG(LogSavePackage, Fatal, TEXT("Could not write to %s!"), *TempFilename);
+				}
+			}
+			{
+				FArchive* Ar = IFileManager::Get().CreateFileWriter(*TempFilenameExports);
+				if (Ar)
+				{
+					Ar->Serialize(Data.Get() + HeaderSize, DataSize - HeaderSize);
+					delete Ar;
+					Ar = nullptr;
+
+					// Clean-up the memory as soon as we save the file to reduce the memory footprint.
+					Data.Reset();
+
+					if (IFileManager::Get().FileSize(*TempFilenameExports) == DataSize - HeaderSize)
+					{
+						if (bUseTempFilename)
+						{
+							if (!IFileManager::Get().Move(*FilenameExports, *TempFilenameExports, true, true, false, false))
+							{
+								UE_LOG(LogSavePackage, Fatal, TEXT("Could not move to %s."), *Filename);
+							}
+
+							// if everything worked, this is not necessary, but we will make every effort to avoid leaving junk in the cache
+							if (FPaths::FileExists(TempFilenameExports))
+							{
+								IFileManager::Get().Delete(*TempFilenameExports);
+							}
+						}
+
+						if (FinalTimeStamp != FDateTime::MinValue())
+						{
+							IFileManager::Get().SetTimeStamp(*FilenameExports, FinalTimeStamp);
+						}
+					}
+					else
+					{
+						UE_LOG(LogSavePackage, Fatal, TEXT("Could not save to %s!"), *TempFilenameExports);
+					}
+				}
+				else
+				{
+					UE_LOG(LogSavePackage, Fatal, TEXT("Could not write to %s!"), *TempFilenameExports);
+				}
+			}
+			OutstandingAsyncWrites.Decrement();
+		}
+
+		FORCEINLINE TStatId GetStatId() const
+		{
+			RETURN_QUICK_DECLARE_CYCLE_STAT(FAsyncWriteWorkerWithSplitExports, STATGROUP_ThreadPoolAsyncTasks);
+		}
+
+	};
+
+	OutstandingAsyncWrites.Increment();
+	(new FAutoDeleteAsyncTask<FAsyncWriteWorkerWithSplitExports>(Filename, MoveTemp(Data), DataSize, HeaderSize, TimeStamp, bUseTempFilename))->StartBackgroundTask();
+}
+#endif
 /** 
  * Helper object for all of the cross object state that is needed only while saving a package.
  */
@@ -626,7 +776,8 @@ public:
 	}
 
 	void ProcessBaseObject( UObject* BaseObject );
-	FArchive& operator<<( UObject*& Obj );
+	virtual FArchive& operator<<(UObject*& Obj) override;
+	virtual FArchive& operator<< (struct FWeakObjectPtr& Value) override;
 
 	/**
 	 * Package we're currently saving.  Only objects contained
@@ -662,6 +813,21 @@ FString FArchiveSaveTagExports::GetArchiveName() const
 		: TEXT("SaveTagExports");
 }
 
+FArchive& FArchiveSaveTagExports::operator<< (struct FWeakObjectPtr& Value)
+{
+	if (IsEventDrivenLoaderEnabledInCookedBuilds() && IsCooking())
+	{
+		// Always serialize weak pointers for the purposes of object tagging
+		UObject* Object = static_cast<UObject*>(Value.Get(true));
+		*this << Object;
+	}
+	else
+	{
+		Value.Serialize(*this);
+	}
+	return *this;
+}
+
 FArchive& FArchiveSaveTagExports::operator<<( UObject*& Obj )
 {
 	check(Outer);
@@ -907,6 +1073,8 @@ public:
 		return *this;
 	}
 
+	virtual FArchive& operator<< (struct FWeakObjectPtr& Value) override;
+
 	/**
 	 * Returns the name of the Archive.  Useful for getting the name of the package a struct or object
 	 * is in when a loading error occurs.
@@ -933,6 +1101,21 @@ FString FArchiveSaveTagImports::GetArchiveName() const
 	return TEXT("SaveTagImports");
 }
 
+FArchive& FArchiveSaveTagImports::operator<< (struct FWeakObjectPtr& Value)
+{
+	if (IsEventDrivenLoaderEnabledInCookedBuilds() && IsCooking())
+	{
+		// Always serialize weak pointers for the purposes of object tagging
+		UObject* Object = static_cast<UObject*>(Value.Get(true));
+		*this << Object;
+	}
+	else
+	{
+		Value.Serialize(*this);
+	}
+	return *this;
+}
+
 FArchive& FArchiveSaveTagImports::operator<<( UObject*& Obj )
 {
 	CheckObjectPriorToSave(*this, Obj, NULL);
@@ -981,6 +1164,43 @@ FArchive& FArchiveSaveTagImports::operator<<( UObject*& Obj )
 				if (!bIsEditorOnly)
 				{
 					Obj->Mark(OBJECTMARK_TagImp);
+
+					if (IsEventDrivenLoaderEnabledInCookedBuilds() && IsCooking() && !bIsNative && Cast<UClass>(Obj))
+					{
+						// we don't want to add this to Dependencies, we simply want it to be an import so that a serialization before creation dependency can be created to the CDO
+						UObject* CDO = Cast<UClass>(Obj)->GetDefaultObject();
+
+						if (CDO)
+						{
+							ConditionallyExcludeObjectForTarget(CDO, *this);
+
+							const bool bCDOIsEditorOnly = IsCookingForNoEditorDataPlatform(*this) && IsEditorOnlyObject(CDO);
+
+							if (!bCDOIsEditorOnly && !CDO->HasAnyFlags(RF_Transient))
+							{
+								CDO->Mark(OBJECTMARK_TagImp);
+
+								// and all subobjects
+								TArray<UObject*> Subobjects;
+								GetObjectsWithOuter(CDO, Subobjects);
+								for (UObject* SubObj : Subobjects)
+								{
+									if (SubObj)
+									{
+										// correctly mark components that shouldn't be cooked
+										ConditionallyExcludeObjectForTarget(SubObj, *this);
+
+										const bool bCompIsEditorOnly = IsCookingForNoEditorDataPlatform(*this) && IsEditorOnlyObject(SubObj);
+
+										if (!bCompIsEditorOnly && !SubObj->HasAnyFlags(RF_Transient))
+										{
+											SubObj->Mark(OBJECTMARK_TagImp);
+										}
+									}
+								}
+							}
+						}
+					}
 #if WITH_EDITOR
 					SavePackageState->AddReplacementsNames(Obj);
 #endif //WITH_EDITOR
@@ -1006,6 +1226,7 @@ FArchive& FArchiveSaveTagImports::operator<<( UObject*& Obj )
 					{
 						*this << Parent;
 					}
+
 				}
 			}
 		}
@@ -2630,6 +2851,7 @@ public:
 					{
 						UClass* AsClass = (UClass*)StructObject;
 						AsClass->GetRequiredPreloadDependencies(StructForceLoadObjects);
+						check(!StructForceLoadObjects.Num()); //@todoio GetRequiredPreloadDependencies is dead code, remove
 					}
 					// append rather than replace (in case we're nested in a 
 					// recursive call)... adding these to ForceLoadObjects 
@@ -3411,7 +3633,10 @@ ESavePackageResult UPackage::Save(UPackage* InOuter, UObject* Base, EObjectFlags
 		FlushAsyncLoading();
 
 		(*GFlushStreamingFunc)();
+
+#if !USE_NEW_ASYNC_IO
 		FIOSystem::Get().BlockTillAllRequestsFinishedAndFlushHandles();
+#endif
 
 		uint32 Time = 0; CLOCK_CYCLES(Time);
 
@@ -3588,11 +3813,11 @@ ESavePackageResult UPackage::Save(UPackage* InOuter, UObject* Base, EObjectFlags
 			
 				{
 					COOK_STAT(FScopedDurationTimer SaveTimer(SavePackageStats::TagPackageExportsTimeSec));
-					// Clear OBJECTMARK_TagExp again as we need to redo tagging below.
-					UnMarkAllObjects((EObjectMark)(OBJECTMARK_TagExp | OBJECTMARK_EditorOnly));
-
-					// We need to serialize objects yet again to tag objects that were created by PreSave as OBJECTMARK_TagExp.
-					PackageExportTagger.TagPackageExports(ExportTaggerArchive, false);
+				// Clear OBJECTMARK_TagExp again as we need to redo tagging below.
+				UnMarkAllObjects((EObjectMark)(OBJECTMARK_TagExp|OBJECTMARK_EditorOnly));
+			
+				// We need to serialize objects yet again to tag objects that were created by PreSave as OBJECTMARK_TagExp.
+				PackageExportTagger.TagPackageExports( ExportTaggerArchive, false );
 				}
 
 
@@ -3620,7 +3845,7 @@ ESavePackageResult UPackage::Save(UPackage* InOuter, UObject* Base, EObjectFlags
 
 				SlowTask.EnterProgressFrame();
 
-				// structure to track what ever export needs to import
+				// structure to track what every export needs to import
 				TMap<UObject*, TArray<UObject*> > ObjectDependencies;
 
 				// and a structure to track non-redirector references
@@ -3770,9 +3995,14 @@ ESavePackageResult UPackage::Save(UPackage* InOuter, UObject* Base, EObjectFlags
 						ImportTagger << Class;
 
 						UObject* Template = Obj->GetArchetype();
-						if ( Template && Template != Class->GetDefaultObject() )
+						if ( Template 
+							)
 						{
-							ImportTagger << Template;
+							// If we're not cooking for the event driven loader, exclude the CDO
+							if (Template != Class->GetDefaultObject() || (IsEventDrivenLoaderEnabledInCookedBuilds() && TargetPlatform))
+							{
+								ImportTagger << Template;
+							}
 						}
 
 						if( Obj->IsIn(GetTransientPackage()) )
@@ -3881,7 +4111,7 @@ ESavePackageResult UPackage::Save(UPackage* InOuter, UObject* Base, EObjectFlags
 							// When Startup is cooked, it will pull in C and A. When M is cooked, it will pull in B, but not A, because
 							// A was already marked by the cooker. M.xxx now has a private import to A, which is normally illegal, hence
 							// the OBJECTMARK_MarkedByCooker check below
-							if( !Obj->HasAnyFlags(RF_Public) )
+							if( !Obj->HasAnyFlags(RF_Public) && !Obj->HasAnyFlags(RF_Transient))
 							{
 								PrivateObjects.Add(Obj);
 							}
@@ -4038,7 +4268,7 @@ ESavePackageResult UPackage::Save(UPackage* InOuter, UObject* Base, EObjectFlags
 
 					// free the file handle and delete the temporary file
 					Linker->Detach();
-					IFileManager::Get().Delete( *TempFilename );
+					IFileManager::Get().Delete(*TempFilename);
 					if (!(SaveFlags & SAVE_NoError))
 					{
 						Error->Logf(ELogVerbosity::Warning, TEXT("Can't save %s: Graph is linked to external private object %s"), Filename, *CulpritString);
@@ -4282,9 +4512,9 @@ ESavePackageResult UPackage::Save(UPackage* InOuter, UObject* Base, EObjectFlags
 				// Sort exports for seek-free loading.
 				{
 					COOK_STAT(FScopedDurationTimer SaveTimer(SavePackageStats::SortExportsSeekfreeInnerTimeSec));
-					FObjectExportSeekFreeSorter SeekFreeSorter;
-					SeekFreeSorter.SortExports(Linker, Conform);
-					Linker->Summary.ExportCount = Linker->ExportMap.Num();
+				FObjectExportSeekFreeSorter SeekFreeSorter;
+				SeekFreeSorter.SortExports( Linker, Conform );
+				Linker->Summary.ExportCount = Linker->ExportMap.Num();
 				}
 
 				if ( EndSavingIfCancelled( Linker, TempFilename ) ) 
@@ -4506,13 +4736,25 @@ ESavePackageResult UPackage::Save(UPackage* InOuter, UObject* Base, EObjectFlags
 				}
 				SlowTask.EnterProgressFrame();
 
-				// save depends map (no need for later patching)
-				check(Linker->DependsMap.Num()==Linker->ExportMap.Num());
-				Linker->Summary.DependsOffset = Linker->Tell();
-				for( int32 i=0; i<Linker->ExportMap.Num(); i++ )
+				if (Linker->IsCooking())
+				{
+					//@todo optimization, this should just be stripped entirely from cooked packages
+					TArray<FPackageIndex> Depends; // empty array
+					for (int32 i = 0; i < Linker->ExportMap.Num(); i++)
+					{
+						*Linker << Depends;
+					}
+				}
+				else
 				{
-					TArray<FPackageIndex>& Depends = Linker->DependsMap[ i ];
-					*Linker << Depends;
+					// save depends map (no need for later patching)
+					check(Linker->DependsMap.Num() == Linker->ExportMap.Num());
+					Linker->Summary.DependsOffset = Linker->Tell();
+					for (int32 i = 0; i < Linker->ExportMap.Num(); i++)
+					{
+						TArray<FPackageIndex>& Depends = Linker->DependsMap[i];
+						*Linker << Depends;
+					}
 				}
 
 
@@ -4544,6 +4786,251 @@ ESavePackageResult UPackage::Save(UPackage* InOuter, UObject* Base, EObjectFlags
 
 				// Save level information used by World browser
 				UPackage::SaveWorldLevelInfo( InOuter, Linker );
+
+
+				for (int32 i = 0; i < Linker->ExportMap.Num(); i++)
+				{
+					FObjectExport& Export = Linker->ExportMap[i];
+					if (Export.Object)
+					{
+						// Set class index.
+						// If this is *exactly* a UClass, store null instead; for anything else, including UClass-derived classes, map it
+						UClass* ObjClass = Export.Object->GetClass();
+						if (ObjClass != UClass::StaticClass())
+						{
+							Export.ClassIndex = Linker->MapObject(ObjClass);
+							check(!Export.ClassIndex.IsNull());
+						}
+						else
+						{
+							Export.ClassIndex = FPackageIndex();
+						}
+
+						if (IsEventDrivenLoaderEnabledInCookedBuilds() && TargetPlatform)
+						{
+							UObject* Archetype = Export.Object->GetArchetype();
+							check(Archetype);
+							check(Archetype->IsA(Export.Object->HasAnyFlags(RF_ClassDefaultObject) ? ObjClass->GetSuperClass() : ObjClass));
+							Export.TemplateIndex = Linker->MapObject(Archetype);
+							UE_CLOG(Export.TemplateIndex.IsNull(), LogSavePackage, Fatal, TEXT("%s was an archetype of %s but returned a null index mapping the object."), *Archetype->GetFullName(), *Export.Object->GetFullName());
+							check(!Export.TemplateIndex.IsNull());
+						}
+
+						// Set the parent index, if this export represents a UStruct-derived object
+						if (UStruct* Struct = dynamic_cast<UStruct*>(Export.Object))
+						{
+							if (Struct->GetSuperStruct() != NULL)
+							{
+								Export.SuperIndex = Linker->MapObject(Struct->GetSuperStruct());
+								check(!Export.SuperIndex.IsNull());
+							}
+							else
+							{
+								Export.SuperIndex = FPackageIndex();
+							}
+						}
+						else
+						{
+							Export.SuperIndex = FPackageIndex();
+						}
+
+						// Set FPackageIndex for this export's Outer. If the export's Outer
+						// is the UPackage corresponding to this package's LinkerRoot, the
+						if (Export.Object->GetOuter() != InOuter)
+						{
+							check(Export.Object->GetOuter());
+							checkf(Export.Object->GetOuter()->IsIn(InOuter),
+								TEXT("Export Object (%s) Outer (%s) mismatch."),
+								*(Export.Object->GetPathName()),
+								*(Export.Object->GetOuter()->GetPathName()));
+							Export.OuterIndex = Linker->MapObject(Export.Object->GetOuter());
+							checkf(!Export.OuterIndex.IsImport(),
+								TEXT("Export Object (%s) Outer (%s) is an Import."),
+								*(Export.Object->GetPathName()),
+								*(Export.Object->GetOuter()->GetPathName()));
+						}
+						else
+						{
+							// this export's Outer is the LinkerRoot for this package
+							Export.OuterIndex = FPackageIndex();
+						}
+					}
+				}
+
+				Linker->Summary.PreloadDependencyOffset = Linker->Tell();
+				Linker->Summary.PreloadDependencyCount = -1;
+
+				if (Linker->IsCooking() && IsEventDrivenLoaderEnabledInCookedBuilds())
+				{
+					const EObjectMark ObjectMarks = UPackage::GetObjectMarksForTargetPlatform(Linker->CookingTarget(), Linker->IsCooking());
+					Linker->Summary.PreloadDependencyCount = 0;
+
+					auto IncludeObjectAsDependency = [Linker,ObjectMarks](TSet<FPackageIndex>& AddTo, UObject* ToTest, UObject* ForObj, bool bMandatory)
+					{
+						// Skip transient, editor only, and excluded client/server objects
+						if (ToTest && !ToTest->HasAllFlags(RF_Transient))
+						{
+							UPackage* Outermost = ToTest->GetOutermost();
+							check(Outermost);
+							if (Outermost->HasAnyPackageFlags(EPackageFlags::PKG_CompiledIn))
+							{
+								return; // we never add dependencies for things that are compiled in
+							}
+							bool bNotFiltered = !ToTest->HasAnyMarks(ObjectMarks) && (!(Linker->Summary.PackageFlags & PKG_FilterEditorOnly) || !ToTest->HasAnyMarks(OBJECTMARK_EditorOnly));
+							if (bMandatory && !bNotFiltered)
+							{
+								UE_LOG(LogSavePackage, Log, TEXT("A dependency '%s' of '%s' was filtered, but is mandatory. This indicates a problem with editor only stripping. We will keep the dependency anyway."), *ToTest->GetFullName(), *ForObj->GetFullName());
+								bNotFiltered = true;
+							}
+							if (bNotFiltered)
+							{
+								FPackageIndex Index = Linker->MapObject(ToTest);
+								if (!Index.IsNull())
+								{
+									AddTo.Add(Index);
+									return;
+								}
+								else
+								{
+									UE_LOG(LogSavePackage, Fatal, TEXT("A dependency '%s' of '%s' was not actually in the linker tables and so will be ignored."), *ToTest->GetFullName(), *ForObj->GetFullName());
+								}
+							}
+							check(!bMandatory);
+						}
+					};
+
+					auto IncludeIndexAsDependency = [Linker](TSet<FPackageIndex>& AddTo, FPackageIndex Dep)
+					{
+						if (!Dep.IsNull())
+						{
+							UObject* ToTest = Dep.IsExport() ? Linker->Exp(Dep).Object : Linker->Imp(Dep).XObject;
+							if (ToTest)
+							{
+								UPackage* Outermost = ToTest->GetOutermost();
+								if (Outermost && !Outermost->HasAnyPackageFlags(EPackageFlags::PKG_CompiledIn))
+								{
+									AddTo.Add(Dep);
+								}
+							}
+						}
+					};
+
+
+					TArray<UObject*> Subobjects;
+					for (int32 i = 0; i < Linker->ExportMap.Num(); i++)
+					{
+						FObjectExport& Export = Linker->ExportMap[i];
+						if (Export.Object)
+						{
+							TSet<FPackageIndex> SerializationBeforeCreateDependencies;
+							{
+								IncludeIndexAsDependency(SerializationBeforeCreateDependencies, Export.ClassIndex);
+								IncludeIndexAsDependency(SerializationBeforeCreateDependencies, Export.SuperIndex);
+								UObject* CDO = Export.Object->GetArchetype();
+								IncludeObjectAsDependency(SerializationBeforeCreateDependencies, CDO, Export.Object, true);
+#if 1 //@todoio delete?
+								Subobjects.Reset();
+								GetObjectsWithOuter(CDO, Subobjects);
+								for (UObject* SubObj : Subobjects)
+								{
+									IncludeObjectAsDependency(SerializationBeforeCreateDependencies, SubObj, Export.Object, false);
+								}
+#endif
+							}
+							TSet<FPackageIndex> SerializationBeforeSerializationDependencies;
+							{
+								TArray<UObject*> Deps;
+								Export.Object->GetPreloadDependencies(Deps);
+								for (UObject* Obj : Deps)
+								{
+									IncludeObjectAsDependency(SerializationBeforeSerializationDependencies, Obj, Export.Object, false);
+								}
+							}
+
+							TSet<FPackageIndex> CreateBeforeSerializationDependencies;
+							{
+								TArray<FPackageIndex>& Depends = Linker->DependsMap[i];
+								for (FPackageIndex Dep : Depends)
+								{
+									UClass* Class = Cast<UClass>(Export.Object);
+									UObject* ToTest = Dep.IsExport() ? Linker->Exp(Dep).Object : Linker->Imp(Dep).XObject;
+									if (Class && ToTest == Class->GetDefaultObject())
+									{
+										continue;
+									}
+									IncludeIndexAsDependency(CreateBeforeSerializationDependencies, Dep);
+								}
+							}
+							TSet<FPackageIndex> CreateBeforeCreateDependencies;
+							{
+								IncludeIndexAsDependency(CreateBeforeCreateDependencies, Export.OuterIndex);
+							}
+
+							for (FPackageIndex Index : SerializationBeforeSerializationDependencies)
+							{
+								if (SerializationBeforeCreateDependencies.Contains(Index))
+								{
+									continue; // if the other thing must be serialized before we create, then this is a redundant dep
+								}
+								if (Export.FirstExportDependency == -1)
+								{
+									Export.FirstExportDependency = Linker->Summary.PreloadDependencyCount;
+									check(Export.SerializationBeforeSerializationDependencies == 0 && Export.CreateBeforeSerializationDependencies == 0 && Export.SerializationBeforeCreateDependencies == 0 && Export.CreateBeforeCreateDependencies == 0);
+								}
+								Linker->Summary.PreloadDependencyCount++;
+								Export.SerializationBeforeSerializationDependencies++;
+								*Linker << Index;
+								Linker->DepListForErrorChecking.Add(Index);
+							}
+							for (FPackageIndex Index : CreateBeforeSerializationDependencies)
+							{
+								if (SerializationBeforeCreateDependencies.Contains(Index))
+								{
+									continue; // if the other thing must be serialized before we create, then this is a redundant dep
+								}
+								if (SerializationBeforeSerializationDependencies.Contains(Index))
+								{
+									continue; // if the other thing must be serialized before we serialize, then this is a redundant dep
+								}
+								if (Export.FirstExportDependency == -1)
+								{
+									Export.FirstExportDependency = Linker->Summary.PreloadDependencyCount;
+									check(Export.SerializationBeforeSerializationDependencies == 0 && Export.CreateBeforeSerializationDependencies == 0 && Export.SerializationBeforeCreateDependencies == 0 && Export.CreateBeforeCreateDependencies == 0);
+								}
+								Linker->Summary.PreloadDependencyCount++;
+								Export.CreateBeforeSerializationDependencies++;
+								*Linker << Index;
+								Linker->DepListForErrorChecking.Add(Index);
+							}
+							for (FPackageIndex Index : SerializationBeforeCreateDependencies)
+							{
+								if (Export.FirstExportDependency == -1)
+								{
+									Export.FirstExportDependency = Linker->Summary.PreloadDependencyCount;
+									check(Export.SerializationBeforeSerializationDependencies == 0 && Export.CreateBeforeSerializationDependencies == 0 && Export.SerializationBeforeCreateDependencies == 0 && Export.CreateBeforeCreateDependencies == 0);
+								}
+								Linker->Summary.PreloadDependencyCount++;
+								Export.SerializationBeforeCreateDependencies++;
+								*Linker << Index;
+								Linker->DepListForErrorChecking.Add(Index);
+							}
+							for (FPackageIndex Index : CreateBeforeCreateDependencies)
+							{
+								if (Export.FirstExportDependency == -1)
+								{
+									Export.FirstExportDependency = Linker->Summary.PreloadDependencyCount;
+									check(Export.SerializationBeforeSerializationDependencies == 0 && Export.CreateBeforeSerializationDependencies == 0 && Export.SerializationBeforeCreateDependencies == 0 && Export.CreateBeforeCreateDependencies == 0);
+								}
+								Linker->Summary.PreloadDependencyCount++;
+								Export.CreateBeforeCreateDependencies++;
+								*Linker << Index;
+								Linker->DepListForErrorChecking.Add(Index);
+							}
+						}
+					}
+					UE_LOG(LogSavePackage, Log, TEXT("Saved %d dependencies for %d exports."), Linker->Summary.PreloadDependencyCount, Linker->ExportMap.Num());
+				}
+
 				
 				Linker->Summary.TotalHeaderSize	= Linker->Tell();
 
@@ -4582,60 +5069,9 @@ ESavePackageResult UPackage::Save(UPackage* InOuter, UObject* Base, EObjectFlags
 						FObjectExport& Export = Linker->ExportMap[i];
 						if( Export.Object )
 						{
-							// Set class index.
-							// If this is *exactly* a UClass, store null instead; for anything else, including UClass-derived classes, map it
-							UClass* ObjClass = Export.Object->GetClass();
-							if (ObjClass != UClass::StaticClass())
-							{
-								Export.ClassIndex = Linker->MapObject(ObjClass);
-								check(!Export.ClassIndex.IsNull());
-							}
-							else
-							{
-								Export.ClassIndex = FPackageIndex();
-							}
-
-							// Set the parent index, if this export represents a UStruct-derived object
-							if( UStruct* Struct = dynamic_cast<UStruct*>(Export.Object) )
-							{
-								if (Struct->GetSuperStruct() != NULL)
-								{
-									Export.SuperIndex = Linker->MapObject(Struct->GetSuperStruct());
-									check(!Export.SuperIndex.IsNull());
-								}
-								else
-								{
-									Export.SuperIndex = FPackageIndex();
-								}
-							}
-							else
-							{
-								Export.SuperIndex = FPackageIndex();
-							}
-
-							// Set FPackageIndex for this export's Outer. If the export's Outer
-							// is the UPackage corresponding to this package's LinkerRoot, the
-							if( Export.Object->GetOuter() != InOuter )
-							{
-								check( Export.Object->GetOuter() );
-								checkf(Export.Object->GetOuter()->IsIn(InOuter),
-									TEXT("Export Object (%s) Outer (%s) mismatch."),
-									*(Export.Object->GetPathName()),
-									*(Export.Object->GetOuter()->GetPathName()));
-								Export.OuterIndex = Linker->MapObject(Export.Object->GetOuter());
-								checkf(!Export.OuterIndex.IsImport(), 
-									TEXT("Export Object (%s) Outer (%s) is an Import."),
-									*(Export.Object->GetPathName()),
-									*(Export.Object->GetOuter()->GetPathName()));
-							}
-							else
-							{
-								// this export's Outer is the LinkerRoot for this package
-								Export.OuterIndex = FPackageIndex();
-							}
-
 							// Save the object data.
 							Export.SerialOffset = Linker->Tell();
+							Linker->CurrentlySavingExport = FPackageIndex::FromExport(i);
 							// UE_LOG(LogSavePackage, Log, TEXT("export %s for %s"), *Export.Object->GetFullName(), *Linker->CookingTarget()->PlatformName());
 							if ( Export.Object->HasAnyFlags(RF_ClassDefaultObject) )
 							{
@@ -4650,6 +5086,7 @@ ESavePackageResult UPackage::Save(UPackage* InOuter, UObject* Base, EObjectFlags
 #endif
 								Export.Object->Serialize( *Linker );
 							}
+							Linker->CurrentlySavingExport = FPackageIndex();
 							Export.SerialSize = Linker->Tell() - Export.SerialOffset;
 
 							// Mark object as having been saved.
@@ -4702,6 +5139,13 @@ ESavePackageResult UPackage::Save(UPackage* InOuter, UObject* Base, EObjectFlags
 							{
 								bEnable = false;
 							}
+#if SPLIT_COOKED_FILES
+							if (IsEventDrivenLoaderEnabledInCookedBuilds())
+							{
+								// Always split bulk data when splitting cooked files
+								bEnable = true;
+							}
+#endif
 						}
 					} ShouldUseSeperateBulkDataFiles;
 
@@ -4711,15 +5155,14 @@ ESavePackageResult UPackage::Save(UPackage* InOuter, UObject* Base, EObjectFlags
 
 					if (bShouldUseSeparateBulkFile)
 					{
+						ExtraBulkDataFlags = BULKDATA_PayloadInSeperateFile;
 						if ( bSaveAsync )
 						{
-							ExtraBulkDataFlags = BULKDATA_PayloadInSeperateFile;
 							TargetArchive = BulkArchive = new FBufferArchive(true);
 						}
 						else
 						{
 							TargetArchive = BulkArchive = IFileManager::Get().CreateFileWriter(*BulkFilename);
-							ExtraBulkDataFlags = BULKDATA_PayloadInSeperateFile;
 						}
 					}
 
@@ -4972,16 +5415,26 @@ ESavePackageResult UPackage::Save(UPackage* InOuter, UObject* Base, EObjectFlags
 						UE_LOG(LogSavePackage, Log,  TEXT("Async saving from memory to '%s'"), *NewPath );
 
 						// Detach archive used for memory saving.
-						if( Linker )
+						if (Linker)
 						{
 							COOK_STAT(FScopedDurationTimer SaveTimer(SavePackageStats::AsyncWriteTimeSec));
 							COOK_STAT(TotalPackageSizeUncompressed += ((FBufferArchive*)(Linker->Saver))->Num());
-														
+
 							FLargeMemoryWriter* Writer = (FLargeMemoryWriter*)(Linker->Saver);
 							int64 DataSize = Writer->TotalSize();
 							FLargeMemoryPtr DataPtr(Writer->GetData());
 							Writer->ReleaseOwnership();
-							AsyncWriteFile(MoveTemp(DataPtr), DataSize, *NewPath, FinalTimeStamp);
+#if SPLIT_COOKED_FILES
+							if (IsEventDrivenLoaderEnabledInCookedBuilds() && Linker->IsCooking())
+							{
+								AsyncWriteFileWithSplitExports(MoveTemp(DataPtr), DataSize, Linker->Summary.TotalHeaderSize, *NewPath, FinalTimeStamp);
+							}
+							else
+#endif
+							{
+								AsyncWriteFile(MoveTemp(DataPtr), DataSize, *NewPath, FinalTimeStamp);
+							}
+
 							Linker->Detach();
 						}
 					}
diff --git a/Engine/Source/Runtime/CoreUObject/Private/UObject/TextProperty.cpp b/Engine/Source/Runtime/CoreUObject/Private/UObject/TextProperty.cpp
--- a/Engine/Source/Runtime/CoreUObject/Private/UObject/TextProperty.cpp
+++ b/Engine/Source/Runtime/CoreUObject/Private/UObject/TextProperty.cpp
@@ -202,6 +202,11 @@ FString UTextProperty::GenerateCppCodeForTextValue(const FText& InValue, const F
 	return CppCode;
 }
 
+FString UTextProperty::GetCPPTypeForwardDeclaration() const
+{
+	return FString();
+}
+
 IMPLEMENT_CORE_INTRINSIC_CLASS(UTextProperty, UProperty,
 {
 }
diff --git a/Engine/Source/Runtime/CoreUObject/Private/UObject/UObjectArchetype.cpp b/Engine/Source/Runtime/CoreUObject/Private/UObject/UObjectArchetype.cpp
--- a/Engine/Source/Runtime/CoreUObject/Private/UObject/UObjectArchetype.cpp
+++ b/Engine/Source/Runtime/CoreUObject/Private/UObject/UObjectArchetype.cpp
@@ -6,6 +6,13 @@
 
 #include "CoreUObjectPrivate.h"
 
+#if !defined(USE_EVENT_DRIVEN_ASYNC_LOAD)
+#error "USE_EVENT_DRIVEN_ASYNC_LOAD must be defined"
+#endif
+
+#if USE_EVENT_DRIVEN_ASYNC_LOAD
+COREUOBJECT_API bool GIgnoreGetArchetypeFromRequiredInfo_RF_NeedsLoad = false;
+#endif
 
 UObject* UObject::GetArchetypeFromRequiredInfo(UClass* Class, UObject* Outer, FName Name, EObjectFlags ObjectFlags)
 {
@@ -33,7 +40,7 @@ UObject* UObject::GetArchetypeFromRequiredInfo(UClass* Class, UObject* Outer, FN
 			else if (!!(ObjectFlags&RF_InheritableComponentTemplate) && Outer->IsA<UClass>())
 			{
 				for (auto SuperClassArchetype = static_cast<UClass*>(Outer)->GetSuperClass();
-					SuperClassArchetype && SuperClassArchetype->HasAllClassFlags(CLASS_CompiledFromBlueprint);
+				SuperClassArchetype && SuperClassArchetype->HasAllClassFlags(CLASS_CompiledFromBlueprint);
 					SuperClassArchetype = SuperClassArchetype->GetSuperClass())
 				{
 					Result = static_cast<UObject*>(FindObjectWithOuter(SuperClassArchetype, Class, Name));
@@ -58,6 +65,16 @@ UObject* UObject::GetArchetypeFromRequiredInfo(UClass* Class, UObject* Outer, FN
 			Result = Class->GetDefaultObject();
 		}
 	}
+#if USE_EVENT_DRIVEN_ASYNC_LOAD
+	//if (!GIgnoreGetArchetypeFromRequiredInfo_RF_NeedsLoad)
+	{
+		if (Result && Result->HasAnyFlags(RF_NeedLoad))
+		{
+			UE_LOG(LogClass, Fatal, TEXT("%s had RF_NeedLoad when being set up as an archetype of %s in %s"), *GetFullNameSafe(Result), *GetFullNameSafe(Class), *GetFullNameSafe(Outer));
+		}
+		check(!Result || !Result->HasAnyFlags(RF_NeedLoad));
+	}
+#endif
 
 	return Result;
 }
diff --git a/Engine/Source/Runtime/CoreUObject/Private/UObject/UObjectArray.cpp b/Engine/Source/Runtime/CoreUObject/Private/UObject/UObjectArray.cpp
--- a/Engine/Source/Runtime/CoreUObject/Private/UObject/UObjectArray.cpp
+++ b/Engine/Source/Runtime/CoreUObject/Private/UObject/UObjectArray.cpp
@@ -52,7 +52,13 @@ void FUObjectArray::OpenDisregardForGC()
 
 void FUObjectArray::CloseDisregardForGC()
 {
+#if THREADSAFE_UOBJECTS
+	FScopeLock ObjObjectsLock(&ObjObjectsCritical);
+#else
+	// Disregard from GC pool is only available from the game thread, at least for now
 	check(IsInGameThread());
+#endif
+
 	check(OpenForDisregardForGC);
 
 	UClass::AssembleReferenceTokenStreams();
@@ -115,8 +121,13 @@ void FUObjectArray::AllocateUObjectIndex(UObjectBase* Object, bool bMergingThrea
 	// Special non- garbage collectable range.
 	if (OpenForDisregardForGC && DisregardForGCEnabled())
 	{
+#if THREADSAFE_UOBJECTS
+		FScopeLock ObjObjectsLock(&ObjObjectsCritical);
+#else
 		// Disregard from GC pool is only available from the game thread, at least for now
 		check(IsInGameThread());
+#endif
+
 		Index = ++ObjLastNonGCIndex;
 		// Check if we're not out of bounds, unless there hasn't been any gc objects yet
 		UE_CLOG(ObjLastNonGCIndex >= MaxObjectsNotConsideredByGC && ObjFirstGCIndex >= 0, LogUObjectArray, Fatal, TEXT("Unable to add more objects to disregard for GC pool (Max: %d)"), MaxObjectsNotConsideredByGC);
diff --git a/Engine/Source/Runtime/CoreUObject/Private/UObject/UObjectBase.cpp b/Engine/Source/Runtime/CoreUObject/Private/UObject/UObjectBase.cpp
--- a/Engine/Source/Runtime/CoreUObject/Private/UObject/UObjectBase.cpp
+++ b/Engine/Source/Runtime/CoreUObject/Private/UObject/UObjectBase.cpp
@@ -262,7 +262,7 @@ bool UObjectBase::IsValidLowLevel() const
 
 bool UObjectBase::IsValidLowLevelFast(bool bRecursive /*= true*/) const
 {
-	// As DEFULT_ALIGNMENT is defined to 0 now, I changed that to the original numerical value here
+	// As DEFAULT_ALIGNMENT is defined to 0 now, I changed that to the original numerical value here
 	const int32 AlignmentCheck = MIN_ALIGNMENT - 1;
 
 	// Check 'this' pointer before trying to access any of the Object's members
@@ -270,7 +270,7 @@ bool UObjectBase::IsValidLowLevelFast(bool bRecursive /*= true*/) const
 	{
 		UE_LOG(LogUObjectBase, Error, TEXT("\'this\' pointer is invalid."));
 		return false;
-	}	
+	}
 	if ((UPTRINT)this & AlignmentCheck)
 	{
 		UE_LOG(LogUObjectBase, Error, TEXT("\'this\' pointer is misaligned."));
diff --git a/Engine/Source/Runtime/CoreUObject/Private/UObject/UObjectBaseUtility.cpp b/Engine/Source/Runtime/CoreUObject/Private/UObject/UObjectBaseUtility.cpp
--- a/Engine/Source/Runtime/CoreUObject/Private/UObject/UObjectBaseUtility.cpp
+++ b/Engine/Source/Runtime/CoreUObject/Private/UObject/UObjectBaseUtility.cpp
@@ -265,14 +265,15 @@ bool UObjectBaseUtility::RootPackageHasAnyFlags( uint32 CheckFlagMask ) const
 /**
  * @return	true if this object is of the specified type.
  */
-#if UCLASS_FAST_ISA_IMPL != 2
-bool UObjectBaseUtility::IsA( const UClass* SomeBase ) const
-{
-	UE_CLOG(!SomeBase, LogObj, Fatal, TEXT("IsA(NULL) cannot yield meaningful results"));
+#if UCLASS_FAST_ISA_COMPARE_WITH_OUTERWALK || UCLASS_FAST_ISA_IMPL == UCLASS_ISA_OUTERWALK
+	bool UObjectBaseUtility::IsA( const UClass* SomeBase ) const
+	{
+		UE_CLOG(!SomeBase, LogObj, Fatal, TEXT("IsA(NULL) cannot yield meaningful results"));
+
+		UClass* ThisClass = GetClass();
 
-	#if UCLASS_FAST_ISA_IMPL & 1
 		bool bOldResult = false;
-		for ( UClass* TempClass=GetClass(); TempClass; TempClass=TempClass->GetSuperClass() )
+		for ( UClass* TempClass=ThisClass; TempClass; TempClass=TempClass->GetSuperClass() )
 		{
 			if ( TempClass == SomeBase )
 			{
@@ -280,22 +281,19 @@ bool UObjectBaseUtility::IsA( const UClass* SomeBase ) const
 				break;
 			}
 		}
-	#endif
 
-	#if UCLASS_FAST_ISA_IMPL & 2
-		bool bNewResult = GetClass()->IsAUsingFastTree(*SomeBase);
+	#if UCLASS_FAST_ISA_IMPL == UCLASS_ISA_INDEXTREE
+		bool bNewResult = ThisClass->IsAUsingFastTree(*SomeBase);
+	#elif UCLASS_FAST_ISA_IMPL == UCLASS_ISA_CLASSARRAY
+		bool bNewResult = ThisClass->IsAUsingClassArray(*SomeBase);
 	#endif
 
-	#if (UCLASS_FAST_ISA_IMPL & 1) && (UCLASS_FAST_ISA_IMPL & 2)
+	#if UCLASS_FAST_ISA_COMPARE_WITH_OUTERWALK
 		ensureMsgf(bOldResult == bNewResult, TEXT("New cast code failed"));
 	#endif
 
-	#if UCLASS_FAST_ISA_IMPL & 1
 		return bOldResult;
-	#else
-		return bNewResult;
-	#endif
-}
+	}
 #endif
 
 
diff --git a/Engine/Source/Runtime/CoreUObject/Private/UObject/UObjectClusters.cpp b/Engine/Source/Runtime/CoreUObject/Private/UObject/UObjectClusters.cpp
--- a/Engine/Source/Runtime/CoreUObject/Private/UObject/UObjectClusters.cpp
+++ b/Engine/Source/Runtime/CoreUObject/Private/UObject/UObjectClusters.cpp
@@ -600,52 +600,79 @@ public:
 	*/
 	FORCEINLINE void HandleTokenStreamObjectReference(TArray<UObject*>& ObjectsToSerialize, UObject* ReferencingObject, UObject*& Object, const int32 TokenIndex, bool bAllowReferenceElimination)
 	{
-#if !(UE_BUILD_TEST || UE_BUILD_SHIPPING)
-		if (Object && !Object->IsValidLowLevelFast())
+		if (Object)
 		{
-			FString TokenDebugInfo;
-			if (UClass *Class = (ReferencingObject ? ReferencingObject->GetClass() : nullptr))
-			{
-				const FTokenInfo& TokenInfo = Class->DebugTokenMap.GetTokenInfo(TokenIndex);
-				TokenDebugInfo = FString::Printf(TEXT("ReferencingObjectClass: %s, Property Name: %s, Offset: %d"),
-					*Class->GetFullName(), *TokenInfo.Name.GetPlainNameString(), TokenInfo.Offset);
-			}
-			else
+#if !(UE_BUILD_TEST || UE_BUILD_SHIPPING)
+			if (
+#if DO_POINTER_CHECKS_ON_GC
+				!IsPossiblyAllocatedUObjectPointer(Object) ||
+#endif
+				!Object->IsValidLowLevelFast())
 			{
-				// This means this objects is most likely being referenced by AddReferencedObjects
-				TokenDebugInfo = TEXT("Native Reference");
-			}
+				FString TokenDebugInfo;
+				if (UClass *Class = (ReferencingObject ? ReferencingObject->GetClass() : nullptr))
+				{
+					const FTokenInfo& TokenInfo = Class->DebugTokenMap.GetTokenInfo(TokenIndex);
+					TokenDebugInfo = FString::Printf(TEXT("ReferencingObjectClass: %s, Property Name: %s, Offset: %d"),
+						*Class->GetFullName(), *TokenInfo.Name.GetPlainNameString(), TokenInfo.Offset);
+				}
+				else
+				{
+					// This means this objects is most likely being referenced by AddReferencedObjects
+					TokenDebugInfo = TEXT("Native Reference");
+				}
 
-			UE_LOG(LogObj, Fatal, TEXT("Invalid object while verifying cluster assumptions: 0x%016llx, ReferencingObject: %s, %s, TokenIndex: %d"),
-				(int64)(PTRINT)Object,
-				ReferencingObject ? *ReferencingObject->GetFullName() : TEXT("NULL"),
-				*TokenDebugInfo, TokenIndex);
-		}
+				UE_LOG(LogObj, Fatal, TEXT("Invalid object while verifying cluster assumptions: 0x%016llx, ReferencingObject: %s, %s, TokenIndex: %d"),
+					(int64)(PTRINT)Object,
+					ReferencingObject ? *ReferencingObject->GetFullName() : TEXT("NULL"),
+					*TokenDebugInfo, TokenIndex);
+			}
 #endif
-		if (Object && !ProcessedObjects.Contains(Object))
-		{
-			ProcessedObjects.Add(Object);
-			FUObjectItem* ObjectItem = GUObjectArray.ObjectToObjectItem(Object);
-			if (ObjectItem->GetOwnerIndex() == 0)
+			if (!ProcessedObjects.Contains(Object))
 			{
-				// We are allowed to reference other clusters, root set objects and objects from diregard for GC pool
-				if (!ObjectItem->HasAnyFlags(EInternalObjectFlags::ClusterRoot|EInternalObjectFlags::RootSet) && !GUObjectArray.IsDisregardForGC(Object) && Object->CanBeInCluster())
+				ProcessedObjects.Add(Object);
+				FUObjectItem* ObjectItem = GUObjectArray.ObjectToObjectItem(Object);
+				if (ObjectItem->GetOwnerIndex() == 0)
 				{
-					UE_LOG(LogObj, Warning, TEXT("Object %s from cluster %s is referencing 0x%016llx %s which is not part of root set or cluster."),
-						*ReferencingObject->GetFullName(),
-						*ClusterRootObject->GetFullName(),
-						(int64)(PTRINT)Object,
-						*Object->GetFullName());
-					bFailed = true;
+					// We are allowed to reference other clusters, root set objects and objects from diregard for GC pool
+					if (!ObjectItem->HasAnyFlags(EInternalObjectFlags::ClusterRoot|EInternalObjectFlags::RootSet) && !GUObjectArray.IsDisregardForGC(Object) && Object->CanBeInCluster())
+					{
+						UE_LOG(LogObj, Warning, TEXT("Object %s from cluster %s is referencing 0x%016llx %s which is not part of root set or cluster."),
+							*ReferencingObject->GetFullName(),
+							*ClusterRootObject->GetFullName(),
+							(int64)(PTRINT)Object,
+							*Object->GetFullName());
+						bFailed = true;
 #if UE_BUILD_DEBUG
-					FReferenceChainSearch RefChainSearch(Object, FReferenceChainSearch::ESearchMode::Shortest | FReferenceChainSearch::ESearchMode::PrintResults);
+						FReferenceChainSearch RefChainSearch(Object, FReferenceChainSearch::ESearchMode::Shortest | FReferenceChainSearch::ESearchMode::PrintResults);
 #endif
+					}
+					else if (ObjectItem->HasAnyFlags(EInternalObjectFlags::ClusterRoot))
+					{
+						// However, clusters need to be referenced by the current cluster otherwise they can also get GC'd too early.
+						const int32 OtherClusterRootIndex = GUObjectArray.ObjectToIndex(Object);
+						const FUObjectItem* OtherClusterRootItem = GUObjectArray.IndexToObjectUnsafeForGC(OtherClusterRootIndex);
+						check(OtherClusterRootItem && OtherClusterRootItem->Object);
+						UObject* OtherClusterRootObject = static_cast<UObject*>(OtherClusterRootItem->Object);
+						UE_CLOG(OtherClusterRootIndex != ClusterRootIndex && !Cluster.ReferencedClusters.Contains(OtherClusterRootIndex), LogObj, Fatal,
+							TEXT("Object %s from source cluster %s is referencing object %s (0x%016llx) from cluster %s which is not referenced by the source cluster."),
+							*ReferencingObject->GetFullName(),
+							*ClusterRootObject->GetFullName(),
+							*Object->GetFullName(),
+							(int64)(PTRINT)Object,
+							*OtherClusterRootObject->GetFullName());
+					}
 				}
-				else if (ObjectItem->HasAnyFlags(EInternalObjectFlags::ClusterRoot))
+				else if (ObjectItem->GetOwnerIndex() == ClusterRootIndex)
 				{
-					// However, clusters need to be referenced by the current cluster otherwise they can also get GC'd too early.
-					const int32 OtherClusterRootIndex = GUObjectArray.ObjectToIndex(Object);
-					const FUObjectItem* OtherClusterRootItem = GUObjectArray.IndexToObjectUnsafeForGC(OtherClusterRootIndex);
+					// If this object belongs to the current cluster, keep processing its references. Otherwise ignore it as it will be processed by its cluster
+					ObjectsToSerialize.Add(Object);
+				}
+				else
+				{
+					// If we're referencing an object from another cluster, make sure the other cluster is actually referenced by this cluster
+					const int32 OtherClusterRootIndex = ObjectItem->GetOwnerIndex();
+					const FUObjectItem* OtherClusterRootItem = GUObjectArray.IndexToObjectUnsafeForGC(OtherClusterRootIndex);				
 					check(OtherClusterRootItem && OtherClusterRootItem->Object);
 					UObject* OtherClusterRootObject = static_cast<UObject*>(OtherClusterRootItem->Object);
 					UE_CLOG(OtherClusterRootIndex != ClusterRootIndex && !Cluster.ReferencedClusters.Contains(OtherClusterRootIndex), LogObj, Fatal,
@@ -657,26 +684,6 @@ public:
 						*OtherClusterRootObject->GetFullName());
 				}
 			}
-			else if (ObjectItem->GetOwnerIndex() == ClusterRootIndex)
-			{
-				// If this object belongs to the current cluster, keep processing its references. Otherwise ignore it as it will be processed by its cluster
-				ObjectsToSerialize.Add(Object);
-			}
-			else
-			{
-				// If we're referencing an object from another cluster, make sure the other cluster is actually referenced by this cluster
-				const int32 OtherClusterRootIndex = ObjectItem->GetOwnerIndex();
-				const FUObjectItem* OtherClusterRootItem = GUObjectArray.IndexToObjectUnsafeForGC(OtherClusterRootIndex);				
-				check(OtherClusterRootItem && OtherClusterRootItem->Object);
-				UObject* OtherClusterRootObject = static_cast<UObject*>(OtherClusterRootItem->Object);
-				UE_CLOG(OtherClusterRootIndex != ClusterRootIndex && !Cluster.ReferencedClusters.Contains(OtherClusterRootIndex), LogObj, Fatal,
-					TEXT("Object %s from source cluster %s is referencing object %s (0x%016llx) from cluster %s which is not referenced by the source cluster."),
-					*ReferencingObject->GetFullName(),
-					*ClusterRootObject->GetFullName(),
-					*Object->GetFullName(),
-					(int64)(PTRINT)Object,
-					*OtherClusterRootObject->GetFullName());
-			}
 		}
 	}
 };
diff --git a/Engine/Source/Runtime/CoreUObject/Private/UObject/UObjectGlobals.cpp b/Engine/Source/Runtime/CoreUObject/Private/UObject/UObjectGlobals.cpp
--- a/Engine/Source/Runtime/CoreUObject/Private/UObject/UObjectGlobals.cpp
+++ b/Engine/Source/Runtime/CoreUObject/Private/UObject/UObjectGlobals.cpp
@@ -17,6 +17,11 @@
 
 DEFINE_LOG_CATEGORY(LogUObjectGlobals);
 
+#if !defined(USE_EVENT_DRIVEN_ASYNC_LOAD)
+#error "USE_EVENT_DRIVEN_ASYNC_LOAD must be defined"
+#endif
+
+#define UE_USE_ASYNCPATH_FOR_LOADPACKAGE (USE_EVENT_DRIVEN_ASYNC_LOAD)
 
 #if USE_MALLOC_PROFILER
 #include "MallocProfiler.h"
@@ -820,33 +825,19 @@ UObject* StaticLoadObjectInternal(UClass* ObjectClass, UObject* InOuter, const T
 
 		if (!Result)
 		{
-			if ((FPlatformProperties::RequiresCookedData() || FPlatformProperties::IsServerOnly()) && GUseSeekFreeLoading)
-			{
-				// Warn only if we allow it and we're sure we're not going to re-try searching with the object name
-				if (bContainsObjectName && (LoadFlags&LOAD_NoWarn) == 0)
-				{
-					UE_LOG(LogUObjectGlobals, Warning, TEXT("StaticLoadObject for %s %s %s couldn't find object in memory!"),
-						*ObjectClass->GetName(),
-						*InOuter->GetName(),
-						*StrName);
-				}
-			}
-			else
-			{
-				// now that we have one asset per package, we load the entire package whenever a single object is requested
-				LoadPackage(NULL, *InOuter->GetOutermost()->GetName(), LoadFlags & ~LOAD_Verify);
+			// now that we have one asset per package, we load the entire package whenever a single object is requested
+			LoadPackage(NULL, *InOuter->GetOutermost()->GetName(), LoadFlags & ~LOAD_Verify);
 
-				// now, find the object in the package
-				Result = StaticFindObjectFast(ObjectClass, InOuter, *StrName);
+			// now, find the object in the package
+			Result = StaticFindObjectFast(ObjectClass, InOuter, *StrName);
 
-				// If the object was not found, check for a redirector and follow it if the class matches
-				if (!Result && !(LoadFlags & LOAD_NoRedirects))
+			// If the object was not found, check for a redirector and follow it if the class matches
+			if (!Result && !(LoadFlags & LOAD_NoRedirects))
+			{
+				UObjectRedirector* Redirector = FindObjectFast<UObjectRedirector>(InOuter, *StrName);
+				if (Redirector && Redirector->DestinationObject && Redirector->DestinationObject->IsA(ObjectClass))
 				{
-					UObjectRedirector* Redirector = FindObjectFast<UObjectRedirector>(InOuter, *StrName);
-					if (Redirector && Redirector->DestinationObject && Redirector->DestinationObject->IsA(ObjectClass))
-					{
-						return Redirector->DestinationObject;
-					}
+					return Redirector->DestinationObject;
 				}
 			}
 		}
@@ -1004,6 +995,7 @@ public:
 
 void ScanPackageDependenciesForLoadOrder(const TCHAR* InLongPackageName, TMap<FName, int32>& InOrderTracker, int32& Order, IAssetRegistryInterface* InAssetRegistry)
 {
+	check(FPlatformTLS::GetCurrentThreadId() == GGameThreadId || !WITH_EDITORONLY_DATA); // only the cooked asset registry is safe to access multithreaded
 	FString FileToLoad;
 	if (InLongPackageName && FCString::Strlen(InLongPackageName) > 0)
 	{
@@ -1055,11 +1047,70 @@ void ScanPackageDependenciesForLoadOrder(const TCHAR* InLongPackageName, TMap<FN
 * @param	ImportLinker	Linker that requests this package through one of its imports
 * @return	Loaded package if successful, NULL otherwise
 */
-static UPackage* LoadPackageInternalInner(UPackage* InOuter, const TCHAR* InLongPackageName, uint32 LoadFlags, FLinkerLoad* ImportLinker, bool bSkipNameChecks = false)
+static UPackage* LoadPackageInternalInner(UPackage* InOuter, const TCHAR* InLongPackageNameOrFilename, uint32 LoadFlags, FLinkerLoad* ImportLinker, bool bSkipNameChecks = false)
 {
 	DECLARE_SCOPE_CYCLE_COUNTER(TEXT("LoadPackageInternal"), STAT_LoadPackageInternal, STATGROUP_ObjectVerbose);
 
-	UPackage* Result = NULL;
+	UPackage* Result = nullptr;
+
+#if UE_USE_ASYNCPATH_FOR_LOADPACKAGE
+	if (FPlatformProperties::RequiresCookedData()
+		&& !GIsInitialLoad //@todoio fix this so we can async load during compiled in object init		
+		)
+	{
+		FString InName;
+		FString InPackageName;
+
+		if (FPackageName::IsPackageFilename(InLongPackageNameOrFilename))
+		{
+			FPackageName::TryConvertFilenameToLongPackageName(InLongPackageNameOrFilename, InPackageName);
+		}
+		else
+		{
+			InPackageName = InLongPackageNameOrFilename;
+		}
+
+		if (InOuter)
+		{
+			InName = InOuter->GetPathName();
+		}
+		else
+		{
+			InName = InPackageName;
+		}
+
+		FName PackageFName(*InPackageName);
+
+		Result = FindObjectFast<UPackage>(nullptr, PackageFName);
+#if USE_EVENT_DRIVEN_ASYNC_LOAD
+		if (!Result || !Result->IsFullyLoaded())
+		{
+			FlushAsyncLoading();
+			Result = FindObjectFast<UPackage>(nullptr, PackageFName);
+			check(!Result || Result->IsFullyLoaded());
+			if (!Result)
+			{
+				LoadPackageAsync(InName, nullptr, *InPackageName);
+				FlushAsyncLoading();
+				Result = FindObjectFast<UPackage>(nullptr, PackageFName);
+			}
+		}
+#else
+		if (!Result || !Result->IsFullyLoaded())
+		{
+			int32 RequestID = LoadPackageAsync(InName, nullptr, *InPackageName);
+			FlushAsyncLoading(RequestID);
+		}
+
+		if (InOuter)
+		{
+			return InOuter;
+		}
+		Result = FindObjectFast<UPackage>(nullptr, PackageFName);
+#endif
+		return Result;
+}
+#endif	
 
 	FString FileToLoad;
 #if WITH_EDITOR
@@ -1067,7 +1118,7 @@ static UPackage* LoadPackageInternalInner(UPackage* InOuter, const TCHAR* InLong
 #endif
 	if (bSkipNameChecks)
 	{
-		FileToLoad = InLongPackageName;
+		FileToLoad = InLongPackageNameOrFilename;
 	}
 	else
 	{
@@ -1075,14 +1126,14 @@ static UPackage* LoadPackageInternalInner(UPackage* InOuter, const TCHAR* InLong
 #if WITH_EDITOR
 		if (LoadFlags & LOAD_ForFileDiff)
 		{
-			FString TempFilenames = InLongPackageName;
+			FString TempFilenames = InLongPackageNameOrFilename;
 			ensure(TempFilenames.Split(TEXT(";"), &FileToLoad, &DiffFileToLoad, ESearchCase::CaseSensitive));
 		}
 		else
 #endif
-		if (InLongPackageName && FCString::Strlen(InLongPackageName) > 0)
+		if (InLongPackageNameOrFilename && FCString::Strlen(InLongPackageNameOrFilename) > 0)
 		{
-			FileToLoad = InLongPackageName;
+			FileToLoad = InLongPackageNameOrFilename;
 		}
 		else if (InOuter)
 		{
@@ -1096,7 +1147,7 @@ static UPackage* LoadPackageInternalInner(UPackage* InOuter, const TCHAR* InLong
 			FName* ScriptPackageName = FPackageName::FindScriptPackageName(*FileToLoad);
 			if (ScriptPackageName)
 			{
-				UE_LOG(LogUObjectGlobals, Warning, TEXT("LoadPackage: %s is a short script package name."), InLongPackageName);
+				UE_LOG(LogUObjectGlobals, Warning, TEXT("LoadPackage: %s is a short script package name."), InLongPackageNameOrFilename);
 				FileToLoad = ScriptPackageName->ToString();
 			}
 			else if (!FPackageName::SearchForPackageOnDisk(FileToLoad, &FileToLoad))
@@ -1233,11 +1284,6 @@ static UPackage* LoadPackageInternalInner(UPackage* InOuter, const TCHAR* InLong
 		GIsEditorLoadingPackage = *IsEditorLoadingPackage;
 #endif
 
-#if WITH_ENGINE
-		// Cancel all texture allocations that haven't been claimed yet.
-		Linker->Summary.TextureAllocations.CancelRemainingAllocations( true );
-#endif		// WITH_ENGINE
-
 		// if we are calculating the script SHA for a package, do the comparison now
 		if (bHasScriptSHAHash)
 		{
@@ -1258,12 +1304,16 @@ static UPackage* LoadPackageInternalInner(UPackage* InOuter, const TCHAR* InLong
 			Result->SetLoadTime( FPlatformTime::Seconds() - StartTime );
 		}
 
+#if USE_NEW_ASYNC_IO
+		Linker->Flush();
+#else
 		// @todo: the next two conditions should check the file limit
 		if (FPlatformProperties::RequiresCookedData())
 		{
 			// give a hint to the IO system that we are done with this file for now
 			FIOSystem::Get().HintDoneWithFile(*Linker->Filename);
 		}
+#endif
 
 		// With UE4 and single asset per package, we load so many packages that some platforms will run out
 		// of file handles. So, this will close the package, but just things like bulk data loading will
@@ -1342,6 +1392,7 @@ UPackage* LoadPackageInternal(UPackage* InOuter, const TCHAR* InLongPackageName,
 		}
 	}
 #endif
+	UPackage* Result = nullptr;
 	if (AssetRegistry)
 	{
 		//MaybeFlushCachedAsyncArchive();
@@ -1351,19 +1402,6 @@ UPackage* LoadPackageInternal(UPackage* InOuter, const TCHAR* InLongPackageName,
 		ScanPackageDependenciesForLoadOrder(InLongPackageName, OrderTracker, Order, AssetRegistry);
 
 		OrderTracker.ValueSort(TLess<int32>());
-#if USE_NEW_ASYNC_IO
-		TArray<FString> Hints;
-		for (auto& Dependency : OrderTracker)
-		{
-			FString PrestreamFilename = GetPrestreamPackageLinkerName(*Dependency.Key.ToString());
-			if (PrestreamFilename.Len() > 0)
-			{
-				HintFutureRead(*PrestreamFilename);
-				Hints.Add(PrestreamFilename);
-			}
-		}
-#endif // USE_NEW_ASYNC_IO
-		UPackage* Result = nullptr;
 		for (auto& Dependency : OrderTracker)
 		{
 			Result = FindObjectFast<UPackage>(nullptr, Dependency.Key, false, false); // might have already loaded this via a circular dependency
@@ -1376,15 +1414,12 @@ UPackage* LoadPackageInternal(UPackage* InOuter, const TCHAR* InLongPackageName,
 				Result = LoadPackageInternalInner(nullptr, *Dependency.Key.ToString(), LoadFlags, ImportLinker, true);
 			}
 		}
-#if USE_NEW_ASYNC_IO
-		for (auto& Dependency : Hints)
-		{
-			HintFutureReadDone(*Dependency);
-		}
-#endif // USE_NEW_ASYNC_IO
-		return Result;
 	}
-	return LoadPackageInternalInner(InOuter, InLongPackageName, LoadFlags, ImportLinker);
+	else
+	{
+		Result = LoadPackageInternalInner(InOuter, InLongPackageName, LoadFlags, ImportLinker);
+	}
+	return Result;
 }
 
 UPackage* LoadPackage(UPackage* InOuter, const TCHAR* InLongPackageName, uint32 LoadFlags)
@@ -2927,6 +2962,11 @@ void FSubobjectPtr::Set(UObject* InObject)
 // Binary initialize object properties to zero or defaults.
 void FObjectInitializer::InitProperties(UObject* Obj, UClass* DefaultsClass, UObject* DefaultData, bool bCopyTransientsFromClassDefaults)
 {
+#if USE_EVENT_DRIVEN_ASYNC_LOAD
+	check(!DefaultsClass || !DefaultsClass->HasAnyFlags(RF_NeedLoad));
+	check(!DefaultData || !DefaultData->HasAnyFlags(RF_NeedLoad));
+#endif
+
 	SCOPE_CYCLE_COUNTER(STAT_InitProperties);
 
 	check(DefaultsClass && Obj);
@@ -2977,7 +3017,11 @@ void FObjectInitializer::InitProperties(UObject* Obj, UClass* DefaultsClass, UOb
 		// As with native classes, we must iterate through all properties (slow path) if default data is pointing at something other than the CDO.
 		bCanUsePostConstructLink &= (DefaultData == Class->GetDefaultObject(false));
 
-		UObject* ClassDefaults = bCopyTransientsFromClassDefaults ? DefaultsClass->GetDefaultObject() : NULL;		
+		UObject* ClassDefaults = bCopyTransientsFromClassDefaults ? DefaultsClass->GetDefaultObject() : NULL;	
+#if USE_EVENT_DRIVEN_ASYNC_LOAD
+		check(!bCopyTransientsFromClassDefaults || !DefaultsClass->GetDefaultObject()->HasAnyFlags(RF_NeedLoad));
+#endif
+
 		for (UProperty* P = bCanUsePostConstructLink ? Class->PostConstructLink : Class->PropertyLink; P; P = bCanUsePostConstructLink ? P->PostConstructLinkNext : P->PropertyLinkNext)
 		{
 			if (bNeedInitialize)
@@ -3114,7 +3158,8 @@ UObject* StaticConstructObject_Internal
 	EInternalObjectFlags InternalSetFlags /*=0*/,
 	UObject*		InTemplate							/*=NULL*/,
 	bool bCopyTransientsFromClassDefaults	/*=false*/,
-	FObjectInstancingGraph* InInstanceGraph				/*=NULL*/
+	FObjectInstancingGraph* InInstanceGraph				/*=NULL*/,
+	bool bAssumeTemplateIsArchetype	/*=false*/
 )
 {
 	SCOPE_CYCLE_COUNTER(STAT_ConstructObject);
@@ -3132,7 +3177,7 @@ UObject* StaticConstructObject_Internal
 	const bool bIsNativeFromCDO = bIsNativeClass &&
 		(
 			!InTemplate || 
-			(InName != NAME_None && InTemplate == UObject::GetArchetypeFromRequiredInfo(InClass, InOuter, InName, InFlags))
+			(InName != NAME_None && (bAssumeTemplateIsArchetype || InTemplate == UObject::GetArchetypeFromRequiredInfo(InClass, InOuter, InName, InFlags)))
 		);
 #if WITH_HOT_RELOAD
 	// Do not recycle subobjects when performing hot-reload as they may contain old property values.
@@ -3552,6 +3597,11 @@ UScriptStruct* GetFallbackStruct()
 	return TBaseStructure<FFallbackStruct>::Get();
 }
 
+#if USE_EVENT_DRIVEN_ASYNC_LOAD
+extern COREUOBJECT_API bool GIgnoreGetArchetypeFromRequiredInfo_RF_NeedsLoad;
+#endif
+
+
 UObject* FObjectInitializer::CreateDefaultSubobject(UObject* Outer, FName SubobjectFName, UClass* ReturnType, UClass* ClassToCreateByDefault, bool bIsRequired, bool bAbstract, bool bIsTransient) const
 {
 	UE_CLOG(!FUObjectThreadContext::Get().IsInConstructor, LogClass, Fatal, TEXT("Subobjects cannot be created outside of UObject constructors. UObject constructing subobjects cannot be created using new or placement new operator."));
@@ -3576,7 +3626,16 @@ UObject* FObjectInitializer::CreateDefaultSubobject(UObject* Outer, FName Subobj
 		{
 			UObject* Template = OverrideClass->GetDefaultObject(); // force the CDO to be created if it hasn't already
 			const EObjectFlags SubobjectFlags = Outer->GetMaskedFlags(RF_PropagateToSubObjects);
-			const bool bOwnerArchetypeIsNotNative = !Outer->GetArchetype()->GetClass()->HasAnyClassFlags(CLASS_Native | CLASS_Intrinsic);
+			bool bOwnerArchetypeIsNotNative;
+			UClass* OuterArchetypeClass;
+
+			{
+#if USE_EVENT_DRIVEN_ASYNC_LOAD
+				TGuardValue<bool> Guard(GIgnoreGetArchetypeFromRequiredInfo_RF_NeedsLoad, true);
+#endif
+				OuterArchetypeClass = Outer->GetArchetype()->GetClass();
+				bOwnerArchetypeIsNotNative = !OuterArchetypeClass->HasAnyClassFlags(CLASS_Native | CLASS_Intrinsic);
+			}
 			const bool bOwnerTemplateIsNotCDO = ObjectArchetype != nullptr && ObjectArchetype != Outer->GetClass()->GetDefaultObject(false) && !Outer->HasAnyFlags(RF_ClassDefaultObject);
 #if !UE_BUILD_SHIPPING
 			// Guard against constructing the same subobject multiple times.
@@ -3602,7 +3661,7 @@ UObject* FObjectInitializer::CreateDefaultSubobject(UObject* Outer, FName Subobj
 				if (!MaybeTemplate)
 				{
 					// The archetype of the outer is not native, so we need to copy properties to the subobjects after the C++ constructor chain for the outer has run (because those sets properties on the subobjects)
-					MaybeTemplate = Outer->GetArchetype()->GetClass()->GetDefaultSubobjectByName(SubobjectFName);
+					MaybeTemplate = OuterArchetypeClass->GetDefaultSubobjectByName(SubobjectFName);
 				}
 				if (MaybeTemplate && MaybeTemplate->IsA(ReturnType) && Template != MaybeTemplate)
 				{
diff --git a/Engine/Source/Runtime/CoreUObject/Private/UObject/UObjectThreadContext.cpp b/Engine/Source/Runtime/CoreUObject/Private/UObject/UObjectThreadContext.cpp
--- a/Engine/Source/Runtime/CoreUObject/Private/UObject/UObjectThreadContext.cpp
+++ b/Engine/Source/Runtime/CoreUObject/Private/UObject/UObjectThreadContext.cpp
@@ -23,4 +23,5 @@ FUObjectThreadContext::FUObjectThreadContext()
 , SerializedImportLinker(nullptr)
 , SerializedExportIndex(INDEX_NONE)
 , SerializedExportLinker(nullptr)
+, AsyncPackage(nullptr)
 {}
\ No newline at end of file
