Copying //UE4/Dev-Core to //UE4/Dev-Main (Source: //UE4/Dev-Core @ 3151653)

diff --git a/Engine/Source/Runtime/CoreUObject/Private/UObject/WeakObjectPtr.cpp b/Engine/Source/Runtime/CoreUObject/Private/UObject/WeakObjectPtr.cpp
--- a/Engine/Source/Runtime/CoreUObject/Private/UObject/WeakObjectPtr.cpp
+++ b/Engine/Source/Runtime/CoreUObject/Private/UObject/WeakObjectPtr.cpp
@@ -96,32 +96,25 @@ UObject* FWeakObjectPtr::GetEvenIfUnreachable() const
 	return Result;
 }
 
-
-/**
- * Weak object pointer serialization.  Weak object pointers only have weak references to objects and
- * won't serialize the object when gathering references for garbage collection.  So in many cases, you
- * don't need to bother serializing weak object pointers.  However, serialization is required if you
- * want to load and save your object.
- */
-FArchive& operator<<( FArchive& Ar, FWeakObjectPtr& WeakObjectPtr )
+void FWeakObjectPtr::Serialize(FArchive& Ar)
 {
+	// NOTE: When changing this function, make sure to update the SavePackage.cpp version in the import and export tagger.
+
 	// We never serialize our reference while the garbage collector is harvesting references
 	// to objects, because we don't want weak object pointers to keep objects from being garbage
 	// collected.  That would defeat the whole purpose of a weak object pointer!
 	// However, when modifying both kinds of references we want to serialize and writeback the updated value.
-	if( !Ar.IsObjectReferenceCollector() || Ar.IsModifyingWeakAndStrongReferences() )
+	if (!Ar.IsObjectReferenceCollector() || Ar.IsModifyingWeakAndStrongReferences())
 	{
 		// Downcast from UObjectBase to UObject
-		UObject* Object = static_cast< UObject* >( WeakObjectPtr.Get(true) );
+		UObject* Object = static_cast<UObject*>(Get(true));
 
 		Ar << Object;
 
-		if( Ar.IsLoading() || Ar.IsModifyingWeakAndStrongReferences() )
+		if (Ar.IsLoading() || Ar.IsModifyingWeakAndStrongReferences())
 		{
-			WeakObjectPtr = Object;
+			*this = Object;
 		}
 	}
-
-	return Ar;
 }
 
diff --git a/Engine/Source/Runtime/CoreUObject/Public/Misc/StartupPackages.h b/Engine/Source/Runtime/CoreUObject/Public/Misc/StartupPackages.h
--- a/Engine/Source/Runtime/CoreUObject/Public/Misc/StartupPackages.h
+++ b/Engine/Source/Runtime/CoreUObject/Public/Misc/StartupPackages.h
@@ -25,9 +25,6 @@ struct FStartupPackages
 
 	/**
 	 * Load all startup packages. If desired preload async followed by serialization from memory.
-	 * Only native script packages are loaded from memory if we're not using the GUseSeekFreeLoading
-	 * codepath as we need to reset the loader on those packages and don't want to keep the bulk
-	 * data around. Native script packages don't have any bulk data so this doesn't matter.
 	 *
 	 * The list of additional packages can be found at Engine.StartupPackages and if seekfree loading
 	 * is enabled, the startup map is going to be preloaded as well.
diff --git a/Engine/Source/Runtime/CoreUObject/Public/Serialization/ArchiveUObject.h b/Engine/Source/Runtime/CoreUObject/Public/Serialization/ArchiveUObject.h
--- a/Engine/Source/Runtime/CoreUObject/Public/Serialization/ArchiveUObject.h
+++ b/Engine/Source/Runtime/CoreUObject/Public/Serialization/ArchiveUObject.h
@@ -15,5 +15,6 @@ public:
 	virtual FArchive& operator<< (class FLazyObjectPtr& Value) override;
 	virtual FArchive& operator<< (class FAssetPtr& Value) override;
 	virtual FArchive& operator<< (struct FStringAssetReference& Value) override;
+	virtual FArchive& operator<< (struct FWeakObjectPtr& Value) override;
 	//~ End FArchive Interface
 };
diff --git a/Engine/Source/Runtime/CoreUObject/Public/Serialization/AsyncLoading.h b/Engine/Source/Runtime/CoreUObject/Public/Serialization/AsyncLoading.h
--- a/Engine/Source/Runtime/CoreUObject/Public/Serialization/AsyncLoading.h
+++ b/Engine/Source/Runtime/CoreUObject/Public/Serialization/AsyncLoading.h
@@ -8,13 +8,332 @@
 
 #include "AsyncPackage.h"
 
+struct FFlushTree;
+
 #define PERF_TRACK_DETAILED_ASYNC_STATS (0)
+
+#ifndef USE_EVENT_DRIVEN_ASYNC_LOAD
+#error "USE_EVENT_DRIVEN_ASYNC_LOAD must be defined""
+#endif
+
+
+#if USE_EVENT_DRIVEN_ASYNC_LOAD
+
+#if !USE_NEW_ASYNC_IO
+#error "USE_EVENT_DRIVEN_ASYNC_LOAD requires USE_NEW_ASYNC_IO"
+#endif
+
+struct FAsyncPackage;
+enum class EAsyncPackageLoadingState : uint8
+{
+	NewPackage,
+	WaitingForSummary,
+	StartImportPackages,
+	WaitingForImportPackages,
+	SetupImports,
+	SetupExports,
+	ProcessNewImportsAndExports,
+	WaitingForPostLoad,
+	ReadyForPostLoad,
+	PostLoad_Etc,
+	PackageComplete,
+};
+
+// This version is an ordinary pointer. We can swap in the safer version to verify assumptions.
+struct FUnsafeWeakAsyncPackagePtr
+{
+	FAsyncPackage* Package;
+	FUnsafeWeakAsyncPackagePtr(FAsyncPackage* InPackage = nullptr)
+		: Package(InPackage)
+	{
+	}
+	FORCEINLINE FAsyncPackage& GetPackage() const
+	{
+		check(Package);
+		return *Package;
+	}
+	FORCEINLINE friend uint32 GetTypeHash(const FUnsafeWeakAsyncPackagePtr& WeakPtr)
+	{
+		return GetTypeHash(WeakPtr.Package);
+	}
+	FORCEINLINE bool operator==(const FUnsafeWeakAsyncPackagePtr& Other) const
+	{
+		return Package == Other.Package;
+	}
+	FName HumanReadableStringForDebugging() const;
+};
+
+struct FWeakAsyncPackagePtr
+{
+	FName PackageName;
+	int32 SerialNumber;
+	FWeakAsyncPackagePtr(FAsyncPackage* Package = nullptr);
+	FAsyncPackage& GetPackage() const;
+	FORCEINLINE friend uint32 GetTypeHash(const FWeakAsyncPackagePtr& WeakPtr)
+	{
+		return GetTypeHash(WeakPtr.PackageName);
+	}
+	FORCEINLINE bool operator==(const FWeakAsyncPackagePtr& Other) const
+	{
+		return SerialNumber == Other.SerialNumber;
+	}
+	FName HumanReadableStringForDebugging() const
+	{
+		return PackageName;
+	}
+};
+
+#define VERIFY_WEAK_ASYNC_PACKAGE_PTRS (0) //(DO_CHECK)
+
+#if VERIFY_WEAK_ASYNC_PACKAGE_PTRS
+typedef FWeakAsyncPackagePtr FCheckedWeakAsyncPackagePtr;
+#else
+typedef FUnsafeWeakAsyncPackagePtr FCheckedWeakAsyncPackagePtr;
+#endif
+
+enum class EEventLoadNode
+{
+	Package_LoadSummary,
+	Package_SetupImports,
+	Package_ExportsSerialized,
+	Package_NumPhases,
+
+	ImportOrExport_Create = 0,
+	ImportOrExport_Serialize,
+	Import_NumPhases,
+
+	Export_StartIO = Import_NumPhases,
+	Export_NumPhases,
+
+	MAX_NumPhases = Package_NumPhases,
+	Invalid_Value = -1
+};
+
+#define USE_IMPLICIT_ARCS (1) // saves time and memory by not actually adding the arcs that are always present, they are:
+
+// Import: EEventLoadNode::ImportOrExport_Create -> EEventLoadNode::ImportOrExport_Serialize: can't consider this import serialized until we hook it up after creation
+// Import: EEventLoadNode::ImportOrExport_Serialize -> EEventLoadNode::Package_ExportsSerialized: can't consider the package done with event driven loading until all imports are serialized
+
+// Export: EEventLoadNode::ImportOrExport_Create -> EEventLoadNode::Export_StartIO: can't do the IO request until it is created
+// Export: EEventLoadNode::Export_StartIO -> EEventLoadNode::ImportOrExport_Serialize: can't serialize until the IO request is ready
+// Import: EEventLoadNode::ImportOrExport_Serialize -> EEventLoadNode::Package_ExportsSerialized: can't consider the package done with event driven loading until all exports are serialized
+
+
+struct FEventLoadNodePtr
+{
+	FCheckedWeakAsyncPackagePtr WaitingPackage;
+	FPackageIndex ImportOrExportIndex;  // IsNull()==true for PACKAGE_ Phases
+	EEventLoadNode Phase;
+
+	FEventLoadNodePtr()
+		: Phase(EEventLoadNode::Invalid_Value)
+	{
+	}
+	FORCEINLINE friend uint32 GetTypeHash(const FEventLoadNodePtr& NodePtr)
+	{
+		return GetTypeHash(NodePtr.WaitingPackage) ^ (GetTypeHash(NodePtr.ImportOrExportIndex) << 13) ^ (uint32(NodePtr.Phase) << 27);
+	}
+
+#if 0
+	FORCEINLINE bool operator<(const FEventLoadNodePtr& Other) const // not a total order
+	{
+		return WaitingPackage.SerialNumber < Other.WaitingPackage.SerialNumber;
+	}
+#endif
+	FORCEINLINE bool operator==(const FEventLoadNodePtr& Other) const
+	{
+		return WaitingPackage == Other.WaitingPackage && 
+			ImportOrExportIndex == Other.ImportOrExportIndex && 
+			Phase == Other.Phase;
+	}
+#if USE_IMPLICIT_ARCS
+	FORCEINLINE int32 NumImplicitArcs()
+	{
+		return ImportOrExportIndex.IsNull() ? 0 : 1; // only import and export nodes have implicit arcs, and every one of them has exactly one
+	}
+	FORCEINLINE FEventLoadNodePtr GetImplicitArc()
+	{
+		check(!ImportOrExportIndex.IsNull()); // package nodes don't have implicit arcs
+		FEventLoadNodePtr Result;
+		Result.WaitingPackage = WaitingPackage;
+		if (Phase == EEventLoadNode::ImportOrExport_Serialize)
+		{
+			Result.Phase = EEventLoadNode::Package_ExportsSerialized;
+			check(Result.ImportOrExportIndex.IsNull());
+		}
+		else
+		{
+			Result.ImportOrExportIndex = ImportOrExportIndex;
+			if (Phase == EEventLoadNode::ImportOrExport_Create)
+			{
+				Result.Phase = ImportOrExportIndex.IsImport() ? EEventLoadNode::ImportOrExport_Serialize : EEventLoadNode::Export_StartIO;
+			}
+			else
+			{
+				check(Phase == EEventLoadNode::Export_StartIO);
+				Result.Phase = EEventLoadNode::ImportOrExport_Serialize;
+			}
+		}
+		return Result;
+	}
+
+#endif
+
+	FString HumanReadableStringForDebugging() const;
+};
+
+struct FEventLoadNode
+{
+#if USE_IMPLICIT_ARCS
+	typedef TArray<FEventLoadNodePtr> TNodesWaitingForMeArray;
+#else
+	typedef TArray<FEventLoadNodePtr, TInlineAllocator<1> > TNodesWaitingForMeArray;
+#endif
+
+	TNodesWaitingForMeArray NodesWaitingForMe;
+	int32 NumPrerequistes;
+	bool bFired;
+	bool bAddedToGraph;
+	FORCEINLINE FEventLoadNode()
+		: NumPrerequistes(0)
+		, bFired(false)
+		, bAddedToGraph(false)
+	{
+	}
+};
+
+struct FEventLoadNodeArray
+{
+	FEventLoadNode PackageNodes[int(EEventLoadNode::Package_NumPhases)];
+	FEventLoadNode* Array;
+	int32 TotalNumberOfImportExportNodes;
+	int32 TotalNumberOfNodesAdded;
+	int32 NumImports;
+	int32 NumExports;
+	int32 OffsetToImports;
+	int32 OffsetToExports;
+
+	FEventLoadNodeArray()
+		: Array(nullptr)
+		, TotalNumberOfImportExportNodes(0)
+		, TotalNumberOfNodesAdded(0)
+		, NumImports(0)
+		, NumExports(0)
+		, OffsetToImports(0)
+		, OffsetToExports(0)
+	{
+	}
+
+
+	bool AddNode(FEventLoadNodePtr Node)
+	{
+		FEventLoadNode& NodeRef(PtrToNode(Node));
+		check(!NodeRef.bAddedToGraph);
+		NodeRef.bAddedToGraph = true;
+		return (++TotalNumberOfNodesAdded) == 1; // first add, the caller will add this to a list of things with outstanding nodes
+	}
+	bool RemoveNode(FEventLoadNodePtr Node)
+	{
+		FEventLoadNode& NodeRef(PtrToNode(Node));
+		check(NodeRef.bAddedToGraph);
+		NodeRef.bAddedToGraph = false;
+		return (--TotalNumberOfNodesAdded) == 0; // first add, the caller will add this to a list of things with outstanding nodes
+	}
+	FORCEINLINE FEventLoadNode& GetNode(FEventLoadNodePtr Node, bool bCheckAdded = true)
+	{
+		FEventLoadNode& NodeRef(PtrToNode(Node));
+		check(!bCheckAdded || NodeRef.bAddedToGraph);
+		return NodeRef;
+	}
+
+	void Init(int32 InNumImports, int32 InNumExports);
+	void Shutdown();
+	void GetAddedNodes(TArray<FEventLoadNodePtr>& OutAddedNodes, FAsyncPackage* Owner);
+private:
+	FORCEINLINE FEventLoadNode& PtrToNode(FEventLoadNodePtr Node)
+	{
+		int32 Index;
+		if (Node.ImportOrExportIndex.IsNull())
+		{
+			Index = int32(Node.Phase);
+			check(Index < int32(EEventLoadNode::Package_NumPhases) && Index >= 0);
+			return PackageNodes[Index];
+		}
+		if (Node.ImportOrExportIndex.IsImport())
+		{
+			check(int32(Node.Phase) < int32(EEventLoadNode::Import_NumPhases));
+			Index = OffsetToImports + Node.ImportOrExportIndex.ToImport() * int32(EEventLoadNode::Import_NumPhases) + int32(Node.Phase);
+			check(Index >= OffsetToImports && Index < OffsetToExports);
+		}
+		else
+		{
+			check(int32(Node.Phase) < int32(EEventLoadNode::Export_NumPhases));
+			Index = OffsetToExports + Node.ImportOrExportIndex.ToExport() * int32(EEventLoadNode::Export_NumPhases) + int32(Node.Phase);
+			check(Index >= OffsetToExports);
+		}
+		check(Array && Index < TotalNumberOfImportExportNodes && Index >= 0);
+		return Array[Index];
+	}
+
+};
+
+
+struct FEventLoadGraph
+{
+	TSet<FCheckedWeakAsyncPackagePtr> PackagesWithNodes;
+
+	FEventLoadNodeArray& GetArray(FEventLoadNodePtr& Node);
+
+	FEventLoadNode& GetNode(FEventLoadNodePtr& NodeToGet);
+
+	void AddNode(FEventLoadNodePtr& NewNode, bool bHoldForLater = false, int32 NumImplicitPrereqs = 0);
+	void DoneAddingPrerequistesFireIfNone(FEventLoadNodePtr& NewNode, bool bWasHeldForLater = false);
+	void AddArc(FEventLoadNodePtr& PrereqisiteNode, FEventLoadNodePtr& DependentNode);
+	void RemoveNode(FEventLoadNodePtr& NodeToRemove);
+	void NodeWillBeFiredExternally(FEventLoadNodePtr& NodeThatWasFired);
+	void CheckForCycles();
+#if !UE_BUILD_SHIPPING
+	bool CheckForCyclesInner(const TMultiMap<FEventLoadNodePtr, FEventLoadNodePtr>& Arcs, TSet<FEventLoadNodePtr>& Visited, TSet<FEventLoadNodePtr>& Stack, const FEventLoadNodePtr& Visit);
+#endif
+};
+
+struct FAsyncLoadEventArgs
+{
+	double TickStartTime;
+	const TCHAR* OutLastTypeOfWorkPerformed;
+	UObject* OutLastObjectWorkWasPerformedOn;
+	float TimeLimit;
+	bool bUseTimeLimit;
+	bool bUseFullTimeLimit;
+
+	FAsyncLoadEventArgs()
+		: TickStartTime(0.0)
+		, OutLastTypeOfWorkPerformed(nullptr)
+		, OutLastObjectWorkWasPerformedOn(nullptr)
+		, TimeLimit(0.0f)
+		, bUseTimeLimit(false)
+		, bUseFullTimeLimit(true)
+	{
+	}
+};
+
+// this is a little wrapper that does random pops for further randomization
+class FImportOrImportIndexArray : public TArray<int32>
+{
+public:
+	void HeapPop(int32& OutItem, bool bAllowShrinking = true);
+};
+
+#endif
+
 /**
- * Structure containing intermediate data required for async loading of all imports and exports of a
- * FLinkerLoad.
- */
-struct FAsyncPackage
+* Structure containing intermediate data required for async loading of all imports and exports of a
+* FLinkerLoad.
+*/
+
+struct FAsyncPackage : FGCObject
 {
+	friend struct FScopedAsyncPackageEvent;
 	/**
 	 * Constructor
 	 */
@@ -30,7 +349,13 @@ struct FAsyncPackage
 	 *
 	 * @return	true if package has finished loading, false otherwise
 	 */
-	EAsyncPackageState::Type Tick( bool bUseTimeLimit, bool bInbUseFullTimeLimit, float& InOutTimeLimit );
+	EAsyncPackageState::Type TickAsyncPackage(bool bUseTimeLimit, bool bInbUseFullTimeLimit, float& InOutTimeLimit, FFlushTree* FlushTree);
+
+	/** Fills the package dependency tree required to flush a specific package */
+	void PopulateFlushTree(struct FFlushTree* FlushTree);	
+	
+	/** Marks a specific request as complete */
+	void MarkRequestIDsAsComplete();
 
 	/**
 	 * @return Estimated load completion percentage.
@@ -79,6 +404,11 @@ struct FAsyncPackage
 		return DependencyRefCount.GetValue();
 	}
 
+	FORCEINLINE UPackage* GetLinkerRoot() const
+	{
+		return LinkerRoot;
+	}
+
 	/** Returns true if the package has finished loading. */
 	FORCEINLINE bool HasFinishedLoading() const
 	{
@@ -103,6 +433,18 @@ struct FAsyncPackage
 		return bLoadHasFailed;
 	}
 
+	/** Marks the threaded loading phase as complete for this package */
+	void ThreadedLoadingHasFinished()
+	{
+		bThreadedLoadingFinished = true;
+	}
+
+	/** Returns true if the threaded loaded phas has completed for this package */
+	bool HasThreadedLoadingFinished() const
+	{
+		return bThreadedLoadingFinished;
+	}
+
 	/** Adds new request ID to the existing package */
 	void AddRequestID(int32 Id);
 
@@ -119,18 +461,37 @@ struct FAsyncPackage
 		DependencyRootPackage = InDependencyRootPackage;
 	}
 
-private:
+	/** Returns true if this package is already being loaded in the current callstack */
+	bool IsBeingProcessedRecursively() const
+	{
+		return ReentryCount > 1;
+	}
+
+	/** FGCObject Interface */
+	virtual void AddReferencedObjects(FReferenceCollector& Collector) override;
+
+	/** Adds a new object referenced by this package */
+	void AddObjectReference(UObject* InObject);
+
+	/** Removes all objects from the list and clears async loading flags */
+	void EmptyReferencedObjects();
+
+private:	
 
 	struct FCompletionCallback
 	{
 		bool bIsInternal;
+		bool bCalled;
 		FLoadPackageAsyncDelegate Callback;
 
 		FCompletionCallback()
+			: bIsInternal(false)
+			, bCalled(false)
 		{
 		}
 		FCompletionCallback(bool bInInternal, FLoadPackageAsyncDelegate InCallback)
 			: bIsInternal(bInInternal)
+			, bCalled(false)
 			, Callback(InCallback)
 		{
 		}
@@ -158,14 +519,16 @@ private:
 	int32							ImportIndex;
 	/** Current index into linkers export table used to spread creation over several frames				*/
 	int32							ExportIndex;
-	/** Current index into GObjLoaded array used to spread routing PreLoad over several frames			*/
-	static int32					PreLoadIndex;
-	/** Current index into GObjLoaded array used to spread routing PreLoad over several frames			*/
-	static int32					PreLoadSortIndex;
-	/** Current index into GObjLoaded array used to spread routing PostLoad over several frames			*/
-	static int32					PostLoadIndex;
+	/** Current index into ObjLoaded array used to spread routing PreLoad over several frames			*/
+	int32							PreLoadIndex;
+	/** Current index into ObjLoaded array used to spread routing PreLoad over several frames			*/
+	int32							PreLoadSortIndex;
+	/** Current index into ObjLoaded array used to spread routing PostLoad over several frames			*/
+	int32							PostLoadIndex;
 	/** Current index into DeferredPostLoadObjects array used to spread routing PostLoad over several frames			*/
 	int32						DeferredPostLoadIndex;
+	/** Current index into DeferredFinalizeObjects array used to spread routing PostLoad over several frames			*/
+	int32						DeferredFinalizeIndex;
 	/** Currently used time limit for this tick.														*/
 	float						TimeLimit;
 	/** Whether we are using a time limit for this tick.												*/
@@ -178,6 +541,8 @@ private:
 	bool						bLoadHasFailed;
 	/** True if our load has finished */
 	bool						bLoadHasFinished;
+	/** True if threaded loading has finished for this package */
+	bool						bThreadedLoadingFinished;
 	/** The time taken when we started the tick.														*/
 	double						TickStartTime;
 	/** Last object work was performed on. Used for debugging/ logging purposes.						*/
@@ -192,15 +557,159 @@ private:
 	TArray<UObject*> DeferredPostLoadObjects;
 	/** Objects to be finalized on the game thread */
 	TArray<UObject*> DeferredFinalizeObjects;
+	/** Objects loaded while loading this package */
+	TArray<UObject*> PackageObjLoaded;
+	/** Packages that were loaded synchronously while async loading this package or packages added by verify import */
+	TArray<FLinkerLoad*> DelayedLinkerClosePackages;
+
 	/** List of all request handles */
 	TArray<int32> RequestIDs;
 #if WITH_EDITORONLY_DATA
 	/** Index of the meta-data object within the linkers export table (unset if not yet processed, although may still be INDEX_NONE if there is no meta-data) */
 	TOptional<int32> MetaDataIndex;
 #endif // WITH_EDITORONLY_DATA
+	/** Number of times we recursed to load this package. */
+	int32 ReentryCount;
+	/** List of objects referenced by this package */
+	TSet<UObject*> ReferencedObjects;
+#if USE_EVENT_DRIVEN_ASYNC_LOAD
+	/** Critical section for referenced objects list */
+	FCriticalSection ReferencedObjectsCritical;
+#endif
 	/** Cached async loading thread object this package was created by */
 	class FAsyncLoadingThread& AsyncLoadingThread;
+
 public:
+
+#if USE_EVENT_DRIVEN_ASYNC_LOAD
+	EAsyncPackageLoadingState AsyncPackageLoadingState;
+	int32 SerialNumber;
+
+	FImportOrImportIndexArray ImportsThatAreNowCreated;
+	FImportOrImportIndexArray ImportsThatAreNowSerialized;
+	FImportOrImportIndexArray ExportsThatCanBeCreated;
+	FImportOrImportIndexArray ExportsThatCanHaveIOStarted;
+	FImportOrImportIndexArray ExportsThatCanBeSerialized;
+	TArray<IAsyncReadRequest*> ReadyPrecacheRequests;
+
+	struct FExportIORequest
+	{
+		int64 Offset;
+		int64 BytesToRead;
+		int32 FirstExportCovered;
+		int32 LastExportCovered;
+		TArray<int32> ExportsToRead;
+
+		FExportIORequest()
+			: Offset(-1)
+			, BytesToRead(-1)
+			, FirstExportCovered(-1)
+			, LastExportCovered(-1)
+		{
+		}
+	};
+
+	TMap<IAsyncReadRequest*, FExportIORequest> PrecacheRequests;
+	TMap<int32, IAsyncReadRequest*> ExportIndexToPrecacheRequest;
+	int64 CurrentBlockOffset;
+	int64 CurrentBlockBytes;
+	TSet<int32> ExportsInThisBlock;
+
+	TMultiMap<FName, FPackageIndex> ObjectNameToImportOrExport;
+
+	TSet<FWeakAsyncPackagePtr> PackagesIMayBeWaitingForBeforePostload; // these need to be reexamined and perhaps deleted or collapsed
+
+	TSet<FWeakAsyncPackagePtr> PackagesIAmWaitingForBeforePostload;    // These are linked with PackagesIAmWaitingForBeforePostload, so we can be sure the other package will let us know
+	TSet<FWeakAsyncPackagePtr> OtherPackagesWaitingForMeBeforePostload;
+
+	TArray<FCheckedWeakAsyncPackagePtr> PackagesWaitingToLinkImports;
+
+	int32 ImportAddNodeIndex;
+	int32 ExportAddNodeIndex;
+
+	bool bProcessImportsAndExportsInFlight;
+	bool bProcessPostloadWaitInFlight;
+	bool bAllExportsSerialized;
+
+	void Event_CreateLinker();
+	void Event_FinishLinker();
+	void Event_StartImportPackages();
+	void Event_SetupImports();
+	void Event_SetupExports();
+	void Event_ProcessImportsAndExports();
+	void Event_ExportsDone();
+	void Event_ProcessPostloadWait();
+	void Event_StartPostload();
+
+	bool AnyImportsAndExportWorkOutstanding();
+	void ConditionalQueueProcessImportsAndExports(bool bRequeueForTimeout = false);
+	void ConditionalQueueProcessPostloadWait();
+
+
+	EAsyncPackageState::Type LoadImports_Event();
+	EAsyncPackageState::Type SetupImports_Event();
+	EAsyncPackageState::Type SetupExports_Event();
+	EAsyncPackageState::Type ProcessImportsAndExports_Event();
+
+	FObjectImport* FindExistingImport(int32 LocalImportIndex);
+	void LinkImport(int32 LocalImportIndex);
+	void EventDrivenCreateExport(int32 LocalExportIndex);
+	void StartPrecacheRequest();
+	void EventDrivenSerializeExport(int32 LocalExportIndex);
+	int64 PrecacheRequestReady(IAsyncReadRequest * Req);
+	void MakeNextPrecacheRequestCurrent();
+	void FlushPrecacheBuffer();
+	void EventDrivenLoadingComplete();
+
+	void DumpDependencies(const TCHAR* Label, UObject* Obj);
+	void DumpDependencies(const TCHAR* Label, FLinkerLoad* DumpLinker, FPackageIndex Index);
+
+	UObject* EventDrivenIndexToObject(FPackageIndex Index, bool bCheckSerialized, FPackageIndex DumpIndex = FPackageIndex());
+	template<class T>
+	T* CastEventDrivenIndexToObject(FPackageIndex Index, bool bCheckSerialized, FPackageIndex DumpIndex = FPackageIndex())
+	{
+		UObject* Result = EventDrivenIndexToObject(Index, bCheckSerialized, DumpIndex);
+		if (!Result)
+		{
+			return nullptr;
+		}
+		return CastChecked<T>(Result);
+	}
+
+	FEventLoadNodeArray EventNodeArray;
+
+	static FEventLoadGraph GlobalEventGraph;
+
+	FORCEINLINE static FEventLoadGraph& GetEventGraph()
+	{
+		return GlobalEventGraph;
+	}
+
+	FEventLoadNodePtr AddNode(EEventLoadNode Phase, FPackageIndex ImportOrExportIndex = FPackageIndex(), bool bHoldForLater = false, int32 NumImplicitPrereqs = 0);
+	void DoneAddingPrerequistesFireIfNone(EEventLoadNode Phase, FPackageIndex ImportOrExportIndex = FPackageIndex(), bool bWasHeldForLater = false);
+	void RemoveNode(EEventLoadNode Phase, FPackageIndex ImportOrExportIndex = FPackageIndex());
+	void NodeWillBeFiredExternally(EEventLoadNode Phase, FPackageIndex ImportOrExportIndex = FPackageIndex());
+	void AddArc(FEventLoadNodePtr& PrereqisiteNode, FEventLoadNodePtr& DependentNode);
+	void RemoveAllNodes();
+	void FireNode(FEventLoadNodePtr& NodeToFire);
+
+	FString GetDebuggingPath(FPackageIndex Idx);
+
+	void SetTimeLimit(FAsyncLoadEventArgs& Args, const TCHAR* WorkType)
+	{
+		Args.OutLastTypeOfWorkPerformed = WorkType;
+		Args.OutLastObjectWorkWasPerformedOn = LinkerRoot;
+		TickStartTime = Args.TickStartTime;
+		LastTypeOfWorkPerformed = WorkType;
+		LastObjectWorkWasPerformedOn = LinkerRoot;
+		TimeLimit = Args.TimeLimit;
+		bUseTimeLimit = Args.bUseTimeLimit;
+		bUseFullTimeLimit = Args.bUseFullTimeLimit;
+	}
+
+#endif
+
+
 #if PERF_TRACK_DETAILED_ASYNC_STATS
 	/** Number of times Tick function has been called.													*/
 	int32							TickCount;
@@ -250,6 +759,9 @@ public:
 	*/
 	EAsyncPackageState::Type PostLoadDeferredObjects(double InTickStartTime, bool bInUseTimeLimit, float& InOutTimeLimit);
 
+	/** Close any linkers that have been open as a result of synchronous load during async loading */
+	void CloseDelayedLinkers();
+
 private:
 	/**
 	 * Gives up time slice if time limit is enabled.
@@ -290,21 +802,17 @@ private:
 	/**
 	 * Loads imported packages..
 	 *
+	 * @param FlushTree Package dependency tree to be flushed
 	 * @return true if we finished loading all imports, false otherwise
 	 */
-	EAsyncPackageState::Type LoadImports();
+	EAsyncPackageState::Type LoadImports(FFlushTree* FlushTree);
 	/**
 	 * Create imports till time limit is exceeded.
 	 *
 	 * @return true if we finished creating all imports, false otherwise
 	 */
 	EAsyncPackageState::Type CreateImports();
-	/**
-	 * Checks if all async texture allocations for this package have been completed.
-	 *
-	 * @return true if all texture allocations have been completed, false otherwise
-	 */
-	EAsyncPackageState::Type FinishTextureAllocations();
+
 #if WITH_EDITORONLY_DATA
 	/**
 	 * Creates and loads meta-data for the package.
@@ -346,20 +854,23 @@ private:
 	 * Adds dependency tree to the list if packages to wait for until their linkers have been created.
 	 *
 	 * @param ImportedPackage Package imported either directly or by one of the imported packages
+	 * @param FlushTree Package dependency tree to be flushed
 	 */
-	void AddDependencyTree(FAsyncPackage& ImportedPackage, TSet<FAsyncPackage*>& SearchedPackages);
+	void AddDependencyTree(FAsyncPackage& ImportedPackage, TSet<FAsyncPackage*>& SearchedPackages, FFlushTree* FlushTree);
 	/**
 	 * Adds a unique package to the list of packages to wait for until their linkers have been created.
 	 *
 	 * @param PendingImport Package imported either directly or by one of the imported packages
+	 * @param FlushTree Package dependency tree to be flushed
 	 */
-	bool AddUniqueLinkerDependencyPackage(FAsyncPackage& PendingImport);
+	bool AddUniqueLinkerDependencyPackage(FAsyncPackage& PendingImport, FFlushTree* FlushTree);
 	/**
 	 * Adds a package to the list of pending import packages.
 	 *
 	 * @param PendingImport Name of the package imported either directly or by one of the imported packages
+	 * @param FlushTree Package dependency tree to be flushed
 	 */
-	void AddImportDependency(const FName& PendingImport);
+	void AddImportDependency(const FName& PendingImport, FFlushTree* FlushTree);
 	/**
 	 * Removes references to any imported packages.
 	 */
@@ -379,4 +890,20 @@ private:
 #endif // PERF_TRACK_DETAILED_ASYNC_STATS
 };
 
+#if USE_EVENT_DRIVEN_ASYNC_LOAD
+struct FScopedAsyncPackageEvent
+{
+	/** Current scope package */
+	FAsyncPackage* Package;
+	/** Outer scope package */
+	FAsyncPackage* PreviousPackage;
+	
+	FScopedAsyncPackageEvent(FAsyncPackage* InPackage);
+	~FScopedAsyncPackageEvent();
+};
+#endif
+
+
+
+
 
diff --git a/Engine/Source/Runtime/CoreUObject/Public/Serialization/BulkData.h b/Engine/Source/Runtime/CoreUObject/Public/Serialization/BulkData.h
--- a/Engine/Source/Runtime/CoreUObject/Public/Serialization/BulkData.h
+++ b/Engine/Source/Runtime/CoreUObject/Public/Serialization/BulkData.h
@@ -31,7 +31,8 @@ enum EBulkDataFlags
 	BULKDATA_PayloadInSeperateFile				= 1 << 8,
 	/** If set, payload is compressed using platform specific bit window			*/
 	BULKDATA_SerializeCompressedBitWindow		= 1<<9,
-
+	/** Only used with USE_NEW_ASYNC_IO, which changes the default to inline unless you opt out */
+	BULKDATA_Force_NOT_InlinePayload = 1 << 10,
 };
 
 /**
diff --git a/Engine/Source/Runtime/CoreUObject/Public/Serialization/ObjectReader.h b/Engine/Source/Runtime/CoreUObject/Public/Serialization/ObjectReader.h
--- a/Engine/Source/Runtime/CoreUObject/Public/Serialization/ObjectReader.h
+++ b/Engine/Source/Runtime/CoreUObject/Public/Serialization/ObjectReader.h
@@ -59,6 +59,7 @@ public:
 	COREUOBJECT_API virtual FArchive& operator<<( FLazyObjectPtr& LazyObjectPtr ) override;
 	COREUOBJECT_API virtual FArchive& operator<<( FAssetPtr& AssetPtr ) override;
 	COREUOBJECT_API virtual FArchive& operator<<(FStringAssetReference& AssetPtr) override;
+	COREUOBJECT_API virtual FArchive& operator<<(struct FWeakObjectPtr& Value) override;
 	COREUOBJECT_API virtual FString GetArchiveName() const override;
 	//~ End FArchive Interface
 
diff --git a/Engine/Source/Runtime/CoreUObject/Public/Serialization/ObjectWriter.h b/Engine/Source/Runtime/CoreUObject/Public/Serialization/ObjectWriter.h
--- a/Engine/Source/Runtime/CoreUObject/Public/Serialization/ObjectWriter.h
+++ b/Engine/Source/Runtime/CoreUObject/Public/Serialization/ObjectWriter.h
@@ -37,6 +37,7 @@ public:
 	COREUOBJECT_API virtual FArchive& operator<<( FLazyObjectPtr& LazyObjectPtr ) override;
 	COREUOBJECT_API virtual FArchive& operator<<( FAssetPtr& AssetPtr ) override;
 	COREUOBJECT_API virtual FArchive& operator<<(FStringAssetReference& AssetPtr) override;
+	COREUOBJECT_API virtual FArchive& operator<<(struct FWeakObjectPtr& Value) override;
 	COREUOBJECT_API virtual FString GetArchiveName() const override;
 	//~ End FArchive Interface
 
diff --git a/Engine/Source/Runtime/CoreUObject/Public/Templates/Casts.h b/Engine/Source/Runtime/CoreUObject/Public/Templates/Casts.h
--- a/Engine/Source/Runtime/CoreUObject/Public/Templates/Casts.h
+++ b/Engine/Source/Runtime/CoreUObject/Public/Templates/Casts.h
@@ -77,7 +77,7 @@ struct TCastFlags
 template <typename From, typename To, bool bFromInterface = TIsIInterface<From>::Value, bool bToInterface = TIsIInterface<To>::Value, uint64 CastClass = TCastFlags<To>::Value>
 struct TGetCastType
 {
-#if UCLASS_FAST_ISA_IMPL == 2
+#if UCLASS_FAST_ISA_IMPL == UCLASS_ISA_INDEXTREE || UCLASS_FAST_ISA_IMPL == UCLASS_ISA_CLASSARRAY
 	static const ECastType Value = ECastType::UObjectToUObject;
 #else
 	static const ECastType Value = ECastType::FromCastFlags;
diff --git a/Engine/Source/Runtime/CoreUObject/Public/UObject/Class.h b/Engine/Source/Runtime/CoreUObject/Public/UObject/Class.h
--- a/Engine/Source/Runtime/CoreUObject/Public/UObject/Class.h
+++ b/Engine/Source/Runtime/CoreUObject/Public/UObject/Class.h
@@ -282,6 +282,7 @@ public:
 	virtual void FinishDestroy() override;
 	virtual void RegisterDependencies() override;
 	static void AddReferencedObjects(UObject* InThis, FReferenceCollector& Collector);
+	virtual void GetPreloadDependencies(TArray<UObject*>& OutDeps) override;
 
 	// UField interface.
 	virtual void AddCppProperty(UProperty* Property) override;
@@ -1767,11 +1768,11 @@ namespace EIncludeSuperFlag
 }
 
 
-#if UCLASS_FAST_ISA_IMPL & 2
+#if UCLASS_FAST_ISA_IMPL == UCLASS_ISA_INDEXTREE
 
 	class FFastIndexingClassTreeRegistrar
 	{
-	public:
+	protected:
 		COREUOBJECT_API FFastIndexingClassTreeRegistrar();
 		COREUOBJECT_API FFastIndexingClassTreeRegistrar(const FFastIndexingClassTreeRegistrar&);
 		COREUOBJECT_API ~FFastIndexingClassTreeRegistrar();
@@ -1792,6 +1793,33 @@ namespace EIncludeSuperFlag
 		uint32 ClassTreeNumChildren;
 	};
 
+#elif UCLASS_FAST_ISA_IMPL == UCLASS_ISA_CLASSARRAY
+
+	class FClassBaseChain
+	{
+	protected:
+		COREUOBJECT_API FClassBaseChain();
+		COREUOBJECT_API ~FClassBaseChain();
+
+		// Non-copyable
+		FClassBaseChain(const FClassBaseChain&) = delete;
+		FClassBaseChain& operator=(const FClassBaseChain&) = delete;
+
+		COREUOBJECT_API void ReinitializeBaseChainArray();
+
+		FORCEINLINE bool IsAUsingClassArray(const FClassBaseChain& Parent) const
+		{
+			int32 NumParentClassBasesInChainMinusOne = Parent.NumClassBasesInChainMinusOne;
+			return NumParentClassBasesInChainMinusOne <= NumClassBasesInChainMinusOne && ClassBaseChainArray[NumParentClassBasesInChainMinusOne] == &Parent;
+		}
+
+	private:
+		FClassBaseChain** ClassBaseChainArray;
+		int32 NumClassBasesInChainMinusOne;
+
+		friend class UClass;
+	};
+
 #endif
 
 
@@ -1799,15 +1827,17 @@ namespace EIncludeSuperFlag
  * An object class.
  */
 class COREUOBJECT_API UClass : public UStruct
-#if UCLASS_FAST_ISA_IMPL & 2
+#if UCLASS_FAST_ISA_IMPL == UCLASS_ISA_INDEXTREE
 	, private FFastIndexingClassTreeRegistrar
+#elif UCLASS_FAST_ISA_IMPL == UCLASS_ISA_CLASSARRAY
+	, private FClassBaseChain
 #endif
 {
 	DECLARE_CASTED_CLASS_INTRINSIC_NO_CTOR(UClass, UStruct, 0, TEXT("/Script/CoreUObject"), CASTCLASS_UClass, NO_API)
 	DECLARE_WITHIN(UPackage)
 
 public:
-#if UCLASS_FAST_ISA_IMPL & 2
+#if UCLASS_FAST_ISA_IMPL == UCLASS_ISA_INDEXTREE
 	friend class FFastIndexingClassTree;
 #endif
 	friend class FRestoreClassInfo;
@@ -2351,10 +2381,18 @@ public:
 	virtual bool HasInstrumentation() const { return false; }
 
 private:
-	#if UCLASS_FAST_ISA_IMPL & 2
+	#if UCLASS_FAST_ISA_IMPL == UCLASS_ISA_INDEXTREE
 		// For UObjectBaseUtility
 		friend class UObjectBaseUtility;
 		using FFastIndexingClassTreeRegistrar::IsAUsingFastTree;
+	#elif UCLASS_FAST_ISA_IMPL == UCLASS_ISA_CLASSARRAY
+		// For UObjectBaseUtility
+		friend class UObjectBaseUtility;
+		using FClassBaseChain::IsAUsingClassArray;
+		using FClassBaseChain::ReinitializeBaseChainArray;
+
+		friend class FClassBaseChain;
+		friend class FBlueprintCompileReinstancer;
 	#endif
 
 	// This signature intentionally hides the method declared in UObjectBaseUtility to make it private.
diff --git a/Engine/Source/Runtime/CoreUObject/Public/UObject/CoreNet.h b/Engine/Source/Runtime/CoreUObject/Public/UObject/CoreNet.h
--- a/Engine/Source/Runtime/CoreUObject/Public/UObject/CoreNet.h
+++ b/Engine/Source/Runtime/CoreUObject/Public/UObject/CoreNet.h
@@ -271,6 +271,7 @@ public:
 	virtual FArchive& operator<<(FName& Name) override;
 	virtual FArchive& operator<<(UObject*& Object) override;
 	virtual FArchive& operator<<(FStringAssetReference& Value) override;
+	virtual FArchive& operator<<(struct FWeakObjectPtr& Value) override;
 };
 
 
@@ -290,6 +291,7 @@ public:
 	virtual FArchive& operator<<(FName& Name) override;
 	virtual FArchive& operator<<(UObject*& Object) override;
 	virtual FArchive& operator<<(FStringAssetReference& Value) override;
+	virtual FArchive& operator<<(struct FWeakObjectPtr& Value) override;
 };
 
 bool FORCEINLINE NetworkGuidSetsAreSame( const TSet< FNetworkGUID >& A, const TSet< FNetworkGUID >& B )
diff --git a/Engine/Source/Runtime/CoreUObject/Public/UObject/GCObject.h b/Engine/Source/Runtime/CoreUObject/Public/UObject/GCObject.h
--- a/Engine/Source/Runtime/CoreUObject/Public/UObject/GCObject.h
+++ b/Engine/Source/Runtime/CoreUObject/Public/UObject/GCObject.h
@@ -23,6 +23,8 @@ class COREUOBJECT_API UGCObjectReferencer : public UObject
 	 * This is the list of objects that are referenced
 	 */
 	TArray<FGCObject*> ReferencedObjects;
+	/** Critical section used when adding and removing objects */
+	FCriticalSection ReferencedObjectsCritical;
 
 public:
 	DECLARE_CASTED_CLASS_INTRINSIC_WITH_API(UGCObjectReferencer, UObject, CLASS_Transient, TEXT("/Script/CoreUObject"), CASTCLASS_None, NO_API);
diff --git a/Engine/Source/Runtime/CoreUObject/Public/UObject/GarbageCollection.h b/Engine/Source/Runtime/CoreUObject/Public/UObject/GarbageCollection.h
--- a/Engine/Source/Runtime/CoreUObject/Public/UObject/GarbageCollection.h
+++ b/Engine/Source/Runtime/CoreUObject/Public/UObject/GarbageCollection.h
@@ -14,6 +14,11 @@
 COREUOBJECT_API DECLARE_LOG_CATEGORY_EXTERN(LogGarbage, Warning, All);
 DECLARE_STATS_GROUP(TEXT("Garbage Collection"), STATGROUP_GC, STATCAT_Advanced);
 
+/**
+ * Do extra checks on GC'd function references to catch uninitialized pointers?
+ */
+#define DO_POINTER_CHECKS_ON_GC WITH_EDITORONLY_DATA
+
 /*-----------------------------------------------------------------------------
 	Realtime garbage collection helper classes.
 -----------------------------------------------------------------------------*/
diff --git a/Engine/Source/Runtime/CoreUObject/Public/UObject/Linker.h b/Engine/Source/Runtime/CoreUObject/Public/UObject/Linker.h
--- a/Engine/Source/Runtime/CoreUObject/Public/UObject/Linker.h
+++ b/Engine/Source/Runtime/CoreUObject/Public/UObject/Linker.h
@@ -11,6 +11,10 @@
 
 DECLARE_LOG_CATEGORY_EXTERN(LogLinker, Log, All);
 
+#if !defined(USE_NEW_ASYNC_IO) || !defined(SPLIT_COOKED_FILES)
+#error "USE_NEW_ASYNC_IO and SPLIT_COOKED_FILES must be defined"
+#endif
+
 /**
  * Information about a compressed chunk in a file.
  */
diff --git a/Engine/Source/Runtime/CoreUObject/Public/UObject/LinkerLoad.h b/Engine/Source/Runtime/CoreUObject/Public/UObject/LinkerLoad.h
--- a/Engine/Source/Runtime/CoreUObject/Public/UObject/LinkerLoad.h
+++ b/Engine/Source/Runtime/CoreUObject/Public/UObject/LinkerLoad.h
@@ -39,13 +39,36 @@ struct FDependencyRef
 	friend uint32 GetTypeHash( const FDependencyRef& Ref  );
 };
 
+/** Helper struct to keep track of the first time CreateImport() is called in the current callstack. */
+struct FScopedCreateImportCounter
+{
+	/**
+	*	Constructor. Called upon CreateImport() entry.
+	*	@param Linker	- Current Linker
+	*	@param Index	- Index of the current Import
+	*/
+	FScopedCreateImportCounter(FLinkerLoad* Linker, int32 Index);
+
+	/** Destructor. Called upon CreateImport() exit. */
+	~FScopedCreateImportCounter();
+
+	/** Previously stored linker */
+	FLinkerLoad* PreviousLinker;
+	/** Previously stored index */
+	int32 PreviousIndex;
+};
+
 /**
  * Handles loading Unreal package files, including reading UObject data from disk.
  */
 #if USE_NEW_ASYNC_IO
-	class FArchiveAsync2;
+class FArchiveAsync2;
+#endif
+class FLinkerLoad
+#if USE_NEW_ASYNC_IO
+	final
 #endif
-class FLinkerLoad : public FLinker, public FArchiveUObject
+	: public FLinker, public FArchiveUObject
 {
 	// Friends.
 	friend class UObject;
@@ -90,6 +113,11 @@ public:
 	bool					bHaveImportsBeenVerified;
 	/** Indicates that this linker was created for a dynamic class package and will not use Loader */
 	bool					bDynamicClassLinker;
+#if USE_EVENT_DRIVEN_ASYNC_LOAD
+	UObject*				TemplateForGetArchetypeFromLoader;
+	bool					bForceSimpleIndexToObject;
+	bool					bLockoutLegacyOperations;
+#endif
 #if USE_NEW_ASYNC_IO
 	/** True if Loader is FArchiveAsync2  */
 	bool					bLoaderIsFArchiveAsync2;
@@ -110,6 +138,11 @@ public:
 	/** Hash table for exports.																								*/
 	int32						ExportHash[256];
 
+	/**
+	* List of imports and exports that must be serialized before other exports...all packed together, see FirstExportDependency
+	*/
+	TArray<FPackageIndex> PreloadDependencies;
+
 	/** OldClassName to NewClassName for ImportMap */
 	static TMap<FName, FName> ObjectNameRedirects;
 	/** Additional info for some ObjectName redirects to also redirect the class and class package  */
@@ -181,8 +214,6 @@ private:
 	int32						ImportMapIndex;
 	/** Current index into export map, used by async linker creation for spreading out serializing exportmap entries.			*/
 	int32						ExportMapIndex;
-	/** Current index into export map, used by async linker creation for spreading out serializing exportmap entries.			*/
-	int32						FirstNotLoadedExportMapIndex;
 	/** Current index into depends map, used by async linker creation for spreading out serializing dependsmap entries.			*/
 	int32						DependsMapIndex;
 	/** Current index into export hash map, used by async linker creation for spreading out hashing exports.					*/
@@ -222,6 +253,7 @@ private:
 	/** Id of the thread that created this linker. This is to guard against using this linker on other threads than the one it was created on **/
 	int32					OwnerThread;
 
+#if !USE_NEW_ASYNC_IO
 	/**
 	 * Helper struct to keep track of background file reads
 	 */
@@ -257,6 +289,17 @@ private:
 	/** Map that keeps track of any precached full package reads															*/
 	static TMap<FString, FPackagePrecacheInfo> PackagePrecacheMap;
 
+	/**
+	* Fills in the passed in TArray with the packages that are in its PrecacheMap
+	*
+	* @param TArray<FString> to be populated
+	*/
+public:
+	COREUOBJECT_API static void GetListOfPackagesInPackagePrecacheMap(TArray<FString>& ListOfPackages);
+private:
+
+#endif
+
 	/** Allows access to UTexture2D::StaticClass() without linking Core with Engine											*/
 	static UClass* UTexture2DStaticClass;
 
@@ -269,7 +312,7 @@ private:
 	/** Test whether we should report progress or not */
 	FORCEINLINE bool ShouldReportProgress() const
 	{
-		return !IsAsyncLoading() && ( LoadFlags & ( LOAD_Quiet | LOAD_SeekFree ) ) == 0;
+		return !IsAsyncLoading() && (LoadFlags & (LOAD_Quiet | LOAD_Async)) == 0;
 	}
 #endif 
 
@@ -300,14 +343,6 @@ public:
 	 */
 	COREUOBJECT_API ELinkerStatus FixupExportMap();
 
-
-	/**
- 	 * Fills in the passed in TArray with the packages that are in its PrecacheMap
-	 *
-	 * @param TArray<FString> to be populated
-	 */
-	COREUOBJECT_API static void GetListOfPackagesInPackagePrecacheMap( TArray<FString>& ListOfPackages );
-
 	/**
 	 * Returns whether linker has finished (potentially) async initialization.
 	 *
@@ -348,6 +383,7 @@ public:
 	COREUOBJECT_API FName GetExportClassPackage( int32 i );
 	virtual FString GetArchiveName() const override;
 
+#if WITH_EDITORONLY_DATA
 	/**
 	 * Recursively gathers the dependencies of a given export (the recursive chain of imports
 	 * and their imports, and so on)
@@ -367,6 +403,7 @@ public:
 	 * @param bSkipLoadedObjects Whether to skip already loaded objects when gathering dependencies
 	 */
 	COREUOBJECT_API void GatherImportDependencies(int32 ImportIndex, TSet<FDependencyRef>& Dependencies, bool bSkipLoadedObjects=true);
+#endif
 
 	/**
 	 * A wrapper around VerifyImportInner. If the VerifyImportInner (previously VerifyImport) fails, this function
@@ -458,12 +495,14 @@ public:
 	 */
 	COREUOBJECT_API bool WillTextureBeLoaded( UClass* Class, int32 ExportIndex );
 
+#if !USE_NEW_ASYNC_IO
 	/**
 	 * Kick off an async load of a package file into memory
 	 * 
 	 * @param PackageName Name of package to read in. Must be the same name as passed into LoadPackage/CreateLinker
 	 */
 	COREUOBJECT_API static void AsyncPreloadPackage(const TCHAR* PackageName);
+#endif
 
 	/**
 	 * Called when an object begins serializing property data using script serialization.
@@ -475,6 +514,10 @@ public:
 	 */
 	virtual void MarkScriptSerializationEnd( const UObject* Obj ) override;
 
+#if USE_EVENT_DRIVEN_ASYNC_LOAD
+	virtual UObject* GetArchetypeFromLoader(const UObject* Obj) override;
+#endif
+
 	/**
 	 * Looks for an existing linker for the given package, without trying to make one if it doesn't exist
 	 */
@@ -602,7 +645,10 @@ private:
 	 * @param	PrecacheSize	Number of bytes to precache
 	 * @return	false if precache operation is still pending, true otherwise
 	 */
-	virtual bool Precache( int64 PrecacheOffset, int64 PrecacheSize ) override;
+	FORCEINLINE virtual bool Precache(int64 PrecacheOffset, int64 PrecacheSize) override
+	{
+		return bDynamicClassLinker || Loader->Precache(PrecacheOffset, PrecacheSize);
+	}
 
 #if WITH_EDITOR
 	/**
@@ -640,17 +686,73 @@ public:
 
 private:
 
-	void Seek( int64 InPos ) override;
-	int64 Tell() override;
-	int64 TotalSize() override;
+	FORCEINLINE virtual void Seek(int64 InPos) override
+	{
+		Loader->Seek(InPos);
+	}
+	FORCEINLINE virtual int64 Tell() override
+	{
+		return Loader->Tell();
+	}
+	FORCEINLINE virtual int64 TotalSize() override
+	{
+		return Loader->TotalSize();
+	}
+
 	// this fixes the warning : 'ULinkerSave::Serialize' hides overloaded virtual function
 	using FLinker::Serialize;
-	virtual void Serialize( void* V, int64 Length ) override;
+	FORCEINLINE virtual void Serialize(void* V, int64 Length) override
+	{
+		checkSlow(FPlatformTLS::GetCurrentThreadId() == OwnerThread);
+		Loader->Serialize(V, Length);
+	}
 	using FArchiveUObject::operator<<; // For visibility of the overloads we don't override
-	virtual FArchive& operator<<( UObject*& Object ) override;
-	virtual FArchive& operator<<( FLazyObjectPtr& LazyObjectPtr) override;
-	virtual FArchive& operator<<( FAssetPtr& AssetPtr) override;
-	virtual FArchive& operator<<( FName& Name ) override;
+	virtual FArchive& operator<<(UObject*& Object) override;
+	FORCEINLINE virtual FArchive& operator<<(FLazyObjectPtr& LazyObjectPtr) override
+	{
+		FArchive& Ar = *this;
+		FUniqueObjectGuid ID;
+		Ar << ID;
+		LazyObjectPtr = ID;
+		return Ar;
+	}
+
+	FORCEINLINE virtual FArchive& operator<<(FAssetPtr& AssetPtr) override
+	{
+		FArchive& Ar = *this;
+		FStringAssetReference ID;
+		ID.Serialize(Ar);
+		AssetPtr = ID;
+		return Ar;
+	}
+	void BadNameIndexError(NAME_INDEX NameIndex);
+	FORCEINLINE virtual FArchive& operator<<(FName& Name) override
+	{
+		Name = NAME_None;
+		FArchive& Ar = *this;
+		NAME_INDEX NameIndex;
+		Ar << NameIndex;
+		int32 Number;
+		Ar << Number;
+
+		if (!NameMap.IsValidIndex(NameIndex))
+		{
+			BadNameIndexError(NameIndex);
+			ArIsError = true;
+			ArIsCriticalError = true;
+		}
+		else
+		{
+			// if the name wasn't loaded (because it wasn't valid in this context)
+			const FName& MappedName = NameMap[NameIndex];
+			if (!MappedName.IsNone())
+			{
+				// simply create the name from the NameMap's name and the serialized instance number
+				Name = FName(MappedName, Number);
+			}
+		}
+		return *this;
+	}
 
 	/**
 	 * Safely verify that an import in the ImportMap points to a good object. This decides whether or not
@@ -678,7 +780,11 @@ private:
 	 *
 	 * @return	new FLinkerLoad object for Parent/ Filename
 	 */
-	COREUOBJECT_API static FLinkerLoad* CreateLinkerAsync( UPackage* Parent, const TCHAR* Filename, uint32 LoadFlags );
+	COREUOBJECT_API static FLinkerLoad* CreateLinkerAsync(UPackage* Parent, const TCHAR* Filename, uint32 LoadFlags
+#if USE_EVENT_DRIVEN_ASYNC_LOAD
+		, TFunction<void()>&& InSummaryReadyCallback
+#endif		
+	);
 
 	/**
 	 * Ticks an in-flight linker and spends InTimeLimit seconds on creation. This is a soft time limit used
@@ -715,7 +821,11 @@ protected: // Daniel L: Made this protected so I can override the constructor an
 	/**
 	 * Creates loader used to serialize content.
 	 */
-	ELinkerStatus CreateLoader();
+	ELinkerStatus CreateLoader(
+#if USE_EVENT_DRIVEN_ASYNC_LOAD
+		TFunction<void()>&& InSummaryReadyCallback
+#endif
+	);
 private:
 	/**
 	 * Serializes the package file summary.
@@ -747,18 +857,10 @@ private:
 	 */
 	ELinkerStatus SerializeExportMap();
 
-#if WITH_ENGINE
-	/**
-	 * Kicks off async memory allocations for all textures that will be loaded from this package.
-	 */
-	ELinkerStatus StartTextureAllocation();
-#endif
-
-	/**
-	 * Serializes the import map.
-	 */
 	ELinkerStatus SerializeDependsMap();
 
+	ELinkerStatus SerializePreloadDependencies();
+
 public:
 	/**
 	 * Serializes the gatherable text data container.
@@ -944,7 +1046,7 @@ private:
 	/** Finds import, tries to fall back to dynamic class if the object could not be found */
 	UObject* FindImport(UClass* ImportClass, UObject* ImportOuter, const TCHAR* Name);
 	/** Finds import, tries to fall back to dynamic class if the object could not be found */
-	UObject* FindImportFast(UClass* ImportClass, UObject* ImportOuter, FName Name);
+	static UObject* FindImportFast(UClass* ImportClass, UObject* ImportOuter, FName Name);
 
 	/** Fills all necessary information for constructing dynamic type package linker */
 	void CreateDynamicTypeLoader();
diff --git a/Engine/Source/Runtime/CoreUObject/Public/UObject/LinkerSave.h b/Engine/Source/Runtime/CoreUObject/Public/UObject/LinkerSave.h
--- a/Engine/Source/Runtime/CoreUObject/Public/UObject/LinkerSave.h
+++ b/Engine/Source/Runtime/CoreUObject/Public/UObject/LinkerSave.h
@@ -29,6 +29,9 @@ public:
 	/** The archive that actually writes the data to disk. */
 	FArchive* Saver;
 
+	FPackageIndex CurrentlySavingExport;
+	TArray<FPackageIndex> DepListForErrorChecking;
+
 	/** Index array - location of the resource for a UObject is stored in the ObjectIndices array using the UObject's Index */
 	TMap<UObject *,FPackageIndex> ObjectIndicesMap;
 
diff --git a/Engine/Source/Runtime/CoreUObject/Public/UObject/Object.h b/Engine/Source/Runtime/CoreUObject/Public/UObject/Object.h
--- a/Engine/Source/Runtime/CoreUObject/Public/UObject/Object.h
+++ b/Engine/Source/Runtime/CoreUObject/Public/UObject/Object.h
@@ -367,6 +367,13 @@ public:
 		return false;
 	}
 
+	/**
+	* Called during cooking. Must return all objects that will be Preload()ed when this is serialized at load time
+	*
+	* @param	OutDeps				all objects that will be preloaded when this is serialized at load time
+	*/
+	virtual void GetPreloadDependencies(TArray<UObject*>& OutDeps);
+
 	/**
 	*	Update the list of classes that we should exclude from dedicated server builds
 	*/
diff --git a/Engine/Source/Runtime/CoreUObject/Public/UObject/ObjectMacros.h b/Engine/Source/Runtime/CoreUObject/Public/UObject/ObjectMacros.h
--- a/Engine/Source/Runtime/CoreUObject/Public/UObject/ObjectMacros.h
+++ b/Engine/Source/Runtime/CoreUObject/Public/UObject/ObjectMacros.h
@@ -23,13 +23,17 @@ typedef	uint64 ScriptPointerType;
 /** Set this to 0 to disable UObject thread safety features */
 #define THREADSAFE_UOBJECTS 1
 
-// 1 = old IsA behavior
-// 2 = new IsA behavior
-// 3 = old IsA behavior with checks against the new behavior
+// Enumeration of different methods of determining class relationships.
+#define UCLASS_ISA_OUTERWALK  1 // walks the class chain                                         - original IsA behavior
+#define UCLASS_ISA_INDEXTREE  2 // uses position in an index-based tree                          - thread-unsafe if one thread does a parental test while the tree is changing, e.g. by async loading a class
+#define UCLASS_ISA_CLASSARRAY 3 // stores an array of parents per class and uses this to compare - faster than 1, slower but comparable with 2, and thread-safe
+
+// UCLASS_FAST_ISA_IMPL sets which implementation of IsA to use.
+#define UCLASS_FAST_ISA_IMPL UCLASS_ISA_CLASSARRAY
+
+// UCLASS_FAST_ISA_COMPARE_WITH_OUTERWALK, if set, does a checked comparison of the current implementation against the outer walk - used for testing.
 #if 0//!(UE_BUILD_SHIPPING || UE_BUILD_TEST)
-	#define UCLASS_FAST_ISA_IMPL 3
-#else
-	#define UCLASS_FAST_ISA_IMPL 2
+	#define UCLASS_FAST_ISA_COMPARE_WITH_OUTERWALK 1
 #endif
 
 /*-----------------------------------------------------------------------------
@@ -42,7 +46,7 @@ typedef	uint64 ScriptPointerType;
 enum ELoadFlags
 {
 	LOAD_None						= 0x00000000,	// No flags.
-	LOAD_SeekFree					= 0x00000001,	// Loads the package via the seek free loading path/ reader
+	LOAD_Async					= 0x00000001,	// Loads the package using async loading path/ reader
 	LOAD_NoWarn						= 0x00000002,	// Don't display warning if load fails.
 	LOAD_EditorOnly					= 0x00000004, // Load for editor-only purposes and by editor-only code
 	LOAD_ResolvingDeferredExports	= 0x00000008,	// Denotes that we should not defer export loading (as we're resolving them)
diff --git a/Engine/Source/Runtime/CoreUObject/Public/UObject/ObjectResource.h b/Engine/Source/Runtime/CoreUObject/Public/UObject/ObjectResource.h
--- a/Engine/Source/Runtime/CoreUObject/Public/UObject/ObjectResource.h
+++ b/Engine/Source/Runtime/CoreUObject/Public/UObject/ObjectResource.h
@@ -4,6 +4,10 @@
 
 #include "UObjectBaseUtility.h"
 
+#if !defined(USE_NEW_ASYNC_IO) || !defined(SPLIT_COOKED_FILES)
+#error "USE_NEW_ASYNC_IO and SPLIT_COOKED_FILES must be defined"
+#endif
+
 /**
  * Wrapper for index into a ULnker's ImportMap or ExportMap.
  * Values greater than zero indicate that this is an index into the ExportMap.  The
@@ -121,6 +125,11 @@ public:
 		Ar << Value.Index;
 		return Ar;
 	}
+
+	FORCEINLINE friend uint32 GetTypeHash(const FPackageIndex& In)
+	{
+		return uint32(In.Index);
+	}
 };
 
 
@@ -188,6 +197,14 @@ struct FObjectExport : public FObjectResource
 	 */
 	FPackageIndex 	SuperIndex;
 
+	/**
+	* Location of the resource for this export's template/archetypes.  Only used
+	* in the new cooked loader. A value of zero indicates that the value of GetArchetype
+	* was zero at cook time, which is more or less impossible and checked.
+	* Serialized
+	*/
+	FPackageIndex 	TemplateIndex;
+
 	/**
 	 * The object flags for the UObject represented by this resource.  Only flags that
 	 * match the RF_Load combination mask will be loaded from disk and applied to the UObject.
@@ -292,6 +309,16 @@ struct FObjectExport : public FObjectResource
 	 */
 	uint32			PackageFlags;
 
+	/**
+	* The export table must serialize as a fixed size, this is use to index into a long list, which is later loaded into the array. -1 means dependencies are not present
+	* These are contiguous blocks, so CreateBeforeSerializationDependencies starts at FirstExportDependency + SerializationBeforeSerializationDependencies
+	*/
+	int32 FirstExportDependency;
+	int32 SerializationBeforeSerializationDependencies;
+	int32 CreateBeforeSerializationDependencies;
+	int32 SerializationBeforeCreateDependencies;
+	int32 CreateBeforeCreateDependencies;
+
 	/**
 	 * Constructors
 	 */
@@ -343,6 +370,11 @@ struct FObjectImport : public FObjectResource
 	 */
 	int32             SourceIndex;
 
+#if USE_EVENT_DRIVEN_ASYNC_LOAD
+	bool			bImportPackageHandled;
+	bool			bImportSearchedFor;
+#endif
+
 	/**
 	 * Constructors
 	 */
diff --git a/Engine/Source/Runtime/CoreUObject/Public/UObject/PackageFileSummary.h b/Engine/Source/Runtime/CoreUObject/Public/UObject/PackageFileSummary.h
--- a/Engine/Source/Runtime/CoreUObject/Public/UObject/PackageFileSummary.h
+++ b/Engine/Source/Runtime/CoreUObject/Public/UObject/PackageFileSummary.h
@@ -33,166 +33,14 @@ struct FGenerationInfo
 	void Serialize(FArchive& Ar, const struct FPackageFileSummary& Summary);
 };
 
-#if WITH_ENGINE
-/**
- * Information about the textures stored in the package.
- */
-struct FTextureAllocations
-{
-	/**
-	 * Stores an export index for each texture of a certain type (size, format, etc).
-	 */
-	struct FTextureType
-	{
-		FTextureType();
-		COREUOBJECT_API FTextureType( int32 SizeX, int32 SizeY, int32 NumMips, uint32 Format, uint32 TexCreateFlags );
-
-		/** Width of the largest mip-level stored in the package. */
-		int32				SizeX;
-		/** Height of the largest mip-level stored in the package. */
-		int32				SizeY;
-		/** Number of mips */
-		int32				NumMips;
-		/** Texture format */
-		uint32			Format;
-		/** ETextureCreateFlags bit flags */
-		uint32			TexCreateFlags;
-		/** Index into the package ExportMap. */
-		TArray<int32>		ExportIndices;
-
-		/** Not serialized. ResourceMems are constructed on load. */
-		TArray<FTexture2DResourceMem*>	Allocations;
-		/** Note serialized. Number of ExportIndices processed during load. */
-		int32				NumExportIndicesProcessed;
-
-		/**
-		 * Checks whether all potential allocations for this TextureType have been considered yet (during load).
-		 *
-		 * @return true if all allocations have been started
-		 */
-		bool	HaveAllAllocationsBeenConsidered() const
-		{
-			return NumExportIndicesProcessed == ExportIndices.Num();
-		}
-
-		/**
-		 * Serializes an FTextureType
-		 */
-		friend FArchive& operator<<( FArchive& Ar, FTextureAllocations::FTextureType& TextureType );
-	};
-
-	FTextureAllocations()
-	:	PendingAllocationSize(0)
-	,	NumTextureTypesConsidered(0)
-	{
-	}
-
-	/** Array of texture types in the package. */
-	TArray<FTextureType>	TextureTypes;
-	/** Number of allocations that haven't been completed yet. */
-	FThreadSafeCounter		PendingAllocationCount;
-	/** Number of allocated bytes that has yet to be claimed by UTexture2D serialization. */
-	int32						PendingAllocationSize;
-	/** Number of TextureTypes that have been considered for potential allocations so far, during load. */
-	int32						NumTextureTypesConsidered;
-
-	/**
-	 * Finds a suitable ResourceMem allocation, removes it from this container and returns it to the user.
-	 *
-	 * @param SizeX				Width of texture
-	 * @param SizeY				Height of texture
-	 * @param NumMips			Number of mips
-	 * @param Format			Texture format (EPixelFormat)
-	 * @param TexCreateFlags	ETextureCreateFlags bit flags
-	 **/
-	FTexture2DResourceMem*	FindAndRemove( int32 SizeX, int32 SizeY, int32 NumMips, uint32 Format, uint32 TexCreateFlags );
-
-	/**
-	 * Cancels any pending ResourceMem allocation that hasn't been claimed by a texture yet,
-	 * just in case there are any mismatches at run-time.
-	 *
-	 * @param bCancelEverything		If true, cancels all allocations. If false, only cancels allocations that haven't been completed yet.
-	 */
-	void	CancelRemainingAllocations( bool bCancelEverything );
-
-	/**
-	 * Checks if all allocations that should be started have been started (during load).
-	 *
-	 * @return true if all allocations have been started
-	 */
-	bool	HaveAllAllocationsBeenConsidered() const
-	{
-		return NumTextureTypesConsidered == TextureTypes.Num();
-	}
-
-	/**
-	 * Checks if all ResourceMem allocations has completed.
-	 *
-	 * @return true if all ResourceMem allocations has completed
-	 */
-	bool	HasCompleted() const
-	{
-		return PendingAllocationCount.GetValue() == 0;
-	}
-
-	/**
-	 * Checks if all ResourceMem allocations have been claimed by a texture.
-	 *
-	 * @return true if there are no more pending ResourceMem allocations waiting for a texture to claim it
-	 */
-	bool	HasBeenFullyClaimed() const
-	{
-		return PendingAllocationSize == 0;
-	}
-
-	/**
-	 * Serializes an FTextureType
-	 */
-	friend FArchive& operator<<( FArchive& Ar, FTextureAllocations::FTextureType& TextureAllocationType );
-
-	/**
-	 * Serializes an FTextureAllocations struct
-	 */
-	friend FArchive& operator<<( FArchive& Ar, FTextureAllocations& TextureAllocations );
-
-	FTextureAllocations( const FTextureAllocations& Other );
-	void operator=(const FTextureAllocations& Other);
-
-private:
-	/**
-	 * Finds a texture type that matches the given specifications.
-	 *
-	 * @param SizeX				Width of the largest mip-level stored in the package
-	 * @param SizeY				Height of the largest mip-level stored in the package
-	 * @param NumMips			Number of mips
-	 * @param Format			Texture format (EPixelFormat)
-	 * @param TexCreateFlags	ETextureCreateFlags bit flags
-	 * @return					Matching texture type, or NULL if none was found
-	 */
-	FTextureType*	FindTextureType( int32 SizeX, int32 SizeY, int32 NumMips, uint32 Format, uint32 TexCreateFlags );
-
-	/**
-	 * Adds a dummy export index (-1) for a specified texture type.
-	 * Creates the texture type entry if needed.
-	 *
-	 * @param SizeX				Width of the largest mip-level stored in the package
-	 * @param SizeY				Height of the largest mip-level stored in the package
-	 * @param NumMips			Number of mips
-	 * @param Format			Texture format (EPixelFormat)
-	 * @param TexCreateFlags	ETextureCreateFlags bit flags
-	 */
-	void			AddResourceMemInfo( int32 SizeX, int32 SizeY, int32 NumMips, uint32 Format, uint32 TexCreateFlags );
-};
-#endif		// WITH_ENGINE
-
 /**
  * A "table of contents" for an Unreal package file.  Stored at the top of the file.
  */
 struct FPackageFileSummary
 {
 	/**
-	 * Magic tag compared against PACKAGE_FILE_TAG to ensure that package is an Unreal package.
-	 */
+	* Magic tag compared against PACKAGE_FILE_TAG to ensure that package is an Unreal package.
+	*/
 	int32		Tag;
 
 private:
@@ -205,59 +53,59 @@ private:
 
 public:
 	/**
-	 * Total size of all information that needs to be read in to create a FLinkerLoad. This includes
-	 * the package file summary, name table and import & export maps.
-	 */
+	* Total size of all information that needs to be read in to create a FLinkerLoad. This includes
+	* the package file summary, name table and import & export maps.
+	*/
 	int32		TotalHeaderSize;
 
 	/**
-	 * The flags for the package
-	 */
+	* The flags for the package
+	*/
 	uint32	PackageFlags;
 
 	/**
-	 * The Generic Browser folder name that this package lives in
-	 */
+	* The Generic Browser folder name that this package lives in
+	*/
 	FString	FolderName;
 
 	/**
-	 * Number of names used in this package
-	 */
+	* Number of names used in this package
+	*/
 	int32		NameCount;
 
 	/**
-	 * Location into the file on disk for the name data
-	 */
+	* Location into the file on disk for the name data
+	*/
 	int32 	NameOffset;
 
 	/**
-	 * Number of gatherable text data items in this package
-	 */
+	* Number of gatherable text data items in this package
+	*/
 	int32	GatherableTextDataCount;
 
 	/**
-	 * Location into the file on disk for the gatherable text data items
-	 */
+	* Location into the file on disk for the gatherable text data items
+	*/
 	int32 	GatherableTextDataOffset;
-	
+
 	/**
-	 * Number of exports contained in this package
-	 */
+	* Number of exports contained in this package
+	*/
 	int32		ExportCount;
 
 	/**
-	 * Location into the file on disk for the ExportMap data
-	 */
+	* Location into the file on disk for the ExportMap data
+	*/
 	int32		ExportOffset;
 
 	/**
-	 * Number of imports contained in this package
-	 */
+	* Number of imports contained in this package
+	*/
 	int32		ImportCount;
 
 	/**
-	 * Location into the file on disk for the ImportMap data
-	 */
+	* Location into the file on disk for the ImportMap data
+	*/
 	int32		ImportOffset;
 
 	/**
@@ -266,90 +114,89 @@ public:
 	int32		DependsOffset;
 
 	/**
-	 * Number of references contained in this package
-	 */
+	* Number of references contained in this package
+	*/
 	int32		StringAssetReferencesCount;
 
 	/**
-	 * Location into the file on disk for the string asset references map data
-	 */
+	* Location into the file on disk for the string asset references map data
+	*/
 	int32		StringAssetReferencesOffset;
 
 	/**
-	 * Thumbnail table offset
-	 */
+	* Thumbnail table offset
+	*/
 	int32		ThumbnailTableOffset;
 
 	/**
-	 * Current id for this package
-	 */
+	* Current id for this package
+	*/
 	FGuid	Guid;
 
 	/**
-	 * Data about previous versions of this package
-	 */
+	* Data about previous versions of this package
+	*/
 	TArray<FGenerationInfo> Generations;
 
 	/**
-	 * Engine version this package was saved with. For hotfix releases and engine versions which maintain strict binary compatibility with another version, this may differ from CompatibleWithEngineVersion.
-	 */
+	* Engine version this package was saved with. For hotfix releases and engine versions which maintain strict binary compatibility with another version, this may differ from CompatibleWithEngineVersion.
+	*/
 	FEngineVersion SavedByEngineVersion;
 
 	/**
-	 * Engine version this package is compatible with. See SavedByEngineVersion.
-	 */
+	* Engine version this package is compatible with. See SavedByEngineVersion.
+	*/
 	FEngineVersion CompatibleWithEngineVersion;
 
 	/**
-	 * Flags used to compress the file on save and uncompress on load.
-	 */
+	* Flags used to compress the file on save and uncompress on load.
+	*/
 	uint32	CompressionFlags;
 
 	/**
-	 * Value that is used to determine if the package was saved by Epic (or licensee) or by a modder, etc
-	 */
+	* Value that is used to determine if the package was saved by Epic (or licensee) or by a modder, etc
+	*/
 	uint32	PackageSource;
 
 	/**
-	 * Array of compressed chunks in case this package was stored compressed.
-	 */
+	* Array of compressed chunks in case this package was stored compressed.
+	*/
 	TArray<FCompressedChunk> CompressedChunks;
 
 	/**
-	 * List of additional packages that are needed to be cooked for this package (ie streaming levels)
-	 */
+	* List of additional packages that are needed to be cooked for this package (ie streaming levels)
+	*/
 	TArray<FString>	AdditionalPackagesToCook;
 
-	/** 
-	 * If true, this file will not be saved with version numbers or was saved without version numbers. In this case they are assumed to be the current version. 
-	 * This is only used for full cooks for distribution because it is hard to guarantee correctness 
-	 **/
-	bool bUnversioned;
-
-#if WITH_ENGINE
 	/**
-	 * Information about the textures stored in the package.
-	 */
-	FTextureAllocations	TextureAllocations;
-#endif		// WITH_ENGINE
+	* If true, this file will not be saved with version numbers or was saved without version numbers. In this case they are assumed to be the current version.
+	* This is only used for full cooks for distribution because it is hard to guarantee correctness
+	**/
+	bool bUnversioned;
 
 	/**
-	 * Location into the file on disk for the asset registry tag data
-	 */
+	* Location into the file on disk for the asset registry tag data
+	*/
 	int32 	AssetRegistryDataOffset;
 
 	/** Offset to the location in the file where the bulkdata starts */
 	int64	BulkDataStartOffset;
 	/**
-	 * Offset to the location in the file where the FWorldTileInfo data starts
-	 */
+	* Offset to the location in the file where the FWorldTileInfo data starts
+	*/
 	int32 	WorldTileInfoDataOffset;
 
 	/**
-	 * Streaming install ChunkIDs
-	 */
+	* Streaming install ChunkIDs
+	*/
 	TArray<int32>	ChunkIDs;
 
+	int32		PreloadDependencyCount;
+
+	/**
+	* Location into the file on disk for the preload dependency data
+	*/
+	int32		PreloadDependencyOffset;
 
 	/** Constructor */
 	COREUOBJECT_API FPackageFileSummary();
@@ -382,5 +229,6 @@ public:
 	}
 
 	/** I/O function */
-	friend COREUOBJECT_API FArchive& operator<<( FArchive& Ar, FPackageFileSummary& Sum );
+	friend COREUOBJECT_API FArchive& operator<<(FArchive& Ar, FPackageFileSummary& Sum);
 };
+
diff --git a/Engine/Source/Runtime/CoreUObject/Public/UObject/TextProperty.h b/Engine/Source/Runtime/CoreUObject/Public/UObject/TextProperty.h
--- a/Engine/Source/Runtime/CoreUObject/Public/UObject/TextProperty.h
+++ b/Engine/Source/Runtime/CoreUObject/Public/UObject/TextProperty.h
@@ -31,10 +31,7 @@ public:
 	virtual void SerializeItem( FArchive& Ar, void* Value, void const* Defaults ) const override;
 	virtual void ExportTextItem( FString& ValueStr, const void* PropertyValue, const void* DefaultValue, UObject* Parent, int32 PortFlags, UObject* ExportRootScope ) const override;
 	virtual const TCHAR* ImportText_Internal( const TCHAR* Buffer, void* Data, int32 PortFlags, UObject* OwnerObject, FOutputDevice* ErrorText ) const override;
-	virtual FString GetCPPTypeForwardDeclaration() const override
-	{
-		return FString();
-	}
+	virtual FString GetCPPTypeForwardDeclaration() const override;
 	// End of UProperty interface
 
 	/** Generate the correct C++ code for the given text value */
diff --git a/Engine/Source/Runtime/CoreUObject/Public/UObject/UObjectBaseUtility.h b/Engine/Source/Runtime/CoreUObject/Public/UObject/UObjectBaseUtility.h
--- a/Engine/Source/Runtime/CoreUObject/Public/UObject/UObjectBaseUtility.h
+++ b/Engine/Source/Runtime/CoreUObject/Public/UObject/UObjectBaseUtility.h
@@ -428,12 +428,16 @@ public:
 	/**
 	 * @return	true if this object is of the specified type.
 	 */
-	#if UCLASS_FAST_ISA_IMPL == 2
+	#if !UCLASS_FAST_ISA_COMPARE_WITH_OUTERWALK && UCLASS_FAST_ISA_IMPL != UCLASS_ISA_OUTERWALK
 	private:
 		template <typename ClassType>
 		static FORCEINLINE bool IsAWorkaround(const ClassType* ObjClass, const ClassType* TestCls)
 		{
-			return ObjClass->IsAUsingFastTree(*TestCls);
+			#if UCLASS_FAST_ISA_IMPL == UCLASS_ISA_INDEXTREE
+				return ObjClass->IsAUsingFastTree(*TestCls);
+			#elif UCLASS_FAST_ISA_IMPL == UCLASS_ISA_CLASSARRAY
+				return ObjClass->IsAUsingClassArray(*TestCls);
+			#endif
 		}
 
 	public:
@@ -573,6 +577,28 @@ public:
 	}
 };
 
+FORCEINLINE bool IsPossiblyAllocatedUObjectPointer(UObject* Ptr)
+{
+	auto CountByteValues = [](UPTRINT Val, UPTRINT ByteVal) -> int32
+	{
+		int32 Result = 0;
+
+		for (int32 I = 0; I != sizeof(UPTRINT); ++I)
+		{
+			if ((Val & 0xFF) == ByteVal)
+			{
+				++Result;
+			}
+			Val >>= 8;
+		}
+
+		return Result;
+	};
+
+	UPTRINT PtrVal = (UPTRINT)Ptr;
+	return PtrVal >= 0x1000 && CountByteValues(PtrVal, 0xCD) < sizeof(UPTRINT) / 2;
+}
+
 /**
  * Returns the name of this object (with no path information)
  * @param Object object to retrieve the name for; NULL gives "None"
diff --git a/Engine/Source/Runtime/CoreUObject/Public/UObject/UObjectGlobals.h b/Engine/Source/Runtime/CoreUObject/Public/UObject/UObjectGlobals.h
--- a/Engine/Source/Runtime/CoreUObject/Public/UObject/UObjectGlobals.h
+++ b/Engine/Source/Runtime/CoreUObject/Public/UObject/UObjectGlobals.h
@@ -221,7 +221,7 @@ COREUOBJECT_API UClass* StaticLoadClass(UClass* BaseClass, UObject* InOuter, con
 *
 * @return	a pointer to a fully initialized object of the specified class.
 */
-COREUOBJECT_API UObject* StaticConstructObject_Internal(UClass* Class, UObject* InOuter = (UObject*)GetTransientPackage(), FName Name = NAME_None, EObjectFlags SetFlags = RF_NoFlags, EInternalObjectFlags InternalSetFlags = EInternalObjectFlags::None, UObject* Template = NULL, bool bCopyTransientsFromClassDefaults = false, struct FObjectInstancingGraph* InstanceGraph = NULL);
+COREUOBJECT_API UObject* StaticConstructObject_Internal(UClass* Class, UObject* InOuter = (UObject*)GetTransientPackage(), FName Name = NAME_None, EObjectFlags SetFlags = RF_NoFlags, EInternalObjectFlags InternalSetFlags = EInternalObjectFlags::None, UObject* Template = NULL, bool bCopyTransientsFromClassDefaults = false, struct FObjectInstancingGraph* InstanceGraph = NULL, bool bAssumeTemplateIsArchetype = false);
 
 /**
  * Create a new instance of an object.  The returned object will be fully initialized.  If InFlags contains RF_NeedsLoad (indicating that the object still needs to load its object data from disk), components
@@ -370,6 +370,11 @@ COREUOBJECT_API int32 LoadPackageAsync(const FString& InName, FLoadPackageAsyncD
 */
 COREUOBJECT_API void CancelAsyncLoading();
 
+/**
+* Returns true if the event driven loader is enabled in cooked builds
+*/
+COREUOBJECT_API bool IsEventDrivenLoaderEnabledInCookedBuilds();
+
 /**
  * Returns the async load percentage for a package in flight with the passed in name or -1 if there isn't one.
  * THIS IS SLOW. MAY BLOCK ASYNC LOADING.
diff --git a/Engine/Source/Runtime/CoreUObject/Public/UObject/UObjectThreadContext.h b/Engine/Source/Runtime/CoreUObject/Public/UObject/UObjectThreadContext.h
--- a/Engine/Source/Runtime/CoreUObject/Public/UObject/UObjectThreadContext.h
+++ b/Engine/Source/Runtime/CoreUObject/Public/UObject/UObjectThreadContext.h
@@ -87,6 +87,8 @@ public:
 	int32 SerializedExportIndex;
 	/** Points to the most recently used Linker for serialization by CreateExport() */
 	FLinkerLoad* SerializedExportLinker;
+	/** Async Package currently processing objects */
+	struct FAsyncPackage* AsyncPackage;
 
 #if !(UE_BUILD_SHIPPING || UE_BUILD_TEST)
 	/** Stack to ensure that PostInitProperties is routed through Super:: calls. **/
diff --git a/Engine/Source/Runtime/CoreUObject/Public/UObject/UnrealType.h b/Engine/Source/Runtime/CoreUObject/Public/UObject/UnrealType.h
--- a/Engine/Source/Runtime/CoreUObject/Public/UObject/UnrealType.h
+++ b/Engine/Source/Runtime/CoreUObject/Public/UObject/UnrealType.h
@@ -113,7 +113,7 @@ public:
 	void ExportCppDeclaration(FOutputDevice& Out, EExportedDeclaration::Type DeclarationType, const TCHAR* ArrayDimOverride = NULL, uint32 AdditionalExportCPPFlags = 0
 		, bool bSkipParameterName = false, const FString* ActualCppType = nullptr, const FString* ActualExtendedType = nullptr, const FString* ActualParameterName = nullptr) const;
 	virtual FString GetCPPMacroType( FString& ExtendedTypeText ) const;
-	virtual bool PassCPPArgsByRef() const { return false; }
+	virtual bool PassCPPArgsByRef() const;
 
 	/**
 	 * Returns the C++ name of the property, including the _DEPRECATED suffix if the 
@@ -203,7 +203,7 @@ public:
 	*
 	* @return	true if the function has handled the tag
 	*/
-	virtual bool ConvertFromType(const FPropertyTag& Tag, FArchive& Ar, uint8* Data, UStruct* DefaultsStruct, bool& bOutAdvanceProperty) { return false; }
+	virtual bool ConvertFromType(const FPropertyTag& Tag, FArchive& Ar, uint8* Data, UStruct* DefaultsStruct, bool& bOutAdvanceProperty);
 
 	/**
 	 * Determines whether the property values are identical.
@@ -452,16 +452,8 @@ public:
 	uint32 GetValueTypeHash(const void* Src) const;
 
 protected:
-	virtual void CopyValuesInternal( void* Dest, void const* Src, int32 Count  ) const
-	{
-		check(0); // if you are not memcpyable, then you need to deal with the virtual call
-	}
-
-	virtual uint32 GetValueTypeHashInternal(const void* Src) const
-	{
-		check(false); // you need to deal with the virtual call
-		return 0;
-	}
+	virtual void CopyValuesInternal( void* Dest, void const* Src, int32 Count  ) const;
+	virtual uint32 GetValueTypeHashInternal(const void* Src) const;
 
 public:
 	/**
@@ -502,10 +494,8 @@ public:
 	 *									SIZE = the ElementSize of this UProperty
 	 * @param	Src					the address of the value to copy from. should be evaluated the same way as Dest
 	 */
-	virtual void CopySingleValueToScriptVM( void* Dest, void const* Src ) const
-	{
-		CopySingleValue(Dest, Src);
-	}
+	virtual void CopySingleValueToScriptVM( void* Dest, void const* Src ) const;
+
 	/**
 	 * Copy the value for all elements of this property. To the script VM.
 	 * 
@@ -514,10 +504,7 @@ public:
 	 *									OFFSET = the Offset of this UProperty
 	 * @param	Src					the address of the value to copy from. should be evaluated the same way as Dest
 	 */
-	virtual void CopyCompleteValueToScriptVM( void* Dest, void const* Src ) const
-	{
-		CopyCompleteValue(Dest, Src);
-	}
+	virtual void CopyCompleteValueToScriptVM( void* Dest, void const* Src ) const;
 
 	/**
 	 * Copy the value for a single element of this property. From the script VM.
@@ -529,10 +516,8 @@ public:
 	 *									SIZE = the ElementSize of this UProperty
 	 * @param	Src					the address of the value to copy from. should be evaluated the same way as Dest
 	 */
-	virtual void CopySingleValueFromScriptVM( void* Dest, void const* Src ) const
-	{
-		CopySingleValue(Dest, Src);
-	}
+	virtual void CopySingleValueFromScriptVM( void* Dest, void const* Src ) const;
+
 	/**
 	 * Copy the value for all elements of this property. From the script VM.
 	 * 
@@ -541,10 +526,7 @@ public:
 	 *									OFFSET = the Offset of this UProperty
 	 * @param	Src					the address of the value to copy from. should be evaluated the same way as Dest
 	 */
-	virtual void CopyCompleteValueFromScriptVM( void* Dest, void const* Src ) const
-	{
-		CopyCompleteValue(Dest, Src);
-	}
+	virtual void CopyCompleteValueFromScriptVM( void* Dest, void const* Src ) const;
 
 	/**
 	 * Zeros the value for this property. The existing data is assumed valid (so for example this calls FString::Empty)
@@ -673,17 +655,18 @@ public:
 	 * @param	Owner				the object that contains this property's data
 	 * @param	InstanceGraph		contains the mappings of instanced objects and components to their templates
 	 */
-	virtual void InstanceSubobjects( void* Data, void const* DefaultData, UObject* Owner, struct FObjectInstancingGraph* InstanceGraph ) {}
+	virtual void InstanceSubobjects( void* Data, void const* DefaultData, UObject* Owner, struct FObjectInstancingGraph* InstanceGraph );
 
-	virtual int32 GetMinAlignment() const { return 1; }
+	virtual int32 GetMinAlignment() const;
 
 	/**
 	 * Returns true if this property, or in the case of e.g. array or struct properties any sub- property, contains a
 	 * UObject reference.
+	 * @param	EncounteredStructProps		used to check for recursion in arrays
 	 *
 	 * @return true if property (or sub- properties) contain a UObject reference, false otherwise
 	 */
-	virtual bool ContainsObjectReference() const;
+	virtual bool ContainsObjectReference(TArray<const UStructProperty*>& EncounteredStructProps) const;
 
 	/**
 	 * Returns true if this property, or in the case of e.g. array or struct properties any sub- property, contains a
@@ -708,7 +691,7 @@ public:
 	 * Emits tokens used by realtime garbage collection code to passed in ReferenceTokenStream. The offset emitted is relative
 	 * to the passed in BaseOffset which is used by e.g. arrays of structs.
 	 */
-	virtual void EmitReferenceInfo(UClass& OwnerClass, int32 BaseOffset);
+	virtual void EmitReferenceInfo(UClass& OwnerClass, int32 BaseOffset, TArray<const UStructProperty*>& EncounteredStructProps);
 
 	FORCEINLINE int32 GetSize() const
 	{
@@ -1129,16 +1112,10 @@ class COREUOBJECT_API UNumericProperty : public UProperty
 	// UNumericProperty interface.
 
 	/** Return true if this property is for a floating point number **/
-	virtual bool IsFloatingPoint() const
-	{
-		return false;
-	}
+	virtual bool IsFloatingPoint() const;
 
 	/** Return true if this property is for a integral or enum type **/
-	virtual bool IsInteger() const
-	{
-		return true;
-	}
+	virtual bool IsInteger() const;
 
 	/** Return true if this property is a UByteProperty with a non-null Enum **/
 	FORCEINLINE bool IsEnum() const
@@ -1146,41 +1123,29 @@ class COREUOBJECT_API UNumericProperty : public UProperty
 		return !!GetIntPropertyEnum();
 	}
 
-	/** Return teh UEnum if this property is a UByteProperty with a non-null Enum **/
-	virtual UEnum* GetIntPropertyEnum() const
-	{
-		return nullptr;
-	}
+	/** Return the UEnum if this property is a UByteProperty with a non-null Enum **/
+	virtual UEnum* GetIntPropertyEnum() const;
 
 	/** 
 	 * Set the value of an unsigned integral property type
 	 * @param Data - pointer to property data to set
 	 * @param Value - Value to set data to
 	**/
-	virtual void SetIntPropertyValue(void* Data, uint64 Value) const
-	{
-		check(0);
-	}
+	virtual void SetIntPropertyValue(void* Data, uint64 Value) const;
 
 	/** 
 	 * Set the value of a signed integral property type
 	 * @param Data - pointer to property data to set
 	 * @param Value - Value to set data to
 	**/
-	virtual void SetIntPropertyValue(void* Data, int64 Value) const
-	{
-		check(0);
-	}
+	virtual void SetIntPropertyValue(void* Data, int64 Value) const;
 
 	/** 
 	 * Set the value of a floating point property type
 	 * @param Data - pointer to property data to set
 	 * @param Value - Value to set data to
 	**/
-	virtual void SetFloatingPointPropertyValue(void* Data, double Value) const
-	{
-		check(0);
-	}
+	virtual void SetFloatingPointPropertyValue(void* Data, double Value) const;
 
 	/** 
 	 * Set the value of any numeric type from a string point
@@ -1188,43 +1153,28 @@ class COREUOBJECT_API UNumericProperty : public UProperty
 	 * @param Value - Value (as a string) to set 
 	 * CAUTION: This routine does not do enum name conversion
 	**/
-	virtual void SetNumericPropertyValueFromString(void* Data, TCHAR const* Value) const
-	{
-		check(0);
-	}
+	virtual void SetNumericPropertyValueFromString(void* Data, TCHAR const* Value) const;
 
 	/** 
 	 * Gets the value of a signed integral property type
 	 * @param Data - pointer to property data to get
 	 * @return Data as a signed int
 	**/
-	virtual int64 GetSignedIntPropertyValue(void const* Data) const
-	{
-		check(0);
-		return 0;
-	}
+	virtual int64 GetSignedIntPropertyValue(void const* Data) const;
 
 	/** 
 	 * Gets the value of an unsigned integral property type
 	 * @param Data - pointer to property data to get
 	 * @return Data as an unsigned int
 	**/
-	virtual uint64 GetUnsignedIntPropertyValue(void const* Data) const
-	{
-		check(0);
-		return 0;
-	}
+	virtual uint64 GetUnsignedIntPropertyValue(void const* Data) const;
 
 	/** 
 	 * Gets the value of an floating point property type
 	 * @param Data - pointer to property data to get
 	 * @return Data as a double
 	**/
-	virtual double GetFloatingPointPropertyValue(void const* Data) const
-	{
-		check(0);
-		return 0.0;
-	}
+	virtual double GetFloatingPointPropertyValue(void const* Data) const;
 
 	/** 
 	 * Get the value of any numeric type and return it as a string
@@ -1232,11 +1182,7 @@ class COREUOBJECT_API UNumericProperty : public UProperty
 	 * @return Data as a string
 	 * CAUTION: This routine does not do enum name conversion
 	**/
-	virtual FString GetNumericPropertyValueToString(void const* Data) const
-	{
-		check(0);
-		return FString();
-	}
+	virtual FString GetNumericPropertyValueToString(void const* Data) const;
 	// End of UNumericProperty interface
 
 	static uint8 ReadEnumAsUint8(FArchive& Ar, UStruct* DefaultsStruct, const FPropertyTag& Tag);
@@ -1438,6 +1384,8 @@ class COREUOBJECT_API UByteProperty : public TProperty_Numeric<uint8>
 	// UObject interface.
 	virtual void Serialize( FArchive& Ar ) override;
 	static void AddReferencedObjects(UObject* InThis, FReferenceCollector& Collector);
+	virtual void GetPreloadDependencies(TArray<UObject*>& OutDeps) override;
+
 	// End of UObject interface
 
 	// UHT interface
@@ -1453,10 +1401,7 @@ class COREUOBJECT_API UByteProperty : public TProperty_Numeric<uint8>
 	// End of UProperty interface
 
 	// UNumericProperty interface.
-	virtual UEnum* GetIntPropertyEnum() const override
-	{
-		return Enum;
-	}
+	virtual UEnum* GetIntPropertyEnum() const override;
 	// End of UNumericProperty interface
 };
 
@@ -1815,6 +1760,8 @@ public:
 		return FieldMask == 0xff;
 	}
 
+	uint32 GetValueTypeHashInternal(const void* Src) const override;
+
 #if HACK_HEADER_GENERATOR
 	// Required by UHT makefiles for internal data serialization.
 	friend struct FBoolPropertyArchiveProxy;
@@ -1869,10 +1816,8 @@ class COREUOBJECT_API UObjectPropertyBase : public UProperty
 	 *									SIZE = the ElementSize of this UProperty
 	 * @param	Src					the address of the value to copy from. should be evaluated the same way as Dest
 	 */
-	virtual void CopySingleValueToScriptVM( void* Dest, void const* Src ) const override
-	{
-		*(UObject**)Dest = GetObjectPropertyValue(Src);
-	}
+	virtual void CopySingleValueToScriptVM( void* Dest, void const* Src ) const override;
+
 	/**
 	 * Copy the value for all elements of this property. To the script VM.
 	 * 
@@ -1881,13 +1826,7 @@ class COREUOBJECT_API UObjectPropertyBase : public UProperty
 	 *									OFFSET = the Offset of this UProperty
 	 * @param	Src					the address of the value to copy from. should be evaluated the same way as Dest
 	 */
-	virtual void CopyCompleteValueToScriptVM( void* Dest, void const* Src ) const override
-	{
-		for (int32 Index = 0; Index < ArrayDim; Index++)
-		{
-			((UObject**)Dest)[Index] = GetObjectPropertyValue(((uint8*)Src) + Index * ElementSize);
-		}
-	}
+	virtual void CopyCompleteValueToScriptVM( void* Dest, void const* Src ) const override;
 
 	/**
 	 * Copy the value for a single element of this property. From the script VM.
@@ -1899,10 +1838,8 @@ class COREUOBJECT_API UObjectPropertyBase : public UProperty
 	 *									SIZE = the ElementSize of this UProperty
 	 * @param	Src					the address of the value to copy from. should be evaluated the same way as Dest
 	 */
-	virtual void CopySingleValueFromScriptVM( void* Dest, void const* Src ) const override
-	{
-		SetObjectPropertyValue(Dest, *(UObject**)Src);
-	}
+	virtual void CopySingleValueFromScriptVM( void* Dest, void const* Src ) const override;
+
 	/**
 	 * Copy the value for all elements of this property. From the script VM.
 	 * 
@@ -1911,14 +1848,7 @@ class COREUOBJECT_API UObjectPropertyBase : public UProperty
 	 *									OFFSET = the Offset of this UProperty
 	 * @param	Src					the address of the value to copy from. should be evaluated the same way as Dest
 	 */
-	virtual void CopyCompleteValueFromScriptVM( void* Dest, void const* Src ) const override
-	{
-		checkSlow(ElementSize == sizeof(UObject*)); // the idea that script pointers are the same size as weak pointers is maybe required, maybe not
-		for (int32 Index = 0; Index < ArrayDim; Index++)
-		{
-			SetObjectPropertyValue(((uint8*)Dest) + Index * ElementSize, ((UObject**)Src)[Index]);
-		}
-	}
+	virtual void CopyCompleteValueFromScriptVM( void* Dest, void const* Src ) const override;
 	// End of UProperty interface
 
 	// UObjectPropertyBase interface
@@ -1946,19 +1876,12 @@ public:
 	// Returns the qualified export path for a given object, parent, and export root scope
 	static FString GetExportPath(const UObject* Object, const UObject* Parent, const UObject* ExportRootScope, const uint32 PortFlags);
 
-	virtual UObject* GetObjectPropertyValue(const void* PropertyValueAddress) const
-	{
-		check(0);
-		return NULL;
-	}
+	virtual UObject* GetObjectPropertyValue(const void* PropertyValueAddress) const;
 	FORCEINLINE UObject* GetObjectPropertyValue_InContainer(const void* PropertyValueAddress, int32 ArrayIndex = 0) const
 	{
 		return GetObjectPropertyValue(ContainerPtrToValuePtr<void>(PropertyValueAddress, ArrayIndex));
 	}
-	virtual void SetObjectPropertyValue(void* PropertyValueAddress, UObject* Value) const
-	{
-		check(0);
-	}
+	virtual void SetObjectPropertyValue(void* PropertyValueAddress, UObject* Value) const;
 	FORCEINLINE void SetObjectPropertyValue_InContainer(void* PropertyValueAddress, UObject* Value, int32 ArrayIndex = 0) const
 	{
 		SetObjectPropertyValue(ContainerPtrToValuePtr<void>(PropertyValueAddress, ArrayIndex), Value);
@@ -1980,10 +1903,7 @@ public:
 #endif // USE_CIRCULAR_DEPENDENCY_LOAD_DEFERRING
 
 protected:
-	virtual bool AllowCrossLevel() const
-	{
-		return false;
-	}
+	virtual bool AllowCrossLevel() const;
 
 	virtual void CheckValidObject(void* Value) const;
 	// End of UObjectPropertyBase interface
@@ -2020,7 +1940,7 @@ public:
 #endif // WITH_HOT_RELOAD_CTORS
 
 	// UProperty interface.
-	virtual bool ContainsObjectReference() const override
+	virtual bool ContainsObjectReference(TArray<const UStructProperty*>& EncounteredStructProps) const override
 	{
 		return !TIsWeakPointerType<InTCppType>::Value;
 	}
@@ -2064,25 +1984,18 @@ class COREUOBJECT_API UObjectProperty : public TUObjectPropertyBase<UObject*>
 
 	// UProperty interface
 	virtual void SerializeItem( FArchive& Ar, void* Value, void const* Defaults ) const override;
-	virtual void EmitReferenceInfo(UClass& OwnerClass, int32 BaseOffset) override;
+	virtual void EmitReferenceInfo(UClass& OwnerClass, int32 BaseOffset, TArray<const UStructProperty*>& EncounteredStructProps) override;
 	virtual const TCHAR* ImportText_Internal(const TCHAR* Buffer, void* Data, int32 PortFlags, UObject* OwnerObject, FOutputDevice* ErrorText) const override;
 	virtual bool ConvertFromType(const FPropertyTag& Tag, FArchive& Ar, uint8* Data, UStruct* DefaultsStruct, bool& bOutAdvanceProperty) override;
 
 private:
-	virtual uint32 GetValueTypeHashInternal(const void* Src) const override
-	{
-		return GetTypeHash(GetPropertyValue(Src));
-	}
+	virtual uint32 GetValueTypeHashInternal(const void* Src) const override;
 public:
 	// End of UProperty interface
 
 	// UObjectPropertyBase interface
-	virtual UObject* GetObjectPropertyValue(const void* PropertyValueAddress) const override
-	{
-		return GetPropertyValue(PropertyValueAddress);
-	}
+	virtual UObject* GetObjectPropertyValue(const void* PropertyValueAddress) const override;
 	virtual void SetObjectPropertyValue(void* PropertyValueAddress, UObject* Value) const override;
-
 	virtual FString GetCPPTypeCustom(FString* ExtendedTypeText, uint32 CPPExportFlags, const FString& InnerNativeTypeName)  const override;
 	// End of UObjectPropertyBase interface
 };
@@ -2115,14 +2028,8 @@ class COREUOBJECT_API UWeakObjectProperty : public TUObjectPropertyBase<FWeakObj
 	// End of UProperty interface
 
 	// UObjectProperty interface
-	virtual UObject* GetObjectPropertyValue(const void* PropertyValueAddress) const override
-	{
-		return GetPropertyValue(PropertyValueAddress).Get();
-	}
-	virtual void SetObjectPropertyValue(void* PropertyValueAddress, UObject* Value) const override
-	{
-		SetPropertyValue(PropertyValueAddress, TCppType(Value));
-	}
+	virtual UObject* GetObjectPropertyValue(const void* PropertyValueAddress) const override;
+	virtual void SetObjectPropertyValue(void* PropertyValueAddress, UObject* Value) const override;
 	// End of UObjectProperty interface
 };
 
@@ -2155,23 +2062,11 @@ class COREUOBJECT_API ULazyObjectProperty : public TUObjectPropertyBase<FLazyObj
 	// End of UProperty interface
 
 	// UObjectProperty interface
-	virtual UObject* GetObjectPropertyValue(const void* PropertyValueAddress) const override
-	{
-		return GetPropertyValue(PropertyValueAddress).Get();
-	}
-	virtual void SetObjectPropertyValue(void* PropertyValueAddress, UObject* Value) const override
-	{
-		SetPropertyValue(PropertyValueAddress, TCppType(Value));
-	}
-	virtual bool AllowCrossLevel() const override
-	{
-		return true;
-	}
+	virtual UObject* GetObjectPropertyValue(const void* PropertyValueAddress) const override;
+	virtual void SetObjectPropertyValue(void* PropertyValueAddress, UObject* Value) const override;
+	virtual bool AllowCrossLevel() const override;
 private:
-	virtual uint32 GetValueTypeHashInternal(const void* Src) const override
-	{
-		return GetTypeHash(GetPropertyValue(Src));
-	}
+	virtual uint32 GetValueTypeHashInternal(const void* Src) const override;
 public:
 	// End of UObjectProperty interface
 };
@@ -2206,48 +2101,21 @@ class COREUOBJECT_API UAssetObjectProperty : public TUObjectPropertyBase<FAssetP
 	// End of UProperty interface
 
 	// UObjectProperty interface
-	virtual UObject* GetObjectPropertyValue(const void* PropertyValueAddress) const override
-	{
-		return GetPropertyValue(PropertyValueAddress).Get();
-	}
-	virtual void SetObjectPropertyValue(void* PropertyValueAddress, UObject* Value) const override
-	{
-		SetPropertyValue(PropertyValueAddress, TCppType(Value));
-	}
-	virtual bool AllowCrossLevel() const override
-	{
-		return true;
-	}
+	virtual UObject* GetObjectPropertyValue(const void* PropertyValueAddress) const override;
+	virtual void SetObjectPropertyValue(void* PropertyValueAddress, UObject* Value) const override;
+	virtual bool AllowCrossLevel() const override;
 	virtual FString GetCPPTypeCustom(FString* ExtendedTypeText, uint32 CPPExportFlags, const FString& InnerNativeTypeName)  const override;
 
 private:
-	virtual uint32 GetValueTypeHashInternal(const void* Src) const override
-	{
-		return GetTypeHash(GetPropertyValue(Src));
-	}
+	virtual uint32 GetValueTypeHashInternal(const void* Src) const override;
 public:
 
 	// ScriptVM should store Asset as FAssetPtr not as UObject.
 
-	virtual void CopySingleValueToScriptVM(void* Dest, void const* Src) const override
-	{
-		CopySingleValue(Dest, Src);
-	}
-
-	virtual void CopyCompleteValueToScriptVM(void* Dest, void const* Src) const override
-	{
-		CopyCompleteValue(Dest, Src);
-	}
-
-	virtual void CopySingleValueFromScriptVM(void* Dest, void const* Src) const override
-	{
-		CopySingleValue(Dest, Src);
-	}
-
-	virtual void CopyCompleteValueFromScriptVM(void* Dest, void const* Src) const override
-	{
-		CopyCompleteValue(Dest, Src);
-	}
+	virtual void CopySingleValueToScriptVM(void* Dest, void const* Src) const override;
+	virtual void CopyCompleteValueToScriptVM(void* Dest, void const* Src) const override;
+	virtual void CopySingleValueFromScriptVM(void* Dest, void const* Src) const override;
+	virtual void CopyCompleteValueFromScriptVM(void* Dest, void const* Src) const override;
 	// End of UObjectProperty interface
 };
 
@@ -2420,13 +2288,13 @@ public:
 	virtual bool NetSerializeItem( FArchive& Ar, UPackageMap* Map, void* Data, TArray<uint8> * MetaData = NULL ) const override;
 	virtual void ExportTextItem( FString& ValueStr, const void* PropertyValue, const void* DefaultValue, UObject* Parent, int32 PortFlags, UObject* ExportRootScope ) const override;
 	virtual const TCHAR* ImportText_Internal( const TCHAR* Buffer, void* Data, int32 PortFlags, UObject* OwnerObject, FOutputDevice* ErrorText ) const override;
-	virtual bool ContainsObjectReference() const override;
+	virtual bool ContainsObjectReference(TArray<const UStructProperty*>& EncounteredStructProps) const override;
 	virtual bool SameType(const UProperty* Other) const override;
 	// End of UProperty interface
 
 	// UObject interface
 	virtual void Serialize( FArchive& Ar ) override;
-	virtual void EmitReferenceInfo(UClass& OwnerClass, int32 BaseOffset) override;
+	virtual void EmitReferenceInfo(UClass& OwnerClass, int32 BaseOffset, TArray<const UStructProperty*>& EncounteredStructProps) override;
 	virtual void BeginDestroy() override;
 	static void AddReferencedObjects(UObject* InThis, FReferenceCollector& Collector);
 	// End of UObject interface
@@ -2479,15 +2347,8 @@ public:
 	virtual void ExportTextItem( FString& ValueStr, const void* PropertyValue, const void* DefaultValue, UObject* Parent, int32 PortFlags, UObject* ExportRootScope ) const override;
 	virtual const TCHAR* ImportText_Internal( const TCHAR* Buffer, void* Data, int32 PortFlags, UObject* OwnerObject, FOutputDevice* ErrorText ) const override;
 	virtual bool ConvertFromType(const FPropertyTag& Tag, FArchive& Ar, uint8* Data, UStruct* DefaultsStruct, bool& bOutAdvanceProperty) override;
-	virtual FString GetCPPTypeForwardDeclaration() const override
-	{
-		return FString();
-	}
-
-	uint32 GetValueTypeHashInternal(const void* Src) const override
-	{
-		return GetTypeHash(*(const FName*)Src);
-	}
+	virtual FString GetCPPTypeForwardDeclaration() const override;
+	uint32 GetValueTypeHashInternal(const void* Src) const override;
 	// End of UProperty interface
 };
 
@@ -2523,15 +2384,8 @@ public:
 	virtual void ExportTextItem( FString& ValueStr, const void* PropertyValue, const void* DefaultValue, UObject* Parent, int32 PortFlags, UObject* ExportRootScope ) const override;
 	virtual const TCHAR* ImportText_Internal( const TCHAR* Buffer, void* Data, int32 PortFlags, UObject* OwnerObject, FOutputDevice* ErrorText ) const override;
 	virtual bool ConvertFromType(const FPropertyTag& Tag, FArchive& Ar, uint8* Data, UStruct* DefaultsStruct, bool& bOutAdvanceProperty) override;
-	virtual FString GetCPPTypeForwardDeclaration() const override
-	{
-		return FString();
-	}
-
-	uint32 GetValueTypeHashInternal(const void* Src) const override
-	{
-		return GetTypeHash(*(const FString*)Src);
-	}
+	virtual FString GetCPPTypeForwardDeclaration() const override;
+	uint32 GetValueTypeHashInternal(const void* Src) const override;
 	// End of UProperty interface
 
 	// Necessary to fix Compiler Error C2026
@@ -2573,6 +2427,7 @@ public:
 	// UObject interface
 	virtual void Serialize( FArchive& Ar ) override;
 	static void AddReferencedObjects(UObject* InThis, FReferenceCollector& Collector);
+	virtual void GetPreloadDependencies(TArray<UObject*>& OutDeps) override;
 	// End of UObject interface
 
 	// UField interface
@@ -2594,9 +2449,9 @@ public:
 	virtual void DestroyValueInternal( void* Dest ) const override;
 	virtual bool PassCPPArgsByRef() const override;
 	virtual void InstanceSubobjects( void* Data, void const* DefaultData, UObject* Owner, struct FObjectInstancingGraph* InstanceGraph ) override;
-	virtual bool ContainsObjectReference() const override;
+	virtual bool ContainsObjectReference(TArray<const UStructProperty*>& EncounteredStructProps) const override;
 	virtual bool ContainsWeakObjectReference() const override;
-	virtual void EmitReferenceInfo(UClass& OwnerClass, int32 BaseOffset) override;
+	virtual void EmitReferenceInfo(UClass& OwnerClass, int32 BaseOffset, TArray<const UStructProperty*>& EncounteredStructProps) override;
 	virtual bool SameType(const UProperty* Other) const override;
 	virtual bool ConvertFromType(const FPropertyTag& Tag, FArchive& Ar, uint8* Data, UStruct* DefaultsStruct, bool& bOutAdvanceProperty) override;
 	// End of UProperty interface
@@ -2625,6 +2480,7 @@ public:
 	// UObject interface
 	virtual void Serialize(FArchive& Ar) override;
 	static void AddReferencedObjects(UObject* InThis, FReferenceCollector& Collector);
+	virtual void GetPreloadDependencies(TArray<UObject*>& OutDeps) override;
 	// End of UObject interface
 
 	// UField interface
@@ -2645,9 +2501,9 @@ public:
 	virtual void DestroyValueInternal(void* Dest) const override;
 	virtual bool PassCPPArgsByRef() const override;
 	virtual void InstanceSubobjects(void* Data, void const* DefaultData, UObject* Owner, struct FObjectInstancingGraph* InstanceGraph) override;
-	virtual bool ContainsObjectReference() const override;
+	virtual bool ContainsObjectReference(TArray<const UStructProperty*>& EncounteredStructProps) const override;
 	virtual bool ContainsWeakObjectReference() const override;
-	virtual void EmitReferenceInfo(UClass& OwnerClass, int32 BaseOffset) override;
+	virtual void EmitReferenceInfo(UClass& OwnerClass, int32 BaseOffset, TArray<const UStructProperty*>& EncounteredStructProps) override;
 	virtual bool SameType(const UProperty* Other) const override;
 	// End of UProperty interface
 
@@ -2674,6 +2530,7 @@ public:
 	// UObject interface
 	virtual void Serialize(FArchive& Ar) override;
 	static void AddReferencedObjects(UObject* InThis, FReferenceCollector& Collector);
+	virtual void GetPreloadDependencies(TArray<UObject*>& OutDeps) override;
 	// End of UObject interface
 
 	// UField interface
@@ -2695,9 +2552,9 @@ public:
 	virtual void DestroyValueInternal(void* Dest) const override;
 	virtual bool PassCPPArgsByRef() const override;
 	virtual void InstanceSubobjects(void* Data, void const* DefaultData, UObject* Owner, struct FObjectInstancingGraph* InstanceGraph) override;
-	virtual bool ContainsObjectReference() const override;
+	virtual bool ContainsObjectReference(TArray<const UStructProperty*>& EncounteredStructProps) const override;
 	virtual bool ContainsWeakObjectReference() const override;
-	virtual void EmitReferenceInfo(UClass& OwnerClass, int32 BaseOffset) override;
+	virtual void EmitReferenceInfo(UClass& OwnerClass, int32 BaseOffset, TArray<const UStructProperty*>& EncounteredStructProps) override;
 	virtual bool SameType(const UProperty* Other) const override;
 	// End of UProperty interface
 
@@ -3849,6 +3706,7 @@ public:
 	// UObject interface
 	virtual void Serialize( FArchive& Ar ) override;
 	static void AddReferencedObjects(UObject* InThis, FReferenceCollector& Collector);
+	virtual void GetPreloadDependencies(TArray<UObject*>& OutDeps) override;
 	// End of UObject interface
 
 	// UProperty interface
@@ -3867,9 +3725,9 @@ public:
 	virtual void InitializeValueInternal( void* Dest ) const override;
 	virtual void InstanceSubobjects( void* Data, void const* DefaultData, UObject* Owner, struct FObjectInstancingGraph* InstanceGraph ) override;
 	virtual int32 GetMinAlignment() const override;
-	virtual bool ContainsObjectReference() const override;
+	virtual bool ContainsObjectReference(TArray<const UStructProperty*>& EncounteredStructProps) const override;
 	virtual bool ContainsWeakObjectReference() const override;
-	virtual void EmitReferenceInfo(UClass& OwnerClass, int32 BaseOffset) override;
+	virtual void EmitReferenceInfo(UClass& OwnerClass, int32 BaseOffset, TArray<const UStructProperty*>& EncounteredStructProps) override;
 	virtual bool SameType(const UProperty* Other) const override;
 	virtual bool ConvertFromType(const FPropertyTag& Tag, FArchive& Ar, uint8* Data, UStruct* DefaultsStruct, bool& bOutAdvanceProperty) override;
 	// End of UProperty interface
diff --git a/Engine/Source/Runtime/CoreUObject/Public/UObject/WeakObjectPtr.h b/Engine/Source/Runtime/CoreUObject/Public/UObject/WeakObjectPtr.h
--- a/Engine/Source/Runtime/CoreUObject/Public/UObject/WeakObjectPtr.h
+++ b/Engine/Source/Runtime/CoreUObject/Public/UObject/WeakObjectPtr.h
@@ -149,7 +149,7 @@ public:
 	 * don't need to bother serializing weak object pointers.  However, serialization is required if you
 	 * want to load and save your object.
 	 */
-	COREUOBJECT_API friend FArchive& operator<<( FArchive& Ar, FWeakObjectPtr& WeakObjectPtr );
+	COREUOBJECT_API void Serialize(FArchive& Ar);
 
 protected:
 
diff --git a/Engine/Source/Runtime/Engine/Classes/Components/ReflectionCaptureComponent.h b/Engine/Source/Runtime/Engine/Classes/Components/ReflectionCaptureComponent.h
--- a/Engine/Source/Runtime/Engine/Classes/Components/ReflectionCaptureComponent.h
+++ b/Engine/Source/Runtime/Engine/Classes/Components/ReflectionCaptureComponent.h
@@ -230,6 +230,7 @@ private:
 	 * This queue should be in the UWorld or the FSceneInterface, but those are not available yet in PostLoad.
 	 */
 	static TArray<UReflectionCaptureComponent*> ReflectionCapturesToUpdateForLoad;
+	static FCriticalSection ReflectionCapturesToUpdateForLoadLock;
 
 	void UpdateDerivedData(FReflectionCaptureFullHDR* NewDerivedData);
 	void SerializeSourceData(FArchive& Ar);
