Copying //UE4/Dev-Core to //UE4/Dev-Main (Source: //UE4/Dev-Core @ 3172069)

==========================
MAJOR FEATURES + CHANGES
==========================

Change 3152650 on 2016/10/05 by [REDACTED_PERSON]

Change 3152652 on 2016/10/05 by [REDACTED_PERSON]

	Fix cases that were using the PropertiesSize to malloc a UStruct to instead use GetStructureSize(). There is a difference because of alignment, so on some platforms it was corrupting memory

Change 3152654 on 2016/10/05 by [REDACTED_PERSON]

	Enable EDL again

Change 3153222 on 2016/10/06 by [REDACTED_PERSON]

	#PR 2830: Fixes a garbage collection crash when copying an empty timeline
	contributed by [REDACTED_PERSON]

Change 3154177 on 2016/10/06 by [REDACTED_PERSON]

	Fix it so pending kill component templates do not end up in import table, they will fail to import
	Fix it so when marking an inheritable component template as unnecessary it also marks it pending kill.
	This fixes it so if GetArchetype will not return an orphaned, non-saving component as the archetype for a grandchild blueprint

Change 3157907 on 2016/10/11 by [REDACTED_PERSON]

	FVariantDataConverter::ConvertScalarUPropertyToVariant fixes for enums.

Change 3158037 on 2016/10/11 by [REDACTED_PERSON]

	PR #2684: Enable full callstacks in non-monolithic builds by default (Contributed by [REDACTED_PERSON])

Change 3159574 on 2016/10/12 by [REDACTED_PERSON]

	Slight change to database profile #rb none

Change 3159689 on 2016/10/12 by [REDACTED_PERSON]

	UE4 - Critical fix for task grpah memory leak.

Change 3159730 on 2016/10/12 by [REDACTED_PERSON]

	Making FindShaderResourceById and FindShaderById return a raw pointer instead of TRefCountPtr (basically undoing CL #2538774) to prevent creating temporary TRefCountPtrs. It's no longer necessary to use TRefCountPtrs as shader serialization has changed (CL #2989898) and shaders are no longer registered on the async loading thread.

Change 3164243 on 2016/10/17 by [REDACTED_PERSON]

	Fixing instances of Ar.Tell() being assigned to int32 instead of int64.

Change 3164388 on 2016/10/17 by [REDACTED_PERSON]

	Fix crashes for FMallocStomp allocations of zero bytes.

Change 3164485 on 2016/10/17 by [REDACTED_PERSON]

	GET_ENUMERATOR_NAME_CHECKED added, which works like GET_MEMBER_NAME_CHECKED but with enums.

Change 3165755 on 2016/10/18 by [REDACTED_PERSON]

	Submitting changes to datamodel to match DB changes. Description field and command line field size increased

	#RNX

Change 3165759 on 2016/10/18 by [REDACTED_PERSON]

	Merging //Tasks/UE4/Dev-LoadTimes/... to //UE4/Dev-Core/...

Change 3166091 on 2016/10/18 by [REDACTED_PERSON]

	UE4 - Fixed typo in check relating to stalling the RHI thread

Change 3167607 on 2016/10/19 by [REDACTED_PERSON]

	UE4 - Never use a cached file wrapper with the EDL

Change 3168738 on 2016/10/20 by [REDACTED_PERSON]

	Fixing CancelAsyncLoading crashes

Change 3168774 on 2016/10/20 by [REDACTED_PERSON]

	Fix for a crash reported in a UDN post:

	https://udn.unrealengine.com/questions/288234/a-problem-on-mobile-after-updating-to-4111.html

	On 32-bit targets, some int64s could be aligned to 4 bytes and then passed into platform atomics functions. On apple platforms, those atomics functions can crash when passed in a non-8 byte aligned int64. Some explict alignment has been added to those variables to for 8 bytes.

Change 3168952 on 2016/10/20 by [REDACTED_PERSON]

	Removing unused packages #rb none

Change 3170562 on 2016/10/21 by [REDACTED_PERSON]

	Fixing CIS warning

Change 3170617 on 2016/10/21 by [REDACTED_PERSON]

	UE4 - first pass at allowing the EDL to work with nativized blueprints

Change 3170662 on 2016/10/21 by [REDACTED_PERSON]

	Re-added duplicate crash checks to CRP (v1.2.2)
	[UE-37565] - Re-add duplicate crash checks to CRP inputs

Change 3171794 on 2016/10/24 by [REDACTED_PERSON]

	Change a file IO log message to verbose

Change 3171831 on 2016/10/24 by [REDACTED_PERSON]

	PUSH_MACRO() and POP_MACRO() for preserving and restoring the state of a macro.

Change 3171934 on 2016/10/24 by [REDACTED_PERSON]

	UE4 - Fix rare crash with particle systems destroying render state. This is a very safe fix.

diff --git a/Engine/Config/BaseEngine.ini b/Engine/Config/BaseEngine.ini
--- a/Engine/Config/BaseEngine.ini
+++ b/Engine/Config/BaseEngine.ini
@@ -1609,7 +1609,7 @@ UseStrictEngineVersioning=True
 UseVTableConstructors=True
 CanStripEditorOnlyExportsAndImports=True
 CanSkipEditorReferencedPackagesWhenCooking=False
-DetailedCallstacksInNonMonolithicBuilds=False
+DetailedCallstacksInNonMonolithicBuilds=True
 UseSeperateBulkDataFiles=False
 HangDuration=0.0
 
diff --git a/Engine/Plugins/Messaging/TcpMessaging/Source/TcpMessaging/Private/Transport/TcpDeserializedMessage.cpp b/Engine/Plugins/Messaging/TcpMessaging/Source/TcpMessaging/Private/Transport/TcpDeserializedMessage.cpp
--- a/Engine/Plugins/Messaging/TcpMessaging/Source/TcpMessaging/Private/Transport/TcpDeserializedMessage.cpp
+++ b/Engine/Plugins/Messaging/TcpMessaging/Source/TcpMessaging/Private/Transport/TcpDeserializedMessage.cpp
@@ -117,7 +117,7 @@ bool FTcpDeserializedMessage::Deserialize(const FArrayReaderPtr& Message)
 	}
 
 	// create message body
-	MessageData = FMemory::Malloc(TypeInfo->PropertiesSize);
+	MessageData = FMemory::Malloc(TypeInfo->GetStructureSize());
 	TypeInfo->InitializeStruct(MessageData);
 
 	// deserialize message body
diff --git a/Engine/Plugins/Messaging/UdpMessaging/Source/UdpMessaging/Private/Transport/UdpDeserializedMessage.cpp b/Engine/Plugins/Messaging/UdpMessaging/Source/UdpMessaging/Private/Transport/UdpDeserializedMessage.cpp
--- a/Engine/Plugins/Messaging/UdpMessaging/Source/UdpMessaging/Private/Transport/UdpDeserializedMessage.cpp
+++ b/Engine/Plugins/Messaging/UdpMessaging/Source/UdpMessaging/Private/Transport/UdpDeserializedMessage.cpp
@@ -117,7 +117,7 @@ bool FUdpDeserializedMessage::Deserialize(const FUdpReassembledMessageRef& Reass
 	}
 
 	// create message body
-	MessageData = FMemory::Malloc(TypeInfo->PropertiesSize);
+	MessageData = FMemory::Malloc(TypeInfo->GetStructureSize());
 	TypeInfo->InitializeStruct(MessageData);
 
 	// deserialize message body
diff --git a/Engine/Plugins/Online/OnlineSubsystem/Source/Private/OnlineKeyValuePair.cpp b/Engine/Plugins/Online/OnlineSubsystem/Source/Private/OnlineKeyValuePair.cpp
--- a/Engine/Plugins/Online/OnlineSubsystem/Source/Private/OnlineKeyValuePair.cpp
+++ b/Engine/Plugins/Online/OnlineSubsystem/Source/Private/OnlineKeyValuePair.cpp
@@ -1083,9 +1083,8 @@ bool FVariantDataConverter::ConvertScalarUPropertyToVariant(UProperty* Property,
 			FString StringValue = EnumDef->GetEnumName(NumericProperty->GetSignedIntPropertyValue(Value));
 			OutVariantData.SetValue(StringValue);
 		}
-
 		// We want to export numbers as numbers
-		if (NumericProperty->IsFloatingPoint())
+		else if (NumericProperty->IsFloatingPoint())
 		{
 			double DoubleValue = NumericProperty->GetFloatingPointPropertyValue(Value);
 			OutVariantData.SetValue(DoubleValue);
diff --git a/Engine/Source/Editor/UnrealEd/Private/DataTableEditorUtils.cpp b/Engine/Source/Editor/UnrealEd/Private/DataTableEditorUtils.cpp
--- a/Engine/Source/Editor/UnrealEd/Private/DataTableEditorUtils.cpp
+++ b/Engine/Source/Editor/UnrealEd/Private/DataTableEditorUtils.cpp
@@ -52,7 +52,7 @@ uint8* FDataTableEditorUtils::AddRow(UDataTable* DataTable, FName RowName)
 	BroadcastPreChange(DataTable, EDataTableChangeInfo::RowList);
 	DataTable->Modify();
 	// Allocate data to store information, using UScriptStruct to know its size
-	uint8* RowData = (uint8*)FMemory::Malloc(DataTable->RowStruct->PropertiesSize);
+	uint8* RowData = (uint8*)FMemory::Malloc(DataTable->RowStruct->GetStructureSize());
 	DataTable->RowStruct->InitializeStruct(RowData);
 	// And be sure to call DestroyScriptStruct later
 
diff --git a/Engine/Source/Programs/CrashReporter/CrashReportWebSite/Controllers/CrashesController.cs b/Engine/Source/Programs/CrashReporter/CrashReportWebSite/Controllers/CrashesController.cs
--- a/Engine/Source/Programs/CrashReporter/CrashReportWebSite/Controllers/CrashesController.cs
+++ b/Engine/Source/Programs/CrashReporter/CrashReportWebSite/Controllers/CrashesController.cs
@@ -262,6 +262,7 @@ namespace Tools.CrashReporter.CrashReportWebSite.Controllers
                 messageBuilder.AppendLine(payloadString);
 
                 FLogger.Global.WriteException(messageBuilder.ToString());
+                _slackWriter.Write(messageBuilder.ToString());
 
                 newCrashResult.Message = messageBuilder.ToString();
                 newCrashResult.bSuccess = false;
@@ -944,6 +945,7 @@ namespace Tools.CrashReporter.CrashReportWebSite.Controllers
                     }
                 }
                 FLogger.Global.WriteException(messageBuilder.ToString());
+                _slackWriter.Write("Create Crash Exception : " + messageBuilder.ToString());
                 throw;
             }
 	        catch (Exception ex)
diff --git a/Engine/Source/Programs/CrashReporter/CrashReportWebSite/Models/CrashDataFormater.cs b/Engine/Source/Programs/CrashReporter/CrashReportWebSite/Models/CrashDataFormater.cs
--- a/Engine/Source/Programs/CrashReporter/CrashReportWebSite/Models/CrashDataFormater.cs
+++ b/Engine/Source/Programs/CrashReporter/CrashReportWebSite/Models/CrashDataFormater.cs
@@ -255,7 +255,7 @@ namespace Tools.CrashReporter.CrashReportWebSite.DataModels
 				LineToSkipUpto = "FDebug::AssertFailed";
 				break;
 			case 3:
-				bSkipping = true;
+				bSkipping = false;
 				LineToSkipUpto = "FDebug::";
 				break;
 			}
diff --git a/Engine/Source/Runtime/Core/Private/Apple/ApplePlatformTime.cpp b/Engine/Source/Runtime/Core/Private/Apple/ApplePlatformTime.cpp
--- a/Engine/Source/Runtime/Core/Private/Apple/ApplePlatformTime.cpp
+++ b/Engine/Source/Runtime/Core/Private/Apple/ApplePlatformTime.cpp
@@ -20,6 +20,7 @@ double FApplePlatformTime::InitTiming(void)
 	mach_timebase_info_data_t Info;
 	verify( mach_timebase_info( &Info ) == 0 );
 	SecondsPerCycle = 1e-9 * (double)Info.numer / (double)Info.denom;
+	SecondsPerCycle64 = 1e-9 * (double)Info.numer / (double)Info.denom;
 	return FPlatformTime::Seconds();
 }
 
diff --git a/Engine/Source/Runtime/Core/Private/Async/TaskGraph.cpp b/Engine/Source/Runtime/Core/Private/Async/TaskGraph.cpp
--- a/Engine/Source/Runtime/Core/Private/Async/TaskGraph.cpp
+++ b/Engine/Source/Runtime/Core/Private/Async/TaskGraph.cpp
@@ -2500,9 +2500,12 @@ static FLockFreePointerListFIFOIntrusive<FGraphEventAndSmallTaskStorage> TheGrap
 
 #endif
 
+#define USE_NIEVE_GRAPH_EVENT_ALLOCTOR (0) // this is useful for tracking leaked handles
+
 FGraphEventRef FGraphEvent::CreateGraphEvent()
 {
 	GraphEventsInUse.Increment();
+#if !USE_NIEVE_GRAPH_EVENT_ALLOCTOR
 	FGraphEvent* Used =  TheGraphEventAllocator.Pop();
 	if (Used)
 	{
@@ -2516,13 +2519,16 @@ FGraphEventRef FGraphEvent::CreateGraphEvent()
 //		Used->LockFreePointerQueueNext = nullptr; // temp
 		return FGraphEventRef(Used);
 	}
+#endif
 	return FGraphEventRef(new FGraphEvent(false));
 }
 
 FGraphEvent* FGraphEvent::CreateGraphEventWithInlineStorage()
 {
 	GraphEventsWithInlineStorageInUse.Increment();
-	FGraphEventAndSmallTaskStorage* Used =  TheGraphEventWithInlineStorageAllocator.Pop();
+	FGraphEventAndSmallTaskStorage* Used;
+#if !USE_NIEVE_GRAPH_EVENT_ALLOCTOR
+	Used = TheGraphEventWithInlineStorageAllocator.Pop();
 	if (Used)
 	{
 		FPlatformMisc::MemoryBarrier();
@@ -2533,6 +2539,7 @@ FGraphEvent* FGraphEvent::CreateGraphEventWithInlineStorage()
 		checkThreadGraph(Used->bInline);
 	}
 	else
+#endif
 	{
 		Used = new FGraphEventAndSmallTaskStorage();
 	}
@@ -2547,12 +2554,20 @@ void FGraphEvent::Recycle(FGraphEvent* ToRecycle)
 	if (ToRecycle->bInline)
 	{
 		GraphEventsWithInlineStorageInUse.Decrement();
+#if USE_NIEVE_GRAPH_EVENT_ALLOCTOR
+		delete ToRecycle;
+#else
 		TheGraphEventWithInlineStorageAllocator.Push((FGraphEventAndSmallTaskStorage*)ToRecycle);
+#endif
 	}
 	else
 	{
 		GraphEventsInUse.Decrement();
+#if USE_NIEVE_GRAPH_EVENT_ALLOCTOR
+		delete ToRecycle;
+#else
 		TheGraphEventAllocator.Push(ToRecycle);
+#endif
 	}
 }
 
diff --git a/Engine/Source/Runtime/Core/Private/HAL/MallocStomp.cpp b/Engine/Source/Runtime/Core/Private/HAL/MallocStomp.cpp
--- a/Engine/Source/Runtime/Core/Private/HAL/MallocStomp.cpp
+++ b/Engine/Source/Runtime/Core/Private/HAL/MallocStomp.cpp
@@ -26,7 +26,7 @@ void* FMallocStomp::Malloc(SIZE_T Size, uint32 Alignment)
 {
 	if(Size == 0U)
 	{
-		return nullptr;
+		Size = 1U;
 	}
 
 	const SIZE_T AlignedSize = (Alignment > 0U) ? ((Size + Alignment - 1U) & -static_cast<int32>(Alignment)) : Size;
diff --git a/Engine/Source/Runtime/Core/Private/Windows/WindowsPlatformTime.cpp b/Engine/Source/Runtime/Core/Private/Windows/WindowsPlatformTime.cpp
--- a/Engine/Source/Runtime/Core/Private/Windows/WindowsPlatformTime.cpp
+++ b/Engine/Source/Runtime/Core/Private/Windows/WindowsPlatformTime.cpp
@@ -11,6 +11,7 @@ double FWindowsPlatformTime::InitTiming(void)
 	LARGE_INTEGER Frequency;
 	verify( QueryPerformanceFrequency(&Frequency) );
 	SecondsPerCycle = 1.0 / Frequency.QuadPart;
+	SecondsPerCycle64 = 1.0 / Frequency.QuadPart;
 
 	// Due to some limitation of the OS, we limit the polling frequency to 4 times per second, 
 	// but it should be enough for longterm CPU usage monitoring.
diff --git a/Engine/Source/Runtime/Core/Public/Apple/ApplePlatformTime.h b/Engine/Source/Runtime/Core/Public/Apple/ApplePlatformTime.h
--- a/Engine/Source/Runtime/Core/Public/Apple/ApplePlatformTime.h
+++ b/Engine/Source/Runtime/Core/Public/Apple/ApplePlatformTime.h
@@ -42,6 +42,12 @@ struct CORE_API FApplePlatformTime : public FGenericPlatformTime
 		return Cycles;
 	}
 
+	static FORCEINLINE uint64 Cycles64()
+	{
+		uint64 Cycles = mach_absolute_time();
+		return Cycles;
+	}
+
 	static FCPUTime GetCPUTime();
 };
 
diff --git a/Engine/Source/Runtime/Core/Public/Containers/LockFreeVoidPointerListBase.h b/Engine/Source/Runtime/Core/Public/Containers/LockFreeVoidPointerListBase.h
--- a/Engine/Source/Runtime/Core/Public/Containers/LockFreeVoidPointerListBase.h
+++ b/Engine/Source/Runtime/Core/Public/Containers/LockFreeVoidPointerListBase.h
@@ -619,7 +619,7 @@
 {
 	enum
 	{
-		NUM_PER_BUNDLE=256,
+		NUM_PER_BUNDLE=32,
 	};
 public:
 
@@ -769,12 +769,14 @@
  * never returns free space, even at shutdown
  * alignment isn't handled, assumes FMemory::Malloc will work
  */
+
+#define USE_NIEVE_TLockFreeFixedSizeAllocator_TLSCacheBase (0) // this is useful for find who really leaked
 template<int32 SIZE, typename TBundleRecycler, typename TTrackingCounter = FNoopCounter>
 class TLockFreeFixedSizeAllocator_TLSCacheBase : public FNoncopyable
 {
 	enum
 	{
-		NUM_PER_BUNDLE=256,
+		NUM_PER_BUNDLE=32,
 	};
 public:
 
@@ -800,6 +802,9 @@
 	 */
 	FORCEINLINE void* Allocate()
 	{
+#if USE_NIEVE_TLockFreeFixedSizeAllocator_TLSCacheBase
+		return FMemory::Malloc(SIZE);
+#else
 		FThreadLocalCache& TLS = GetTLS();
 
 		if (!TLS.PartialBundle)
@@ -835,6 +840,7 @@
 		TLS.NumPartial--;
 		check(TLS.NumPartial >= 0 && ((!!TLS.NumPartial) == (!!TLS.PartialBundle)));
 		return Result;
+#endif
 	}
 
 	/**
@@ -845,6 +851,9 @@
 	 */
 	FORCEINLINE void Free(void *Item)
 	{
+#if USE_NIEVE_TLockFreeFixedSizeAllocator_TLSCacheBase
+		return FMemory::Free(Item);
+#else
 		NumUsed.Decrement();
 		NumFree.Increment();
 		FThreadLocalCache& TLS = GetTLS();
@@ -862,6 +871,7 @@
 		*(void**)Item = (void*)TLS.PartialBundle;
 		TLS.PartialBundle = (void**)Item;
 		TLS.NumPartial++;
+#endif
 	}
 
 	/**
@@ -1136,14 +1146,156 @@
 };
 
 template<int TPaddingForCacheContention>
+class FSpinLocked_LIFO
+{
+public:
+	void** Pop()
+	{
+		int32 SpinCount = 0;
+		while (true)
+		{
+			if (!Lock && FPlatformAtomics::InterlockedCompareExchange((volatile int32*)&Lock, 1, 0) == 0)
+			{
+				TestCriticalStall();
+				void** Result = nullptr;
+				if (Queue.Num())
+				{
+					Result = Queue.Pop(false);
+				}
+				FPlatformMisc::MemoryBarrier();
+				Lock = 0;
+				return Result;
+			}
+			LockFreeCriticalSpin(SpinCount);
+		}
+	}
+	void Push(void** Item)
+	{
+		int32 SpinCount = 0;
+		while (true)
+		{
+			if (!Lock && FPlatformAtomics::InterlockedCompareExchange((volatile int32*)&Lock, 1, 0) == 0)
+			{
+				TestCriticalStall();
+				Queue.Add(Item);
+				FPlatformMisc::MemoryBarrier();
+				Lock = 0;
+				return;
+			}
+			LockFreeCriticalSpin(SpinCount);
+		}
+	}
+private:
+	TArray<void **> Queue;
+	uint8 PadToAvoidContention1[TPaddingForCacheContention];
+	int32 Lock;
+	uint8 PadToAvoidContention2[TPaddingForCacheContention];
+};
+
+template<>
+class FSpinLocked_LIFO<0>
+{
+public:
+	void** Pop()
+	{
+		int32 SpinCount = 0;
+		while (true)
+		{
+			if (!Lock && FPlatformAtomics::InterlockedCompareExchange((volatile int32*)&Lock, 1, 0) == 0)
+			{
+				TestCriticalStall();
+				void** Result = nullptr;
+				if (Queue.Num())
+				{
+					Result = Queue.Pop(false);
+				}
+				FPlatformMisc::MemoryBarrier();
+				Lock = 0;
+				return Result;
+			}
+			LockFreeCriticalSpin(SpinCount);
+		}
+	}
+	void Push(void** Item)
+	{
+		int32 SpinCount = 0;
+		while (true)
+		{
+			if (!Lock && FPlatformAtomics::InterlockedCompareExchange((volatile int32*)&Lock, 1, 0) == 0)
+			{
+				TestCriticalStall();
+				Queue.Add(Item);
+				FPlatformMisc::MemoryBarrier();
+				Lock = 0;
+				return;
+			}
+			LockFreeCriticalSpin(SpinCount);
+		}
+	}
+private:
+	TArray<void **> Queue;
+	int32 Lock;
+};
+
+class FLockFreePointerQueueBaseLinkAllocator : public FNoncopyable
+{
+public:
+	static CORE_API FLockFreePointerQueueBaseLinkAllocator& Get();
+
+	FORCEINLINE void Free(FLockFreeLink* Link)
+	{
+		Link->~FLockFreeLink();
+		InnerAllocator.Free(Link);
+	}
+
+	FORCEINLINE FLockFreeLink* Alloc()
+	{
+		return new (InnerAllocator.Allocate()) FLockFreeLink();
+	}
+
+	/**
+	* Gets the number of allocated memory blocks that are currently in use.
+	*
+	* @return Number of used memory blocks.
+	*/
+	const int32 GetNumUsed() const
+	{
+		return InnerAllocator.GetNumUsed().GetValue();
+	}
+
+	/**
+	* Gets the number of allocated memory blocks that are currently unused.
+	*
+	* @return Number of allocated, but unused memory blocks.
+	*/
+	const int32 GetNumFree() const
+	{
+		return InnerAllocator.GetNumFree().GetValue();
+	}
+
+private:
+	TLockFreeFixedSizeAllocator_TLSCacheBase<sizeof(FLockFreeLink), FSpinLocked_LIFO<PLATFORM_CACHE_LINE_SIZE>,
+#if 1 || UE_BUILD_SHIPPING || UE_BUILD_TEST
+		FNoopCounter
+#else
+		FThreadSafeCounter
+#endif
+	> InnerAllocator;
+};
+
+template<int TPaddingForCacheContention>
 class FLockFreePointerQueueBaseSingleConsumer : public FNoncopyable
 {
 public:
-	FLockFreePointerQueueBaseSingleConsumer(FLockFreeLink* Stub)
-		: Head(Stub)
-		, Tail(Stub)
+	FLockFreePointerQueueBaseSingleConsumer()
+		: Head(FLockFreePointerQueueBaseLinkAllocator::Get().Alloc())
+		, Tail(Head)
 	{
-
+	}
+	~FLockFreePointerQueueBaseSingleConsumer()
+	{
+		check(IsEmpty());
+		FLockFreePointerQueueBaseLinkAllocator::Get().Free(Head);
 	}
 
 	void Push(FLockFreeLink* Link)
@@ -1202,11 +1354,16 @@
 class FLockFreePointerQueueBaseSingleConsumer<0> : public FNoncopyable
 {
 public:
-	FLockFreePointerQueueBaseSingleConsumer(FLockFreeLink* Stub)
-		: Head(Stub)
-		, Tail(Stub)
+	FLockFreePointerQueueBaseSingleConsumer()
+		: Head(FLockFreePointerQueueBaseLinkAllocator::Get().Alloc())
+		, Tail(Head)
 	{
-
 	}
 
+	~FLockFreePointerQueueBaseSingleConsumer()
+	{
+		check(IsEmpty());
+		FLockFreePointerQueueBaseLinkAllocator::Get().Free(Head);
+	}
+
 	void Push(FLockFreeLink* Link)
@@ -1262,13 +1419,19 @@
 class FCloseableLockFreePointerQueueBaseSingleConsumer : public FNoncopyable
 {
 public:
-	FCloseableLockFreePointerQueueBaseSingleConsumer(FLockFreeLink* Stub)
-		: Head(Stub)
-		, Tail(Stub)
+	FCloseableLockFreePointerQueueBaseSingleConsumer()
+		: Head(FLockFreePointerQueueBaseLinkAllocator::Get().Alloc())
+		, Tail(Head)
 	{
 
 	}
 
+	~FCloseableLockFreePointerQueueBaseSingleConsumer()
+	{
+		check(IsEmpty());
+		FLockFreePointerQueueBaseLinkAllocator::Get().Free(Head);
+	}
+
 	/**	
 	 *	Push an item onto the head of the list, unless the list is closed
 	 *	@param NewItem, the new item to push on the list, cannot be nullptr
@@ -1442,11 +1605,15 @@
 class FCloseableLockFreePointerQueueBaseSingleConsumer<0> : public FNoncopyable
 {
 public:
-	FCloseableLockFreePointerQueueBaseSingleConsumer(FLockFreeLink* Stub)
-		: Head(Stub)
-		, Tail(Stub)
+	FCloseableLockFreePointerQueueBaseSingleConsumer()
+		: Head(FLockFreePointerQueueBaseLinkAllocator::Get().Alloc())
+		, Tail(Head)
 	{
-
+	}
+	~FCloseableLockFreePointerQueueBaseSingleConsumer()
+	{
+		check(IsEmpty());
+		FLockFreePointerQueueBaseLinkAllocator::Get().Free(Head);
 	}
 
 	/**
@@ -2129,143 +2296,6 @@
 	TCounter NumPushed;
 };
 
-template<int TPaddingForCacheContention>
-class FSpinLocked_LIFO
-{
-public:
-	void** Pop()
-	{
-		int32 SpinCount = 0;
-		while (true)
-		{
-			if (!Lock && FPlatformAtomics::InterlockedCompareExchange((volatile int32*)&Lock, 1, 0) == 0)
-			{
-				TestCriticalStall();
-				void** Result = nullptr;
-				if (Queue.Num())
-				{
-					Result = Queue.Pop(false);
-				}
-				FPlatformMisc::MemoryBarrier();
-				Lock = 0;
-				return Result;
-			}
-			LockFreeCriticalSpin(SpinCount);
-		}
-	}
-	void Push(void** Item)
-	{
-		int32 SpinCount = 0;
-		while (true)
-		{
-			if (!Lock && FPlatformAtomics::InterlockedCompareExchange((volatile int32*)&Lock, 1, 0) == 0)
-			{
-				TestCriticalStall();
-				Queue.Add(Item);
-				FPlatformMisc::MemoryBarrier();
-				Lock = 0;
-				return;
-			}
-			LockFreeCriticalSpin(SpinCount);
-		}
-	}	
-private:
-	TArray<void **> Queue;
-	uint8 PadToAvoidContention1[TPaddingForCacheContention];
-	int32 Lock;
-	uint8 PadToAvoidContention2[TPaddingForCacheContention];
-};
-
-template<>
-class FSpinLocked_LIFO<0>
-{
-public:
-	void** Pop()
-	{
-		int32 SpinCount = 0;
-		while (true)
-		{
-			if (!Lock && FPlatformAtomics::InterlockedCompareExchange((volatile int32*)&Lock, 1, 0) == 0)
-			{
-				TestCriticalStall();
-				void** Result = nullptr;
-				if (Queue.Num())
-				{
-					Result = Queue.Pop(false);
-				}
-				FPlatformMisc::MemoryBarrier();
-				Lock = 0;
-				return Result;
-			}
-			LockFreeCriticalSpin(SpinCount);
-		}
-	}
-	void Push(void** Item)
-	{
-		int32 SpinCount = 0;
-		while (true)
-		{
-			if (!Lock && FPlatformAtomics::InterlockedCompareExchange((volatile int32*)&Lock, 1, 0) == 0)
-			{
-				TestCriticalStall();
-				Queue.Add(Item);
-				FPlatformMisc::MemoryBarrier();
-				Lock = 0;
-				return;
-			}
-			LockFreeCriticalSpin(SpinCount);
-		}
-	}
-private:
-	TArray<void **> Queue;
-	int32 Lock;
-};
-
-class FLockFreePointerQueueBaseLinkAllocator : public FNoncopyable
-{
-public:
-	static CORE_API FLockFreePointerQueueBaseLinkAllocator& Get();
-
-	FORCEINLINE void Free(FLockFreeLink* Link)
-	{
-		Link->~FLockFreeLink();
-		InnerAllocator.Free(Link);
-	}
-
-	FORCEINLINE FLockFreeLink* Alloc()
-	{
-		return new (InnerAllocator.Allocate()) FLockFreeLink();
-	}
-
-	/**
-	 * Gets the number of allocated memory blocks that are currently in use.
-	 *
-	 * @return Number of used memory blocks.
-	 */
-	const int32 GetNumUsed() const
-	{
-		return InnerAllocator.GetNumUsed().GetValue();
-	}
-
-	/**
-	 * Gets the number of allocated memory blocks that are currently unused.
-	 *
-	 * @return Number of allocated, but unused memory blocks.
-	 */
-	const int32 GetNumFree() const
-	{
-		return InnerAllocator.GetNumFree().GetValue();
-	}
-
-private:
-	TLockFreeFixedSizeAllocator_TLSCacheBase<sizeof(FLockFreeLink), FSpinLocked_LIFO<PLATFORM_CACHE_LINE_SIZE>,
-#if 1 || UE_BUILD_SHIPPING || UE_BUILD_TEST
-		FNoopCounter
-#else
-		FThreadSafeCounter
-#endif
-	> InnerAllocator;
-};
 
 
 template<class T, int TPaddingForCacheContention>
@@ -2560,7 +2590,6 @@
 
 	FLockFreePointerListFIFOBase()
 		: LinkAllocator(FLockFreePointerQueueBaseLinkAllocator::Get())
-		, IncomingQueue(LinkAllocator.Alloc())
 		, DequeueLock(0)
 	{
 	}
@@ -2696,7 +2725,6 @@
 
 	FLockFreePointerListFIFOBase()
 		: LinkAllocator(FLockFreePointerQueueBaseLinkAllocator::Get())
-		, IncomingQueue(LinkAllocator.Alloc())
 		, DequeueLock(0)
 	{
 	}
@@ -2832,10 +2860,9 @@
 
 	FCloseableLockFreePointerListFIFOBaseSingleConsumer()
 		: LinkAllocator(FLockFreePointerQueueBaseLinkAllocator::Get())
-		, IncomingQueue(LinkAllocator.Alloc())
 	{
 	}
-	/**	
+	/**
 	 *	Push an item onto the head of the list, unless the list is closed
 	 *
 	 *	@param NewItem, the new item to push on the list, cannot be NULL
@@ -2949,7 +2976,6 @@
 
 	FCloseableLockFreePointerListFIFOBaseSingleConsumer()
 		: LinkAllocator(FLockFreePointerQueueBaseLinkAllocator::Get())
-		, IncomingQueue(LinkAllocator.Alloc())
 	{
 	}
 	/**
diff --git a/Engine/Source/Runtime/Core/Public/HAL/IPlatformFileCachedWrapper.h b/Engine/Source/Runtime/Core/Public/HAL/IPlatformFileCachedWrapper.h
--- a/Engine/Source/Runtime/Core/Public/HAL/IPlatformFileCachedWrapper.h
+++ b/Engine/Source/Runtime/Core/Public/HAL/IPlatformFileCachedWrapper.h
@@ -229,6 +229,9 @@ public:
 	}
 	virtual bool ShouldBeUsed(IPlatformFile* Inner, const TCHAR* CmdLine) const override
 	{
+#if USE_EVENT_DRIVEN_ASYNC_LOAD
+		return false; // this stuff is awful for the EDL
+#else
 		// Default to false on platforms that already do platform file level caching
 		bool bResult = !PLATFORM_PS4 && !PLATFORM_WINDOWS && !PLATFORM_IOS && !PLATFORM_ANDROID && FPlatformProperties::RequiresCookedData();
 
@@ -249,6 +252,7 @@ public:
 		}
 #endif
 		return bResult;
+#endif
 	}
 	IPlatformFile* GetLowerLevel() override
 	{
diff --git a/Engine/Source/Runtime/Core/Public/HTML5/HTML5PlatformTime.h b/Engine/Source/Runtime/Core/Public/HTML5/HTML5PlatformTime.h
--- a/Engine/Source/Runtime/Core/Public/HTML5/HTML5PlatformTime.h
+++ b/Engine/Source/Runtime/Core/Public/HTML5/HTML5PlatformTime.h
@@ -40,6 +40,13 @@ struct CORE_API FHTML5PlatformTime : public FGenericPlatformTime
 		return tb.time * 1000.0 + tb.millitm;
 	}
 
+	static FORCEINLINE uint64 Cycles64()
+	{
+		struct _timeb tb;
+		_ftime(&tb);
+		return static_cast<uint64>( tb.time * 1000.0 + tb.millitm );
+	}
+
 	static void SystemTime( int32& Year, int32& Month, int32& DayOfWeek, int32& Day, int32& Hour, int32& Min, int32& Sec, int32& MSec );
 
 	static void UtcTime( int32& Year, int32& Month, int32& DayOfWeek, int32& Day, int32& Hour, int32& Min, int32& Sec, int32& MSec );
@@ -63,6 +70,7 @@ struct CORE_API FHTML5PlatformTime : public FGenericPlatformTime
 	{
 		emscripten_t0 = emscripten_get_now(); 
 		SecondsPerCycle = 1.0f / 1000000.0f;
+		SecondsPerCycle64 = 1.0 / 1000000.0;
 		return FHTML5PlatformTime::Seconds();
 	}
 
@@ -77,6 +85,10 @@ struct CORE_API FHTML5PlatformTime : public FGenericPlatformTime
 	{
 		return (uint32)(Seconds() * 1000000);
 	}
+	static FORCEINLINE uint64 Cycles64()
+	{
+		return (uint64)(Seconds() * 1000000);
+	}
 
 };
 
diff --git a/Engine/Source/Runtime/Core/Public/Linux/LinuxPlatformTime.h b/Engine/Source/Runtime/Core/Public/Linux/LinuxPlatformTime.h
--- a/Engine/Source/Runtime/Core/Public/Linux/LinuxPlatformTime.h
+++ b/Engine/Source/Runtime/Core/Public/Linux/LinuxPlatformTime.h
@@ -28,6 +28,13 @@ struct CORE_API FLinuxTime : public FGenericPlatformTime
 		return static_cast<uint32>(static_cast<uint64>(ts.tv_sec) * 1000000ULL + static_cast<uint64>(ts.tv_nsec) / 1000ULL);
 	}
 
+	static FORCEINLINE uint64 Cycles64()
+	{
+		struct timespec ts;
+		clock_gettime(CLOCK_MONOTONIC_RAW, &ts);
+		return static_cast<uint64>(static_cast<uint64>(ts.tv_sec) * 1000000ULL + static_cast<uint64>(ts.tv_nsec) / 1000ULL);
+	}
+
 	static FCPUTime GetCPUTime();	
 
 	/**
diff --git a/Engine/Source/Runtime/Core/Public/Misc/AssertionMacros.h b/Engine/Source/Runtime/Core/Public/Misc/AssertionMacros.h
--- a/Engine/Source/Runtime/Core/Public/Misc/AssertionMacros.h
+++ b/Engine/Source/Runtime/Core/Public/Misc/AssertionMacros.h
@@ -217,6 +217,10 @@ namespace UE4Asserts_Private
 	bool GetMemberNameCheckedJunk(const volatile T&);
 }
 
+// Returns FName(TEXT("EnumeratorName")), while statically verifying that the enumerator exists in the enum
+#define GET_ENUMERATOR_NAME_CHECKED(EnumName, EnumeratorName) \
+	((void)sizeof(UE4Asserts_Private::GetMemberNameCheckedJunk(EnumName::EnumeratorName)), FName(TEXT(#EnumeratorName)))
+
 // Returns FName(TEXT("MemberName")), while statically verifying that the member exists in ClassName
 #define GET_MEMBER_NAME_CHECKED(ClassName, MemberName) \
 	((void)sizeof(UE4Asserts_Private::GetMemberNameCheckedJunk(((ClassName*)0)->MemberName)), FName(TEXT(#MemberName)))
diff --git a/Engine/Source/Runtime/Core/Public/Misc/CoreMiscDefines.h b/Engine/Source/Runtime/Core/Public/Misc/CoreMiscDefines.h
--- a/Engine/Source/Runtime/Core/Public/Misc/CoreMiscDefines.h
+++ b/Engine/Source/Runtime/Core/Public/Misc/CoreMiscDefines.h
@@ -121,3 +121,12 @@ const FPlatformUserId PLATFORMUSERID_NONE = INDEX_NONE;
 
 // When passed to pragma message will result in clickable warning in VS
 #define WARNING_LOCATION(Line) __FILE__ "(" PREPROCESSOR_TO_STRING(Line) ")"
+
+// Push and pop macro definitions
+#ifdef __clang__
+	#define PUSH_MACRO(name) _Pragma(PREPROCESSOR_TO_STRING(push_macro(PREPROCESSOR_TO_STRING(name))))
+	#define POP_MACRO(name) _Pragma(PREPROCESSOR_TO_STRING(pop_macro(PREPROCESSOR_TO_STRING(name))))
+#else
+	#define PUSH_MACRO(name) __pragma(push_macro(PREPROCESSOR_TO_STRING(name)))
+	#define POP_MACRO(name) __pragma(pop_macro(PREPROCESSOR_TO_STRING(name)))
+#endif
diff --git a/Engine/Source/Runtime/Core/Public/Windows/WindowsPlatformTime.h b/Engine/Source/Runtime/Core/Public/Windows/WindowsPlatformTime.h
--- a/Engine/Source/Runtime/Core/Public/Windows/WindowsPlatformTime.h
+++ b/Engine/Source/Runtime/Core/Public/Windows/WindowsPlatformTime.h
@@ -35,6 +35,14 @@ struct CORE_API FWindowsPlatformTime
 		return Cycles.QuadPart;
 	}
 
+	static FORCEINLINE uint64 Cycles64()
+	{
+		LARGE_INTEGER Cycles;
+		QueryPerformanceCounter(&Cycles);
+		return Cycles.QuadPart;
+	}
+
+
 	static void SystemTime( int32& Year, int32& Month, int32& DayOfWeek, int32& Day, int32& Hour, int32& Min, int32& Sec, int32& MSec );
 	static void UtcTime( int32& Year, int32& Month, int32& DayOfWeek, int32& Day, int32& Hour, int32& Min, int32& Sec, int32& MSec );
 
diff --git a/Engine/Source/Runtime/CoreUObject/Private/Blueprint/BlueprintSupport.cpp b/Engine/Source/Runtime/CoreUObject/Private/Blueprint/BlueprintSupport.cpp
--- a/Engine/Source/Runtime/CoreUObject/Private/Blueprint/BlueprintSupport.cpp
+++ b/Engine/Source/Runtime/CoreUObject/Private/Blueprint/BlueprintSupport.cpp
@@ -10,7 +10,9 @@
 #include "UObject/UObjectThreadContext.h"
 #include "ModuleManager.h"
 
-/** 
+DEFINE_LOG_CATEGORY_STATIC(LogBlueprintSupport, Log, All);
+
+/**
  * Defined in BlueprintSupport.cpp
  * Duplicates all fields of a class in depth-first order. It makes sure that everything contained
  * in a class is duplicated before the class itself, as well as all function parameters before the
@@ -1718,6 +1720,9 @@ void FLinkerLoad::CreateDynamicTypeLoader()
 	TArray<FBlueprintDependencyData> ImportDependencyPackages;
 	FConvertedBlueprintsDependencies::Get().GetAssets(LinkerRoot->GetFName(), ImportDependencyPackages);
 
+	const FName DynamicClassName = UDynamicClass::StaticClass()->GetFName();
+	const FName DynamicClassPackageName = UDynamicClass::StaticClass()->GetOuterUPackage()->GetFName();
+
 	// Create Imports
 	for (FBlueprintDependencyData& Import : ImportDependencyPackages)
 	{
@@ -1731,9 +1736,25 @@ void FLinkerLoad::CreateDynamicTypeLoader()
 		OuterImport->ClassName = NAME_Package;
 		OuterImport->ClassPackage = GLongCoreUObjectPackageName;
 		OuterImport->ObjectName = Import.PackageName;
+
+#if !USE_EVENT_DRIVEN_ASYNC_LOAD
+		if((Import.ClassName == DynamicClassName) && (Import.ClassPackageName == DynamicClassPackageName))
+		{
+			const FString DynamicClassPath = Import.PackageName.ToString() + TEXT(".") + Import.ObjectName.ToString();
+			const FName DynamicClassPathName(*DynamicClassPath);
+			FDynamicClassStaticData* ClassConstructFn = GetDynamicClassMap().Find(DynamicClassPathName);
+			if (ensure(ClassConstructFn))
+			{
+				// The class object is created here. The class is not fully constructed yet (no CLASS_Constructed flag), ZConstructor will do that later.
+				// The class object is needed to resolve circular dependencies. Regular native classes use deferred initialization/registration to avoid them.
+				ClassConstructFn->StaticClassFn();
+			}
+		}
+#endif
 	}
 
-	// Create Exports
+	// Create Export
+	int32 RunningIndex = 0;
 	{
 		FObjectExport* Export = new (ExportMap)FObjectExport();
 		const FName* TypeNamePtr = GetConvertedDynamicPackageNameToTypeName().Find(LinkerRoot->GetFName());
@@ -1742,9 +1763,44 @@ void FLinkerLoad::CreateDynamicTypeLoader()
 		// This allows us to skip creating two additional imports for UDynamicClass and its package
 		Export->bDynamicClass = true;
 		Export->ObjectFlags |= RF_Public;
+
+#if USE_EVENT_DRIVEN_ASYNC_LOAD
+		Export->FirstExportDependency = RunningIndex;
+		int32 LocalImportIndex = 0;
+		for (FBlueprintDependencyData& Import : ImportDependencyPackages)
+		{
+			if (Import.bSerializationBeforeSerializationDependency)
+			{
+				Export->SerializationBeforeSerializationDependencies++;
+				PreloadDependencies.Add(FPackageIndex::FromImport(LocalImportIndex));
+				RunningIndex++;
+			}
+			if (Import.bCreateBeforeSerializationDependency)
+			{
+				Export->CreateBeforeSerializationDependencies++;
+				PreloadDependencies.Add(FPackageIndex::FromImport(LocalImportIndex));
+				RunningIndex++;
+			}
+			if (Import.bSerializationBeforeCreateDependency)
+			{
+				Export->SerializationBeforeCreateDependencies++;
+				PreloadDependencies.Add(FPackageIndex::FromImport(LocalImportIndex));
+				RunningIndex++;
+			}
+			if (Import.bCreateBeforeCreateDependency)
+			{
+				Export->CreateBeforeCreateDependencies++;
+				PreloadDependencies.Add(FPackageIndex::FromImport(LocalImportIndex));
+				RunningIndex++;
+			}
+			LocalImportIndex += 2;
+		}
+#endif
 	}
 
+#if !USE_EVENT_DRIVEN_ASYNC_LOAD
 	LinkerRoot->SetPackageFlags(LinkerRoot->GetPackageFlags() | PKG_CompiledIn);
+#endif
 }
 
 /*******************************************************************************
@@ -2085,7 +2141,14 @@ FBlueprintDependencyData::FBlueprintDependencyData(const TCHAR* InPackageFolder
 	, ObjectName(InObjectName)
 	, ClassPackageName(InClassPackageName)
 	, ClassName(InClassName)
-{}
+{
+	// for now, we will assume that all imports are bCreateBeforeSerializationDependency, which is incorrect
+	bSerializationBeforeSerializationDependency = false;
+	bCreateBeforeSerializationDependency = true;
+	bSerializationBeforeCreateDependency = false;
+	bCreateBeforeCreateDependency = false;
+
+}
 
 FConvertedBlueprintsDependencies& FConvertedBlueprintsDependencies::Get()
 {
@@ -2119,10 +2182,30 @@ void FConvertedBlueprintsDependencies::FillUsedAssetsInDynamicClass(UDynamicClas
 	TArray<FBlueprintDependencyData> UsedAssetdData;
 	GetUsedAssets(UsedAssetdData);
 
+#if USE_EVENT_DRIVEN_ASYNC_LOAD
+	FLinkerLoad* Linker = DynamicClass->GetOutermost()->LinkerLoad;
+	if (Linker)
+	{
+		int32 ImportIndex = 0;
+		for (FBlueprintDependencyData& ItData : UsedAssetdData)
+		{
+			FObjectImport& Import = Linker->Imp(FPackageIndex::FromImport(ImportIndex));
+			check(Import.ObjectName == ItData.ObjectName);
+			UObject* TheAsset = Import.XObject;
+			UE_CLOG(!TheAsset, LogBlueprintSupport, Error, TEXT("Could not find UDynamicClass dependent asset (EDL) %s in %s"), *ItData.ObjectName.ToString(), *ItData.PackageName.ToString());
+			DynamicClass->UsedAssets.Add(TheAsset);
+			ImportIndex += 2;
+		}
+		return;
+	}
+	check(0);
+#endif
+
 	for (FBlueprintDependencyData& ItData : UsedAssetdData)
 	{
 		const FString PathToObj = FString::Printf(TEXT("%s.%s"), *ItData.PackageName.ToString(), *ItData.ObjectName.ToString());
 		UObject* TheAsset = LoadObject<UObject>(nullptr, *PathToObj);
+		UE_CLOG(!TheAsset, LogBlueprintSupport, Error, TEXT("Could not find UDynamicClass dependent asset (non-EDL) %s in %s"), *ItData.ObjectName.ToString(), *ItData.PackageName.ToString());
 		DynamicClass->UsedAssets.Add(TheAsset);
 	}
 }
diff --git a/Engine/Source/Runtime/CoreUObject/Private/Serialization/AsyncLoading.cpp b/Engine/Source/Runtime/CoreUObject/Private/Serialization/AsyncLoading.cpp
--- a/Engine/Source/Runtime/CoreUObject/Private/Serialization/AsyncLoading.cpp
+++ b/Engine/Source/Runtime/CoreUObject/Private/Serialization/AsyncLoading.cpp
@@ -300,13 +300,23 @@ void FAsyncLoadingThread::CancelAsyncLoadingInternal()
 	{
 		// Packages we started processing, need to be canceled.
 		// Accessed only in async thread, no need to protect region.
-		for (FAsyncPackage* AsyncPackage : AsyncPackages)
+		// Move first so that we remove the package from these lists BEFORE we delete it otherwise we will assert in the package dtor.
+		TArray<FAsyncPackage*> PackagesToDeleteCopy = MoveTemp(PackagesToDelete);
+
+		// This is accessed on the game thread but it should be blocked at this point
+		TArray<FAsyncPackage*> AsyncPackagesCopy = MoveTemp(AsyncPackages);
+
+		for (FAsyncPackage* Package : PackagesToDeleteCopy)
+		{
+			Package->Cancel();
+			delete Package;
+		}
+
+		for (FAsyncPackage* AsyncPackage : AsyncPackagesCopy)
 		{
 			AsyncPackage->Cancel();
 			delete AsyncPackage;
 		}
-
-		AsyncPackages.Reset();
 		AsyncPackageNameLookup.Reset();
 	}
 
@@ -1527,7 +1537,18 @@ FObjectImport* FAsyncPackage::FindExistingImport(int32 LocalImportIndex)
 				Import->XObject = StaticFindObjectFast(UObject::StaticClass(), ImportOuter->XObject, Import->ObjectName, false, true);
 				if (Import->XObject)
 				{
-					check(Import->XObject->GetClass()->GetFName() == Import->ClassName);
+//native blueprint 
+					FName NameImportClass(Import->ClassName);
+					FName NameActualImportClass(Import->XObject->GetClass()->GetFName());
+					if (NameActualImportClass != NameImportClass)
+					{
+						static const FName NAME_BlueprintGeneratedClass("BlueprintGeneratedClass");
+						static const FName NAME_DynamicClass("DynamicClass");
+						bool bBroken = (NameImportClass != NAME_BlueprintGeneratedClass && NameImportClass != NAME_DynamicClass) ||
+							(NameActualImportClass != NAME_BlueprintGeneratedClass && NameActualImportClass != NAME_DynamicClass);
+						UE_CLOG(bBroken, LogStreaming, Fatal, TEXT("FAsyncPackage::FindExistingImport class mismatch %s != %s"), *NameActualImportClass.ToString(), *NameImportClass.ToString());
+					}
+
 					AddObjectReference(Import->XObject);
 				}
 			}
@@ -1606,9 +1627,11 @@ EAsyncPackageState::Type FAsyncPackage::LoadImports_Event()
 			ExistingPackage = CastChecked<UPackage>(Import->XObject);
 			PendingPackage = ExistingPackage->LinkerLoad ? ExistingPackage->LinkerLoad->AsyncRoot : nullptr;
 		}
+		bool bCompiledIn = ExistingPackage && ExistingPackage->HasAnyPackageFlags(PKG_CompiledIn);
 		// Our import package name is the import name
 		const FName ImportPackageFName(Import->ObjectName);
-		if (!PendingPackage && (!ExistingPackage || !ExistingPackage->HasAnyPackageFlags(PKG_CompiledIn)))
+		check(!PendingPackage || !bCompiledIn); // we should never have a pending package for something that is compiled in
+		if (!PendingPackage && !bCompiledIn)
 		{
 			PendingPackage = FAsyncLoadingThread::Get().FindAsyncPackage(ImportPackageFName);
 		}
@@ -1711,6 +1734,33 @@ void FAsyncPackage::Event_SetupImports()
 	AsyncLoadingThread.QueueEvent_SetupExports(this);
 }
 
+static bool IsFullyLoadedObj(UObject* Obj)
+{
+	if (!Obj)
+	{
+		return false;
+	}
+	if (Obj->HasAllFlags(RF_WasLoaded | RF_LoadCompleted))
+	{
+		return true;
+	}
+	if (Obj->HasAnyFlags(RF_WasLoaded | RF_NeedLoad))
+	{
+		return false;
+	}
+//native blueprint 
+	UDynamicClass* UD = Cast<UDynamicClass>(Obj);
+	if (!UD)
+	{
+		return true;
+	}
+	if (UD->GetDefaultObject(false))
+	{
+		UE_CLOG(!UD->HasAnyClassFlags(CLASS_TokenStreamAssembled), LogStreaming, Fatal, TEXT("Class %s is fully loaded, but does not have its token stream assembled."), *UD->GetFullName());
+		return true;
+	}
+	return false;
+}
 
 EAsyncPackageState::Type FAsyncPackage::SetupImports_Event()
 {
@@ -1739,7 +1789,7 @@ EAsyncPackageState::Type FAsyncPackage::SetupImports_Event()
 
 		if (Import.OuterIndex.IsNull())
 		{
-			UPackage* ImportPackage = (UPackage*)Import.XObject;
+			UPackage* ImportPackage = Import.XObject ? CastChecked<UPackage>(Import.XObject) : nullptr;
 			if (!ImportPackage)
 			{
 				ImportPackage = FindObjectFast<UPackage>(NULL, Import.ObjectName, false, false);
@@ -1754,6 +1804,7 @@ EAsyncPackageState::Type FAsyncPackage::SetupImports_Event()
 			if (ImportPackage)
 			{
 				FLinkerLoad* ImportLinker = ImportPackage->LinkerLoad;
+				check(!ImportLinker || !ImportPackage->HasAnyPackageFlags(PKG_CompiledIn)); // we should never have a linker on a compiled in package
 				if (ImportLinker && ImportLinker->AsyncRoot)
 				{
 					check(ImportLinker->AsyncRoot != this);
@@ -1781,7 +1832,7 @@ EAsyncPackageState::Type FAsyncPackage::SetupImports_Event()
 			}
 			FObjectImport& OuterMostImport = Linker->Imp(OuterMostIndex);
 			check(OuterMostImport.OuterIndex.IsNull());
-			UPackage* ImportPackage = (UPackage*)OuterMostImport.XObject;
+			UPackage* ImportPackage = OuterMostImport.XObject ? CastChecked<UPackage>(OuterMostImport.XObject) : nullptr;
 			if (!ImportPackage)
 			{
 				ImportPackage = FindObjectFast<UPackage>(NULL, OuterMostImport.ObjectName, false, false);
@@ -1797,14 +1848,14 @@ EAsyncPackageState::Type FAsyncPackage::SetupImports_Event()
 			if (ImportPackage)
 			{
 				FLinkerLoad* ImportLinker = ImportPackage->LinkerLoad;
+				check(!ImportLinker || !ImportPackage->HasAnyPackageFlags(PKG_CompiledIn)); // we should never have a linker on a compiled in package
+				bool bDynamicImport = ImportLinker && ImportLinker->bDynamicClassLinker;
 				if (!ImportLinker || !ImportLinker->AsyncRoot)
 				{
 					FindExistingImport(LocalImportIndex);
-					if (!Import.XObject || (!Import.XObject->HasAnyFlags(RF_LoadCompleted) && Import.XObject->HasAnyFlags(RF_WasLoaded)))
-					{
-						UE_LOG(LogStreaming, Error, TEXT("Found package without a linker and could not find %s in %s."), *Import.ObjectName.ToString(), *ImportPackage->GetName());
-					}
-
+					bool bFinishedLoading = IsFullyLoadedObj(Import.XObject);
+					UE_CLOG(!bFinishedLoading && Import.XObject, LogStreaming, Fatal, TEXT("Found package without a linker, could find %s in %s, but somehow wasn't finished loading."), *Import.ObjectName.ToString(), *ImportPackage->GetName());
+					UE_CLOG(!bFinishedLoading && !Import.XObject, LogStreaming, Fatal, TEXT("Found package without a linker, could not find %s in %s."), *Import.ObjectName.ToString(), *ImportPackage->GetName());
 				}
 				if (ImportLinker && ImportLinker->AsyncRoot)
 				{
@@ -1813,6 +1864,7 @@ EAsyncPackageState::Type FAsyncPackage::SetupImports_Event()
 					check(Import.OuterIndex.IsImport());
 					FName OuterName = Import.OuterIndex == OuterMostIndex ? NAME_None : // if the outer is the package, then that is just a null outer in the export table
 						Linker->Imp(Import.OuterIndex).ObjectName;
+
 					FPackageIndex LocalExportIndex;
 					for (auto It = ImportLinker->AsyncRoot->ObjectNameToImportOrExport.CreateKeyIterator(Import.ObjectName); It; ++It)
 					{
@@ -1838,13 +1890,59 @@ EAsyncPackageState::Type FAsyncPackage::SetupImports_Event()
 							}
 						}
 					}
+
+//native blueprint 
+					bool bDynamicSomethingMissingFromTheFakeExportTable = bDynamicImport && LocalExportIndex.IsNull();
+
+					// this is a hack because the fake export table is missing lots
+					if (bDynamicSomethingMissingFromTheFakeExportTable)
+					{
+						LocalExportIndex = FPackageIndex::FromExport(0);
+					}
+
 					UE_CLOG(LocalExportIndex.IsNull(), LogStreaming, Warning, TEXT("Could not find import %s.%s in package %s, this is probably related to adding stuff we don't need to import map under the 'IsEventDrivenLoaderEnabledInCookedBuilds'."), *OuterName.ToString(), *Import.ObjectName.ToString(), *ImportPackage->GetName());
+					UE_CLOG(bDynamicImport && LocalExportIndex.IsNull(), LogStreaming, Fatal, TEXT("Could not find dynamic import %s.%s in package %s."), *OuterName.ToString(), *Import.ObjectName.ToString(), *ImportPackage->GetName());
 					if (!LocalExportIndex.IsNull())
 					{
 						FObjectExport& Export = ImportLinker->Exp(LocalExportIndex);
-						Import.SourceIndex = LocalExportIndex.ToExport();
-						Import.SourceLinker = ImportLinker;
+						if (bDynamicSomethingMissingFromTheFakeExportTable)
+						{
+//native blueprint 
+
+							// we can't set Import.SourceIndex because they would be incorrect
+
+							// We hope this things is available when the class is constructed
+							if (!IsFullyLoadedObj(Export.Object))
+							{
+								FEventLoadNodePtr MyDependentNode;
+								MyDependentNode.WaitingPackage = WeakThis;
+								MyDependentNode.ImportOrExportIndex = FPackageIndex::FromImport(LocalImportIndex);
+								MyDependentNode.Phase = EEventLoadNode::ImportOrExport_Create;
+
+								{
+									FEventLoadNodePtr PrereqisiteNode;
+									PrereqisiteNode.WaitingPackage = FCheckedWeakAsyncPackagePtr(ImportLinker->AsyncRoot);
+									PrereqisiteNode.ImportOrExportIndex = LocalExportIndex;
+									PrereqisiteNode.Phase = EEventLoadNode::ImportOrExport_Serialize;
+
+									// can't consider an import serialized until the corresponding export is serialized
+									AddArc(PrereqisiteNode, MyDependentNode);
+								}
+
+								{
+									FEventLoadNodePtr DependentNode;
+									DependentNode.WaitingPackage = FCheckedWeakAsyncPackagePtr(ImportLinker->AsyncRoot);
+									DependentNode.Phase = EEventLoadNode::Package_ExportsSerialized; // this could be much later...really all we care about is that the linker isn't destroyed.
+
+									// The other package should not leave the event driven loader until we have linked this import
+									AddArc(MyDependentNode, DependentNode);
+								}
+							}
+						}
+						else
 						{
+							Import.SourceIndex = LocalExportIndex.ToExport();
+							Import.SourceLinker = ImportLinker;
 							if (!Export.Object)
 							{
 								FEventLoadNodePtr MyDependentNode;
@@ -1881,8 +1979,7 @@ EAsyncPackageState::Type FAsyncPackage::SetupImports_Event()
 									AddObjectReference(Import.XObject);
 								}
 							}
-							if (!Export.Object || Export.Object->HasAnyFlags(RF_NeedLoad) 
-								|| !Export.Object->HasAnyFlags(RF_LoadCompleted)) // if this export was created by some external process (i.e. CDO), even if it doesn't have RF_NeedLoad, it will be loaded later
+							if (!IsFullyLoadedObj(Export.Object))
 							{
 								FEventLoadNodePtr MyDependentNode;
 								MyDependentNode.WaitingPackage = WeakThis;
@@ -2087,10 +2184,10 @@ EAsyncPackageState::Type FAsyncPackage::SetupExports_Event()
 		check(!Export.Object); // we should not have this yet
 		if (!Export.Object)
 		{
-			if (!Linker->FilterExport(Export) && !Export.ClassIndex.IsNull())
+//native blueprint 
+			if (!Linker->FilterExport(Export) && (!Export.ClassIndex.IsNull() || Linker->bDynamicClassLinker))
 			{
 				check(Export.ObjectName != NAME_None || !(Export.ObjectFlags&RF_Public));
-				check(!Export.bDynamicClass); // @todoio we don't handle dynamic classes
 
 				int32 RunningIndex = Export.FirstExportDependency;
 				if (RunningIndex >= 0)
@@ -2179,7 +2276,10 @@ void FAsyncPackage::LinkImport(int32 LocalImportIndex)
 	FObjectImport& Import = Linker->ImportMap[LocalImportIndex];
 	if (!Import.XObject)
 	{
-		Linker->GetFArchiveAsync2Loader()->LogItem(TEXT("LinkImport"));
+		if (Linker->GetFArchiveAsync2Loader())
+		{
+			Linker->GetFArchiveAsync2Loader()->LogItem(TEXT("LinkImport"));
+		}
 		if (Import.SourceLinker)
 		{
 			Import.XObject = Import.SourceLinker->ExportMap[Import.SourceIndex].Object;
@@ -2218,7 +2318,19 @@ void FAsyncPackage::LinkImport(int32 LocalImportIndex)
 						UClass*	FindClass = FindObjectFast<UClass>(ClassPackage, Import.ClassName, false, false);
 						if (FindClass)
 						{
-							Import.XObject = FLinkerLoad::FindImportFast(FindClass, ImportPackage, Import.ObjectName);
+							UObject* Outer = ImportPackage;
+							if (OuterMostIndex != Import.OuterIndex)
+							{
+								FObjectImport& OuterImport = Linker->Imp(Import.OuterIndex);
+								if (!OuterImport.XObject)
+								{
+									LinkImport(Import.OuterIndex.ToImport());
+								}
+								Outer = OuterImport.XObject;
+								check(Outer); // Bad dependencies, we somehow are trying to link an import when the outer of the import has not been created yet
+							}
+							Import.XObject = FLinkerLoad::FindImportFast(FindClass, Outer, Import.ObjectName);
+							UE_CLOG(!Import.XObject, LogStreaming, Fatal, TEXT("Missing import: %s in %s was not found, but the package exists."), *Import.ObjectName.ToString(), *ImportPackage->GetFullName());
 						}
 					}
 				}
@@ -2353,8 +2465,7 @@ UObject* FAsyncPackage::EventDrivenIndexToObject(FPackageIndex Index, bool bChec
 			}
 		}
 	}
-	else if (bCheckSerialized && (Result->HasAnyFlags(RF_NeedLoad) 
-		|| (Result->HasAnyFlags(RF_WasLoaded) && !Result->HasAnyFlags(RF_LoadCompleted)))) 
+	if (bCheckSerialized && !IsFullyLoadedObj(Result))
 	{
 		FEventLoadNodePtr MyDependentNode;
 		MyDependentNode.WaitingPackage = FCheckedWeakAsyncPackagePtr(this); //unnecessary
@@ -2371,7 +2482,11 @@ UObject* FAsyncPackage::EventDrivenIndexToObject(FPackageIndex Index, bool bChec
 			DumpDependencies(TEXT("Dependencies"), Linker, DumpIndex);
 		}
 
-		if (EventNodeArray.GetNode(MyDependentNode, false).bAddedToGraph || !EventNodeArray.GetNode(MyDependentNode, false).bFired)
+		if (!Result)
+		{
+			UE_LOG(LogStreaming, Fatal, TEXT("Missing Dependency, request for %s but it hasn't been created yet."), *Linker->ImpExp(Index).ObjectName.ToString());
+		}
+		else if (EventNodeArray.GetNode(MyDependentNode, false).bAddedToGraph || !EventNodeArray.GetNode(MyDependentNode, false).bFired)
 		{
 			UE_LOG(LogStreaming, Fatal, TEXT("Missing Dependency, request for %s but it was still waiting for serialization."), *Linker->ImpExp(Index).ObjectName.ToString());
 		}
@@ -2394,225 +2509,255 @@ void FAsyncPackage::EventDrivenCreateExport(int32 LocalExportIndex)
 	{
 		if (!Export.Object && !Linker->FilterExport(Export)) // for some acceptable position, it was not "not for" 
 		{
-			Linker->GetFArchiveAsync2Loader()->LogItem(TEXT("EventDrivenCreateExport"), Export.SerialOffset, Export.SerialSize);
+			if (Linker->GetFArchiveAsync2Loader())
+			{
+				Linker->GetFArchiveAsync2Loader()->LogItem(TEXT("EventDrivenCreateExport"), Export.SerialOffset, Export.SerialSize);
+			}
 			LastTypeOfWorkPerformed = TEXT("EventDrivenCreateExport");
 			LastObjectWorkWasPerformedOn = nullptr;
 			check(Export.ObjectName != NAME_None || !(Export.ObjectFlags&RF_Public));
 			check(IsLoading());
-			check(!Export.bDynamicClass); //@todo support dynamic classes
-
-			UClass* LoadClass = nullptr;
-			if (Export.ClassIndex.IsNull())
+			if (Export.bDynamicClass)
 			{
-				LoadClass = UClass::StaticClass();
+//native blueprint 
+
+				Export.Object = ConstructDynamicType(*Linker->GetExportPathName(LocalExportIndex));
+				check(Export.Object);
+				UDynamicClass* DC = Cast<UDynamicClass>(Export.Object);
+				UObject* DCD = DC ? DC->GetDefaultObject() : nullptr;
+				if (GIsInitialLoad || GUObjectArray.IsOpenForDisregardForGC())
+				{
+					Export.Object->AddToRoot();
+					if (DCD)
+					{
+						DCD->AddToRoot();
+					}
+				}
+				if (DCD)
+				{
+					AddObjectReference(DCD);
+				}
+				UE_LOG(LogStreaming, Verbose, TEXT("EventDrivenCreateExport: Created dynamic class %s"), *Export.Object->GetFullName());
+				if (Export.Object)
+				{
+					Export.Object->SetLinker(Linker, LocalExportIndex);
+				}
 			}
 			else
 			{
-				LoadClass = CastEventDrivenIndexToObject<UClass>(Export.ClassIndex, true, FPackageIndex::FromExport(LocalExportIndex));
-			}
-			if (!LoadClass)
-			{
-				UE_LOG(LogStreaming, Error, TEXT("Could not find class %s to create %s"), *Linker->ImpExp(Export.ClassIndex).ObjectName.ToString(), *Export.ObjectName.ToString());
-				Export.bExportLoadFailed = true;
-				return;
-			}
-			UStruct* SuperStruct = nullptr;
-			if (!Export.SuperIndex.IsNull())
-			{
-				SuperStruct = CastEventDrivenIndexToObject<UStruct>(Export.SuperIndex, true, FPackageIndex::FromExport(LocalExportIndex));
-				if (!SuperStruct)
+				UClass* LoadClass = nullptr;
+				if (Export.ClassIndex.IsNull())
 				{
-					// see FLinkerLoad::CreateExport, there may be some more we can do here
-					UE_LOG(LogStreaming, Error, TEXT("Could not find SuperStruct %s to create %s"), *Linker->ImpExp(Export.SuperIndex).ObjectName.ToString(), *Export.ObjectName.ToString());
+					LoadClass = UClass::StaticClass();
+				}
+				else
+				{
+					LoadClass = CastEventDrivenIndexToObject<UClass>(Export.ClassIndex, true, FPackageIndex::FromExport(LocalExportIndex));
+				}
+				if (!LoadClass)
+				{
+					UE_LOG(LogStreaming, Error, TEXT("Could not find class %s to create %s"), *Linker->ImpExp(Export.ClassIndex).ObjectName.ToString(), *Export.ObjectName.ToString());
 					Export.bExportLoadFailed = true;
 					return;
 				}
-			}
-			UObject* ThisParent = nullptr;
-			if (!Export.OuterIndex.IsNull())
-			{
-				ThisParent = EventDrivenIndexToObject(Export.OuterIndex, false, FPackageIndex::FromExport(LocalExportIndex));
-			}
-			else if (Export.bForcedExport)
-			{
-				// see FLinkerLoad::CreateExport, there may be some more we can do here
-				check(!Export.bForcedExport); // this is leftover from seekfree loading I think
-			}
-			else
-			{
-				check(LinkerRoot);
-				ThisParent = LinkerRoot;
-			}
-			check(!dynamic_cast<UObjectRedirector*>(ThisParent));
-			if (!ThisParent)
-			{
-				UE_LOG(LogStreaming, Error, TEXT("Could not find outer %s to create %s"), *Linker->ImpExp(Export.OuterIndex).ObjectName.ToString(), *Export.ObjectName.ToString());
-				Export.bExportLoadFailed = true;
-				return;
-			}
-
-			// Try to find existing object first in case we're a forced export to be able to reconcile. Also do it for the
-			// case of async loading as we cannot in-place replace objects.
-
-			UObject* ActualObjectWithTheName = StaticFindObjectFastInternal(NULL, ThisParent, Export.ObjectName, true);
-
-			// if we require cooked data, attempt to find exports in memory first...which we always do for event driven loading
-			check(FPlatformProperties::RequiresCookedData()
-				|| IsAsyncLoading()
-				|| Export.bForcedExport
-				|| LinkerRoot->ShouldFindExportsInMemoryFirst()
-				);
-
-			if (ActualObjectWithTheName && (ActualObjectWithTheName->GetClass() == LoadClass))
-			{
-				Export.Object = ActualObjectWithTheName;
-			}
-
-			// Object is found in memory.
-			if (Export.Object)
-			{
-				// Mark that we need to dissociate forced exports later on if we are a forced export.
-				if (Export.bForcedExport)
+				UStruct* SuperStruct = nullptr;
+				if (!Export.SuperIndex.IsNull())
+				{
+					SuperStruct = CastEventDrivenIndexToObject<UStruct>(Export.SuperIndex, true, FPackageIndex::FromExport(LocalExportIndex));
+					if (!SuperStruct)
+					{
+						// see FLinkerLoad::CreateExport, there may be some more we can do here
+						UE_LOG(LogStreaming, Error, TEXT("Could not find SuperStruct %s to create %s"), *Linker->ImpExp(Export.SuperIndex).ObjectName.ToString(), *Export.ObjectName.ToString());
+						Export.bExportLoadFailed = true;
+						return;
+					}
+				}
+				UObject* ThisParent = nullptr;
+				if (!Export.OuterIndex.IsNull())
+				{
+					ThisParent = EventDrivenIndexToObject(Export.OuterIndex, false, FPackageIndex::FromExport(LocalExportIndex));
+				}
+				else if (Export.bForcedExport)
 				{
 					// see FLinkerLoad::CreateExport, there may be some more we can do here
 					check(!Export.bForcedExport); // this is leftover from seekfree loading I think
 				}
-				// Associate linker with object to avoid detachment mismatches.
 				else
 				{
-					Export.Object->SetLinker(Linker, LocalExportIndex);
-
-					// If this object was allocated but never loaded (components created by a constructor, CDOs, etc) make sure it gets loaded
-					// Do this for all subobjects created in the native constructor.
-					if (!Export.Object->HasAnyFlags(RF_LoadCompleted))
-					{
-						UE_LOG(LogStreaming, Verbose, TEXT("Note: %s was found in memory but created by some other process and needs loading."), *Export.Object->GetFullName());
-						Export.Object->SetFlags(RF_NeedLoad | RF_NeedPostLoad | RF_NeedPostLoadSubobjects | RF_WasLoaded);
-					}
+					check(LinkerRoot);
+					ThisParent = LinkerRoot;
 				}
-			}
-			else
-			{
-				if (ActualObjectWithTheName && !ActualObjectWithTheName->GetClass()->IsChildOf(LoadClass))
+				check(!dynamic_cast<UObjectRedirector*>(ThisParent));
+				if (!ThisParent)
 				{
-					UE_LOG(LogLinker, Error, TEXT("Failed import: class '%s' name '%s' outer '%s'. There is another object (of '%s' class) at the path."),
-						*LoadClass->GetName(), *Export.ObjectName.ToString(), *ThisParent->GetName(), *ActualObjectWithTheName->GetClass()->GetName());
-					Export.bExportLoadFailed = true; // I am not sure if this is an actual fail or not...looked like it in the original code
+					UE_LOG(LogStreaming, Error, TEXT("Could not find outer %s to create %s"), *Linker->ImpExp(Export.OuterIndex).ObjectName.ToString(), *Export.ObjectName.ToString());
+					Export.bExportLoadFailed = true;
 					return;
 				}
 
-				// Find the Archetype object for the one we are loading.
-				check(!Export.TemplateIndex.IsNull());
-				UObject* Template = EventDrivenIndexToObject(Export.TemplateIndex, true, FPackageIndex::FromExport(LocalExportIndex));
-				check(Template);
-				// we also need to ensure that the template has set up any instances
-				Template->ConditionalPostLoadSubobjects();
+				// Try to find existing object first in case we're a forced export to be able to reconcile. Also do it for the
+				// case of async loading as we cannot in-place replace objects.
 
+				UObject* ActualObjectWithTheName = StaticFindObjectFastInternal(NULL, ThisParent, Export.ObjectName, true);
 
-				check(!GVerifyObjectReferencesOnly); // not supported with the event driven loader
-				// Create the export object, marking it with the appropriate flags to
-				// indicate that the object's data still needs to be loaded.
-				EObjectFlags ObjectLoadFlags = Export.ObjectFlags;
-				ObjectLoadFlags = EObjectFlags(ObjectLoadFlags | RF_NeedLoad | RF_NeedPostLoad | RF_NeedPostLoadSubobjects | RF_WasLoaded);
+				// if we require cooked data, attempt to find exports in memory first...which we always do for event driven loading
+				check(FPlatformProperties::RequiresCookedData()
+					|| IsAsyncLoading()
+					|| Export.bForcedExport
+					|| LinkerRoot->ShouldFindExportsInMemoryFirst()
+					);
 
-				FName NewName = Export.ObjectName;
+				if (ActualObjectWithTheName && (ActualObjectWithTheName->GetClass() == LoadClass))
+				{
+					Export.Object = ActualObjectWithTheName;
+				}
 
-				// If we are about to create a CDO, we need to ensure that all parent sub-objects are loaded
-				// to get default value initialization to work.
-#if DO_CHECK
-				if ((ObjectLoadFlags & RF_ClassDefaultObject) != 0)
+				// Object is found in memory.
+				if (Export.Object)
 				{
-					UClass* SuperClass = LoadClass->GetSuperClass();
-					UObject* SuperCDO = SuperClass ? SuperClass->GetDefaultObject() : nullptr;
-					check(!SuperCDO || Template == SuperCDO); // the template for a CDO is the CDO of the super
-					if (SuperClass && !SuperClass->IsNative())
+					// Mark that we need to dissociate forced exports later on if we are a forced export.
+					if (Export.bForcedExport)
 					{
-						check(SuperCDO);
-						if (SuperClass->HasAnyFlags(RF_NeedLoad))
-						{
-							UE_LOG(LogStreaming, Fatal, TEXT("Super %s had RF_NeedLoad while creating %s"), *SuperClass->GetFullName(), *Export.ObjectName.ToString());
-							Export.bExportLoadFailed = true;
-							return;
-						}
-						if (SuperCDO->HasAnyFlags(RF_NeedLoad))
+						// see FLinkerLoad::CreateExport, there may be some more we can do here
+						check(!Export.bForcedExport); // this is leftover from seekfree loading I think
+					}
+					// Associate linker with object to avoid detachment mismatches.
+					else
+					{
+						Export.Object->SetLinker(Linker, LocalExportIndex);
+
+						// If this object was allocated but never loaded (components created by a constructor, CDOs, etc) make sure it gets loaded
+						// Do this for all subobjects created in the native constructor.
+						if (!Export.Object->HasAnyFlags(RF_LoadCompleted))
 						{
-							UE_LOG(LogStreaming, Fatal, TEXT("Super CDO %s had RF_NeedLoad while creating %s"), *SuperCDO->GetFullName(), *Export.ObjectName.ToString());
-							Export.bExportLoadFailed = true;
-							return;
+							UE_LOG(LogStreaming, Verbose, TEXT("Note: %s was found in memory but created by some other process and needs loading."), *Export.Object->GetFullName());
+							Export.Object->SetFlags(RF_NeedLoad | RF_NeedPostLoad | RF_NeedPostLoadSubobjects | RF_WasLoaded);
 						}
-						TArray<UObject*> SuperSubObjects;
-						GetObjectsWithOuter(SuperCDO, SuperSubObjects, /*bIncludeNestedObjects=*/ false, /*ExclusionFlags=*/ RF_NoFlags, /*InternalExclusionFlags=*/ EInternalObjectFlags::Native);
+					}
+				}
+				else
+				{
+					if (ActualObjectWithTheName && !ActualObjectWithTheName->GetClass()->IsChildOf(LoadClass))
+					{
+						UE_LOG(LogLinker, Error, TEXT("Failed import: class '%s' name '%s' outer '%s'. There is another object (of '%s' class) at the path."),
+							*LoadClass->GetName(), *Export.ObjectName.ToString(), *ThisParent->GetName(), *ActualObjectWithTheName->GetClass()->GetName());
+						Export.bExportLoadFailed = true; // I am not sure if this is an actual fail or not...looked like it in the original code
+						return;
+					}
+
+					// Find the Archetype object for the one we are loading.
+					check(!Export.TemplateIndex.IsNull());
+					UObject* Template = EventDrivenIndexToObject(Export.TemplateIndex, true, FPackageIndex::FromExport(LocalExportIndex));
+					check(Template);
+					// we also need to ensure that the template has set up any instances
+					Template->ConditionalPostLoadSubobjects();
+
+
+					check(!GVerifyObjectReferencesOnly); // not supported with the event driven loader
+					// Create the export object, marking it with the appropriate flags to
+					// indicate that the object's data still needs to be loaded.
+					EObjectFlags ObjectLoadFlags = Export.ObjectFlags;
+					ObjectLoadFlags = EObjectFlags(ObjectLoadFlags | RF_NeedLoad | RF_NeedPostLoad | RF_NeedPostLoadSubobjects | RF_WasLoaded);
+
+					FName NewName = Export.ObjectName;
 
-						for (UObject* SubObject : SuperSubObjects)
+					// If we are about to create a CDO, we need to ensure that all parent sub-objects are loaded
+					// to get default value initialization to work.
+#if DO_CHECK
+					if ((ObjectLoadFlags & RF_ClassDefaultObject) != 0)
+					{
+						UClass* SuperClass = LoadClass->GetSuperClass();
+						UObject* SuperCDO = SuperClass ? SuperClass->GetDefaultObject() : nullptr;
+						check(!SuperCDO || Template == SuperCDO); // the template for a CDO is the CDO of the super
+						if (SuperClass && !SuperClass->IsNative())
 						{
-							if (SubObject->HasAnyFlags(RF_NeedLoad))
+							check(SuperCDO);
+							if (SuperClass->HasAnyFlags(RF_NeedLoad))
+							{
+								UE_LOG(LogStreaming, Fatal, TEXT("Super %s had RF_NeedLoad while creating %s"), *SuperClass->GetFullName(), *Export.ObjectName.ToString());
+								Export.bExportLoadFailed = true;
+								return;
+							}
+							if (SuperCDO->HasAnyFlags(RF_NeedLoad))
 							{
-								UE_LOG(LogStreaming, Fatal, TEXT("Super CDO subobject %s had RF_NeedLoad while creating %s"), *SubObject->GetFullName(), *Export.ObjectName.ToString());
+								UE_LOG(LogStreaming, Fatal, TEXT("Super CDO %s had RF_NeedLoad while creating %s"), *SuperCDO->GetFullName(), *Export.ObjectName.ToString());
 								Export.bExportLoadFailed = true;
 								return;
 							}
+							TArray<UObject*> SuperSubObjects;
+							GetObjectsWithOuter(SuperCDO, SuperSubObjects, /*bIncludeNestedObjects=*/ false, /*ExclusionFlags=*/ RF_NoFlags, /*InternalExclusionFlags=*/ EInternalObjectFlags::Native);
+
+							for (UObject* SubObject : SuperSubObjects)
+							{
+								if (SubObject->HasAnyFlags(RF_NeedLoad))
+								{
+									UE_LOG(LogStreaming, Fatal, TEXT("Super CDO subobject %s had RF_NeedLoad while creating %s"), *SubObject->GetFullName(), *Export.ObjectName.ToString());
+									Export.bExportLoadFailed = true;
+									return;
+								}
+							}
+						}
+						else
+						{
+							check(Template->IsA(LoadClass));
 						}
 					}
-					else
+#endif
+					if (LoadClass->HasAnyFlags(RF_NeedLoad))
 					{
-						check(Template->IsA(LoadClass));
+						UE_LOG(LogStreaming, Fatal, TEXT("LoadClass %s had RF_NeedLoad while creating %s"), *LoadClass->GetFullName(), *Export.ObjectName.ToString());
+						Export.bExportLoadFailed = true;
+						return;
 					}
-				}
-#endif
-				if (LoadClass->HasAnyFlags(RF_NeedLoad))
-				{
-					UE_LOG(LogStreaming, Fatal, TEXT("LoadClass %s had RF_NeedLoad while creating %s"), *LoadClass->GetFullName(), *Export.ObjectName.ToString());
-					Export.bExportLoadFailed = true;
-					return;
-				}
-				{
-					UObject* LoadCDO = LoadClass->GetDefaultObject();
-					if (LoadCDO->HasAnyFlags(RF_NeedLoad))
 					{
-						UE_LOG(LogStreaming, Fatal, TEXT("Class CDO %s had RF_NeedLoad while creating %s"), *LoadCDO->GetFullName(), *Export.ObjectName.ToString());
+						UObject* LoadCDO = LoadClass->GetDefaultObject();
+						if (LoadCDO->HasAnyFlags(RF_NeedLoad))
+						{
+							UE_LOG(LogStreaming, Fatal, TEXT("Class CDO %s had RF_NeedLoad while creating %s"), *LoadCDO->GetFullName(), *Export.ObjectName.ToString());
+							Export.bExportLoadFailed = true;
+							return;
+						}
+					}
+					if (Template->HasAnyFlags(RF_NeedLoad))
+					{
+						UE_LOG(LogStreaming, Fatal, TEXT("Template %s had RF_NeedLoad while creating %s"), *Template->GetFullName(), *Export.ObjectName.ToString());
 						Export.bExportLoadFailed = true;
 						return;
 					}
-				}
-				if (Template->HasAnyFlags(RF_NeedLoad))
-				{
-					UE_LOG(LogStreaming, Fatal, TEXT("Template %s had RF_NeedLoad while creating %s"), *Template->GetFullName(), *Export.ObjectName.ToString());
-					Export.bExportLoadFailed = true;
-					return;
-				}
 
-				Export.Object = StaticConstructObject_Internal
-					(
-						LoadClass,
-						ThisParent,
-						NewName,
-						ObjectLoadFlags,
-						EInternalObjectFlags::None,
-						Template,
-						false,
-						nullptr,
-						true
-					);
-				if (GIsInitialLoad || GUObjectArray.IsOpenForDisregardForGC())
-				{
-					Export.Object->AddToRoot();
-				}
-				Export.Object->SetLinker(Linker, LocalExportIndex);
-				check(Export.Object->GetClass() == LoadClass);
-				check(NewName == Export.ObjectName);
-
-				// If it's a struct or class, set its parent.
-				if (UStruct* Struct = dynamic_cast<UStruct*>(Export.Object))
-				{
-					if (SuperStruct)
+					Export.Object = StaticConstructObject_Internal
+						(
+							LoadClass,
+							ThisParent,
+							NewName,
+							ObjectLoadFlags,
+							EInternalObjectFlags::None,
+							Template,
+							false,
+							nullptr,
+							true
+							);
+					if (GIsInitialLoad || GUObjectArray.IsOpenForDisregardForGC())
 					{
-						Struct->SetSuperStruct(SuperStruct);
+						Export.Object->AddToRoot();
 					}
+					Export.Object->SetLinker(Linker, LocalExportIndex);
+					check(Export.Object->GetClass() == LoadClass);
+					check(NewName == Export.ObjectName);
 
-					// If it's a class, bind it to C++.
-					if (UClass* ClassObject = dynamic_cast<UClass*>(Export.Object))
+					// If it's a struct or class, set its parent.
+					if (UStruct* Struct = dynamic_cast<UStruct*>(Export.Object))
 					{
-						ClassObject->Bind();
+						if (SuperStruct)
+						{
+							Struct->SetSuperStruct(SuperStruct);
+						}
+
+						// If it's a class, bind it to C++.
+						if (UClass* ClassObject = dynamic_cast<UClass*>(Export.Object))
+						{
+							ClassObject->Bind();
+						}
 					}
 				}
 			}
@@ -2632,7 +2777,17 @@ void FAsyncPackage::EventDrivenSerializeExport(int32 LocalExportIndex)
 
 	FObjectExport& Export = Linker->ExportMap[LocalExportIndex];
 	UObject* Object = Export.Object;
-	if (Object && Object->HasAnyFlags(RF_NeedLoad))
+	if (Object && Linker->bDynamicClassLinker)
+	{
+//native blueprint 
+
+		UDynamicClass* UD = Cast<UDynamicClass>(Object);
+		if (UD)
+		{
+			verify(UD->GetDefaultObject(true));
+		}
+	}
+	else if (Object && Object->HasAnyFlags(RF_NeedLoad))
 	{
 		Linker->GetFArchiveAsync2Loader()->LogItem(TEXT("EventDrivenSerializeExport"), Export.SerialOffset, Export.SerialSize);
 
@@ -2726,6 +2881,18 @@ void FAsyncPackage::EventDrivenSerializeExport(int32 LocalExportIndex)
 
 void FAsyncPackage::StartPrecacheRequest()
 {
+	if (Linker->bDynamicClassLinker)
+	{
+//native blueprint 
+
+		// there is no IO for these
+		for (int32 LocalExportIndex : ExportsThatCanHaveIOStarted)
+		{
+			RemoveNode(EEventLoadNode::Export_StartIO, FPackageIndex::FromExport(LocalExportIndex));
+		}
+		ExportsThatCanHaveIOStarted.Empty();
+		return;
+	}
 	int32 LocalExportIndex = -1;
 	while (true)
 	{
@@ -2886,9 +3053,12 @@ void FAsyncPackage::FlushPrecacheBuffer()
 {
 	CurrentBlockOffset = -1;
 	CurrentBlockBytes = -1;
-	FArchiveAsync2* FAA2 = Linker->GetFArchiveAsync2Loader();
-	check(FAA2);
-	FAA2->FlushPrecacheBlock();
+	if (!Linker->bDynamicClassLinker)
+	{
+		FArchiveAsync2* FAA2 = Linker->GetFArchiveAsync2Loader();
+		check(FAA2);
+		FAA2->FlushPrecacheBlock();
+	}
 }
 
 int32 GCurrentExportIndex = -1;
@@ -2984,7 +3154,9 @@ EAsyncPackageState::Type FAsyncPackage::ProcessImportsAndExports_Event()
 			int32 LocalExportIndex = -1;
 			ExportsThatCanBeSerialized.HeapPop(LocalExportIndex, false);
 
-			if (ExportsInThisBlock.Remove(LocalExportIndex) == 1)
+//native blueprint 
+			if (Linker->bDynamicClassLinker || // dynamic things aren't actually in any block
+				ExportsInThisBlock.Remove(LocalExportIndex) == 1)
 			{
 				FGCScopeGuard GCGuard;
 				GCurrentExportIndex = LocalExportIndex;
@@ -2995,6 +3167,7 @@ EAsyncPackageState::Type FAsyncPackage::ProcessImportsAndExports_Event()
 			{
 				check(!Linker->ExportMap[LocalExportIndex].Object || !Linker->ExportMap[LocalExportIndex].Object->HasAnyFlags(RF_NeedLoad));
 			}
+
 			RemoveNode(EEventLoadNode::ImportOrExport_Serialize, FPackageIndex::FromExport(LocalExportIndex));
 		}
 		if (bDidSomething)
@@ -4060,7 +4233,7 @@ void FAsyncLoadingThread::Stop()
 
 void FAsyncLoadingThread::CancelAsyncLoading()
 {
-	checkSlow(IsInGameThread());	
+	check(IsInGameThread());	
 
 	bShouldCancelLoading = true;
 	if (IsMultithreaded())
@@ -4268,9 +4441,7 @@ void FAsyncPackage::AddObjectReference(UObject* InObject)
 	{
 		UE_CLOG(!IsInGameThread() && !IsGarbageCollectionLocked(), LogStreaming, Fatal, TEXT("Trying to add an object %s to FAsyncPackage referenced objects list outside of a FGCScopeLock."), *InObject->GetFullName());
 		{
-#if USE_EVENT_DRIVEN_ASYNC_LOAD
 			FScopeLock ReferencedObjectsLock(&ReferencedObjectsCritical);
-#endif
 			if (!ReferencedObjects.Contains(InObject))
 			{
 				ReferencedObjects.Add(InObject);
@@ -4282,11 +4453,8 @@ void FAsyncPackage::AddObjectReference(UObject* InObject)
 
 void FAsyncPackage::EmptyReferencedObjects()
 {
-	check(IsInGameThread());
 	const EInternalObjectFlags AsyncFlags = EInternalObjectFlags::Async | EInternalObjectFlags::AsyncLoading;
-#if USE_EVENT_DRIVEN_ASYNC_LOAD
 	FScopeLock ReferencedObjectsLock(&ReferencedObjectsCritical);
-#endif
 	for (UObject* Obj : ReferencedObjects)
 	{
 		// Temporary fatal messages instead of checks to find the cause for a one-time crash in shipping config
@@ -4687,6 +4855,12 @@ EAsyncPackageState::Type FAsyncPackage::CreateLinker()
 			if (Linker)
 			{
 				GPrecacheCallbackHandler.RegisterNewSummaryRequest(this);
+				if (Linker->bDynamicClassLinker)
+				{
+//native blueprint 
+					check(!Linker->GetFArchiveAsync2Loader());
+					GPrecacheCallbackHandler.SummaryComplete(WeakPtr);
+				}
 			}
 #else
 			Linker = FLinkerLoad::CreateLinkerAsync(Package, *PackageFileName, LinkerFlags);
@@ -5515,6 +5689,16 @@ EAsyncPackageState::Type FAsyncPackage::PostLoadDeferredObjects(double InTickSta
 			// The CDO is created at this point, because now it's safe to solve cyclic dependencies.
 			if (UDynamicClass* DynamicClass = Cast<UDynamicClass>(Object))
 			{
+				check((DynamicClass->ClassFlags & CLASS_Constructed) != 0);
+#if USE_EVENT_DRIVEN_ASYNC_LOAD
+//native blueprint 
+
+				check(DynamicClass->HasAnyClassFlags(CLASS_TokenStreamAssembled));
+				// this block should be removed entirely when and if we add the CDO to the fake export table
+				CDOToHandle = DynamicClass->GetDefaultObject(false);
+				UE_CLOG(CDOToHandle, LogStreaming, Fatal, TEXT("EDL did not create the CDO for %s before it finished loading."), *DynamicClass->GetFullName());
+				CDOToHandle->AtomicallyClearInternalFlags(EInternalObjectFlags::AsyncLoading);
+#else
 				UObject* OldCDO = DynamicClass->GetDefaultObject(false);
 				UObject* NewCDO = DynamicClass->GetDefaultObject(true);
 				const bool bCDOWasJustCreated = (OldCDO != NewCDO);
@@ -5523,6 +5707,7 @@ EAsyncPackageState::Type FAsyncPackage::PostLoadDeferredObjects(double InTickSta
 					NewCDO->AtomicallyClearInternalFlags(EInternalObjectFlags::AsyncLoading);
 					CDOToHandle = NewCDO;
 				}
+#endif
 			}
 			else
 			{
@@ -5665,8 +5850,8 @@ void FAsyncPackage::CloseDelayedLinkers()
 		check(LinkerToReset == LinkerToClose);
 		if (LinkerToReset && LinkerToReset->AsyncRoot)
 		{
-			UE_LOG(LogTemp, Error, TEXT("Linker cannot be reset right now...leaking %s"), *LinkerToReset->GetArchiveName())
-				continue;
+			UE_LOG(LogStreaming, Error, TEXT("Linker cannot be reset right now...leaking %s"), *LinkerToReset->GetArchiveName());
+			continue;
 		}
 	#endif
 		FLinkerManager::Get().ResetLoaders(LinkerToClose->LinkerRoot);
@@ -5697,8 +5882,23 @@ void FAsyncPackage::Cancel()
 	const EAsyncLoadingResult::Type Result = EAsyncLoadingResult::Canceled;
 	for (int32 CallbackIndex = 0; CallbackIndex < CompletionCallbacks.Num(); CallbackIndex++)
 	{
-		CompletionCallbacks[CallbackIndex].Callback.ExecuteIfBound(Desc.Name, nullptr, Result);
+		if (!CompletionCallbacks[CallbackIndex].bCalled)
+		{
+			CompletionCallbacks[CallbackIndex].Callback.ExecuteIfBound(Desc.Name, nullptr, Result);
+		}
 	}
+	{
+		// Clear load flags from any referenced objects
+		FScopeLock RereferncedObjectsLock(&ReferencedObjectsCritical);
+		const EObjectFlags ObjectLoadFlags = EObjectFlags(RF_NeedLoad | RF_NeedPostLoad | RF_NeedPostLoadSubobjects | RF_WasLoaded);
+		for (UObject* ObjRef : ReferencedObjects)
+		{			
+			ObjRef->AtomicallyClearFlags(ObjectLoadFlags);
+		}
+		// Release references
+		EmptyReferencedObjects();
+	}
+
 	bLoadHasFailed = true;
 	if (LinkerRoot)
 	{
diff --git a/Engine/Source/Runtime/CoreUObject/Private/UObject/Class.cpp b/Engine/Source/Runtime/CoreUObject/Private/UObject/Class.cpp
--- a/Engine/Source/Runtime/CoreUObject/Private/UObject/Class.cpp
+++ b/Engine/Source/Runtime/CoreUObject/Private/UObject/Class.cpp
@@ -1149,7 +1149,7 @@ void UStruct::SerializeTaggedProperties(FArchive& Ar, uint8* Data, UStruct* Defa
 						Ar << Tag;
 
 						// need to know how much data this call to SerializeTaggedProperty consumes, so mark where we are
-						int32 DataOffset = Ar.Tell();
+						int64 DataOffset = Ar.Tell();
 
 						// if using it, save the current custom property list and switch to its sub property list (in case of UStruct serialization)
 						const FCustomPropertyListNode* SavedCustomPropertyList = nullptr;
@@ -1225,7 +1225,7 @@ void UStruct::Serialize( FArchive& Ar )
 	else
 	{
 		int32 ScriptBytecodeSize = Script.Num();
-		int32 ScriptStorageSizeOffset = INDEX_NONE;
+		int64 ScriptStorageSizeOffset = INDEX_NONE;
 
 		if (Ar.IsSaving())
 		{
@@ -1245,7 +1245,7 @@ void UStruct::Serialize( FArchive& Ar )
 
 			// no bytecode patch for this struct - serialize normally [i.e. from disk]
 			int32 iCode = 0;
-			int32 const BytecodeStartOffset = Ar.Tell();
+			int64 const BytecodeStartOffset = Ar.Tell();
 
 			if (Ar.IsPersistent() && Ar.GetLinker())
 			{
@@ -1292,7 +1292,7 @@ void UStruct::Serialize( FArchive& Ar )
 			{
 				FArchive::FScopeSetDebugSerializationFlags S(Ar, DSF_IgnoreDiff);
 
-				int32 const BytecodeEndOffset = Ar.Tell();
+				int64 const BytecodeEndOffset = Ar.Tell();
 
 				// go back and write on-disk size
 				Ar.Seek(ScriptStorageSizeOffset);
@@ -4278,8 +4278,6 @@ IMPLEMENT_CORE_INTRINSIC_CLASS(UClass, UStruct,
 	}
 );
 
-UClass::StaticClassFunctionType GetDynamicClassConstructFn(FName ClassPathName);
-
 void GetPrivateStaticClassBody(
 	const TCHAR* PackageName,
 	const TCHAR* Name,
@@ -4388,14 +4386,6 @@ void GetPrivateStaticClassBody(
 
 	// Register the class's native functions.
 	RegisterNativeFunc();
-
-	if (bIsDynamic)
-	{
-		// Now call the UHT-generated Z_Construct* function for the dynamic class
-		UClass::StaticClassFunctionType ZConstructDynamicClassFn = GetDynamicClassConstructFn(*ReturnClass->GetPathName());
-		check(ZConstructDynamicClassFn);
-		ZConstructDynamicClassFn();
-	}
 }
 
 /*-----------------------------------------------------------------------------
diff --git a/Engine/Source/Runtime/CoreUObject/Private/UObject/FastReferenceCollector.h b/Engine/Source/Runtime/CoreUObject/Private/UObject/FastReferenceCollector.h
--- a/Engine/Source/Runtime/CoreUObject/Private/UObject/FastReferenceCollector.h
+++ b/Engine/Source/Runtime/CoreUObject/Private/UObject/FastReferenceCollector.h
@@ -218,7 +218,12 @@ private:
 						ObjectClass->AssembleReferenceTokenStream();
 					}
 				}
-				check(CurrentObject->GetClass()->HasAnyClassFlags(CLASS_TokenStreamAssembled));
+#if DO_CHECK
+				if (!CurrentObject->GetClass()->HasAnyClassFlags(CLASS_TokenStreamAssembled))
+				{
+					UE_LOG(LogGarbage, Fatal, TEXT("%s does not yet have a token stream assembled."), *GetFullNameSafe(CurrentObject->GetClass()));
+				}
+#endif
 
 				// Get pointer to token stream and jump to the start.
 				FGCReferenceTokenStream* RESTRICT TokenStream = &CurrentObject->GetClass()->ReferenceTokenStream;
diff --git a/Engine/Source/Runtime/CoreUObject/Private/UObject/LinkerLoad.cpp b/Engine/Source/Runtime/CoreUObject/Private/UObject/LinkerLoad.cpp
--- a/Engine/Source/Runtime/CoreUObject/Private/UObject/LinkerLoad.cpp
+++ b/Engine/Source/Runtime/CoreUObject/Private/UObject/LinkerLoad.cpp
@@ -598,7 +598,7 @@ FLinkerLoad::ELinkerStatus FLinkerLoad::Tick( float InTimeLimit, bool bInUseTime
 		do
 		{
 #if USE_EVENT_DRIVEN_ASYNC_LOAD
-			check(Loader);
+			check(Loader || bDynamicClassLinker);
 			if (true)
 #else
 			// Create loader, aka FArchive used for serialization and also precache the package file summary.
diff --git a/Engine/Source/Runtime/CoreUObject/Private/UObject/SavePackage.cpp b/Engine/Source/Runtime/CoreUObject/Private/UObject/SavePackage.cpp
--- a/Engine/Source/Runtime/CoreUObject/Private/UObject/SavePackage.cpp
+++ b/Engine/Source/Runtime/CoreUObject/Private/UObject/SavePackage.cpp
@@ -1176,7 +1176,7 @@ FArchive& FArchiveSaveTagImports::operator<<( UObject*& Obj )
 
 							const bool bCDOIsEditorOnly = IsCookingForNoEditorDataPlatform(*this) && IsEditorOnlyObject(CDO);
 
-							if (!bCDOIsEditorOnly && !CDO->HasAnyFlags(RF_Transient))
+							if (!bCDOIsEditorOnly && !CDO->HasAnyFlags(RF_Transient) && !CDO->IsPendingKill())
 							{
 								CDO->Mark(OBJECTMARK_TagImp);
 
@@ -1192,7 +1192,7 @@ FArchive& FArchiveSaveTagImports::operator<<( UObject*& Obj )
 
 										const bool bCompIsEditorOnly = IsCookingForNoEditorDataPlatform(*this) && IsEditorOnlyObject(SubObj);
 
-										if (!bCompIsEditorOnly && !SubObj->HasAnyFlags(RF_Transient))
+										if (!bCompIsEditorOnly && !SubObj->HasAnyFlags(RF_Transient) && !SubObj->IsPendingKill())
 										{
 											SubObj->Mark(OBJECTMARK_TagImp);
 										}
@@ -4868,7 +4868,7 @@ ESavePackageResult UPackage::Save(UPackage* InOuter, UObject* Base, EObjectFlags
 					auto IncludeObjectAsDependency = [Linker,ObjectMarks](TSet<FPackageIndex>& AddTo, UObject* ToTest, UObject* ForObj, bool bMandatory)
 					{
 						// Skip transient, editor only, and excluded client/server objects
-						if (ToTest && !ToTest->HasAllFlags(RF_Transient))
+						if (ToTest && !ToTest->HasAllFlags(RF_Transient) && !ToTest->IsPendingKill())
 						{
 							UPackage* Outermost = ToTest->GetOutermost();
 							check(Outermost);
@@ -4876,10 +4876,10 @@ ESavePackageResult UPackage::Save(UPackage* InOuter, UObject* Base, EObjectFlags
 							{
 								return; // we never add dependencies for things that are compiled in
 							}
-							bool bNotFiltered = !ToTest->HasAnyMarks(ObjectMarks) && (!(Linker->Summary.PackageFlags & PKG_FilterEditorOnly) || !ToTest->HasAnyMarks(OBJECTMARK_EditorOnly));
+							bool bNotFiltered = !ToTest->HasAllMarks(ObjectMarks) && (!(Linker->Summary.PackageFlags & PKG_FilterEditorOnly) || !ToTest->HasAnyMarks(OBJECTMARK_EditorOnly));
 							if (bMandatory && !bNotFiltered)
 							{
-								UE_LOG(LogSavePackage, Log, TEXT("A dependency '%s' of '%s' was filtered, but is mandatory. This indicates a problem with editor only stripping. We will keep the dependency anyway."), *ToTest->GetFullName(), *ForObj->GetFullName());
+								UE_LOG(LogSavePackage, Warning, TEXT("A dependency '%s' of '%s' was filtered, but is mandatory. This indicates a problem with editor only stripping. We will keep the dependency anyway."), *ToTest->GetFullName(), *ForObj->GetFullName());
 								bNotFiltered = true;
 							}
 							if (bNotFiltered)
@@ -4928,14 +4928,12 @@ ESavePackageResult UPackage::Save(UPackage* InOuter, UObject* Base, EObjectFlags
 								IncludeIndexAsDependency(SerializationBeforeCreateDependencies, Export.SuperIndex);
 								UObject* CDO = Export.Object->GetArchetype();
 								IncludeObjectAsDependency(SerializationBeforeCreateDependencies, CDO, Export.Object, true);
-#if 1 //@todoio delete?
 								Subobjects.Reset();
 								GetObjectsWithOuter(CDO, Subobjects);
 								for (UObject* SubObj : Subobjects)
 								{
 									IncludeObjectAsDependency(SerializationBeforeCreateDependencies, SubObj, Export.Object, false);
 								}
-#endif
 							}
 							TSet<FPackageIndex> SerializationBeforeSerializationDependencies;
 							{
diff --git a/Engine/Source/Runtime/CoreUObject/Private/UObject/UObjectArchetype.cpp b/Engine/Source/Runtime/CoreUObject/Private/UObject/UObjectArchetype.cpp
--- a/Engine/Source/Runtime/CoreUObject/Private/UObject/UObjectArchetype.cpp
+++ b/Engine/Source/Runtime/CoreUObject/Private/UObject/UObjectArchetype.cpp
@@ -44,7 +44,8 @@ UObject* UObject::GetArchetypeFromRequiredInfo(UClass* Class, UObject* Outer, FN
 					SuperClassArchetype = SuperClassArchetype->GetSuperClass())
 				{
 					Result = static_cast<UObject*>(FindObjectWithOuter(SuperClassArchetype, Class, Name));
-					if (Result)
+					// We can have invalid archetypes halfway through the hierarchy, keep looking if it's pending kill
+					if (Result && !Result->IsPendingKill())
 					{
 						break;
 					}
diff --git a/Engine/Source/Runtime/CoreUObject/Private/UObject/UObjectBase.cpp b/Engine/Source/Runtime/CoreUObject/Private/UObject/UObjectBase.cpp
--- a/Engine/Source/Runtime/CoreUObject/Private/UObject/UObjectBase.cpp
+++ b/Engine/Source/Runtime/CoreUObject/Private/UObject/UObjectBase.cpp
@@ -1144,24 +1144,12 @@ UEnum* FindExistingEnumIfHotReloadOrDynamic(UObject* Outer, const TCHAR* EnumNam
 	return Result;
 }
 
-UClass::StaticClassFunctionType GetDynamicClassConstructFn(FName ClassPathName)
-{
-	FDynamicClassStaticData* ClassConstructFn = GetDynamicClassMap().Find(ClassPathName);
-	UE_CLOG(!ClassConstructFn, LogUObjectBase, Fatal, TEXT("Unable to find construct function pointer for dynamic class %s. Make sure dynamic class exists."), *ClassPathName.ToString());
-	if (ClassConstructFn)
-	{
-		return ClassConstructFn->ZConstructFn;
-	}
-	// We should never get here. We either find the class or assert.
-	return nullptr;
-}
-
 UObject* ConstructDynamicType(FName TypePathName)
 {
 	UObject* Result = nullptr;
 	if (FDynamicClassStaticData* ClassConstructFn = GetDynamicClassMap().Find(TypePathName))
 	{
-		UClass* DynamicClass = ClassConstructFn->StaticClassFn();
+		UClass* DynamicClass = ClassConstructFn->ZConstructFn();
 		check(DynamicClass);
 		DynamicClass->AssembleReferenceTokenStream();
 		Result = DynamicClass;
@@ -1201,7 +1189,9 @@ UPackage* FindOrConstructDynamicTypePackage(const TCHAR* PackageName)
 	if (!Package)
 	{
 		Package = CreatePackage(nullptr, PackageName);
+#if !USE_EVENT_DRIVEN_ASYNC_LOAD
 		Package->SetPackageFlags(PKG_CompiledIn);
+#endif
 	}
 	check(Package);
 	return Package;
diff --git a/Engine/Source/Runtime/CoreUObject/Private/UObject/UObjectGlobals.cpp b/Engine/Source/Runtime/CoreUObject/Private/UObject/UObjectGlobals.cpp
--- a/Engine/Source/Runtime/CoreUObject/Private/UObject/UObjectGlobals.cpp
+++ b/Engine/Source/Runtime/CoreUObject/Private/UObject/UObjectGlobals.cpp
@@ -1543,17 +1543,17 @@ void EndLoad()
 			ThreadContext.ObjLoaded.Empty();
 
 			// Sort by Filename and Offset.
-			ObjLoaded.Sort( FCompareUObjectByLinkerAndOffset() );
+			ObjLoaded.Sort(FCompareUObjectByLinkerAndOffset());
 
 			// Finish loading everything.
-			for( int32 i=0; i<ObjLoaded.Num(); i++ )
+			for (int32 i = 0; i < ObjLoaded.Num(); i++)
 			{
 				// Preload.
 				UObject* Obj = ObjLoaded[i];
-				if( Obj->HasAnyFlags(RF_NeedLoad) )
+				if (Obj->HasAnyFlags(RF_NeedLoad))
 				{
 					check(Obj->GetLinker());
-					Obj->GetLinker()->Preload( Obj );
+					Obj->GetLinker()->Preload(Obj);
 				}
 			}
 
@@ -1570,12 +1570,12 @@ void EndLoad()
 			SlowTask.CurrentFrameScope = 0;
 #endif
 
-			if ( GIsEditor )
+			if (GIsEditor)
 			{
-				for( int32 i=0; i<ObjLoaded.Num(); i++ )
+				for (int32 i = 0; i < ObjLoaded.Num(); i++)
 				{
 					UObject* Obj = ObjLoaded[i];
-					if ( Obj->GetLinker() )
+					if (Obj->GetLinker())
 					{
 						LoadedLinkers.Add(Obj->GetLinker());
 					}
@@ -1587,9 +1587,9 @@ void EndLoad()
 				TGuardValue<bool> GuardIsRoutingPostLoad(FUObjectThreadContext::Get().IsRoutingPostLoad, true);
 
 				// Postload objects.
-				for(int32 i = 0; i < ObjLoaded.Num(); i++)
+				for (int32 i = 0; i < ObjLoaded.Num(); i++)
 				{
-				
+
 					UObject* Obj = ObjLoaded[i];
 					check(Obj);
 #if WITH_EDITOR
@@ -1599,16 +1599,34 @@ void EndLoad()
 				}
 			}
 
-			// Dynamic Class doesn't require/use pre-loading (or post-loading). 
-			// The CDO is created at this point, because now it's safe to solve cyclic dependencies.
-			for (UObject* Obj : ObjLoaded)
+#if USE_EVENT_DRIVEN_ASYNC_LOAD
+			if (!GIsInitialLoad)
 			{
-				if (UDynamicClass* DynamicClass = Cast<UDynamicClass>(Obj))
+#if DO_CHECK
+				for (UObject* Obj : ObjLoaded)
 				{
-					DynamicClass->GetDefaultObject(true);
+					if (UDynamicClass* DynamicClass = Cast<UDynamicClass>(Obj))
+					{
+						check((DynamicClass->ClassFlags & CLASS_Constructed) != 0);
+						check(DynamicClass->GetDefaultObject(false)); // this should have already been done
+					}
+				}
+#endif
+			}
+			else
+#endif
+			{
+				// Dynamic Class doesn't require/use pre-loading (or post-loading). 
+				// The CDO is created at this point, because now it's safe to solve cyclic dependencies.
+				for (UObject* Obj : ObjLoaded)
+				{
+					if (UDynamicClass* DynamicClass = Cast<UDynamicClass>(Obj))
+					{
+						check((DynamicClass->ClassFlags & CLASS_Constructed) != 0);
+						DynamicClass->GetDefaultObject(true);
+					}
 				}
 			}
-
 			// Create clusters after all objects have been loaded
 			extern int32 GCreateGCClusters;
 			if (FPlatformProperties::RequiresCookedData() && !GIsInitialLoad && GCreateGCClusters && !GUObjectArray.IsOpenForDisregardForGC())
diff --git a/Engine/Source/Runtime/CoreUObject/Public/Blueprint/BlueprintSupport.h b/Engine/Source/Runtime/CoreUObject/Public/Blueprint/BlueprintSupport.h
--- a/Engine/Source/Runtime/CoreUObject/Public/Blueprint/BlueprintSupport.h
+++ b/Engine/Source/Runtime/CoreUObject/Public/Blueprint/BlueprintSupport.h
@@ -218,6 +218,12 @@ struct COREUOBJECT_API FBlueprintDependencyData
 	FName ClassPackageName;
 	FName ClassName;
 
+	/* If we have more than one element in the fake export table, then this will need to be an array or otherwise we will need to know which export is needs which imports of which kind of dependency*/
+	bool bSerializationBeforeSerializationDependency;
+	bool bCreateBeforeSerializationDependency;
+	bool bSerializationBeforeCreateDependency;
+	bool bCreateBeforeCreateDependency;
+
 	FBlueprintDependencyData() {}
 
 	FORCENOINLINE FBlueprintDependencyData(const TCHAR* InPackageFolder
diff --git a/Engine/Source/Runtime/CoreUObject/Public/Serialization/AsyncLoading.h b/Engine/Source/Runtime/CoreUObject/Public/Serialization/AsyncLoading.h
--- a/Engine/Source/Runtime/CoreUObject/Public/Serialization/AsyncLoading.h
+++ b/Engine/Source/Runtime/CoreUObject/Public/Serialization/AsyncLoading.h
@@ -572,10 +572,8 @@ private:
 	int32 ReentryCount;
 	/** List of objects referenced by this package */
 	TSet<UObject*> ReferencedObjects;
-#if USE_EVENT_DRIVEN_ASYNC_LOAD
 	/** Critical section for referenced objects list */
 	FCriticalSection ReferencedObjectsCritical;
-#endif
 	/** Cached async loading thread object this package was created by */
 	class FAsyncLoadingThread& AsyncLoadingThread;
 
diff --git a/Engine/Source/Runtime/Engine/Classes/Components/TimelineComponent.h b/Engine/Source/Runtime/Engine/Classes/Components/TimelineComponent.h
--- a/Engine/Source/Runtime/Engine/Classes/Components/TimelineComponent.h
+++ b/Engine/Source/Runtime/Engine/Classes/Components/TimelineComponent.h
@@ -260,7 +260,8 @@ public:
 	, bPlaying( false )
 	, PlayRate( 1.f )
 	, Position( 0.0f )	
-	, PropertySetObject(NULL)
+	, PropertySetObject(nullptr)
+	, DirectionProperty(nullptr)
 	{
 	}
 
diff --git a/Engine/Source/Runtime/Engine/Private/DataTable.cpp b/Engine/Source/Runtime/Engine/Private/DataTable.cpp
--- a/Engine/Source/Runtime/Engine/Private/DataTable.cpp
+++ b/Engine/Source/Runtime/Engine/Private/DataTable.cpp
@@ -58,7 +58,7 @@ void UDataTable::LoadStructData(FArchive& Ar)
 		Ar << RowName;
 
 		// Load row data
-		uint8* RowData = (uint8*)FMemory::Malloc(LoadUsingStruct->PropertiesSize);
+		uint8* RowData = (uint8*)FMemory::Malloc(LoadUsingStruct->GetStructureSize());
 
 		// And be sure to call DestroyScriptStruct later
 		LoadUsingStruct->InitializeStruct(RowData);
@@ -185,7 +185,7 @@ void UDataTable::GetResourceSizeEx(FResourceSizeEx& CumulativeResourceSize)
 	CumulativeResourceSize.AddDedicatedSystemMemoryBytes(RowMap.GetAllocatedSize());
 	if (RowStruct)
 	{
-		CumulativeResourceSize.AddDedicatedSystemMemoryBytes(RowMap.Num() * RowStruct->PropertiesSize);
+		CumulativeResourceSize.AddDedicatedSystemMemoryBytes(RowMap.Num() * RowStruct->GetStructureSize());
 	}
 }
 
diff --git a/Engine/Source/Runtime/Engine/Private/DataTableCSV.cpp b/Engine/Source/Runtime/Engine/Private/DataTableCSV.cpp
--- a/Engine/Source/Runtime/Engine/Private/DataTableCSV.cpp
+++ b/Engine/Source/Runtime/Engine/Private/DataTableCSV.cpp
@@ -163,7 +163,7 @@ bool FDataTableImporterCSV::ReadTable()
 		}
 
 		// Allocate data to store information, using UScriptStruct to know its size
-		uint8* RowData = (uint8*)FMemory::Malloc(DataTable->RowStruct->PropertiesSize);
+		uint8* RowData = (uint8*)FMemory::Malloc(DataTable->RowStruct->GetStructureSize());
 		DataTable->RowStruct->InitializeStruct(RowData);
 		// And be sure to call DestroyScriptStruct later
 
diff --git a/Engine/Source/Runtime/Engine/Private/DataTableJSON.cpp b/Engine/Source/Runtime/Engine/Private/DataTableJSON.cpp
--- a/Engine/Source/Runtime/Engine/Private/DataTableJSON.cpp
+++ b/Engine/Source/Runtime/Engine/Private/DataTableJSON.cpp
@@ -322,7 +322,7 @@ bool FDataTableImporterJSON::ReadRow(const TSharedRef<FJsonObject>& InParsedTabl
 	}
 
 	// Allocate data to store information, using UScriptStruct to know its size
-	uint8* RowData = (uint8*)FMemory::Malloc(DataTable->RowStruct->PropertiesSize);
+	uint8* RowData = (uint8*)FMemory::Malloc(DataTable->RowStruct->GetStructureSize());
 	DataTable->RowStruct->InitializeStruct(RowData);
 	// And be sure to call DestroyScriptStruct later
 
diff --git a/Engine/Source/Runtime/Engine/Private/InheritableComponentHandler.cpp b/Engine/Source/Runtime/Engine/Private/InheritableComponentHandler.cpp
--- a/Engine/Source/Runtime/Engine/Private/InheritableComponentHandler.cpp
+++ b/Engine/Source/Runtime/Engine/Private/InheritableComponentHandler.cpp
@@ -161,6 +161,12 @@ void UInheritableComponentHandler::ValidateTemplates()
 				}
 				else
 				{
+					// Set pending kill so this object does not get used as an archetype for subclasses
+					if (Record.ComponentTemplate)
+					{
+						Record.ComponentTemplate->MarkPendingKill();
+					}
+
 					UE_LOG(LogBlueprint, Log, TEXT("ValidateTemplates '%s': overridden template is unnecessary and will be removed - component '%s' from '%s'"),
 						*GetPathNameSafe(this), *VarName.ToString(), *GetPathNameSafe(ComponentKey.GetComponentOwner()));
 				}
diff --git a/Engine/Source/Runtime/Engine/Private/Particles/ParticleComponents.cpp b/Engine/Source/Runtime/Engine/Private/Particles/ParticleComponents.cpp
--- a/Engine/Source/Runtime/Engine/Private/Particles/ParticleComponents.cpp
+++ b/Engine/Source/Runtime/Engine/Private/Particles/ParticleComponents.cpp
@@ -3642,7 +3642,10 @@ void UParticleSystemComponent::DestroyRenderState_Concurrent()
 		ResetParticles();
 	}
 
-	Super::DestroyRenderState_Concurrent();
+	if (bRenderStateCreated)
+	{
+		Super::DestroyRenderState_Concurrent();
+	}
 }
 
 
diff --git a/Engine/Source/Runtime/Engine/Private/World.cpp b/Engine/Source/Runtime/Engine/Private/World.cpp
--- a/Engine/Source/Runtime/Engine/Private/World.cpp
+++ b/Engine/Source/Runtime/Engine/Private/World.cpp
@@ -2765,8 +2765,7 @@ void UWorld::FlushLevelStreaming(EFlushLevelStreamingType FlushType)
 		}
 	}
 	
-	// Update level streaming one last time to make sure all RemoveFromWorld requests made it.
-	UpdateLevelStreaming();
+	check(!IsVisibilityRequestPending());
 
 	// we need this, or the traces will be abysmally slow
 	EnsureCollisionTreeIsBuilt();
diff --git a/Engine/Source/Runtime/Online/BuildPatchServices/Private/BuildPatchDownloader.h b/Engine/Source/Runtime/Online/BuildPatchServices/Private/BuildPatchDownloader.h
--- a/Engine/Source/Runtime/Online/BuildPatchServices/Private/BuildPatchDownloader.h
+++ b/Engine/Source/Runtime/Online/BuildPatchServices/Private/BuildPatchDownloader.h
@@ -140,7 +140,10 @@ private:
 	float ChunkSuccessRate;
 
 	// Used to help time health states
-	volatile int64 CyclesAtLastHealthState;
+	MS_ALIGN(8) volatile int64 CyclesAtLastHealthState;
+
+	// The time in cycles since we received data
+	MS_ALIGN(8) volatile int64 CyclesAtLastData;
 
 	// Timers in seconds for how long we were in each health state
 	TArray<float> HealthStateTimes;
@@ -151,9 +154,6 @@ private:
 	// A critical section to protect the flags and rate
 	FCriticalSection FlagsLock;
 
-	// The time in cycles since we received data
-	volatile int64 CyclesAtLastData;
-
 	// Store a record of each download we made for info
 	TArray< FBuildPatchDownloadRecord > DownloadRecords;
 
diff --git a/Engine/Source/Runtime/PakFile/Public/IPlatformFilePak.h b/Engine/Source/Runtime/PakFile/Public/IPlatformFilePak.h
--- a/Engine/Source/Runtime/PakFile/Public/IPlatformFilePak.h
+++ b/Engine/Source/Runtime/PakFile/Public/IPlatformFilePak.h
@@ -1139,11 +1139,11 @@ public:
 		if (IsNonPakFilenameAllowed(Filename))
 		{
 #if USE_NEW_ASYNC_IO
-			double StartTime(UE_LOG_ACTIVE(LogPakFile, Error) ? FPlatformTime::Seconds() : 0.0);
+			double StartTime(UE_LOG_ACTIVE(LogPakFile, Verbose) ? FPlatformTime::Seconds() : 0.0);
 #endif
 			Result = LowerLevel->GetTimeStamp(Filename);
 #if USE_NEW_ASYNC_IO
-			UE_LOG(LogPakFile, Error, TEXT("GetTimeStamp on disk (!!) for %s took %6.2fms."), Filename, float(FPlatformTime::Seconds() - StartTime) * 1000.0f);
+			UE_LOG(LogPakFile, Verbose, TEXT("GetTimeStamp on disk (!!) for %s took %6.2fms."), Filename, float(FPlatformTime::Seconds() - StartTime) * 1000.0f);
 #endif
 		}
 		return Result;
diff --git a/Engine/Source/Runtime/RHI/Private/RHICommandList.cpp b/Engine/Source/Runtime/RHI/Private/RHICommandList.cpp
--- a/Engine/Source/Runtime/RHI/Private/RHICommandList.cpp
+++ b/Engine/Source/Runtime/RHI/Private/RHICommandList.cpp
@@ -1576,7 +1576,7 @@ bool FRHICommandListImmediate::StallRHIThread()
 
 void FRHICommandListImmediate::UnStallRHIThread()
 {
-	check(IsInRenderingThread() || GRHIThread && GRHIThreadStallTask.GetReference() && !GRHIThreadStallTask->IsComplete() && GRHIThreadStallEvent);
+	check(IsInRenderingThread() && GRHIThread && GRHIThreadStallTask.GetReference() && !GRHIThreadStallTask->IsComplete() && GRHIThreadStallEvent);
 	GRHIThreadStallEvent->Trigger();
 	SCOPE_CYCLE_COUNTER(STAT_SpinWaitRHIThreadUnstall);
 	while (!GRHIThreadStallTask->IsComplete())
diff --git a/Engine/Source/Runtime/ShaderCore/Private/Shader.cpp b/Engine/Source/Runtime/ShaderCore/Private/Shader.cpp
--- a/Engine/Source/Runtime/ShaderCore/Private/Shader.cpp
+++ b/Engine/Source/Runtime/ShaderCore/Private/Shader.cpp
@@ -287,10 +287,10 @@ FArchive& operator<<(FArchive& Ar,FShaderType*& Ref)
 }
 
 
-TRefCountPtr<FShader> FShaderType::FindShaderById(const FShaderId& Id)
+FShader* FShaderType::FindShaderById(const FShaderId& Id)
 {
 	check(IsInGameThread());
-	TRefCountPtr<FShader> Result = ShaderIdMap.FindRef(Id);
+	FShader* Result = ShaderIdMap.FindRef(Id);
 	return Result;
 }
 
@@ -495,10 +495,10 @@ void FShaderResource::Release()
 }
 
 
-TRefCountPtr<FShaderResource> FShaderResource::FindShaderResourceById(const FShaderResourceId& Id)
+FShaderResource* FShaderResource::FindShaderResourceById(const FShaderResourceId& Id)
 {
 	check(IsInGameThread());
-	TRefCountPtr<FShaderResource> Result = ShaderResourceIdMap.FindRef(Id);
+	FShaderResource* Result = ShaderResourceIdMap.FindRef(Id);
 	return Result;
 }
 
@@ -970,7 +970,7 @@ bool FShader::SerializeBase(FArchive& Ar, bool bShadersInline)
 			ResourceId.SpecificShaderTypeName = Type->LimitShaderResourceToThisType() ? Type->GetName() : NULL;
 
 			// use it to look up in the registered resource map
-			TRefCountPtr<FShaderResource> ExistingResource = FShaderResource::FindShaderResourceById(ResourceId);
+			FShaderResource* ExistingResource = FShaderResource::FindShaderResourceById(ResourceId);
 			SetResource(ExistingResource);
 		}
 	}
@@ -1037,7 +1037,7 @@ void FShader::RegisterSerializedResource()
 {
 	if (SerializedResource)
 	{
-		TRefCountPtr<FShaderResource> ExistingResource = FShaderResource::FindShaderResourceById(SerializedResource->GetId());
+		FShaderResource* ExistingResource = FShaderResource::FindShaderResourceById(SerializedResource->GetId());
 
 		// Reuse an existing shader resource if a matching one already exists in memory
 		if (ExistingResource)
diff --git a/Engine/Source/Runtime/ShaderCore/Public/Shader.h b/Engine/Source/Runtime/ShaderCore/Public/Shader.h
--- a/Engine/Source/Runtime/ShaderCore/Public/Shader.h
+++ b/Engine/Source/Runtime/ShaderCore/Public/Shader.h
@@ -200,7 +200,7 @@ public:
 	virtual void FinishCleanup();
 
 	/** Finds a matching shader resource in memory if possible. */
-	SHADERCORE_API static TRefCountPtr<FShaderResource> FindShaderResourceById(const FShaderResourceId& Id);
+	SHADERCORE_API static FShaderResource* FindShaderResourceById(const FShaderResourceId& Id);
 
 	/** 
 	 * Finds a matching shader resource in memory or creates a new one with the given compiler output.  
@@ -811,7 +811,7 @@ public:
 	 * Finds a shader of this type by ID.
 	 * @return NULL if no shader with the specified ID was found.
 	 */
-	TRefCountPtr<FShader> FindShaderById(const FShaderId& Id);
+	FShader* FindShaderById(const FShaderId& Id);
 
 	/** Constructs a new instance of the shader type for deserialization. */
 	FShader* ConstructForDeserialization() const;
@@ -1709,12 +1709,12 @@ public:
 			Shader->RegisterSerializedResource();
 
 			FShaderType* Type = Shader->GetType();
-			TRefCountPtr<FShader> ExistingShader = Type->FindShaderById(Shader->GetId());
+			FShader* ExistingShader = Type->FindShaderById(Shader->GetId());
 
-			if (ExistingShader.IsValid())
+			if (ExistingShader != nullptr)
 			{
 				delete Shader;
-				Shader = ExistingShader.GetReference();
+				Shader = ExistingShader;
 			}
 			else
 			{
