[NO NON-PII COMMIT MESSAGE]

diff --git a/Engine/Source/Runtime/CoreUObject/Private/UObject/SavePackage.cpp b/Engine/Source/Runtime/CoreUObject/Private/UObject/SavePackage.cpp
--- a/Engine/Source/Runtime/CoreUObject/Private/UObject/SavePackage.cpp
+++ b/Engine/Source/Runtime/CoreUObject/Private/UObject/SavePackage.cpp
@@ -874,6 +874,11 @@ FArchive& FArchiveSaveTagExports::operator<<( UObject*& Obj )
 				// shouldn't be saved, as it does not contain any data that should be saved
 				Obj->UnMark(OBJECTMARK_TagExp);
 			}
+			else if (IsEventDrivenLoaderEnabledInCookedBuilds() && IsCooking())
+			{
+				// In EDL we don't always load this object, so mark it appropriately client/server only
+				ConditionallyExcludeObjectForTarget(Obj, *this);
+			}
 		}
 		else
 		{
