Copying //UE4/Dev-Core to //UE4/Dev-Main (Source: //UE4/Dev-Core @ 3208226)

diff --git a/Engine/Source/Runtime/Core/Private/HAL/IPlatformFileLogWrapper.cpp b/Engine/Source/Runtime/Core/Private/HAL/IPlatformFileLogWrapper.cpp
--- a/Engine/Source/Runtime/Core/Private/HAL/IPlatformFileLogWrapper.cpp
+++ b/Engine/Source/Runtime/Core/Private/HAL/IPlatformFileLogWrapper.cpp
@@ -2,6 +2,7 @@
 
 #include "CorePrivatePCH.h"
 #include "IPlatformFileLogWrapper.h"
+#include "UniquePtr.h"
 
 bool bSuppressFileLog = false;
 DEFINE_LOG_CATEGORY(LogPlatformFile);
@@ -29,7 +30,7 @@ public:
 		return false;
 	}
 };
-static TAutoPtr<FFileLogExec> GFileLogExec;
+static TUniquePtr<FFileLogExec> GFileLogExec;
 
 #endif // !UE_BUILD_SHIPPING
 
@@ -46,7 +47,7 @@ bool FLoggedPlatformFile::Initialize(IPlatformFile* Inner, const TCHAR* CommandL
 	LowerLevel = Inner;
 
 #if !UE_BUILD_SHIPPING
-	GFileLogExec = new FFileLogExec(*this);
+	GFileLogExec = MakeUnique<FFileLogExec>(*this);
 #endif
 
 	return !!LowerLevel;
diff --git a/Engine/Source/Runtime/Core/Private/HAL/PlatformFileManager.cpp b/Engine/Source/Runtime/Core/Private/HAL/PlatformFileManager.cpp
--- a/Engine/Source/Runtime/Core/Private/HAL/PlatformFileManager.cpp
+++ b/Engine/Source/Runtime/Core/Private/HAL/PlatformFileManager.cpp
@@ -11,7 +11,7 @@
 #include "../HAL/IPlatformFileCachedWrapper.h"
 #include "../HAL/IPlatformFileModule.h"
 #include "../HAL/IPlatformFileOpenLogWrapper.h"
-
+#include "UniquePtr.h"
 
 FPlatformFileManager::FPlatformFileManager()
 	: TopmostPlatformFile(NULL)
@@ -52,35 +52,35 @@ IPlatformFile* FPlatformFileManager::GetPlatformFile(const TCHAR* Name)
 	// Check Core platform files (Profile, Log) by name.
 	if (FCString::Strcmp(FLoggedPlatformFile::GetTypeName(), Name) == 0)
 	{
-		static TScopedPointer<IPlatformFile> AutoDestroySingleton(new FLoggedPlatformFile());
-		PlatformFile = AutoDestroySingleton.GetOwnedPointer();
+		static TUniquePtr<IPlatformFile> AutoDestroySingleton(new FLoggedPlatformFile());
+		PlatformFile = AutoDestroySingleton.Get();
 	}
 #if !UE_BUILD_SHIPPING
 	else if (FCString::Strcmp(TProfiledPlatformFile<FProfiledFileStatsFileDetailed>::GetTypeName(), Name) == 0)
 	{
-		static TScopedPointer<IPlatformFile> AutoDestroySingleton(new TProfiledPlatformFile<FProfiledFileStatsFileDetailed>());
-		PlatformFile = AutoDestroySingleton.GetOwnedPointer();
+		static TUniquePtr<IPlatformFile> AutoDestroySingleton(new TProfiledPlatformFile<FProfiledFileStatsFileDetailed>());
+		PlatformFile = AutoDestroySingleton.Get();
 	}
 	else if (FCString::Strcmp(TProfiledPlatformFile<FProfiledFileStatsFileSimple>::GetTypeName(), Name) == 0)
 	{
-		static TScopedPointer<IPlatformFile> AutoDestroySingleton(new TProfiledPlatformFile<FProfiledFileStatsFileSimple>());
-		PlatformFile = AutoDestroySingleton.GetOwnedPointer();
+		static TUniquePtr<IPlatformFile> AutoDestroySingleton(new TProfiledPlatformFile<FProfiledFileStatsFileSimple>());
+		PlatformFile = AutoDestroySingleton.Get();
 	}
 	else if (FCString::Strcmp(FPlatformFileReadStats::GetTypeName(), Name) == 0)
 	{
-		static TScopedPointer<IPlatformFile> AutoDestroySingleton(new FPlatformFileReadStats());
-		PlatformFile = AutoDestroySingleton.GetOwnedPointer();
+		static TUniquePtr<IPlatformFile> AutoDestroySingleton(new FPlatformFileReadStats());
+		PlatformFile = AutoDestroySingleton.Get();
 	}
 	else if (FCString::Strcmp(FPlatformFileOpenLog::GetTypeName(), Name) == 0)
 	{
-		static TScopedPointer<IPlatformFile> AutoDestroySingleton(new FPlatformFileOpenLog());
-		PlatformFile = AutoDestroySingleton.GetOwnedPointer();
+		static TUniquePtr<IPlatformFile> AutoDestroySingleton(new FPlatformFileOpenLog());
+		PlatformFile = AutoDestroySingleton.Get();
 	}
 #endif
 	else if (FCString::Strcmp(FCachedReadPlatformFile::GetTypeName(), Name) == 0)
 	{
-		static TScopedPointer<IPlatformFile> AutoDestroySingleton(new FCachedReadPlatformFile());
-		PlatformFile = AutoDestroySingleton.GetOwnedPointer();
+		static TUniquePtr<IPlatformFile> AutoDestroySingleton(new FCachedReadPlatformFile());
+		PlatformFile = AutoDestroySingleton.Get();
 	}
 	else if (FModuleManager::Get().ModuleExists(Name))
 	{
diff --git a/Engine/Source/Runtime/Core/Private/HAL/ThreadHeartBeat.cpp b/Engine/Source/Runtime/Core/Private/HAL/ThreadHeartBeat.cpp
--- a/Engine/Source/Runtime/Core/Private/HAL/ThreadHeartBeat.cpp
+++ b/Engine/Source/Runtime/Core/Private/HAL/ThreadHeartBeat.cpp
@@ -114,7 +114,7 @@ uint32 FThreadHeartBeat::Run()
 
 			// First verify we're not reporting the same hang over and over again
 			uint32 CallstackCRC = FCrc::StrCrc32(StackTrace);
-			if (CallstackCRC != LastHangCallstackCRC && ThreadThatHung != LastHungThreadId)
+			if (CallstackCRC != LastHangCallstackCRC || ThreadThatHung != LastHungThreadId)
 			{
 				LastHangCallstackCRC = CallstackCRC;
 				LastHungThreadId = ThreadThatHung;
diff --git a/Engine/Source/Runtime/Core/Private/Internationalization/TextLocalizationManager.cpp b/Engine/Source/Runtime/Core/Private/Internationalization/TextLocalizationManager.cpp
--- a/Engine/Source/Runtime/Core/Private/Internationalization/TextLocalizationManager.cpp
+++ b/Engine/Source/Runtime/Core/Private/Internationalization/TextLocalizationManager.cpp
@@ -2,6 +2,7 @@
 
 #include "CorePrivatePCH.h"
 #include "Misc/App.h"
+#include "UniquePtr.h"
 
 DEFINE_LOG_CATEGORY_STATIC(LogTextLocalizationManager, Log, All);
 
@@ -142,7 +143,7 @@ void FTextLocalizationManager::FLocalizationEntryTracker::LoadFromDirectory(cons
 
 bool FTextLocalizationManager::FLocalizationEntryTracker::LoadFromFile(const FString& FilePath)
 {
-	TScopedPointer<FArchive> Reader(IFileManager::Get().CreateFileReader( *FilePath ));
+	TUniquePtr<FArchive> Reader(IFileManager::Get().CreateFileReader( *FilePath ));
 	if( !Reader )
 	{
 		return false;
diff --git a/Engine/Source/Runtime/Core/Private/Misc/AES.cpp b/Engine/Source/Runtime/Core/Private/Misc/AES.cpp
--- a/Engine/Source/Runtime/Core/Private/Misc/AES.cpp
+++ b/Engine/Source/Runtime/Core/Private/Misc/AES.cpp
@@ -19,8 +19,6 @@
 
 // It may be useful to integrate this into a class derived from FArchive to allow streaming encryption/decryption - but that's a mission for another day.
 
-#ifdef AES_KEY
-
 #define TEST_ENCRYPTION			0
 #define AES_KEYBITS				256
 
@@ -29,7 +27,9 @@
 #define NROUNDS( keybits )			( ( keybits ) / 32 + 6 )
 
 // This is quite insecure and could do with some obfuscation work
+#ifdef AES_KEY
 static ANSICHAR* DefinedKey = (ANSICHAR*)AES_KEY;
+#endif
 
 static const uint32 Te0[256] =
 {
@@ -977,7 +977,7 @@ static void rijndaelEncrypt( const uint32 *rk, int32 nrounds, const uint8 plaint
 		}
 	}
 
-	rk += nrounds << 2;
+	rk += ((UPTRINT)nrounds) << 2;
 
 	/*
 	 * apply last round and
@@ -1100,7 +1100,7 @@ static void rijndaelDecrypt( const uint32 *rk, int32 nrounds, const uint8 cipher
 			t3 = Td0[s3 >> 24] ^ Td1[(s2 >> 16) & 0xff] ^ Td2[(s1 >>  8) & 0xff] ^ Td3[s0 & 0xff] ^ rk[55];
 		}
 	}
-	rk += nrounds << 2;
+	rk += (UPTRINT)nrounds << 2;
 	/*
 	 * apply last round and
 	 * map cipher state to byte array block:
@@ -1137,7 +1137,6 @@ static void rijndaelDecrypt( const uint32 *rk, int32 nrounds, const uint8 cipher
 		rk[3];
 	PUTU32( plaintext + 12, s3 );
 }
-#endif
 
 void FAES::EncryptData(uint8 *Contents, uint32 NumBytes)
 {
@@ -1148,11 +1147,11 @@ void FAES::EncryptData(uint8 *Contents, uint32 NumBytes)
 
 void FAES::EncryptData( uint8 *Contents, uint32 NumBytes, ANSICHAR* Key )
 {
-#ifdef AES_KEY
 	uint32 rk[RKLENGTH(AES_KEYBITS)] = { 0 };
 	int32 nrounds;
 
 	checkf( ( NumBytes & ( AESBlockSize - 1 ) ) == 0, TEXT( "NumBytes needs to be a multiple of 16 bytes" ) );
+	checkf( Key, TEXT("No encryption key specified") );
 	checkf( TCString<ANSICHAR>::Strlen( Key ) >= KEYLENGTH( AES_KEYBITS ), TEXT( "AES_KEY needs to be at least %d characters" ), KEYLENGTH( AES_KEYBITS ) );
 
 #if TEST_ENCRYPTION
@@ -1178,8 +1177,6 @@ void FAES::EncryptData( uint8 *Contents, uint32 NumBytes, ANSICHAR* Key )
 	DecryptData( DecryptedBlob.GetData(), DecryptedBlob.Num() );
 	checkf( DecryptedBlob == OriginalBlob, TEXT( "Decrypted data different from original data - this is bad" ) );
 #endif
-
-#endif
 }
 
 void FAES::DecryptData(uint8 *Contents, uint32 NumBytes)
@@ -1191,7 +1188,6 @@ void FAES::DecryptData(uint8 *Contents, uint32 NumBytes)
 
 void FAES::DecryptData( uint8 *Contents, uint32 NumBytes, ANSICHAR* Key )
 {
-#ifdef AES_KEY
 	uint32 rk[RKLENGTH(AES_KEYBITS)] = { 0 };
 	int32 nrounds;
 
@@ -1206,5 +1202,4 @@ void FAES::DecryptData( uint8 *Contents, uint32 NumBytes, ANSICHAR* Key )
 	{
 		rijndaelDecrypt( rk, nrounds, Contents + Offset, Contents + Offset );
 	}
-#endif
 }
diff --git a/Engine/Source/Runtime/Core/Private/Misc/DefaultValueHelper.cpp b/Engine/Source/Runtime/Core/Private/Misc/DefaultValueHelper.cpp
--- a/Engine/Source/Runtime/Core/Private/Misc/DefaultValueHelper.cpp
+++ b/Engine/Source/Runtime/Core/Private/Misc/DefaultValueHelper.cpp
@@ -599,6 +599,20 @@ bool FDefaultValueHelper::ParseInt(const FString& Source, int32& OutVal)
 }
 
 
+bool FDefaultValueHelper::ParseInt64(const FString& Source, int64& OutVal)
+{
+	int32 Base;
+	if( !Source.IsEmpty() && IsStringValidInteger( StartOf(Source), EndOf(Source), Base ) )
+	{
+		const bool bHasWhitespace = HasWhitespaces(Source);
+		const FString NoWhitespace = bHasWhitespace ? RemoveWhitespaces(Source) : FString();
+		OutVal = FCString::Strtoui64( bHasWhitespace ? *NoWhitespace : *Source , NULL, Base );
+		return true;
+	}
+	return false;
+}
+
+
 bool FDefaultValueHelper::ParseFloat(const FString& Source, float& OutVal)
 {
 	if( IsStringValidFloat( Source ) )
diff --git a/Engine/Source/Runtime/Core/Private/ProfilingDebugging/Windows/AQtimeExternalProfiler.cpp b/Engine/Source/Runtime/Core/Private/ProfilingDebugging/Windows/AQtimeExternalProfiler.cpp
--- a/Engine/Source/Runtime/Core/Private/ProfilingDebugging/Windows/AQtimeExternalProfiler.cpp
+++ b/Engine/Source/Runtime/Core/Private/ProfilingDebugging/Windows/AQtimeExternalProfiler.cpp
@@ -3,7 +3,7 @@
 #include "CorePrivatePCH.h"
 #include "ExternalProfiler.h"
 #include "Runtime/Core/Public/Features/IModularFeatures.h"
-
+#include "UniquePtr.h"
 
 
 /**
@@ -111,7 +111,7 @@ namespace AQtimeProfiler
 	{
 		FAtModuleInit()
 		{
-			static TScopedPointer<FAQtimeExternalProfiler> ProfilerAQtime( new FAQtimeExternalProfiler() );
+			static TUniquePtr<FAQtimeExternalProfiler> ProfilerAQtime = MakeUnique<FAQtimeExternalProfiler>();
 			if( !ProfilerAQtime->Initialize() )
 			{
 				ProfilerAQtime.Reset();
diff --git a/Engine/Source/Runtime/Core/Private/ProfilingDebugging/Windows/VSPerfExternalProfiler.cpp b/Engine/Source/Runtime/Core/Private/ProfilingDebugging/Windows/VSPerfExternalProfiler.cpp
--- a/Engine/Source/Runtime/Core/Private/ProfilingDebugging/Windows/VSPerfExternalProfiler.cpp
+++ b/Engine/Source/Runtime/Core/Private/ProfilingDebugging/Windows/VSPerfExternalProfiler.cpp
@@ -3,6 +3,7 @@
 #include "CorePrivatePCH.h"
 #include "ExternalProfiler.h"
 #include "Runtime/Core/Public/Features/IModularFeatures.h"
+#include "UniquePtr.h"
 
 // Not all versions of Visual Studio include the profiler SDK headers
 #if WITH_VS_PERF_PROFILER
@@ -146,7 +147,7 @@ namespace VSPerfProfiler
 	{
 		FAtModuleInit()
 		{
-			static TScopedPointer<FVSPerfExternalProfiler> ProfilerVSPerf( new FVSPerfExternalProfiler() );
+			static TUniquePtr<FVSPerfExternalProfiler> ProfilerVSPerf = MakeUnique<FVSPerfExternalProfiler>();
 			if( !ProfilerVSPerf->Initialize() )
 			{
 				ProfilerVSPerf.Reset();
diff --git a/Engine/Source/Runtime/Core/Private/ProfilingDebugging/Windows/VTuneExternalProfiler.cpp b/Engine/Source/Runtime/Core/Private/ProfilingDebugging/Windows/VTuneExternalProfiler.cpp
--- a/Engine/Source/Runtime/Core/Private/ProfilingDebugging/Windows/VTuneExternalProfiler.cpp
+++ b/Engine/Source/Runtime/Core/Private/ProfilingDebugging/Windows/VTuneExternalProfiler.cpp
@@ -3,7 +3,7 @@
 #include "CorePrivatePCH.h"
 #include "ExternalProfiler.h"
 #include "Runtime/Core/Public/Features/IModularFeatures.h"
-
+#include "UniquePtr.h"
 
 
 /**
@@ -141,7 +141,7 @@ namespace VTuneProfiler
 	{
 		FAtModuleInit()
 		{
-			static TScopedPointer<FVTuneExternalProfiler> ProfilerVTune( new FVTuneExternalProfiler() );
+			static TUniquePtr<FVTuneExternalProfiler> ProfilerVTune = MakeUnique<FVTuneExternalProfiler>();
 			if( !ProfilerVTune->Initialize() )
 			{
 				ProfilerVTune.Reset();
diff --git a/Engine/Source/Runtime/Core/Private/Stats/StatsFile.cpp b/Engine/Source/Runtime/Core/Private/Stats/StatsFile.cpp
--- a/Engine/Source/Runtime/Core/Private/Stats/StatsFile.cpp
+++ b/Engine/Source/Runtime/Core/Private/Stats/StatsFile.cpp
@@ -5,6 +5,7 @@
 =============================================================================*/
 
 #include "CorePrivatePCH.h"
+#include "UniquePtr.h"
 
 #if	STATS
 
@@ -1183,7 +1184,7 @@ void FCommandStatsFile::TestLastSaved()
 		}
 	};
 
-	TAutoPtr<FStatsTestReader> Instance( FStatsReader<FStatsTestReader>::Create( *FilePath ) );
+	TUniquePtr<FStatsTestReader> Instance( FStatsReader<FStatsTestReader>::Create( *FilePath ) );
 
 	extern void DumpHistoryFrame( FStatsThreadState const& StatsData, int64 TargetFrame, float InDumpCull = 0.0f, int32 MaxDepth = MAX_int32, TCHAR const* Filter = NULL );
 
diff --git a/Engine/Source/Runtime/Core/Private/UObject/DevObjectVersion.cpp b/Engine/Source/Runtime/Core/Private/UObject/DevObjectVersion.cpp
--- a/Engine/Source/Runtime/Core/Private/UObject/DevObjectVersion.cpp
+++ b/Engine/Source/Runtime/Core/Private/UObject/DevObjectVersion.cpp
@@ -90,3 +90,8 @@ FDevVersionRegistration GRegisterSequencerObjectVersion(FSequencerObjectVersion:
 const FGuid FVRObjectVersion::GUID(0xD7296918, 0x1DD64BDD, 0x9DE264A8, 0x3CC13884);
 // Register VR custom version with Core
 FDevVersionRegistration GRegisterVRObjectVersion(FVRObjectVersion::GUID, FVRObjectVersion::LatestVersion, TEXT("Dev-VR"));
+
+// Unique Load Times version id
+const FGuid FLoadTimesObjectVersion::GUID(0xC2A15278, 0xBFE74AFE, 0x6C1790FF, 0x531DF755);
+// Register LoadTimes custom version with Core
+FDevVersionRegistration GRegisterLoadTimesObjectVersion(FLoadTimesObjectVersion::GUID, FLoadTimesObjectVersion::LatestVersion, TEXT("Dev-LoadTimes"));
diff --git a/Engine/Source/Runtime/Core/Private/Windows/WindowsApplication.cpp b/Engine/Source/Runtime/Core/Private/Windows/WindowsApplication.cpp
--- a/Engine/Source/Runtime/Core/Private/Windows/WindowsApplication.cpp
+++ b/Engine/Source/Runtime/Core/Private/Windows/WindowsApplication.cpp
@@ -10,6 +10,7 @@
 #include "IInputDevice.h"
 #include "IHapticDevice.h"
 #include "HAL/ThreadHeartBeat.h"
+#include "UniquePtr.h"
 
 #if WITH_EDITOR
 #include "ModuleManager.h"
@@ -883,11 +884,11 @@ int32 FWindowsApplication::ProcessMessage( HWND hwnd, uint32 msg, WPARAM wParam,
 				uint32 Size = 0;
 				::GetRawInputData((HRAWINPUT)lParam, RID_INPUT, NULL, &Size, sizeof(RAWINPUTHEADER));
 
-				TScopedPointer<uint8> RawData(new uint8[Size]);
+				TUniquePtr<uint8[]> RawData = MakeUnique<uint8[]>(Size);
 
-				if (::GetRawInputData((HRAWINPUT)lParam, RID_INPUT, RawData.GetOwnedPointer(), &Size, sizeof(RAWINPUTHEADER)) == Size )
+				if (::GetRawInputData((HRAWINPUT)lParam, RID_INPUT, RawData.Get(), &Size, sizeof(RAWINPUTHEADER)) == Size )
 				{
-					const RAWINPUT* const Raw = (const RAWINPUT* const)RawData.GetOwnedPointer();
+					const RAWINPUT* const Raw = (const RAWINPUT* const)RawData.Get();
 
 					if (Raw->header.dwType == RIM_TYPEMOUSE) 
 					{
@@ -1761,12 +1762,12 @@ int32 FWindowsApplication::ProcessDeferredMessage( const FDeferredWindowsMessage
 				UINT InputCount = LOWORD( wParam );
 				if ( InputCount > 0 )
 				{
-					TUniquePtr<TOUCHINPUT> Inputs( new TOUCHINPUT[InputCount] );
+					TUniquePtr<TOUCHINPUT[]> Inputs = MakeUnique<TOUCHINPUT[]>( InputCount );
 					if ( GetTouchInputInfo( (HTOUCHINPUT)lParam, InputCount, Inputs.Get(), sizeof(TOUCHINPUT) ) )
 					{
 						for ( uint32 i = 0; i < InputCount; i++ )
 						{
-							TOUCHINPUT Input = Inputs.Get()[i];
+							TOUCHINPUT Input = Inputs[i];
 							FVector2D Location( Input.x / 100.0f, Input.y / 100.0f );
 							if ( Input.dwFlags & TOUCHEVENTF_DOWN )
 							{
@@ -2373,19 +2374,19 @@ void FWindowsApplication::QueryConnectedMice()
 	for (const auto& Device : DeviceList)
 	{
 		UINT NameLen = 0;
-		TAutoPtr<char> Name;
+		TUniquePtr<char[]> Name;
 		if (Device.dwType != RIM_TYPEMOUSE)
 			continue;
 		//Force the use of ANSI versions of these calls
 		if (GetRawInputDeviceInfoA(Device.hDevice, RIDI_DEVICENAME, nullptr, &NameLen) == static_cast<UINT>(-1))
 			continue;
 
-		Name.Reset(new char[NameLen+1]);
-		if (GetRawInputDeviceInfoA(Device.hDevice, RIDI_DEVICENAME, Name.GetOwnedPointer(), &NameLen) == static_cast<UINT>(-1))
+		Name = MakeUnique<char[]>(NameLen+1);
+		if (GetRawInputDeviceInfoA(Device.hDevice, RIDI_DEVICENAME, Name.Get(), &NameLen) == static_cast<UINT>(-1))
 			continue;
 
 		Name[NameLen] = 0;
-		FString WName = ANSI_TO_TCHAR(Name);
+		FString WName = ANSI_TO_TCHAR(Name.Get());
 		WName.ReplaceInline(TEXT("#"), TEXT("\\"), ESearchCase::CaseSensitive);
 		/*
 		 * Name XP starts with \??\, vista+ starts \\?\ 
diff --git a/Engine/Source/Runtime/Core/Private/Windows/WindowsPlatformCrashContext.cpp b/Engine/Source/Runtime/Core/Private/Windows/WindowsPlatformCrashContext.cpp
--- a/Engine/Source/Runtime/Core/Private/Windows/WindowsPlatformCrashContext.cpp
+++ b/Engine/Source/Runtime/Core/Private/Windows/WindowsPlatformCrashContext.cpp
@@ -9,6 +9,7 @@
 #include "HAL/ExceptionHandling.h"
 #include "HAL/ThreadHeartBeat.h"
 #include "AllowWindowsPlatformTypes.h"
+#include "UniquePtr.h"
 
 #include <strsafe.h>
 #include <dbghelp.h>
@@ -529,14 +530,14 @@ private:
 
 #include "HideWindowsPlatformTypes.h"
 
-TAutoPtr<FCrashReportingThread> GCrashReportingThread(new FCrashReportingThread());
+TUniquePtr<FCrashReportingThread> GCrashReportingThread = MakeUnique<FCrashReportingThread>();
 
 // #CrashReport: 2015-05-28 This should be named EngineCrashHandler
 int32 ReportCrash( LPEXCEPTION_POINTERS ExceptionInfo )
 {
 	// Only create a minidump the first time this function is called.
 	// (Can be called the first time from the RenderThread, then a second time from the MainThread.)
-	if (GCrashReportingThread.IsValid())
+	if (GCrashReportingThread)
 	{
 		if (FPlatformAtomics::InterlockedIncrement(&ReportCrashCallCount) == 1)
 		{
diff --git a/Engine/Source/Runtime/Core/Public/Android/AndroidString.h b/Engine/Source/Runtime/Core/Public/Android/AndroidString.h
--- a/Engine/Source/Runtime/Core/Public/Android/AndroidString.h
+++ b/Engine/Source/Runtime/Core/Public/Android/AndroidString.h
@@ -344,6 +344,31 @@ struct FAndroidPlatformString : public FGenericPlatformString
 		return Res;
 	}
 
+	static FORCEINLINE int64 Strtoi64( const WIDECHAR* Start, WIDECHAR** End, int32 Base ) 
+	{
+		int StartLen = Strlen(Start);
+		ANSICHAR* AnsiStart = (ANSICHAR*)FMemory_Alloca(StartLen+1);
+		CopyWideToAnsi(AnsiStart, Start);
+
+		ANSICHAR* AnsiEnd = NULL;
+
+		uint64 Res = strtoll(AnsiStart, &AnsiEnd, Base);
+
+		if (End)
+		{
+			if (AnsiEnd == NULL)
+			{
+				*End = NULL;
+			}
+			else
+			{
+				*End = (WIDECHAR*)(Start + (AnsiEnd - AnsiStart));
+			}
+		}
+
+		return Res;
+	}
+
 	static FORCEINLINE uint64 Strtoui64( const WIDECHAR* Start, WIDECHAR** End, int32 Base ) 
 	{
 		int StartLen = Strlen(Start);
@@ -518,6 +543,11 @@ struct FAndroidPlatformString : public FGenericPlatformString
 		return strtol( Start, End, Base ); 
 	}
 
+	static FORCEINLINE int64 Strtoi64( const ANSICHAR* Start, ANSICHAR** End, int32 Base ) 
+	{
+		return strtoll(Start, End, Base);;
+	}
+
 	static FORCEINLINE uint64 Strtoui64( const ANSICHAR* Start, ANSICHAR** End, int32 Base ) 
 	{
 		return strtoull(Start, End, Base);;
diff --git a/Engine/Source/Runtime/Core/Public/Containers/EnumAsByte.h b/Engine/Source/Runtime/Core/Public/Containers/EnumAsByte.h
--- a/Engine/Source/Runtime/Core/Public/Containers/EnumAsByte.h
+++ b/Engine/Source/Runtime/Core/Public/Containers/EnumAsByte.h
@@ -5,13 +5,9 @@
 #include "Templates/UnrealTypeTraits.h"
 #include "Templates/IsEnumClass.h"
 
-#define DEPRECATE_ENUM_AS_BYTE_FOR_ENUM_CLASSES 0
-
-#if DEPRECATE_ENUM_AS_BYTE_FOR_ENUM_CLASSES
-	template <bool> struct TEnumAsByte_EnumClass;
-	template <> struct DEPRECATED(4.15, "TEnumAsByte is not intended for use with enum classes - please derive your enum class from uint8 instead.") TEnumAsByte_EnumClass<true> {};
-	template <> struct TEnumAsByte_EnumClass<false> {};
-#endif
+template <bool> struct TEnumAsByte_EnumClass;
+template <> struct DEPRECATED(4.15, "TEnumAsByte is not intended for use with enum classes - please derive your enum class from uint8 instead.") TEnumAsByte_EnumClass<true> {};
+template <> struct TEnumAsByte_EnumClass<false> {};
 
 /**
  * Template to store enumeration values as bytes in a type-safe way.
@@ -19,9 +15,7 @@
 template<class TEnum>
 class TEnumAsByte
 {
-#if DEPRECATE_ENUM_AS_BYTE_FOR_ENUM_CLASSES
 	typedef TEnumAsByte_EnumClass<TIsEnumClass<TEnum>::Value> Check;
-#endif
 
 public:
 	typedef TEnum EnumType;
diff --git a/Engine/Source/Runtime/Core/Public/Containers/SparseArray.h b/Engine/Source/Runtime/Core/Public/Containers/SparseArray.h
--- a/Engine/Source/Runtime/Core/Public/Containers/SparseArray.h
+++ b/Engine/Source/Runtime/Core/Public/Containers/SparseArray.h
@@ -267,7 +267,7 @@ public:
 			int32 ElementIndex = Data.AddUninitialized(ElementsToAdd);
 
 			// now mark the new elements as free
-			for ( int32 FreeIndex = ElementIndex; FreeIndex < ExpectedNumElements; FreeIndex++ )
+			for ( int32 FreeIndex = ExpectedNumElements - 1; FreeIndex >= ElementIndex; --FreeIndex )
 			{
 				if(NumFreeIndices)
 				{
diff --git a/Engine/Source/Runtime/Core/Public/GenericPlatform/GenericPlatformMath.h b/Engine/Source/Runtime/Core/Public/GenericPlatform/GenericPlatformMath.h
--- a/Engine/Source/Runtime/Core/Public/GenericPlatform/GenericPlatformMath.h
+++ b/Engine/Source/Runtime/Core/Public/GenericPlatform/GenericPlatformMath.h
@@ -390,6 +390,25 @@ struct FGenericPlatformMath
 		}*/
 	}
 
+	/**
+	 * Computes the base 2 logarithm for a 64-bit value that is greater than 0.
+	 * The result is rounded down to the nearest integer.
+	 *
+	 * @param Value		The value to compute the log of
+	 * @return			Log2 of Value. 0 if Value is 0.
+	 */	
+	static FORCEINLINE uint64 FloorLog2_64(uint64 Value) 
+	{
+		uint64 pos = 0;
+		if (Value >= 1ull<<32) { Value >>= 32; pos += 32; }
+		if (Value >= 1ull<<16) { Value >>= 16; pos += 16; }
+		if (Value >= 1ull<< 8) { Value >>=  8; pos +=  8; }
+		if (Value >= 1ull<< 4) { Value >>=  4; pos +=  4; }
+		if (Value >= 1ull<< 2) { Value >>=  2; pos +=  2; }
+		if (Value >= 1ull<< 1) {				pos +=  1; }
+		return (Value == 0) ? 0 : pos;
+	}
+
 	/**
 	 * Counts the number of leading zeros in the bit representation of the value
 	 *
@@ -403,6 +422,19 @@ struct FGenericPlatformMath
 		return 31 - FloorLog2(Value);
 	}
 
+	/**
+	 * Counts the number of leading zeros in the bit representation of the 64-bit value
+	 *
+	 * @param Value the value to determine the number of leading zeros for
+	 *
+	 * @return the number of zeros before the first "on" bit
+	 */
+	static FORCEINLINE uint64 CountLeadingZeros64(uint64 Value)
+	{
+		if (Value == 0) return 64;
+		return 63 - FloorLog2_64(Value);
+	}
+
 	/**
 	 * Counts the number of trailing zeros in the bit representation of the value
 	 *
@@ -435,6 +467,12 @@ struct FGenericPlatformMath
 		return (32 - CountLeadingZeros(Arg - 1)) & (~Bitmask);
 	}
 
+	static FORCEINLINE uint64 CeilLogTwo64( uint64 Arg )
+	{
+		int64 Bitmask = ((int64)(CountLeadingZeros64(Arg) << 57)) >> 63;
+		return (64 - CountLeadingZeros64(Arg - 1)) & (~Bitmask);
+	}
+
 	/** @return Rounds the given number up to the next highest power of two. */
 	static FORCEINLINE uint32 RoundUpToPowerOfTwo(uint32 Arg)
 	{
diff --git a/Engine/Source/Runtime/Core/Public/GenericPlatform/StandardPlatformString.h b/Engine/Source/Runtime/Core/Public/GenericPlatform/StandardPlatformString.h
--- a/Engine/Source/Runtime/Core/Public/GenericPlatform/StandardPlatformString.h
+++ b/Engine/Source/Runtime/Core/Public/GenericPlatform/StandardPlatformString.h
@@ -141,6 +141,15 @@ public:
 		return wcstol( Start, End, Base );
 	}
 
+	static FORCEINLINE int64 Strtoi64( const WIDECHAR* Start, WIDECHAR** End, int32 Base ) 
+	{
+#if PLATFORM_HTML5_WIN32
+		return _wtoi64(Start);
+#else
+		return wcstoll( Start, End, Base );
+#endif
+	}
+
 	static FORCEINLINE uint64 Strtoui64( const WIDECHAR* Start, WIDECHAR** End, int32 Base ) 
 	{
 #if PLATFORM_HTML5_WIN32
@@ -299,6 +308,15 @@ public:
 		return strtol( Start, End, Base ); 
 	}
 
+	static FORCEINLINE int64 Strtoi64( const ANSICHAR* Start, ANSICHAR** End, int32 Base ) 
+	{
+#if PLATFORM_HTML5_WIN32
+		return _atoi64(Start);
+#else
+		return strtoll(Start, End, Base);
+#endif
+	}
+
 	static FORCEINLINE uint64 Strtoui64( const ANSICHAR* Start, ANSICHAR** End, int32 Base ) 
 	{
 #if PLATFORM_HTML5_WIN32
diff --git a/Engine/Source/Runtime/Core/Public/HAL/FileManagerGeneric.h b/Engine/Source/Runtime/Core/Public/HAL/FileManagerGeneric.h
--- a/Engine/Source/Runtime/Core/Public/HAL/FileManagerGeneric.h
+++ b/Engine/Source/Runtime/Core/Public/HAL/FileManagerGeneric.h
@@ -2,6 +2,7 @@
 
 #pragma once
 
+#include "UniquePtr.h"
 
 /**
  * Base class for file managers.
@@ -234,7 +235,7 @@ protected:
 	int64 Pos;
 	int64 BufferBase;
 	int64 BufferCount;
-	TAutoPtr<IFileHandle> Handle;
+	TUniquePtr<IFileHandle> Handle;
 	uint8 Buffer[1024];
 };
 
@@ -296,7 +297,7 @@ protected:
 	FString Filename;
 	int64 Pos;
 	int64 BufferCount;
-	TAutoPtr<IFileHandle> Handle;
+	TUniquePtr<IFileHandle> Handle;
 	uint8 Buffer[4096];
 	bool bLoggingError;
 };
diff --git a/Engine/Source/Runtime/Core/Public/HAL/IPlatformFileCachedWrapper.h b/Engine/Source/Runtime/Core/Public/HAL/IPlatformFileCachedWrapper.h
--- a/Engine/Source/Runtime/Core/Public/HAL/IPlatformFileCachedWrapper.h
+++ b/Engine/Source/Runtime/Core/Public/HAL/IPlatformFileCachedWrapper.h
@@ -2,6 +2,8 @@
 
 #pragma once
 
+#include "UniquePtr.h"
+
 class CORE_API FCachedFileHandle : public IFileHandle
 {
 public:
@@ -195,7 +197,7 @@ private:
 		}
 	}
 
-	TAutoPtr<IFileHandle>	FileHandle;
+	TUniquePtr<IFileHandle>	FileHandle;
 	int64					FilePos; /* Desired position in the file stream, this can be different to FilePos due to the cache */
 	int64					TellPos; /* Actual position in the file,  this can be different to FilePos */
 	int64					FileSize;
diff --git a/Engine/Source/Runtime/Core/Public/HAL/IPlatformFileLogWrapper.h b/Engine/Source/Runtime/Core/Public/HAL/IPlatformFileLogWrapper.h
--- a/Engine/Source/Runtime/Core/Public/HAL/IPlatformFileLogWrapper.h
+++ b/Engine/Source/Runtime/Core/Public/HAL/IPlatformFileLogWrapper.h
@@ -2,6 +2,8 @@
 
 #pragma once
 
+#include "UniquePtr.h"
+
 
 /**
  * Wrapper to log the low level file system
@@ -22,7 +24,7 @@ class FLoggedPlatformFile;
 
 class CORE_API FLoggedFileHandle : public IFileHandle
 {
-	TAutoPtr<IFileHandle>	FileHandle;
+	TUniquePtr<IFileHandle>	FileHandle;
 	FString					Filename;
 #if !UE_BUILD_SHIPPING
 	FLoggedPlatformFile& PlatformFile;
diff --git a/Engine/Source/Runtime/Core/Public/HAL/IPlatformFileProfilerWrapper.h b/Engine/Source/Runtime/Core/Public/HAL/IPlatformFileProfilerWrapper.h
--- a/Engine/Source/Runtime/Core/Public/HAL/IPlatformFileProfilerWrapper.h
+++ b/Engine/Source/Runtime/Core/Public/HAL/IPlatformFileProfilerWrapper.h
@@ -2,6 +2,7 @@
 
 #pragma once
 
+#include "UniquePtr.h"
 
 #if !UE_BUILD_SHIPPING
 
@@ -137,7 +138,7 @@ struct FProfiledFileStatsFileSimple : public FProfiledFileStatsFileBase
 template< typename StatType >
 class TProfiledFileHandle : public IFileHandle
 {
-	TAutoPtr<IFileHandle> FileHandle;
+	TUniquePtr<IFileHandle> FileHandle;
 	FString Filename;
 	StatType* FileStats;
 
@@ -491,7 +492,7 @@ inline const TCHAR* TProfiledPlatformFile<FProfiledFileStatsFileSimple>::GetType
 
 class FPlatformFileReadStatsHandle : public IFileHandle
 {
-	TAutoPtr<IFileHandle> FileHandle;
+	TUniquePtr<IFileHandle> FileHandle;
 	FString Filename;
 	volatile int32* BytesPerSecCounter;
 	volatile int32* BytesReadCounter;
diff --git a/Engine/Source/Runtime/Core/Public/Misc/CString.h b/Engine/Source/Runtime/Core/Public/Misc/CString.h
--- a/Engine/Source/Runtime/Core/Public/Misc/CString.h
+++ b/Engine/Source/Runtime/Core/Public/Misc/CString.h
@@ -307,6 +307,11 @@ struct TCString
 	 */
 	static FORCEINLINE int32 Strtoi( const CharType* Start, CharType** End, int32 Base );
 
+	/**
+	 * strtoi wrapper
+	 */
+	static FORCEINLINE int64 Strtoi64( const CharType* Start, CharType** End, int32 Base );
+
 	/**
 	 * strtoui wrapper
 	 */
@@ -716,6 +721,12 @@ int32 TCString<T>::Strtoi( const CharType* Start, CharType** End, int32 Base )
 	return FPlatformString::Strtoi(Start, End, Base);
 }
 
+template <typename T> FORCEINLINE
+int64 TCString<T>::Strtoi64( const CharType* Start, CharType** End, int32 Base ) 
+{ 
+	return FPlatformString::Strtoi64(Start, End, Base);
+}
+
 template <typename T> FORCEINLINE
 uint64 TCString<T>::Strtoui64( const CharType* Start, CharType** End, int32 Base ) 
 { 
diff --git a/Engine/Source/Runtime/Core/Public/Misc/CoreMiscDefines.h b/Engine/Source/Runtime/Core/Public/Misc/CoreMiscDefines.h
--- a/Engine/Source/Runtime/Core/Public/Misc/CoreMiscDefines.h
+++ b/Engine/Source/Runtime/Core/Public/Misc/CoreMiscDefines.h
@@ -84,10 +84,6 @@
 	#define CA_CONSTANT_IF(Condition) if (Condition)
 #endif
 
-#ifndef USING_SIGNED_CONTENT
-	#define USING_SIGNED_CONTENT 0
-#endif
-
 enum {INDEX_NONE	= -1				};
 enum {UNICODE_BOM   = 0xfeff			};
 
diff --git a/Engine/Source/Runtime/Core/Public/Misc/DefaultValueHelper.h b/Engine/Source/Runtime/Core/Public/Misc/DefaultValueHelper.h
--- a/Engine/Source/Runtime/Core/Public/Misc/DefaultValueHelper.h
+++ b/Engine/Source/Runtime/Core/Public/Misc/DefaultValueHelper.h
@@ -90,6 +90,16 @@ public:
 	 */
 	static bool ParseInt(const FString& Source, int32& OutVal);
 
+	/**
+	 * Converts a string into a int64.
+	 *
+	 * @param Source the input string to try to convert
+	 * @param OutVal the output integer
+	 *
+	 * @return true if the conversion happened, false otherwise
+	 */
+	static bool ParseInt64(const FString& Source, int64& OutVal);
+
 	/**
 	 * Converts a string into a float.
 	 *
diff --git a/Engine/Source/Runtime/Core/Public/Serialization/Archive.h b/Engine/Source/Runtime/Core/Public/Serialization/Archive.h
--- a/Engine/Source/Runtime/Core/Public/Serialization/Archive.h
+++ b/Engine/Source/Runtime/Core/Public/Serialization/Archive.h
@@ -918,17 +918,17 @@ public:
 	 *
 	 * @return The container of custom versions in the archive.
 	 */
-	const FCustomVersionContainer& GetCustomVersions() const;
+	virtual const FCustomVersionContainer& GetCustomVersions() const;
 
 	/**
 	 * Sets the custom version numbers for this archive.
 	 *
 	 * @param CustomVersionContainer - The container of custom versions to copy into the archive.
 	 */
-	void SetCustomVersions(const FCustomVersionContainer& CustomVersionContainer);
+	virtual void SetCustomVersions(const FCustomVersionContainer& CustomVersionContainer);
 
 	/** Resets the custom version numbers for this archive. */
-	void ResetCustomVersions();
+	virtual void ResetCustomVersions();
 
 	/**
 	 * Sets a specific custom version
diff --git a/Engine/Source/Runtime/Core/Public/Serialization/ArchiveProxy.h b/Engine/Source/Runtime/Core/Public/Serialization/ArchiveProxy.h
--- a/Engine/Source/Runtime/Core/Public/Serialization/ArchiveProxy.h
+++ b/Engine/Source/Runtime/Core/Public/Serialization/ArchiveProxy.h
@@ -149,6 +149,22 @@ public:
 	{
 		return InnerArchive.IsCloseComplete(bHasError);
 	}
+
+	virtual const FCustomVersionContainer& GetCustomVersions() const override
+	{
+		return InnerArchive.GetCustomVersions();
+	}
+
+	virtual void SetCustomVersions(const FCustomVersionContainer& NewVersions) override
+	{
+		InnerArchive.SetCustomVersions(NewVersions);
+	}
+
+	virtual void ResetCustomVersions() override
+	{
+		InnerArchive.ResetCustomVersions();
+	}
+
 #if WITH_EDITOR
 	virtual void PushDebugDataString(const FName& DebugData) override
 	{
diff --git a/Engine/Source/Runtime/Core/Public/Templates/HasInserterOperator.h b/Engine/Source/Runtime/Core/Public/Templates/HasInserterOperator.h
new file mode 100644
--- /dev/null
+++ b/Engine/Source/Runtime/Core/Public/Templates/HasInserterOperator.h
@@ -0,0 +1,47 @@
+// Copyright 1998-2016 Epic Games, Inc. All Rights Reserved.
+
+#pragma once
+
+namespace UE4HasInserterOperator_Private
+{
+	struct FNotSpecified {};
+
+	template <typename T>
+	struct FReturnValueCheck
+	{
+		static char (&Func())[2];
+	};
+
+	template <>
+	struct FReturnValueCheck<FNotSpecified>
+	{
+		static char (&Func())[1];
+	};
+
+	template <typename Dest, typename T>
+	FNotSpecified operator<<(Dest&&, T&&);
+
+	template <typename T>
+	T&& Make();
+
+	template <typename Dest, typename T, typename = decltype(Make<Dest&>() << Make<T&>())>
+	struct THasInserterOperator
+	{
+		enum { Value = true };
+	};
+
+	template <typename Dest, typename T>
+	struct THasInserterOperator<Dest, T, FNotSpecified>
+	{
+		enum { Value = false };
+	};
+}
+
+/**
+ * Traits class which tests if a type has an operator<< overload between two types.
+ */
+template <typename Dest, typename T>
+struct THasInserterOperator
+{
+	enum { Value = UE4HasInserterOperator_Private::THasInserterOperator<Dest, T>::Value };
+};
diff --git a/Engine/Source/Runtime/Core/Public/Templates/ScopedPointer.h b/Engine/Source/Runtime/Core/Public/Templates/ScopedPointer.h
--- a/Engine/Source/Runtime/Core/Public/Templates/ScopedPointer.h
+++ b/Engine/Source/Runtime/Core/Public/Templates/ScopedPointer.h
@@ -22,6 +22,7 @@ private:
 public:
 
 	/** Initialization constructor. */
+	DEPRECATED(4.15, "TScopedPointer has been deprecated and should be replaced with TUniquePtr.")
 	explicit TScopedPointer(ReferencedType* InReference = nullptr)
 		: Reference(InReference)
 	{ }
@@ -170,6 +171,7 @@ private:
 public:
 
 	/** Initialization constructor. */
+	DEPRECATED(4.15, "TAutoPtr has been deprecated and should be replaced with TUniquePtr.")
 	explicit TAutoPtr(ReferencedType* InReference = nullptr)
 		:	Reference(InReference)
 	{}
diff --git a/Engine/Source/Runtime/Core/Public/Templates/SharedPointer.h b/Engine/Source/Runtime/Core/Public/Templates/SharedPointer.h
--- a/Engine/Source/Runtime/Core/Public/Templates/SharedPointer.h
+++ b/Engine/Source/Runtime/Core/Public/Templates/SharedPointer.h
@@ -121,8 +121,6 @@ FORCEINLINE TSharedRef< CastToType, Mode > StaticCastSharedRef( TSharedRef< Cast
 }
 
 
-class UObjectBase;
-
 /**
  * TSharedRef is a non-nullable, non-intrusive reference-counted authoritative object reference.
  *
@@ -132,9 +130,6 @@ class UObjectBase;
 template< class ObjectType, ESPMode Mode >
 class TSharedRef
 {
-	// TSharedRefs with UObjects are illegal.
-	static_assert(!TPointerIsConvertibleFromTo<ObjectType, const UObjectBase>::Value, "You cannot use TSharedRef with UObjects.");
-
 public:
 
 	// NOTE: TSharedRef has no default constructor as it does not support empty references.  You must
@@ -255,20 +250,23 @@ public:
 		: Object( const_cast< ObjectType* >( InSharedRef.Object ) )
 		, SharedReferenceCount( InSharedRef.SharedReferenceCount )
 	{ }
-	  
+
 	/**
-	 * Special constructor used internally to create a shared reference from an existing shared reference,
-	 * while using the specified object reference instead of the incoming shared reference's object
-	 * pointer.  This is used by with the TSharedFromThis feature (by UpdateWeakReferenceInternal)
+	 * Aliasing constructor used to create a shared reference which shares its reference count with
+	 * another shared object, but pointing to a different object, typically a subobject.
 	 *
-	 * @param  OtherSharedRef  The shared reference whose reference count 
-	 * @param  InObject  The object pointer to use (instead of the incoming shared reference's object)
+	 * @param  OtherSharedRef  The shared reference whose reference count should be shared.
+	 * @param  InObject  The object pointer to use (instead of the incoming shared pointer's object)
 	 */
 	template <typename OtherType>
 	FORCEINLINE TSharedRef( TSharedRef< OtherType, Mode > const& OtherSharedRef, ObjectType* InObject )
 		: Object( InObject )
 		, SharedReferenceCount( OtherSharedRef.SharedReferenceCount )
-	{ }
+	{
+		// If the following assert goes off, it means a TSharedRef was initialized from a nullptr object pointer.
+		// Shared references must never be nullptr, so either pass a valid object or consider using TSharedPtr instead.
+		check( InObject != nullptr );
+	}
 
 	FORCEINLINE TSharedRef( TSharedRef const& InSharedRef )
 		: Object( InSharedRef.Object )
@@ -514,9 +512,6 @@ TSharedRef<InObjectType, InMode> MakeShared(InArgTypes&&... Args);
 template< class ObjectType, ESPMode Mode >
 class TSharedPtr
 {
-	// TSharedPtrs with UObjects are illegal.
-	static_assert(!TPointerIsConvertibleFromTo<ObjectType, const UObjectBase>::Value, "You cannot use TSharedPtr or TWeakPtr with UObjects. Consider a UPROPERTY() pointer or TWeakObjectPtr.");
-
 	enum
 	{
 		ObjectTypeHasSameModeSharedFromThis     = TPointerIsConvertibleFromTo<ObjectType, TSharedFromThis<ObjectType, Mode>>::Value,
@@ -670,13 +665,12 @@ public:
 		: Object( const_cast< ObjectType* >( InSharedPtr.Object ) )
 		, SharedReferenceCount( InSharedPtr.SharedReferenceCount )
 	{ }
-  
+
 	/**
-	 * Special constructor used internally to create a shared pointer from an existing shared pointer,
-	 * while using the specified object pointer instead of the incoming shared pointer's object
-	 * pointer.  This is used by with the TSharedFromThis feature (by UpdateWeakReferenceInternal)
+	 * Aliasing constructor used to create a shared pointer which shares its reference count with
+	 * another shared object, but pointing to a different object, typically a subobject.
 	 *
-	 * @param  OtherSharedPtr  The shared pointer whose reference count 
+	 * @param  OtherSharedPtr  The shared pointer whose reference count should be shared.
 	 * @param  InObject  The object pointer to use (instead of the incoming shared pointer's object)
 	 */
 	template <typename OtherType>
@@ -685,6 +679,34 @@ public:
 		, SharedReferenceCount( OtherSharedPtr.SharedReferenceCount )
 	{ }
 
+	/**
+	 * Aliasing constructor used to create a shared pointer which shares its reference count with
+	 * another shared object, but pointing to a different object, typically a subobject.
+	 *
+	 * @param  OtherSharedPtr  The shared pointer whose reference count should be shared.
+	 * @param  InObject  The object pointer to use (instead of the incoming shared pointer's object)
+	 */
+	template <typename OtherType>
+	FORCEINLINE TSharedPtr( TSharedPtr< OtherType, Mode >&& OtherSharedPtr, ObjectType* InObject )
+		: Object( InObject )
+		, SharedReferenceCount( MoveTemp(OtherSharedPtr.SharedReferenceCount) )
+	{
+		OtherSharedPtr.Object = nullptr;
+	}
+
+	/**
+	 * Aliasing constructor used to create a shared pointer which shares its reference count with
+	 * another shared object, but pointing to a different object, typically a subobject.
+	 *
+	 * @param  OtherSharedRef  The shared reference whose reference count should be shared.
+	 * @param  InObject  The object pointer to use (instead of the incoming shared pointer's object)
+	 */
+	template <typename OtherType>
+	FORCEINLINE TSharedPtr( TSharedRef< OtherType, Mode > const& OtherSharedRef, ObjectType* InObject )
+		: Object( InObject )
+		, SharedReferenceCount( OtherSharedRef.SharedReferenceCount )
+	{ }
+
 	/**
 	 * Assignment to a nullptr pointer.  The object currently referenced by this shared pointer will no longer be
 	 * referenced and will be deleted if there are no other referencers.
diff --git a/Engine/Source/Runtime/Core/Public/Templates/UniquePtr.h b/Engine/Source/Runtime/Core/Public/Templates/UniquePtr.h
--- a/Engine/Source/Runtime/Core/Public/Templates/UniquePtr.h
+++ b/Engine/Source/Runtime/Core/Public/Templates/UniquePtr.h
@@ -20,9 +20,10 @@
 template <typename T>
 struct TDefaultDelete
 {
-	TDefaultDelete()
-	{
-	}
+	TDefaultDelete() = default;
+	TDefaultDelete(const TDefaultDelete&) = default;
+	TDefaultDelete& operator=(const TDefaultDelete&) = default;
+	~TDefaultDelete() = default;
 
 	template <
 		typename U,
@@ -41,10 +42,6 @@ struct TDefaultDelete
 		return *this;
 	}
 
-	~TDefaultDelete()
-	{
-	}
-
 	void operator()(T* Ptr) const
 	{
 		delete Ptr;
@@ -54,9 +51,10 @@ struct TDefaultDelete
 template <typename T>
 struct TDefaultDelete<T[]>
 {
-	TDefaultDelete()
-	{
-	}
+	TDefaultDelete() = default;
+	TDefaultDelete(const TDefaultDelete&) = default;
+	TDefaultDelete& operator=(const TDefaultDelete&) = default;
+	~TDefaultDelete() = default;
 
 	template <
 		typename U,
@@ -75,10 +73,6 @@ struct TDefaultDelete<T[]>
 		return *this;
 	}
 
-	~TDefaultDelete()
-	{
-	}
-
 	template <typename U>
 	typename TEnableIf<TPointerIsConvertibleFromTo<U[], T[]>::Value>::Type operator()(U* Ptr) const
 	{
diff --git a/Engine/Source/Runtime/Core/Public/UObject/AutoPointer.h b/Engine/Source/Runtime/Core/Public/UObject/AutoPointer.h
--- a/Engine/Source/Runtime/Core/Public/UObject/AutoPointer.h
+++ b/Engine/Source/Runtime/Core/Public/UObject/AutoPointer.h
@@ -5,32 +5,37 @@
 
 // TAutoPointer wraps a smart-pointer and adds an implicit conversion to raw pointer
 // Its main use is for converting a variable from raw pointer to a smart pointer without breaking existing code
-// Currently only used in TAutoWeakObjectPtr
 // Not the same thing as TAutoPtr :(
 template<class T, class TBASE>
 class TAutoPointer : public TBASE
 {
 public:
 	/** NULL constructor **/
+	DEPRECATED(4.15, "TAutoPointer has been deprecated - please remove its usage from your project") 
 	FORCEINLINE TAutoPointer()
 	{
 	}
 	/** Construct from a single argument (I'd use inheriting constructors if all our compilers supported it) **/
 	template<typename X>
+	DEPRECATED(4.15, "TAutoPointer has been deprecated - please remove its usage from your project") 
 	FORCEINLINE TAutoPointer(X&& Target)
 		: TBASE(Forward<X>(Target))
 	{
 	}
 
+	DEPRECATED(4.15, "TAutoPointer has been deprecated - please remove its usage from your project") 
 	FORCEINLINE operator T* () const
 	{
 		return TBASE::Get();
 	}
+
+	DEPRECATED(4.15, "TAutoPointer has been deprecated - please remove its usage from your project") 
 	FORCEINLINE operator const T* () const
 	{
 		return (const T*)TBASE::Get();
 	}
 
+	DEPRECATED(4.15, "TAutoPointer has been deprecated - please remove its usage from your project") 
 	FORCEINLINE explicit operator bool() const
 	{
 		return TBASE::Get() != nullptr;
diff --git a/Engine/Source/Runtime/Core/Public/UObject/CoreObjectVersion.h b/Engine/Source/Runtime/Core/Public/UObject/CoreObjectVersion.h
--- a/Engine/Source/Runtime/Core/Public/UObject/CoreObjectVersion.h
+++ b/Engine/Source/Runtime/Core/Public/UObject/CoreObjectVersion.h
@@ -9,6 +9,7 @@ struct CORE_API FCoreObjectVersion
 		// Before any version changes were made
 		BeforeCustomVersionWasAdded = 0,
 		MaterialInputNativeSerialize,
+		EnumProperties,
 
 		// -----<new versions can be added above this line>-------------------------------------------------
 		VersionPlusOne,
diff --git a/Engine/Source/Runtime/Core/Public/UObject/DevObjectVersion.h b/Engine/Source/Runtime/Core/Public/UObject/DevObjectVersion.h
--- a/Engine/Source/Runtime/Core/Public/UObject/DevObjectVersion.h
+++ b/Engine/Source/Runtime/Core/Public/UObject/DevObjectVersion.h
@@ -15,6 +15,7 @@
 #include "RenderingObjectVersion.h"
 #include "SequencerObjectVersion.h"
 #include "VRObjectVersion.h"
+#include "LoadTimesObjectVersion.h"
 
 class CORE_API FDevVersionRegistration :  public FCustomVersionRegistration
 {
diff --git a/Engine/Source/Runtime/Core/Public/UObject/LoadTimesObjectVersion.h b/Engine/Source/Runtime/Core/Public/UObject/LoadTimesObjectVersion.h
new file mode 100644
--- /dev/null
+++ b/Engine/Source/Runtime/Core/Public/UObject/LoadTimesObjectVersion.h
@@ -0,0 +1,25 @@
+// Copyright 1998-2015 Epic Games, Inc. All Rights Reserved.
+#pragma once
+
+// Custom serialization version for changes made in Dev-Rendering stream
+struct CORE_API FLoadTimesObjectVersion
+{
+	enum Type
+	{
+		// Before any version changes were made
+		BeforeCustomVersionWasAdded = 0,
+
+		// Allows uncompressed reflection captures for cooked builds
+		UncompressedReflectionCapturesForCookedBuilds,
+
+		// -----<new versions can be added above this line>-------------------------------------------------
+		VersionPlusOne,
+		LatestVersion = VersionPlusOne - 1
+	};
+
+	// The GUID for this custom version number
+	const static FGuid GUID;
+
+private:
+	FLoadTimesObjectVersion() {}
+};
diff --git a/Engine/Source/Runtime/Core/Public/UObject/ScriptDelegates.h b/Engine/Source/Runtime/Core/Public/UObject/ScriptDelegates.h
--- a/Engine/Source/Runtime/Core/Public/UObject/ScriptDelegates.h
+++ b/Engine/Source/Runtime/Core/Public/UObject/ScriptDelegates.h
@@ -567,7 +567,9 @@ protected:
 
 	// 
 	friend class FCallDelegateHelper;
+
+	friend struct TIsZeroConstructType<TMulticastScriptDelegate<TWeakPtr> >;
 };
 
 
-template<typename TWeakPtr> struct TIsZeroConstructType<TMulticastScriptDelegate<TWeakPtr> > { enum { Value = TIsZeroConstructType<TWeakPtr>::Value }; };
+template<typename TWeakPtr> struct TIsZeroConstructType<TMulticastScriptDelegate<TWeakPtr> > { enum { Value = TIsZeroConstructType<typename TMulticastScriptDelegate<TWeakPtr>::FInvocationList>::Value }; };
diff --git a/Engine/Source/Runtime/Core/Public/UObject/UObjectHierarchyFwd.h b/Engine/Source/Runtime/Core/Public/UObject/UObjectHierarchyFwd.h
--- a/Engine/Source/Runtime/Core/Public/UObject/UObjectHierarchyFwd.h
+++ b/Engine/Source/Runtime/Core/Public/UObject/UObjectHierarchyFwd.h
@@ -34,6 +34,9 @@ class				UTextProperty;
 class				UArrayProperty;
 class				UDelegateProperty;
 class				UMulticastDelegateProperty;
+class				UMapProperty;
+class				USetProperty;
+class				UEnumProperty;
 class			UStruct;
 class				UFunction;
 class				UClass;
diff --git a/Engine/Source/Runtime/Core/Public/UObject/UnrealNames.inl b/Engine/Source/Runtime/Core/Public/UObject/UnrealNames.inl
--- a/Engine/Source/Runtime/Core/Public/UObject/UnrealNames.inl
+++ b/Engine/Source/Runtime/Core/Public/UObject/UnrealNames.inl
@@ -39,6 +39,9 @@ REGISTER_NAME(31,Engine)
 REGISTER_NAME(32,Editor)
 REGISTER_NAME(33,CoreUObject)
 
+// More class properties
+REGISTER_NAME(34,EnumProperty)
+
 // Special types.
 REGISTER_NAME(50,Cylinder)
 REGISTER_NAME(51,BoxSphereBounds)
diff --git a/Engine/Source/Runtime/Core/Public/UObject/WeakObjectPtrTemplates.h b/Engine/Source/Runtime/Core/Public/UObject/WeakObjectPtrTemplates.h
--- a/Engine/Source/Runtime/Core/Public/UObject/WeakObjectPtrTemplates.h
+++ b/Engine/Source/Runtime/Core/Public/UObject/WeakObjectPtrTemplates.h
@@ -2,7 +2,8 @@
 
 #pragma once
 
-#include "AutoPointer.h"
+#include "Templates/IsPointer.h"
+#include "Templates/PointerIsConvertibleFromTo.h"
 
 
 /***
@@ -319,32 +320,45 @@ struct TWeakObjectPtrMapKeyFuncs : public TDefaultMapKeyFuncs<KeyType, ValueType
 	}
 };
 
-
 /**
-  * Automatic version of the weak object pointer
-**/
+ * Automatic version of the weak object pointer
+ */
 template<class T> 
-class TAutoWeakObjectPtr : public TAutoPointer<T, TWeakObjectPtr<T> >
+class TAutoWeakObjectPtr : public TWeakObjectPtr<T>
 {
 public:
 	/** NULL constructor **/
+	DEPRECATED(4.15, "TAutoWeakObjectPtr has been deprecated - use TWeakObjectPtr instead")
 	FORCEINLINE TAutoWeakObjectPtr()
 	{
 	}
 	/** Construct from a raw pointer **/
-	FORCEINLINE TAutoWeakObjectPtr(const T* Target) 
-		: TAutoPointer<T, TWeakObjectPtr<T> >(Target)
+	DEPRECATED(4.15, "TAutoWeakObjectPtr has been deprecated - use TWeakObjectPtr instead")
+	FORCEINLINE TAutoWeakObjectPtr(const T* Target)
+		: TWeakObjectPtr<T>(Target)
 	{
 	}
 	/**  Construct from the base type **/
+	DEPRECATED(4.15, "TAutoWeakObjectPtr has been deprecated - use TWeakObjectPtr instead")
 	FORCEINLINE TAutoWeakObjectPtr(const TWeakObjectPtr<T>& Other) 
-		: TAutoPointer<T, TWeakObjectPtr<T> >(Other)
+		: TWeakObjectPtr<T>(Other)
+	{
+	}
+	DEPRECATED(4.15, "Implicit conversion from TAutoWeakObjectPtr to the pointer type has been deprecated - use Get() instead")
+	FORCEINLINE operator T* () const
+	{
+		return this->Get();
+	}
+	DEPRECATED(4.15, "Implicit conversion from TAutoWeakObjectPtr to the pointer type has been deprecated - use Get() instead")
+	FORCEINLINE operator const T* () const
 	{
+		return (const T*)this->Get();
 	}
-	/**  Construct from another auto pointer **/
-	FORCEINLINE TAutoWeakObjectPtr(const TAutoWeakObjectPtr& Other) 
-		: TAutoPointer<T, TWeakObjectPtr<T> >(Other)
+
+	DEPRECATED(4.15, "Implicit conversion from TAutoWeakObjectPtr to the pointer type has been deprecated - use Get() instead")
+	FORCEINLINE explicit operator bool() const
 	{
+		return this->Get() != nullptr;
 	}
 };
 
diff --git a/Engine/Source/Runtime/Core/Public/Windows/WindowsPlatformMath.h b/Engine/Source/Runtime/Core/Public/Windows/WindowsPlatformMath.h
--- a/Engine/Source/Runtime/Core/Public/Windows/WindowsPlatformMath.h
+++ b/Engine/Source/Runtime/Core/Public/Windows/WindowsPlatformMath.h
@@ -118,6 +118,25 @@ struct FWindowsPlatformMath : public FGenericPlatformMath
 	{
 		return 1 << CeilLogTwo(Arg);
 	}
+#if PLATFORM_64BITS
+	static FORCEINLINE uint64 CeilLogTwo64(uint64 Arg)
+	{
+		int64 Bitmask = ((int64)(CountLeadingZeros64(Arg) << 57)) >> 63;
+		return (64 - CountLeadingZeros64(Arg - 1)) & (~Bitmask);
+	}
+	static FORCEINLINE uint64 CountLeadingZeros64(uint64 Value)
+	{
+		// Use BSR to return the log2 of the integer
+		DWORD Log2;
+		if (_BitScanReverse64(&Log2, Value) != 0)
+		{
+			return 63 - Log2;
+		}
+
+		return 64;
+	}
+#endif
+
 #endif
 };
 
diff --git a/Engine/Source/Runtime/CoreUObject/Private/Misc/StringAssetReference.cpp b/Engine/Source/Runtime/CoreUObject/Private/Misc/StringAssetReference.cpp
--- a/Engine/Source/Runtime/CoreUObject/Private/Misc/StringAssetReference.cpp
+++ b/Engine/Source/Runtime/CoreUObject/Private/Misc/StringAssetReference.cpp
@@ -14,16 +14,6 @@ FStringAssetReference::FStringAssetReference(const UObject* InObject)
 }
 
 PRAGMA_DISABLE_DEPRECATION_WARNINGS
-FStringAssetReference::~FStringAssetReference()
-{
-
-}
-
-const FString& FStringAssetReference::ToString() const
-{
-	return AssetLongPathname;
-}
-
 void FStringAssetReference::SetPath(FString Path)
 {
 	if (Path.IsEmpty())
@@ -74,14 +64,18 @@ bool FStringAssetReference::Serialize(FArchive& Ar)
 
 	return true;
 }
+
 bool FStringAssetReference::operator==(FStringAssetReference const& Other) const
 {
 	return ToString() == Other.ToString();
 }
-void FStringAssetReference::operator=(FStringAssetReference const& Other)
+
+FStringAssetReference& FStringAssetReference::operator=(FStringAssetReference const& Other)
 {
 	SetPath(Other.ToString());
+	return *this;
 }
+
 bool FStringAssetReference::ExportTextItem(FString& ValueStr, FStringAssetReference const& DefaultValue, UObject* Parent, int32 PortFlags, UObject* ExportRootScope) const
 {
 	if (0 != (PortFlags & EPropertyPortFlags::PPF_ExportCpp))
diff --git a/Engine/Source/Runtime/CoreUObject/Private/Misc/WorldCompositionUtility.cpp b/Engine/Source/Runtime/CoreUObject/Private/Misc/WorldCompositionUtility.cpp
--- a/Engine/Source/Runtime/CoreUObject/Private/Misc/WorldCompositionUtility.cpp
+++ b/Engine/Source/Runtime/CoreUObject/Private/Misc/WorldCompositionUtility.cpp
@@ -4,6 +4,7 @@
 	WorldCompositionUtility.cpp : Support structures for world composition
 =============================================================================*/
 #include "CoreUObjectPrivate.h"
+#include "UniquePtr.h"
 
 FArchive& operator<<( FArchive& Ar, FWorldTileLayer& D )
 {
@@ -68,7 +69,7 @@ bool FWorldTileInfo::Read(const FString& InPackageFileName, FWorldTileInfo& OutI
 	OutInfo = FWorldTileInfo();
 
 	// Create a file reader to load the file
-	TScopedPointer<FArchive> FileReader(IFileManager::Get().CreateFileReader(*InPackageFileName));
+	TUniquePtr<FArchive> FileReader(IFileManager::Get().CreateFileReader(*InPackageFileName));
 	if (FileReader == NULL)
 	{
 		// Couldn't open the file
@@ -97,7 +98,7 @@ bool FWorldTileInfo::Read(const FString& InPackageFileName, FWorldTileInfo& OutI
 			check(FileSummary.CompressedChunks.Num() > 0);
 			if (!FileReader->SetCompressionMap(&FileSummary.CompressedChunks, (ECompressionFlags)FileSummary.CompressionFlags))
 			{
-				FileReader = new FArchiveAsync(*InPackageFileName); // re-assign scope pointer
+				FileReader = MakeUnique<FArchiveAsync>(*InPackageFileName); // re-assign scope pointer
 				check(!FileReader->IsError());
 				verify(FileReader->SetCompressionMap(&FileSummary.CompressedChunks, (ECompressionFlags)FileSummary.CompressionFlags));
 			}
diff --git a/Engine/Source/Runtime/CoreUObject/Private/UObject/Class.cpp b/Engine/Source/Runtime/CoreUObject/Private/UObject/Class.cpp
--- a/Engine/Source/Runtime/CoreUObject/Private/UObject/Class.cpp
+++ b/Engine/Source/Runtime/CoreUObject/Private/UObject/Class.cpp
@@ -946,12 +946,12 @@ void UStruct::SerializeTaggedProperties(FArchive& Ar, uint8* Data, UStruct* Defa
 					
 					FName EachName = GetFName();
 					// Search the current class first, then work up the class hierarchy to see if theres a match for our fixup.
-					UStruct* Owner = GetOwnerStruct();		
+					UStruct* Owner = GetOwnerStruct();
 					if( Owner )
 					{
 						UStruct* SuperClass = Owner->GetSuperStruct();
 						while( EachName != NAME_None)
-						{							
+						{
 							const TMap<FName, FName>* ClassTaggedPropertyRedirects = TaggedPropertyRedirects.Find( EachName );
 							if (ClassTaggedPropertyRedirects)
 							{
@@ -1019,16 +1019,16 @@ void UStruct::SerializeTaggedProperties(FArchive& Ar, uint8* Data, UStruct* Defa
 			// No need to check redirects on platforms where everything is cooked. Always check for save games
 			if (!FPlatformProperties::RequiresCookedData() || Ar.IsSaveGame())
 			{
-			if (Tag.Type == NAME_StructProperty && PropID == NAME_StructProperty)
-			{
-				FName* NewName = FLinkerLoad::StructNameRedirects.Find(Tag.StructName);
-				FName StructName = CastChecked<UStructProperty>(Property)->Struct->GetFName();
-					if (NewName != nullptr && *NewName == StructName)
+				if (Tag.Type == NAME_StructProperty && PropID == NAME_StructProperty)
 				{
-					Tag.StructName = *NewName;
+					FName* NewName = FLinkerLoad::StructNameRedirects.Find(Tag.StructName);
+					FName StructName = CastChecked<UStructProperty>(Property)->Struct->GetFName();
+						if (NewName != nullptr && *NewName == StructName)
+					{
+						Tag.StructName = *NewName;
+					}
 				}
 			}
-			}
 
 			const int64 StartOfProperty = Ar.Tell();
 			if( !Property )
@@ -1057,9 +1057,6 @@ void UStruct::SerializeTaggedProperties(FArchive& Ar, uint8* Data, UStruct* Defa
 				UE_CLOG((Ar.IsPersistent() && FPlatformProperties::RequiresCookedData()), LogClass, Warning, TEXT("Skipping saved property %s of %s since it is no longer serializable for asset:  %s. (Maybe resave asset?)"), *Tag.Name.ToString(), *GetName(), *Ar.GetArchiveName() );
 			}
 
-			// Convert properties from old type to new type automatically if types are compatible
-			// If you add an entry to this, you will also need to add an entry to the array case below
-			// For converting to a struct, you can just implement SerializeFromMismatchedTag on the struct
 			else if (Property->ConvertFromType(Tag, Ar, Data, DefaultsStruct, bAdvanceProperty))
 			{
 				if (bAdvanceProperty)
@@ -2639,7 +2636,10 @@ UObject* UClass::CreateDefaultObject()
 			UObjectForceRegistration(ParentClass);
 			ParentDefaultObject = ParentClass->GetDefaultObject(); // Force the default object to be constructed if it isn't already
 #if USE_EVENT_DRIVEN_ASYNC_LOAD
-			check(ParentDefaultObject && !ParentDefaultObject->HasAnyFlags(RF_NeedLoad));
+			if (!GIsInitialLoad)
+			{ 
+				check(ParentDefaultObject && !ParentDefaultObject->HasAnyFlags(RF_NeedLoad));
+			}
 #endif
 		}
 
@@ -3352,7 +3352,10 @@ void UClass::SerializeSuperStruct(FArchive& Ar)
 		FFastIndexingClassTree::Register(this);
 	}
 #elif UCLASS_FAST_ISA_IMPL == UCLASS_ISA_CLASSARRAY
-	this->ReinitializeBaseChainArray();
+	if (Ar.IsLoading())
+	{
+		this->ReinitializeBaseChainArray();
+	}
 #endif
 }
 
diff --git a/Engine/Source/Runtime/CoreUObject/Private/UObject/Enum.cpp b/Engine/Source/Runtime/CoreUObject/Private/UObject/Enum.cpp
--- a/Engine/Source/Runtime/CoreUObject/Private/UObject/Enum.cpp
+++ b/Engine/Source/Runtime/CoreUObject/Private/UObject/Enum.cpp
@@ -2,6 +2,7 @@
 
 #include "CoreUObjectPrivate.h"
 #include "UObjectThreadContext.h"
+#include "UObject/DevObjectVersion.h"
 
 DEFINE_LOG_CATEGORY_STATIC(LogEnum, Log, All);
 
@@ -15,20 +16,39 @@ TMap<FName,TMap<FString,FString> > UEnum::EnumSubstringRedirects;
 
 void UEnum::Serialize( FArchive& Ar )
 {
+	Ar.UsingCustomVersion(FCoreObjectVersion::GUID);
+
 	Super::Serialize(Ar);
-	if (Ar.IsLoading() && Ar.UE4Ver() < VER_UE4_TIGHTLY_PACKED_ENUMS)
+	if (Ar.IsLoading())
 	{
-		TArray<FName> TempNames;
-		Ar << TempNames;
-		uint8 Value = 0;
-		for (const FName& TempName : TempNames)
+		if (Ar.UE4Ver() < VER_UE4_TIGHTLY_PACKED_ENUMS)
 		{
-			Names.Add(TPairInitializer<FName, uint8>(TempName, Value++));
+			TArray<FName> TempNames;
+			Ar << TempNames;
+			int64 Value = 0;
+			for (const FName& TempName : TempNames)
+			{
+				Names.Add(TPairInitializer<FName, uint8>(TempName, Value++));
+			}
+		}
+		else if (Ar.CustomVer(FCoreObjectVersion::GUID) < FCoreObjectVersion::EnumProperties)
+		{
+			TArray<TPair<FName, uint8>> OldNames;
+			Ar << OldNames;
+			Names.Reset(OldNames.Num());
+			for (const TPair<FName, uint8>& OldName : OldNames)
+			{
+				Names.Add(TPairInitializer<FName, uint8>(OldName.Key, OldName.Value));
+			}
+		}
+		else
+		{
+			Ar << Names;
 		}
 	}
 	else
 	{
-	Ar << Names;
+		Ar << Names;
 	}
 
 	if (Ar.UE4Ver() < VER_UE4_ENUM_CLASS_SUPPORT)
@@ -91,7 +111,7 @@ void UEnum::RenameNamesAfterDuplication()
 		FString ThisName = GetName();
 
 		// Replace all usages of base class name to the duplicated one.
-		for (TPair<FName, uint8>& Kvp : Names)
+		for (TPair<FName, int64>& Kvp : Names)
 		{
 			FString NameString = Kvp.Key.ToString();
 			NameString.ReplaceInline(*BaseEnumName, *ThisName);
@@ -100,7 +120,7 @@ void UEnum::RenameNamesAfterDuplication()
 	}
 }
 
-int32 UEnum::ResolveEnumerator(FArchive& Ar, int32 EnumeratorIndex) const
+int64 UEnum::ResolveEnumerator(FArchive& Ar, int64 EnumeratorIndex) const
 {
 	return EnumeratorIndex;
 }
@@ -110,7 +130,7 @@ FString UEnum::GenerateFullEnumName(const TCHAR* InEnumName) const
 	return (CppForm != ECppForm::Regular) ? GenerateFullEnumName(this, InEnumName) : InEnumName;
 }
 
-FName UEnum::GetNameByIndex(uint8 Index) const
+FName UEnum::GetNameByIndex(int32 Index) const
 {
 	if (Names.IsValidIndex(Index))
 	{
@@ -120,15 +140,15 @@ FName UEnum::GetNameByIndex(uint8 Index) const
 	return NAME_None;
 }
 
-uint8 UEnum::GetValueByIndex(uint8 Index) const
+int64 UEnum::GetValueByIndex(int32 Index) const
 {
 	check(Names.IsValidIndex(Index));
 	return Names[Index].Value;
 }
 
-FName UEnum::GetNameByValue(uint8 InValue) const
+FName UEnum::GetNameByValue(int64 InValue) const
 {
-	for (TPair<FName, uint8> Kvp : Names)
+	for (TPair<FName, int64> Kvp : Names)
 	{
 		if (Kvp.Value == InValue)
 		{
@@ -153,7 +173,7 @@ int32 UEnum::GetIndexByName(FName InName) const
 	return INDEX_NONE;
 }
 
-int32 UEnum::GetValueByName(FName InName) const
+int64 UEnum::GetValueByName(FName InName) const
 {
 	FString InNameString = InName.ToString();
 	FString DoubleColon = TEXT("::");
@@ -172,7 +192,7 @@ int32 UEnum::GetValueByName(FName InName) const
 		InName = FName(*InNameString.Mid(DoubleColonPosition + 2, InNameString.Len() - DoubleColonPosition + 2));
 	}
 
-	for (TPair<FName, uint8> Kvp : Names)
+	for (TPair<FName, int64> Kvp : Names)
 	{
 		if (Kvp.Key == InName)
 		{
@@ -184,7 +204,7 @@ int32 UEnum::GetValueByName(FName InName) const
 }
 
 
-uint8 UEnum::GetMaxEnumValue() const
+int64 UEnum::GetMaxEnumValue() const
 {
 	int32 NamesNum = Names.Num();
 	if (NamesNum == 0)
@@ -192,10 +212,10 @@ uint8 UEnum::GetMaxEnumValue() const
 		return 0;
 	}
 
-	uint8 MaxValue = Names[0].Value;
+	int64 MaxValue = Names[0].Value;
 	for (int32 i = 0; i < NamesNum; ++i)
 	{
-		uint8 CurrentValue = Names[i].Value;
+		int64 CurrentValue = Names[i].Value;
 		if (CurrentValue > MaxValue)
 		{
 			MaxValue = CurrentValue;
@@ -205,12 +225,12 @@ uint8 UEnum::GetMaxEnumValue() const
 	return MaxValue;
 }
 
-bool UEnum::IsValidEnumValue(uint8 InValue) const
+bool UEnum::IsValidEnumValue(int64 InValue) const
 {
 	int32 NamesNum = Names.Num();
 	for (int32 i = 0; i < NamesNum; ++i)
 	{
-		uint8 CurrentValue = Names[i].Value;
+		int64 CurrentValue = Names[i].Value;
 		if (CurrentValue == InValue)
 		{
 			return true;
@@ -247,7 +267,7 @@ FString UEnum::GenerateFullEnumName(const UEnum* InEnum, const TCHAR* InEnumName
 
 void UEnum::AddNamesToMasterList()
 {
-	for (TPair<FName, uint8> Kvp : Names)
+	for (TPair<FName, int64> Kvp : Names)
 	{
 		UEnum* Enum = AllEnumNames.FindRef(Kvp.Key);
 		if (Enum == nullptr)
@@ -263,7 +283,7 @@ void UEnum::AddNamesToMasterList()
 
 void UEnum::RemoveNamesFromMasterList()
 {
-	for (TPair<FName, uint8> Kvp : Names)
+	for (TPair<FName, int64> Kvp : Names)
 	{
 		UEnum* Enum = AllEnumNames.FindRef(Kvp.Key);
 		if (Enum == this)
@@ -459,13 +479,13 @@ int32 UEnum::FindEnumRedirects(const UEnum* Enum, FName EnumEntryName)
  *
  * @return	true unless the MAX enum already exists.
  */
-bool UEnum::SetEnums(TArray<TPair<FName, uint8>>& InNames, UEnum::ECppForm InCppForm, bool bAddMaxKeyIfMissing)
+bool UEnum::SetEnums(TArray<TPair<FName, int64>>& InNames, UEnum::ECppForm InCppForm, bool bAddMaxKeyIfMissing)
 {
 	if (Names.Num() > 0)
 	{
 		RemoveNamesFromMasterList();
 	}
-	Names   = InNames;
+	Names  = InNames;
 	CppForm = InCppForm;
 
 	if (bAddMaxKeyIfMissing)
@@ -484,13 +504,7 @@ bool UEnum::SetEnums(TArray<TPair<FName, uint8>>& InNames, UEnum::ECppForm InCpp
 				return false;
 			}
 
-#if HACK_HEADER_GENERATOR
-			if (GetMaxEnumValue() == 255)
-			{
-				FError::Throwf(TEXT("Enum %s doesn't contain user defined %s_MAX entry and its maximum entry value equals 255. Autogenerated %s_MAX enum entry will equal 0. Change maximum entry value to 254 or define %s_MAX entry in enum."), *GetName(), *EnumPrefix, *EnumPrefix, *EnumPrefix, *EnumPrefix);
-			}
-#endif
-			Names.Add(TPairInitializer<FName, uint8>(MaxEnumItem, GetMaxEnumValue() + 1));
+			Names.Add(TPairInitializer<FName, int64>(MaxEnumItem, GetMaxEnumValue() + 1));
 		}
 	}
 	AddNamesToMasterList();
@@ -499,7 +513,7 @@ bool UEnum::SetEnums(TArray<TPair<FName, uint8>>& InNames, UEnum::ECppForm InCpp
 }
 
 #if WITH_EDITOR
-FText UEnum::GetDisplayNameTextByValue(int32 Value) const
+FText UEnum::GetDisplayNameTextByValue(int64 Value) const
 {
 	int32 Index = GetIndexByValue(Value);
 	return GetDisplayNameText(Index);
@@ -707,14 +721,14 @@ void UEnum::RemoveMetaData( const TCHAR* Key, int32 NameIndex/*=INDEX_NONE*/) co
 
 #endif
 
-int32 UEnum::ParseEnum(const TCHAR*& Str)
+int64 UEnum::ParseEnum(const TCHAR*& Str)
 {
 	FString Token;
 	const TCHAR* ParsedStr = Str;
 	if (FParse::AlnumToken(ParsedStr, Token))
 	{
 		FName TheName = FName(*Token, FNAME_Find);
-		int32 Result = LookupEnumName(TheName);
+		int64 Result = LookupEnumName(TheName);
 		if (Result != INDEX_NONE)
 		{
 			Str = ParsedStr;
diff --git a/Engine/Source/Runtime/CoreUObject/Private/UObject/EnumProperty.cpp b/Engine/Source/Runtime/CoreUObject/Private/UObject/EnumProperty.cpp
new file mode 100644
--- /dev/null
+++ b/Engine/Source/Runtime/CoreUObject/Private/UObject/EnumProperty.cpp
@@ -0,0 +1,398 @@
+// Copyright 1998-2016 Epic Games, Inc. All Rights Reserved.
+
+#include "CoreUObjectPrivate.h"
+#include "UObject/EnumProperty.h"
+#include "PropertyTag.h"
+#include "Templates/ChooseClass.h"
+#include "Templates/IsSigned.h"
+
+namespace UE4EnumProperty_Private
+{
+	template <typename OldIntType>
+	void ConvertIntToEnumProperty(FArchive& Ar, UEnumProperty* EnumProp, UNumericProperty* UnderlyingProp, UEnum* Enum, void* Obj)
+	{
+		OldIntType OldValue;
+		Ar << OldValue;
+
+		using LargeIntType = typename TChooseClass<TIsSigned<OldIntType>::Value, int64, uint64>::Result;
+
+		LargeIntType NewValue = OldValue;
+		if (!UnderlyingProp->CanHoldValue(NewValue) || !Enum->IsValidEnumValue(NewValue))
+		{
+			UE_LOG(
+				LogClass,
+				Warning,
+				TEXT("Failed to find valid enum value '%s' for enum type '%s' when converting property '%s' during property loading - setting to '%s'"),
+				*LexicalConversion::ToString(OldValue),
+				*Enum->GetName(),
+				*EnumProp->GetName(),
+				*Enum->GetNameByValue(Enum->GetMaxEnumValue()).ToString()
+			);
+
+			NewValue = Enum->GetMaxEnumValue();
+		}
+
+		UnderlyingProp->SetIntPropertyValue(Obj, NewValue);
+	}
+
+	struct FEnumPropertyFriend
+	{
+		static const int32 EnumOffset = STRUCT_OFFSET(UEnumProperty, Enum);
+		static const int32 UnderlyingPropOffset = STRUCT_OFFSET(UEnumProperty, UnderlyingProp);
+	};
+}
+
+UEnumProperty::UEnumProperty(const FObjectInitializer& ObjectInitializer, UEnum* InEnum)
+	: UProperty(ObjectInitializer)
+	, Enum(InEnum)
+{
+	// This is expected to be set post-construction by AddCppProperty
+	UnderlyingProp = nullptr;
+}
+
+UEnumProperty::UEnumProperty(const FObjectInitializer& ObjectInitializer, ECppProperty, int32 InOffset, uint64 InFlags, UEnum* InEnum)
+	: UProperty(ObjectInitializer, EC_CppProperty, InOffset, InFlags)
+	, Enum(InEnum)
+{
+	// This is expected to be set post-construction by AddCppProperty
+	UnderlyingProp = nullptr;
+}
+
+void UEnumProperty::AddCppProperty(UProperty* Inner)
+{
+	check(!UnderlyingProp);
+
+	UnderlyingProp = CastChecked<UNumericProperty>(Inner);
+}
+
+void UEnumProperty::SerializeItem( FArchive& Ar, void* Value, void const* Defaults ) const
+{
+	check(UnderlyingProp);
+	check(Enum);
+
+	if(Ar.UseToResolveEnumerators())
+	{
+		int64 IntValue = UnderlyingProp->GetSignedIntPropertyValue(Value);
+		int64 ResolvedIndex = Enum->ResolveEnumerator(Ar, IntValue);
+		UnderlyingProp->SetIntPropertyValue(Value, ResolvedIndex);
+		return;
+	}
+
+	// Loading
+	if (Ar.IsLoading())
+	{
+		FName EnumValueName;
+		Ar << EnumValueName;
+
+		// Make sure enum is properly populated
+		if( Enum->HasAnyFlags(RF_NeedLoad) )
+		{
+			Ar.Preload(Enum);
+		}
+
+		// There's no guarantee EnumValueName is still present in Enum, in which case Value will be set to the enum's max value.
+		// On save, it will then be serialized as NAME_None.
+		int32 EnumIndex = Enum->FindEnumIndex(EnumValueName);
+		int64 NewEnumValue;
+		if (EnumIndex == INDEX_NONE)
+		{
+			NewEnumValue = Enum->GetMaxEnumValue();
+		}
+		else
+		{
+			NewEnumValue = Enum->GetValueByIndex(EnumIndex);
+		}
+
+		UnderlyingProp->SetIntPropertyValue(Value, NewEnumValue);
+	}
+	// Saving
+	else if (Ar.IsSaving())
+	{
+		FName EnumValueName;
+		int64 IntValue = UnderlyingProp->GetSignedIntPropertyValue(Value);
+
+		// subtract 1 because the last entry in the enum's Names array
+		// is the _MAX entry
+		if (Enum->IsValidEnumValue(IntValue))
+		{
+			EnumValueName = Enum->GetNameByValue(IntValue);
+		}
+		else
+		{
+			EnumValueName = NAME_None;
+		}
+
+		Ar << EnumValueName;
+	}
+	else
+	{
+		UnderlyingProp->SerializeItem(Ar, Value, Defaults);
+	}
+}
+
+bool UEnumProperty::NetSerializeItem(FArchive& Ar, UPackageMap* Map, void* Data, TArray<uint8>* MetaData) const
+{
+	Ar.SerializeBits(Data, FMath::CeilLogTwo64(Enum->GetMaxEnumValue()));
+	return 1;
+}
+
+void UEnumProperty::Serialize( FArchive& Ar )
+{
+	Super::Serialize(Ar);
+	Ar << Enum;
+	if (Enum != nullptr)
+	{
+		Ar.Preload(Enum);
+	}
+	Ar << UnderlyingProp;
+	if (UnderlyingProp != nullptr)
+	{
+		Ar.Preload(UnderlyingProp);
+	}
+}
+
+void UEnumProperty::AddReferencedObjects(UObject* InThis, FReferenceCollector& Collector)
+{
+	UEnumProperty* This = CastChecked<UEnumProperty>(InThis);
+	Collector.AddReferencedObject(This->Enum, This);
+	Collector.AddReferencedObject(This->UnderlyingProp, This);
+	Super::AddReferencedObjects(InThis, Collector);
+}
+
+FString UEnumProperty::GetCPPType(FString* ExtendedTypeText, uint32 CPPExportFlags) const
+{
+	check(Enum);
+	check(UnderlyingProp);
+
+	const bool bNonNativeEnum = Enum->GetClass() != UEnum::StaticClass(); // cannot use RF_Native flag, because in UHT the flag is not set
+
+	if (!Enum->CppType.IsEmpty())
+	{
+		return Enum->CppType;
+	}
+
+	FString EnumName = Enum->GetName();
+
+	// This would give the wrong result if it's a namespaced type and the CppType hasn't
+	// been set, but we do this here in case existing code relies on it... somehow.
+	if ((CPPExportFlags & CPPF_BlueprintCppBackend) && bNonNativeEnum)
+	{
+		ensure(Enum->CppType.IsEmpty());
+		FString Result = ::UnicodeToCPPIdentifier(EnumName, false, TEXT("E__"));
+		return Result;
+	}
+
+	return EnumName;
+}
+
+void UEnumProperty::ExportTextItem(FString& ValueStr, const void* PropertyValue, const void* DefaultValue, UObject* Parent, int32 PortFlags, UObject* ExportRootScope) const
+{
+	check(Enum);
+	check(UnderlyingProp);
+
+	UNumericProperty* LocalUnderlyingProp = UnderlyingProp;
+
+	if (PortFlags & PPF_ExportCpp)
+	{
+		const int64 ActualValue = LocalUnderlyingProp->GetSignedIntPropertyValue(PropertyValue);
+		const int64 MaxValue = Enum->GetMaxEnumValue();
+		const int64 GoodValue = Enum->IsValidEnumValue(ActualValue) ? ActualValue : MaxValue;
+		const bool bNonNativeEnum = Enum->GetClass() != UEnum::StaticClass();
+		ensure(!bNonNativeEnum || Enum->CppType.IsEmpty());
+		const FString FullyQualifiedEnumName = bNonNativeEnum ? ::UnicodeToCPPIdentifier(Enum->GetName(), false, TEXT("E__"))
+			: (Enum->CppType.IsEmpty() ? Enum->GetName() : Enum->CppType);
+		if (GoodValue == MaxValue)
+		{
+			// not all native enums have Max value declared
+			ValueStr += FString::Printf(TEXT("(%s)(%ull)"), *FullyQualifiedEnumName, ActualValue);
+		}
+		else
+		{
+			ValueStr += FString::Printf(TEXT("%s::%s"), *FullyQualifiedEnumName,
+				*Enum->GetEnumName(Enum->GetIndexByValue(GoodValue)));
+		}
+		return;
+	}
+
+	if (!(PortFlags & PPF_ConsoleVariable))
+	{
+		int64 Value = LocalUnderlyingProp->GetSignedIntPropertyValue(PropertyValue);
+
+		// if the value is the max value (the autogenerated *_MAX value), export as "INVALID", unless we're exporting text for copy/paste (for copy/paste,
+		// the property text value must actually match an entry in the enum's names array)
+		bool bIsValid = Enum->IsValidEnumValue(Value);
+		bool bIsMax = Value == Enum->GetMaxEnumValue();
+		if (bIsValid && (!bIsMax || (PortFlags & PPF_Copy)))
+		{
+			// We do not want to export the enum text for non-display uses, localization text is very dynamic and would cause issues on import
+			if (PortFlags & PPF_PropertyWindow)
+			{
+				ValueStr += Enum->GetEnumTextByValue(Value).ToString();
+			}
+			else
+			{
+				ValueStr += Enum->GetEnumNameStringByValue(Value);
+			}
+		}
+		else
+		{
+			ValueStr += TEXT("(INVALID)");
+		}
+	}
+	else
+	{
+		UnderlyingProp->ExportTextItem(ValueStr, PropertyValue, DefaultValue, Parent, PortFlags, ExportRootScope);
+	}
+}
+
+const TCHAR* UEnumProperty::ImportText_Internal(const TCHAR* InBuffer, void* Data, int32 PortFlags, UObject* Parent, FOutputDevice* ErrorText) const
+{
+	check(Enum);
+	check(UnderlyingProp);
+	
+	if (!(PortFlags & PPF_ConsoleVariable))
+	{
+		FString Temp;
+		if (const TCHAR* Buffer = UPropertyHelpers::ReadToken(InBuffer, Temp, true))
+		{
+			int32 EnumIndex = Enum->FindEnumIndex(*Temp);
+			if (EnumIndex != INDEX_NONE)
+			{
+				UnderlyingProp->SetIntPropertyValue(Data, Enum->GetValueByIndex(EnumIndex));
+				return Buffer;
+			}
+		}
+	}
+
+	const TCHAR* Result = UnderlyingProp->ImportText(InBuffer, Data, PortFlags, Parent, ErrorText);
+	return Result;
+}
+
+FString UEnumProperty::GetCPPMacroType(FString& ExtendedTypeText) const
+{
+	ExtendedTypeText = Enum->GetName();
+	return TEXT("ENUM");
+}
+
+FString UEnumProperty::GetCPPTypeForwardDeclaration() const
+{
+	check(Enum);
+	check(Enum->GetCppForm() == UEnum::ECppForm::EnumClass);
+
+	return FString::Printf(TEXT("enum class %s : %s;"), *Enum->GetName(), *UnderlyingProp->GetCPPType());
+}
+
+void UEnumProperty::GetPreloadDependencies(TArray<UObject*>& OutDeps)
+{
+	Super::GetPreloadDependencies(OutDeps);
+	OutDeps.Add(UnderlyingProp);
+	OutDeps.Add(Enum);
+}
+
+void UEnumProperty::LinkInternal(FArchive& Ar)
+{
+	check(UnderlyingProp);
+
+	Ar.Preload(UnderlyingProp);
+	UnderlyingProp->Link(Ar);
+
+	this->ElementSize = UnderlyingProp->ElementSize;
+	this->PropertyFlags |= CPF_IsPlainOldData | CPF_NoDestructor | CPF_ZeroConstructor;
+}
+
+bool UEnumProperty::Identical(const void* A, const void* B, uint32 PortFlags) const
+{
+	return UnderlyingProp->Identical(A, B, PortFlags);
+}
+
+int32 UEnumProperty::GetMinAlignment() const
+{
+	return UnderlyingProp->GetMinAlignment();
+}
+
+bool UEnumProperty::SameType(const UProperty* Other) const
+{
+	return Super::SameType(Other) && static_cast<const UEnumProperty*>(Other)->Enum == Enum;
+}
+
+bool UEnumProperty::ConvertFromType(const FPropertyTag& Tag, FArchive& Ar, uint8* Data, UStruct* DefaultsStruct, bool& bOutAdvanceProperty)
+{
+	check(Enum);
+	check(UnderlyingProp);
+
+	bOutAdvanceProperty = true;
+
+	if (Tag.Type == NAME_ByteProperty)
+	{
+		uint8 PreviousValue;
+		if (Tag.EnumName == NAME_None)
+		{
+			// If we're a nested property the EnumName tag got lost. Handle this case for backward compatibility reasons
+			UProperty* const PropertyOwner = Cast<UProperty>(GetOuterUField());
+
+			if (PropertyOwner)
+			{
+				FPropertyTag InnerPropertyTag;
+				InnerPropertyTag.Type = Tag.Type;
+				InnerPropertyTag.EnumName = Enum->GetFName();
+				InnerPropertyTag.ArrayIndex = 0;
+
+				PreviousValue = UNumericProperty::ReadEnumAsUint8(Ar, DefaultsStruct, InnerPropertyTag);
+			}
+			else
+			{
+				// a byte property gained an enum
+				Ar << PreviousValue;
+			}
+		}
+		else
+		{
+			// attempt to find the old enum and get the byte value from the serialized enum name
+			PreviousValue = UNumericProperty::ReadEnumAsUint8(Ar, DefaultsStruct, Tag);
+		}
+
+		// now copy the value into the object's address space
+		UnderlyingProp->SetIntPropertyValue(ContainerPtrToValuePtr<void>(Data, Tag.ArrayIndex), (uint64)PreviousValue);
+	}
+	else if (Tag.Type == NAME_Int8Property)
+	{
+		UE4EnumProperty_Private::ConvertIntToEnumProperty<int8>(Ar, this, UnderlyingProp, Enum, ContainerPtrToValuePtr<void>(Data, Tag.ArrayIndex));
+	}
+	else if (Tag.Type == NAME_Int16Property)
+	{
+		UE4EnumProperty_Private::ConvertIntToEnumProperty<int16>(Ar, this, UnderlyingProp, Enum, ContainerPtrToValuePtr<void>(Data, Tag.ArrayIndex));
+	}
+	else if (Tag.Type == NAME_IntProperty)
+	{
+		UE4EnumProperty_Private::ConvertIntToEnumProperty<int32>(Ar, this, UnderlyingProp, Enum, ContainerPtrToValuePtr<void>(Data, Tag.ArrayIndex));
+	}
+	else if (Tag.Type == NAME_Int64Property)
+	{
+		UE4EnumProperty_Private::ConvertIntToEnumProperty<int64>(Ar, this, UnderlyingProp, Enum, ContainerPtrToValuePtr<void>(Data, Tag.ArrayIndex));
+	}
+	else if (Tag.Type == NAME_UInt16Property)
+	{
+		UE4EnumProperty_Private::ConvertIntToEnumProperty<uint16>(Ar, this, UnderlyingProp, Enum, ContainerPtrToValuePtr<void>(Data, Tag.ArrayIndex));
+	}
+	else if (Tag.Type == NAME_UInt32Property)
+	{
+		UE4EnumProperty_Private::ConvertIntToEnumProperty<uint32>(Ar, this, UnderlyingProp, Enum, ContainerPtrToValuePtr<void>(Data, Tag.ArrayIndex));
+	}
+	else if (Tag.Type == NAME_UInt64Property)
+	{
+		UE4EnumProperty_Private::ConvertIntToEnumProperty<uint64>(Ar, this, UnderlyingProp, Enum, ContainerPtrToValuePtr<void>(Data, Tag.ArrayIndex));
+	}
+	else
+	{
+		bOutAdvanceProperty = false;
+	}
+
+	return bOutAdvanceProperty;
+}
+
+IMPLEMENT_CORE_INTRINSIC_CLASS(UEnumProperty, UProperty,
+	{
+		Class->EmitObjectReference(UE4EnumProperty_Private::FEnumPropertyFriend::EnumOffset, TEXT("Enum"));
+		Class->EmitObjectReference(UE4EnumProperty_Private::FEnumPropertyFriend::UnderlyingPropOffset, TEXT("UnderlyingProp"));
+	}
+);
diff --git a/Engine/Source/Runtime/CoreUObject/Private/UObject/GarbageCollection.cpp b/Engine/Source/Runtime/CoreUObject/Private/UObject/GarbageCollection.cpp
--- a/Engine/Source/Runtime/CoreUObject/Private/UObject/GarbageCollection.cpp
+++ b/Engine/Source/Runtime/CoreUObject/Private/UObject/GarbageCollection.cpp
@@ -1824,6 +1824,12 @@ void UClass::EmitFixedArrayEnd()
 
 void UClass::AssembleReferenceTokenStream(bool bForce)
 {
+	const bool bNotNative = !(ClassFlags & CLASS_Native);
+	if (bNotNative)
+	{
+		ReferenceTokenStreamCritical.Lock();
+	}
+
 	if (!HasAnyClassFlags(CLASS_TokenStreamAssembled) || bForce)
 	{
 		if (bForce)
@@ -1882,6 +1888,10 @@ void UClass::AssembleReferenceTokenStream(bool bForce)
 		check(!HasAnyClassFlags(CLASS_TokenStreamAssembled)); // recursion here is probably bad
 		ClassFlags |= CLASS_TokenStreamAssembled;
 	}
+	if (bNotNative)
+	{
+		ReferenceTokenStreamCritical.Unlock();
+	}
 }
 
 
diff --git a/Engine/Source/Runtime/CoreUObject/Private/UObject/Linker.cpp b/Engine/Source/Runtime/CoreUObject/Private/UObject/Linker.cpp
--- a/Engine/Source/Runtime/CoreUObject/Private/UObject/Linker.cpp
+++ b/Engine/Source/Runtime/CoreUObject/Private/UObject/Linker.cpp
@@ -614,6 +614,7 @@ FLinkerLoad* GetPackageLinker
 		UPackage* FilenamePkg = (ExistingPackage ? ExistingPackage : (CreatedPackage = CreatePackage(nullptr, *PackageNameToCreate)));
 		if (FilenamePkg && FilenamePkg != ExistingPackage && (LoadFlags & LOAD_PackageForPIE))
 		{
+			check(FilenamePkg);
 			FilenamePkg->SetPackageFlags(PKG_PlayInEditor);
 		}
 
diff --git a/Engine/Source/Runtime/CoreUObject/Private/UObject/LinkerLoad.cpp b/Engine/Source/Runtime/CoreUObject/Private/UObject/LinkerLoad.cpp
--- a/Engine/Source/Runtime/CoreUObject/Private/UObject/LinkerLoad.cpp
+++ b/Engine/Source/Runtime/CoreUObject/Private/UObject/LinkerLoad.cpp
@@ -1826,9 +1826,9 @@ FLinkerLoad::ELinkerStatus FLinkerLoad::SerializeThumbnails( bool bForceEnableIn
 
 
 		// Allocate a new thumbnail map if we need one
-		if( !LinkerRoot->ThumbnailMap.IsValid() )
+		if( !LinkerRoot->ThumbnailMap )
 		{
-			LinkerRoot->ThumbnailMap.Reset( new FThumbnailMap() );
+			LinkerRoot->ThumbnailMap = MakeUnique<FThumbnailMap>();
 		}
 
 
@@ -4420,7 +4420,15 @@ void FLinkerLoad::DetachExport( int32 i )
 			UE_LOG(LogLinker, Fatal, TEXT("Linker object %s %s.%s mislinked!"), *GetExportClassName(i).ToString(), *LinkerRoot->GetName(), *E.ObjectName.ToString());
 		}
 	}
-	check(E.Object->GetLinkerIndex() == i);
+
+	if (E.Object->GetLinkerIndex() == -1)
+	{
+		UE_LOG(LogLinker, Warning, TEXT("Linker object %s %s.%s was already detached."), *GetExportClassName(i).ToString(), *LinkerRoot->GetName(), *E.ObjectName.ToString());
+	}
+	else
+	{
+		checkf(E.Object->GetLinkerIndex() == i, TEXT("Mismatched linker index in FLinkerLoad::DetachExport for %s in %s. Linker index was supposed to be %d, was %d"), *GetExportClassName(i).ToString(), *LinkerRoot->GetName(), i, E.Object->GetLinkerIndex());
+	}
 	ExportMap[i].Object->SetLinker( NULL, INDEX_NONE );
 }
 
diff --git a/Engine/Source/Runtime/CoreUObject/Private/UObject/LinkerManager.cpp b/Engine/Source/Runtime/CoreUObject/Private/UObject/LinkerManager.cpp
--- a/Engine/Source/Runtime/CoreUObject/Private/UObject/LinkerManager.cpp
+++ b/Engine/Source/Runtime/CoreUObject/Private/UObject/LinkerManager.cpp
@@ -9,8 +9,8 @@
 
 FLinkerManager& FLinkerManager::Get()
 {
-	static TAutoPtr<FLinkerManager> Singleton(new FLinkerManager());
-	return *Singleton.GetOwnedPointer();
+	static TUniquePtr<FLinkerManager> Singleton = MakeUnique<FLinkerManager>();
+	return *Singleton;
 }
 
 bool FLinkerManager::Exec(class UWorld* InWorld, const TCHAR* Cmd, FOutputDevice& Ar)
diff --git a/Engine/Source/Runtime/CoreUObject/Private/UObject/Obj.cpp b/Engine/Source/Runtime/CoreUObject/Private/UObject/Obj.cpp
--- a/Engine/Source/Runtime/CoreUObject/Private/UObject/Obj.cpp
+++ b/Engine/Source/Runtime/CoreUObject/Private/UObject/Obj.cpp
@@ -1381,7 +1381,7 @@ void UObject::FAssetRegistryTag::GetAssetRegistryTagsFromSearchableProperties(co
 			else if ( Class->IsChildOf(UByteProperty::StaticClass()) )
 			{
 				// bytes are numerical, enums are alphabetical
-				UByteProperty* ByteProp = dynamic_cast<UByteProperty*>(*FieldIt);
+				UByteProperty* ByteProp = static_cast<UByteProperty*>(*FieldIt);
 				if ( ByteProp->Enum )
 				{
 					TagType = FAssetRegistryTag::TT_Alphabetical;
@@ -1391,6 +1391,11 @@ void UObject::FAssetRegistryTag::GetAssetRegistryTagsFromSearchableProperties(co
 					TagType = FAssetRegistryTag::TT_Numerical;
 				}
 			}
+			else if ( Class->IsChildOf(UEnumProperty::StaticClass()) )
+			{
+				// enums are alphabetical
+				TagType = FAssetRegistryTag::TT_Alphabetical;
+			}
 			else if ( Class->IsChildOf(UArrayProperty::StaticClass()) )
 			{
 				// arrays are hidden, it is often too much information to display and sort
diff --git a/Engine/Source/Runtime/CoreUObject/Private/UObject/PropertyBool.cpp b/Engine/Source/Runtime/CoreUObject/Private/UObject/PropertyBool.cpp
--- a/Engine/Source/Runtime/CoreUObject/Private/UObject/PropertyBool.cpp
+++ b/Engine/Source/Runtime/CoreUObject/Private/UObject/PropertyBool.cpp
@@ -235,6 +235,15 @@ bool UBoolProperty::ConvertFromType(const FPropertyTag& Tag, FArchive& Ar, uint8
 		// if the byte property was an enum we won't allow a conversion to bool
 		if (Tag.EnumName == NAME_None)
 		{
+			// If we're a nested property the EnumName tag got lost, don't allow this
+			UProperty* const PropertyOwner = Cast<UProperty>(GetOuterUField());
+
+			if (PropertyOwner)
+			{
+				bOutAdvanceProperty = false;
+				return bOutAdvanceProperty;
+			}
+
 			LoadFromType<uint8>(this, Tag, Ar, Data);
 		}
 		else
diff --git a/Engine/Source/Runtime/CoreUObject/Private/UObject/PropertyByte.cpp b/Engine/Source/Runtime/CoreUObject/Private/UObject/PropertyByte.cpp
--- a/Engine/Source/Runtime/CoreUObject/Private/UObject/PropertyByte.cpp
+++ b/Engine/Source/Runtime/CoreUObject/Private/UObject/PropertyByte.cpp
@@ -17,7 +17,7 @@ void UByteProperty::SerializeItem( FArchive& Ar, void* Value, void const* Defaul
 {
 	if(Enum && Ar.UseToResolveEnumerators())
 	{
-		 const int32 ResolvedIndex = Enum->ResolveEnumerator(Ar, *(uint8*)Value);
+		 const int64 ResolvedIndex = Enum->ResolveEnumerator(Ar, *(uint8*)Value);
 		 *(uint8*)Value = static_cast<uint8>(ResolvedIndex);
 		 return;
 	}
@@ -172,6 +172,15 @@ bool UByteProperty::ConvertFromType(const FPropertyTag& Tag, FArchive& Ar, uint8
 		uint8 PreviousValue;
 		if (Tag.EnumName == NAME_None)
 		{
+			// If we're a nested property the EnumName tag got lost. Fail to read in this case
+			UProperty* const PropertyOwner = Cast<UProperty>(GetOuterUField());
+
+			if (PropertyOwner)
+			{
+				bOutAdvanceProperty = false;
+				return bOutAdvanceProperty;
+			}
+
 			// simply pretend the property still doesn't have an enum and serialize the single byte
 			Ar << PreviousValue;
 		}
@@ -184,6 +193,15 @@ bool UByteProperty::ConvertFromType(const FPropertyTag& Tag, FArchive& Ar, uint8
 		// now copy the value into the object's address space
 		SetPropertyValue_InContainer(Data, PreviousValue, Tag.ArrayIndex);
 	}
+	else if (Tag.Type == NAME_EnumProperty && (Enum == nullptr || Tag.EnumName == Enum->GetFName()))
+	{
+		// an enum property became a byte
+		// attempt to find the old enum and get the byte value from the serialized enum name
+		uint8 PreviousValue = ReadEnumAsUint8(Ar, DefaultsStruct, Tag);
+
+		// now copy the value into the object's address space
+		SetPropertyValue_InContainer(Data, PreviousValue, Tag.ArrayIndex);
+	}
 	else if (Tag.Type == NAME_Int8Property)
 	{
 		if (Enum)
diff --git a/Engine/Source/Runtime/CoreUObject/Private/UObject/PropertyNumeric.cpp b/Engine/Source/Runtime/CoreUObject/Private/UObject/PropertyNumeric.cpp
--- a/Engine/Source/Runtime/CoreUObject/Private/UObject/PropertyNumeric.cpp
+++ b/Engine/Source/Runtime/CoreUObject/Private/UObject/PropertyNumeric.cpp
@@ -5,7 +5,7 @@
 
 uint8 UNumericProperty::ReadEnumAsUint8(FArchive& Ar, UStruct* DefaultsStruct, const FPropertyTag& Tag)
 {
-	//@warning: mirrors loading code in UByteProperty::SerializeItem()
+	//@warning: mirrors loading code in UByteProperty::SerializeItem() and UEnumProperty::SerializeItem()
 	FName EnumName;
 	Ar << EnumName;
 
@@ -23,7 +23,7 @@ uint8 UNumericProperty::ReadEnumAsUint8(FArchive& Ar, UStruct* DefaultsStruct, c
 
 	Ar.Preload(Enum);
 
-	uint8 Result = Enum->GetValueByName(EnumName);
+	int64 Result = Enum->GetValueByName(EnumName);
 	if (!Enum->IsValidEnumValue(Result))
 	{
 		UE_LOG(
@@ -51,10 +51,10 @@ const TCHAR* UNumericProperty::ImportText_Internal( const TCHAR* Buffer, void* D
 		{
 			if (FChar::IsAlpha(*Buffer))
 			{
-				int32 EnumValue = UEnum::ParseEnum(Buffer);
+				int64 EnumValue = UEnum::ParseEnum(Buffer);
 				if (EnumValue != INDEX_NONE)
 				{
-					SetIntPropertyValue(Data, int64(EnumValue));
+					SetIntPropertyValue(Data, EnumValue);
 					return Buffer;
 				}
 				else
diff --git a/Engine/Source/Runtime/CoreUObject/Private/UObject/PropertyTag.cpp b/Engine/Source/Runtime/CoreUObject/Private/UObject/PropertyTag.cpp
--- a/Engine/Source/Runtime/CoreUObject/Private/UObject/PropertyTag.cpp
+++ b/Engine/Source/Runtime/CoreUObject/Private/UObject/PropertyTag.cpp
@@ -2,6 +2,7 @@
 
 #include "CoreUObjectPrivate.h"
 #include "PropertyTag.h"
+#include "EnumProperty.h"
 #include "UObject/BlueprintsObjectVersion.h"
 
 /*-----------------------------------------------------------------------------
@@ -44,6 +45,10 @@ FPropertyTag::FPropertyTag( FArchive& InSaveAr, UProperty* Property, int32 InInd
 			StructName = StructProperty->Struct->GetFName();
 			StructGuid = StructProperty->Struct->GetCustomGuid();
 		}
+		else if (UEnumProperty* EnumProp = Cast<UEnumProperty>(Property))
+		{
+			EnumName = EnumProp->GetEnum()->GetFName();
+		}
 		else if (UByteProperty* ByteProp = Cast<UByteProperty>(Property))
 		{
 			if (ByteProp->Enum != NULL)
@@ -116,8 +121,8 @@ FArchive& operator<<( FArchive& Ar, FPropertyTag& Tag )
 	{
 		Ar << Tag.BoolVal;
 	}
-	// only need to serialize this for bytes
-	else if (Tag.Type == NAME_ByteProperty)
+	// only need to serialize this for bytes/enums
+	else if (Tag.Type == NAME_ByteProperty || Tag.Type == NAME_EnumProperty)
 	{
 		Ar << Tag.EnumName;
 	}
diff --git a/Engine/Source/Runtime/CoreUObject/Private/UObject/SavePackage.cpp b/Engine/Source/Runtime/CoreUObject/Private/UObject/SavePackage.cpp
--- a/Engine/Source/Runtime/CoreUObject/Private/UObject/SavePackage.cpp
+++ b/Engine/Source/Runtime/CoreUObject/Private/UObject/SavePackage.cpp
@@ -10,6 +10,7 @@
 #include "DebugSerializationFlags.h"
 #include "UObject/GCScopeLock.h"
 #include "CookStats.h"
+#include "EnumProperty.h"
 
 DEFINE_LOG_CATEGORY_STATIC(LogSavePackage, Log, All);
 
@@ -745,7 +746,7 @@ static void ConditionallyExcludeObjectForTarget(UObject* Obj, FArchive& Ar)
 			Obj->Mark(OBJECTMARK_NotForServer);
 		}
 
-		if (Search->HasAnyFlags(RF_Public))
+		if (Search->IsA(UPackage::StaticClass()))
 		{
 			break;
 		}
@@ -886,7 +887,8 @@ FArchive& FArchiveSaveTagExports::operator<<( UObject*& Obj )
 			// The base virtual methods will leave the flags unchanged.
 			// Once should not read the flags directly (other than at load) because the virtual methods are authoritative.
 
-			// if anything in the outer chain is NotFor, then we are also NotFor. Stop this search at public objects.
+			// If anything in the outer chain is NotFor, then we are also NotFor.
+			// Stop this search at packages, as nested public objects will have the wrong path.
 			ConditionallyExcludeObjectForTarget(Obj, *this);
 
 			{
@@ -1185,39 +1187,53 @@ FArchive& FArchiveSaveTagImports::operator<<( UObject*& Obj )
 				if (!bIsEditorOnly)
 				{
 					Obj->Mark(OBJECTMARK_TagImp);
+					UClass* ClassObj = Cast<UClass>(Obj);
 
-					if (IsEventDrivenLoaderEnabledInCookedBuilds() && IsCooking() && !bIsNative && Cast<UClass>(Obj))
+					if (IsEventDrivenLoaderEnabledInCookedBuilds() && IsCooking() && !bIsNative && ClassObj)
 					{
-						// we don't want to add this to Dependencies, we simply want it to be an import so that a serialization before creation dependency can be created to the CDO
-						UObject* CDO = Cast<UClass>(Obj)->GetDefaultObject();
-
+						// We don't want to add this to Dependencies, we simply want it to be an import so that a serialization before creation dependency can be created to the CDO
+						UObject* CDO = ClassObj->GetDefaultObject();
+						
 						if (CDO)
 						{
-							ConditionallyExcludeObjectForTarget(CDO, *this);
+							// Gets all subobjects defined in a class, including the CDO, CDO components and blueprint-created components
+							TArray<UObject*> ObjectTemplates;
+							ObjectTemplates.Add(CDO);
 
-							const bool bCDOIsEditorOnly = IsCookingForNoEditorDataPlatform(*this) && IsEditorOnlyObject(CDO);
+							TArray<UObject*> CurrentSubobjects;
+							TArray<UObject*> NextSubobjects;
 
-							if (!bCDOIsEditorOnly && !CDO->HasAnyFlags(RF_Transient) && !CDO->IsPendingKill())
-							{
-								CDO->Mark(OBJECTMARK_TagImp);
+							// Recursively search for subobjects. Only care about ones that have a full subobject chain as some nested objects are set wrong
+							GetObjectsWithOuter(ClassObj, NextSubobjects, false);
+							GetObjectsWithOuter(CDO, NextSubobjects, false);
 
-								// and all subobjects
-								TArray<UObject*> Subobjects;
-								GetObjectsWithOuter(CDO, Subobjects);
-								for (UObject* SubObj : Subobjects)
+							while (NextSubobjects.Num() > 0)
+							{
+								CurrentSubobjects = NextSubobjects;
+								NextSubobjects.Empty();
+								for (UObject* SubObj : CurrentSubobjects)
 								{
-									if (SubObj)
+									if (SubObj->HasAnyFlags(RF_DefaultSubObject | RF_ArchetypeObject))
 									{
-										// correctly mark components that shouldn't be cooked
-										ConditionallyExcludeObjectForTarget(SubObj, *this);
+										ObjectTemplates.Add(SubObj);
+										GetObjectsWithOuter(SubObj, NextSubobjects, false);
+									}
+								}
+							}
 
-										const bool bCompIsEditorOnly = IsCookingForNoEditorDataPlatform(*this) && IsEditorOnlyObject(SubObj);
+							for (UObject* ObjTemplate : ObjectTemplates)
+							{
+								if (IsCookingForNoEditorDataPlatform(*this) && IsEditorOnlyObject(ObjTemplate))
+								{
+									// Explicitly mark as editor only so it will be skipped later
+									ObjTemplate->Mark(OBJECTMARK_EditorOnly);
+								}
+								else if (!ObjTemplate->HasAnyFlags(RF_Transient) && !ObjTemplate->IsPendingKill())
+								{
+									// Correctly mark objects that shouldn't be cooked
+									ConditionallyExcludeObjectForTarget(ObjTemplate, *this);
 
-										if (!bCompIsEditorOnly && !SubObj->HasAnyFlags(RF_Transient) && !SubObj->IsPendingKill())
-										{
-											SubObj->Mark(OBJECTMARK_TagImp);
-										}
-									}
+									ObjTemplate->Mark(OBJECTMARK_TagImp);
 								}
 							}
 						}
@@ -2401,21 +2417,22 @@ class FExportReferenceSorter : public FArchiveUObject
 				UObjectPropertyBase::StaticClass(),
 				ULazyObjectProperty::StaticClass(),
 				UAssetObjectProperty::StaticClass(),
-				UAssetClassProperty::StaticClass()
+				UAssetClassProperty::StaticClass(),
+				UMapProperty::StaticClass(),
+				USetProperty::StaticClass(),
+				UEnumProperty::StaticClass()
 			};
 
-			for (int32 CoreClassIndex = 0; CoreClassIndex < ARRAY_COUNT(CoreClassList); CoreClassIndex++)
+			for (UClass* CoreClass : CoreClassList)
 			{
-				UClass* CoreClass = CoreClassList[CoreClassIndex];
 				CoreClasses.AddUnique(CoreClass);
 
 				ReferencedObjects.Add(CoreClass);
 				ReferencedObjects.Add(CoreClass->GetDefaultObject());
 			}
 
-			for (int32 CoreClassIndex = 0; CoreClassIndex < CoreClasses.Num(); CoreClassIndex++)
+			for (UClass* CoreClass : CoreClasses)
 			{
-				UClass* CoreClass = CoreClasses[CoreClassIndex];
 				ProcessStruct(CoreClass);
 			}
 
@@ -2508,21 +2525,22 @@ class FExportReferenceSorter : public FArchiveUObject
 			UObjectPropertyBase::StaticClass(),
 			ULazyObjectProperty::StaticClass(),
 			UAssetObjectProperty::StaticClass(),
-			UAssetClassProperty::StaticClass()
+			UAssetClassProperty::StaticClass(),
+			UMapProperty::StaticClass(),
+			USetProperty::StaticClass(),
+			UEnumProperty::StaticClass()
 		};
 
-		for (int32 CoreClassIndex = 0; CoreClassIndex < ARRAY_COUNT(CoreClassList); CoreClassIndex++)
+		for (UClass* CoreClass : CoreClassList)
 		{
-			UClass* CoreClass = CoreClassList[CoreClassIndex];
 			CoreClasses.AddUnique(CoreClass);
 
 			ReferencedObjects.Add(CoreClass);
 			ReferencedObjects.Add(CoreClass->GetDefaultObject());
 		}
 
-		for (int32 CoreClassIndex = 0; CoreClassIndex < CoreClasses.Num(); CoreClassIndex++)
+		for (UClass* CoreClass : CoreClasses)
 		{
-			UClass* CoreClass = CoreClasses[CoreClassIndex];
 			ProcessStruct(CoreClass);
 		}
 
@@ -2728,11 +2746,28 @@ public:
 							else
 							{
 								// byte properties that are enum references need their enums loaded first so that config importing works
-								UByteProperty* ByteProp = dynamic_cast<UByteProperty*>(Object);
-								if (ByteProp != NULL && ByteProp->Enum != NULL)
 								{
-									HandleDependency(ByteProp->Enum, /*bProcessObject =*/true);
+									UEnum* Enum = nullptr;
+
+									if (UEnumProperty* EnumProp = dynamic_cast<UEnumProperty*>(Object))
+									{
+										Enum = EnumProp->GetEnum();
+									}
+									else
+									{
+										UByteProperty* ByteProp = dynamic_cast<UByteProperty*>(Object);
+										if (ByteProp && ByteProp->Enum)
+										{
+											Enum = ByteProp->Enum;
+										}
+									}
+
+									if (Enum)
+									{
+										HandleDependency(Enum, /*bProcessObject =*/true);
+									}
 								}
+
 								// a normal field - property, enum, const; just insert it into the list and keep going
 								ProcessedObjects.Add(Object);
 								
@@ -4894,6 +4929,12 @@ ESavePackageResult UPackage::Save(UPackage* InOuter, UObject* Base, EObjectFlags
 								TEXT("Export Object (%s) Outer (%s) is an Import."),
 								*(Export.Object->GetPathName()),
 								*(Export.Object->GetOuter()->GetPathName()));
+
+							if (Linker->IsCooking() && IsEventDrivenLoaderEnabledInCookedBuilds())
+							{
+								// Only packages are allowed to have no outer
+								ensureMsgf(Export.OuterIndex != FPackageIndex() || Export.Object->IsA(UPackage::StaticClass()), TEXT("Export %s has no valid outer when cooking!"), *Export.Object->GetPathName());
+							}
 						}
 						else
 						{
@@ -4978,7 +5019,11 @@ ESavePackageResult UPackage::Save(UPackage* InOuter, UObject* Base, EObjectFlags
 								GetObjectsWithOuter(CDO, Subobjects);
 								for (UObject* SubObj : Subobjects)
 								{
-									IncludeObjectAsDependency(SerializationBeforeCreateDependencies, SubObj, Export.Object, false);
+									// Only include subobject archetypes
+									if (SubObj->HasAnyFlags(RF_DefaultSubObject | RF_ArchetypeObject))
+									{
+										IncludeObjectAsDependency(SerializationBeforeCreateDependencies, SubObj, Export.Object, false);
+									}
 								}
 							}
 							TSet<FPackageIndex> SerializationBeforeSerializationDependencies;
@@ -5329,6 +5374,12 @@ ESavePackageResult UPackage::Save(UPackage* InOuter, UObject* Base, EObjectFlags
 							{
 								Import.OuterIndex = Linker->MapObject(Import.XObject->GetOuter());
 							}
+
+							if (Linker->IsCooking() && IsEventDrivenLoaderEnabledInCookedBuilds())
+							{
+								// Only package imports are allowed to have no outer
+								ensureMsgf(Import.OuterIndex != FPackageIndex() || Import.ClassName == NAME_Package, TEXT("Import %s has no valid outer when cooking!"), *Import.XObject->GetPathName());
+							}
 						}
 					}
 					else
@@ -5874,4 +5925,4 @@ bool UPackage::IsEmptyPackage (UPackage* Package, const UObject* LastReferencer)
 
 	// Invalid package
 	return false;
-}
\ No newline at end of file
+}
diff --git a/Engine/Source/Runtime/CoreUObject/Private/UObject/UObjectArchetype.cpp b/Engine/Source/Runtime/CoreUObject/Private/UObject/UObjectArchetype.cpp
--- a/Engine/Source/Runtime/CoreUObject/Private/UObject/UObjectArchetype.cpp
+++ b/Engine/Source/Runtime/CoreUObject/Private/UObject/UObjectArchetype.cpp
@@ -68,6 +68,7 @@ UObject* UObject::GetArchetypeFromRequiredInfo(UClass* Class, UObject* Outer, FN
 	}
 #if USE_EVENT_DRIVEN_ASYNC_LOAD
 	//if (!GIgnoreGetArchetypeFromRequiredInfo_RF_NeedsLoad)
+	if (!GIsInitialLoad)
 	{
 		if (Result && Result->HasAnyFlags(RF_NeedLoad))
 		{
diff --git a/Engine/Source/Runtime/CoreUObject/Private/UObject/UObjectBase.cpp b/Engine/Source/Runtime/CoreUObject/Private/UObject/UObjectBase.cpp
--- a/Engine/Source/Runtime/CoreUObject/Private/UObject/UObjectBase.cpp
+++ b/Engine/Source/Runtime/CoreUObject/Private/UObject/UObjectBase.cpp
@@ -321,8 +321,10 @@ void UObjectBase::EmitBaseReferences(UClass *RootClass)
 /** Enqueue the registration for this object. */
 void UObjectBase::Register(const TCHAR* PackageName,const TCHAR* InName)
 {
+	TMap<UObjectBase*, FPendingRegistrantInfo>& PendingRegistrants = FPendingRegistrantInfo::GetMap();
+
 	FPendingRegistrant* PendingRegistration = new FPendingRegistrant(this);
-	FPendingRegistrantInfo::GetMap().Add(this, FPendingRegistrantInfo(InName, PackageName));
+	PendingRegistrants.Add(this, FPendingRegistrantInfo(InName, PackageName));
 	if(GLastPendingRegistrant)
 	{
 		GLastPendingRegistrant->NextAutoRegister = PendingRegistration;
@@ -381,12 +383,14 @@ static void UObjectProcessRegistrants()
 
 void UObjectForceRegistration(UObjectBase* Object)
 {
-	FPendingRegistrantInfo* Info = FPendingRegistrantInfo::GetMap().Find(Object);
+	TMap<UObjectBase*, FPendingRegistrantInfo>& PendingRegistrants = FPendingRegistrantInfo::GetMap();
+
+	FPendingRegistrantInfo* Info = PendingRegistrants.Find(Object);
 	if (Info)
 	{
 		const TCHAR* PackageName = Info->PackageName;
 		const TCHAR* Name = Info->Name;
-		FPendingRegistrantInfo::GetMap().Remove(Object);  // delete this first so that it doesn't try to do it twice
+		PendingRegistrants.Remove(Object);  // delete this first so that it doesn't try to do it twice
 		Object->DeferredRegister(UClass::StaticClass(),PackageName,Name);
 	}
 }
diff --git a/Engine/Source/Runtime/CoreUObject/Private/UObject/UObjectGlobals.cpp b/Engine/Source/Runtime/CoreUObject/Private/UObject/UObjectGlobals.cpp
--- a/Engine/Source/Runtime/CoreUObject/Private/UObject/UObjectGlobals.cpp
+++ b/Engine/Source/Runtime/CoreUObject/Private/UObject/UObjectGlobals.cpp
@@ -3668,10 +3668,15 @@ UObject* FObjectInitializer::CreateDefaultSubobject(UObject* Outer, FName Subobj
 		if (!OverrideClass->HasAnyClassFlags(CLASS_Abstract) || !bAbstract)
 		{
 			UObject* Template = OverrideClass->GetDefaultObject(); // force the CDO to be created if it hasn't already
-			const EObjectFlags SubobjectFlags = Outer->GetMaskedFlags(RF_PropagateToSubObjects);
+			EObjectFlags SubobjectFlags = Outer->GetMaskedFlags(RF_PropagateToSubObjects);
 			bool bOwnerArchetypeIsNotNative;
 			UClass* OuterArchetypeClass;
 
+			if (bIsTransient)
+			{
+				SubobjectFlags |= RF_Transient;
+			}
+
 			{
 #if USE_EVENT_DRIVEN_ASYNC_LOAD
 				TGuardValue<bool> Guard(GIgnoreGetArchetypeFromRequiredInfo_RF_NeedsLoad, true);
