Copying //UE4/Dev-Core to //UE4/Dev-Main (Source: //UE4/Dev-Core @ 3208226)

diff --git a/Engine/Source/Runtime/CoreUObject/Public/CoreUObject.h b/Engine/Source/Runtime/CoreUObject/Public/CoreUObject.h
--- a/Engine/Source/Runtime/CoreUObject/Public/CoreUObject.h
+++ b/Engine/Source/Runtime/CoreUObject/Public/CoreUObject.h
@@ -34,6 +34,7 @@
 #include "Package.h"					// Package class definition
 #include "MetaData.h"					// Metadata for packages
 #include "UnrealType.h"					// Base property type.
+#include "EnumProperty.h"				// enum property type
 #include "TextProperty.h"				// FText property type
 #include "Stack.h"						// Script stack frame definition.
 #include "ObjectRedirector.h"			// Cross-package object redirector
diff --git a/Engine/Source/Runtime/CoreUObject/Public/Misc/StringAssetReference.h b/Engine/Source/Runtime/CoreUObject/Public/Misc/StringAssetReference.h
--- a/Engine/Source/Runtime/CoreUObject/Public/Misc/StringAssetReference.h
+++ b/Engine/Source/Runtime/CoreUObject/Public/Misc/StringAssetReference.h
@@ -6,7 +6,10 @@ class UObject;
 struct FPropertyTag;
 
 /**
- * A struct that contains a string reference to an asset, can be used to make soft references to assets
+ * A struct that contains a string reference to an asset on disk.
+ * This can be used to make soft references to assets that are loaded on demand.
+ * This is stored internally as a string of the form /package/path.assetname[.objectname]
+ * If the MetaClass metadata is applied to a UProperty with this the UI will restrict to that type of asset.
  */
 struct COREUOBJECT_API FStringAssetReference
 {
@@ -18,87 +21,93 @@ struct COREUOBJECT_API FStringAssetReference
 	{
 	}
 
+	/** Construct from another asset reference */
 	FStringAssetReference(const FStringAssetReference& Other)
 	{
 		SetPath(Other.ToString());
 	}
 
-	/**
-	 * Construct from a path string
-	 */
+	/** Construct from a path string */
 	FStringAssetReference(FString PathString)
 	{
 		SetPath(MoveTemp(PathString));
 	}
 
-	/**
-	 * Construct from an existing object, will do some string processing
-	 */
+	/** Construct from an existing object in memory */
 	FStringAssetReference(const UObject* InObject);
 
-	~FStringAssetReference();
+PRAGMA_DISABLE_DEPRECATION_WARNINGS
+	
+	/** Destructor, in deprecation block to avoid warnings */
+	~FStringAssetReference() {}
 
-	/**
-	 * Converts in to a string
-	 */
-	const FString& ToString() const;
+	/** Returns string representation of reference, in form /package/path.assetname[.objectname] */
+	FORCEINLINE const FString& ToString() const
+	{
+		return AssetLongPathname;		
+	}
 
-	const FString GetLongPackageName() const
+PRAGMA_ENABLE_DEPRECATION_WARNINGS
+
+	/** Returns /package/path, leaving off the asset name */
+	FString GetLongPackageName() const
 	{
 		FString PackageName;
-		ToString().Split(TEXT("."), &PackageName, nullptr, ESearchCase::CaseSensitive, ESearchDir::FromEnd);
+		ToString().Split(TEXT("."), &PackageName, nullptr, ESearchCase::CaseSensitive, ESearchDir::FromStart);
 		return PackageName;
 	}
-
+	
+	/** Returns assetname string, leaving off the /package/path part */
 	const FString GetAssetName() const
 	{
 		FString AssetName;
-		ToString().Split(TEXT("."), nullptr, &AssetName, ESearchCase::CaseSensitive, ESearchDir::FromEnd);
+		ToString().Split(TEXT("."), nullptr, &AssetName, ESearchCase::CaseSensitive, ESearchDir::FromStart);
 		return AssetName;
-	}
+	}	
 
-	/**
-	 * Sets asset path of this reference.
-	 *
-	 * @param Path The path to the asset.
-	 */
+	/** Sets asset path of this reference based on a string path */
 	void SetPath(FString Path);
 
 	/**
-	 * Attempts to load the asset.
-	 * @return Loaded UObject, or null if the asset fails to load, or if the reference is not valid.
+	 * Attempts to load the asset, this will call LoadObject which can be very slow
+	 *
+	 * @return Loaded UObject, or nullptr if the reference is null or the asset fails to load
 	 */
 	UObject* TryLoad() const;
 
 	/**
-	 * Attempts to find a currently loaded object that matches this object ID
-	 * @return Found UObject, or NULL if not currently loaded
+	 * Attempts to find a currently loaded object that matches this path
+	 *
+	 * @return Found UObject, or nullptr if not currently in memory
 	 */
 	UObject* ResolveObject() const;
 
-	/**
-	 * Resets reference to point to NULL
-	 */
+	/** Resets reference to point to null */
 	void Reset()
 	{		
 		SetPath(TEXT(""));
 	}
 	
-	/**
-	 * Check if this could possibly refer to a real object, or was initialized to NULL
-	 */
+	/** Check if this could possibly refer to a real object, or was initialized to null */
 	bool IsValid() const
 	{
 		return ToString().Len() > 0;
 	}
 
+	/** Checks to see if this is initialized to null */
+	bool IsNull() const
+	{
+		return ToString().Len() == 0;
+	}
+
+	/** Struct overrides */
 	bool Serialize(FArchive& Ar);
 	bool operator==(FStringAssetReference const& Other) const;
 	bool operator!=(FStringAssetReference const& Other) const
 	{
 		return !(*this == Other);
 	}
-	void operator=(FStringAssetReference const& Other);
+	FStringAssetReference& operator=(FStringAssetReference const& Other);
 	bool ExportTextItem(FString& ValueStr, FStringAssetReference const& DefaultValue, UObject* Parent, int32 PortFlags, UObject* ExportRootScope) const;
 	bool ImportTextItem( const TCHAR*& Buffer, int32 PortFlags, UObject* Parent, FOutputDevice* ErrorText );
 	bool SerializeFromMismatchedTag(struct FPropertyTag const& Tag, FArchive& Ar);
@@ -108,8 +117,7 @@ struct COREUOBJECT_API FStringAssetReference
 		return GetTypeHash(This.ToString());
 	}
 
-	/** Code needed by AssetPtr to track rather object references should be rechecked */
-
+	/** Code needed by FAssetPtr internals */
 	static int32 GetCurrentTag()
 	{
 		return CurrentTag.GetValue();
@@ -118,22 +126,22 @@ struct COREUOBJECT_API FStringAssetReference
 	{
 		return CurrentTag.Increment();
 	}
-
-	static FStringAssetReference GetOrCreateIDForObject(const UObject *Object);
-
+	static FStringAssetReference GetOrCreateIDForObject(const UObject* Object);
+	
+	/** Enables special PIE path handling, to remove pie prefix as needed */
 	static void SetPackageNamesBeingDuplicatedForPIE(const TArray<FString>& InPackageNamesBeingDuplicatedForPIE);
+	
+	/** Disables special PIE path handling */
 	static void ClearPackageNamesBeingDuplicatedForPIE();
 
-protected:
-
 private:
-	/**
-	 * Fixes up this StringAssetReference to add or remove the PIE prefix depending on what is currently active
-	 */
+	/** Fixes up this StringAssetReference to add or remove the PIE prefix depending on what is currently active */
 	void FixupForPIE();
 
-	/** Global counter that determines when we need to re-search for GUIDs because more objects have been loaded **/
+	/** Global counter that determines when we need to re-search based on path because more objects have been loaded **/
 	static FThreadSafeCounter CurrentTag;
+
+	/** Package names currently being duplicated, needed by FixupForPIE */
 	static TArray<FString> PackageNamesBeingDuplicatedForPIE;
 };
 
diff --git a/Engine/Source/Runtime/CoreUObject/Public/Serialization/AsyncLoading.h b/Engine/Source/Runtime/CoreUObject/Public/Serialization/AsyncLoading.h
--- a/Engine/Source/Runtime/CoreUObject/Public/Serialization/AsyncLoading.h
+++ b/Engine/Source/Runtime/CoreUObject/Public/Serialization/AsyncLoading.h
@@ -259,6 +259,7 @@ private:
 			check(Index < int32(EEventLoadNode::Package_NumPhases) && Index >= 0);
 			return PackageNodes[Index];
 		}
+		check(TotalNumberOfImportExportNodes); // otherwise Init was not called yet
 		if (Node.ImportOrExportIndex.IsImport())
 		{
 			check(int32(Node.Phase) < int32(EEventLoadNode::Import_NumPhases));
diff --git a/Engine/Source/Runtime/CoreUObject/Public/Templates/Casts.h b/Engine/Source/Runtime/CoreUObject/Public/Templates/Casts.h
--- a/Engine/Source/Runtime/CoreUObject/Public/Templates/Casts.h
+++ b/Engine/Source/Runtime/CoreUObject/Public/Templates/Casts.h
@@ -342,6 +342,7 @@ DECLARE_CAST_BY_FLAG(USkeletalMeshComponent)			\
 DECLARE_CAST_BY_FLAG(UBlueprint)						\
 DECLARE_CAST_BY_FLAG(UDelegateFunction)					\
 DECLARE_CAST_BY_FLAG(UStaticMeshComponent)				\
+DECLARE_CAST_BY_FLAG(UEnumProperty)						\
 FINISH_DECLARING_CAST_FLAGS		// This is here to hopefully remind people to include the "\" in all declarations above, especially when copy/pasting the final line.
 
 // Now actually declare the flags
diff --git a/Engine/Source/Runtime/CoreUObject/Public/Templates/SubclassOf.h b/Engine/Source/Runtime/CoreUObject/Public/Templates/SubclassOf.h
--- a/Engine/Source/Runtime/CoreUObject/Public/Templates/SubclassOf.h
+++ b/Engine/Source/Runtime/CoreUObject/Public/Templates/SubclassOf.h
@@ -5,88 +5,86 @@
 #include "Class.h"
 
 /**
-* Template to allow UClass's to be passed around with type safety
-**/
+ * Template to allow UClass's to be passed around with type safety 
+ */
 template<class TClass>
 class TSubclassOf
 {
+	template <class TClassA>
+	friend class TSubclassOf;
+
 public:
-	/**
-	* Default Constructor, defaults to NULL
-	**/
+	/** Default Constructor, defaults to null */
 	FORCEINLINE TSubclassOf() :
-		Class(NULL)
+		Class(nullptr)
 	{
 	}
-	/**
-	* Constructor that takes a UClass and does a runtime check to make sure this is a compatible class
-	* @param From UClass to assign to this TSubclassOf
-	**/
+
+	/** Constructor that takes a UClass and does a runtime check to make sure this is a compatible class */
 	FORCEINLINE TSubclassOf(UClass* From) :
 		Class(From)
 	{
 	}
-	/**
-	* Copy Constructor
-	* @param From TSubclassOf to copy
-	**/
-	template<class TClassA>
+
+	/** Copy Constructor, will only compile if types are compatible */
+	template <class TClassA, class = typename TEnableIf<TPointerIsConvertibleFromTo<TClassA, TClass>::Value>::Type>
 	FORCEINLINE TSubclassOf(const TSubclassOf<TClassA>& From) :
 		Class(*From)
 	{
-		TClass* Test = (TClassA*)0; // verify that TClassA derives from TClass
 	}
 
-	/**
-	* Assignment operator
-	* @param From TSubclassOf to copy
-	**/
-	template<class TClassA>
+	/** Assignment operator, will only compile if types are compatible */
+	template <class TClassA, class = typename TEnableIf<TPointerIsConvertibleFromTo<TClassA, TClass>::Value>::Type>
 	FORCEINLINE TSubclassOf& operator=(const TSubclassOf<TClassA>& From)
 	{
 		Class = *From;
-		TClass* Test = (TClassA*)0; // verify that TClassA derives from TClass
 		return *this;
 	}
-	/**
-	* Assignment operator from UClass, runtime check for compatibility
-	* @param From UClass to assign to this TSubclassOf
-	**/
+	
+	/** Assignment operator from UClass, the type is checked on get not on set */
 	FORCEINLINE TSubclassOf& operator=(UClass* From)
 	{
 		Class = From;
 		return *this;
 	}
-	/**
-	* Dereference back into a UClass
-	* @return	the embedded UClass
-	* The body for this method is in Class.h
-	*/
-	FORCEINLINE UClass* operator*() const;
+	
+	/** Dereference back into a UClass, does runtime type checking */
+	FORCEINLINE UClass* operator*() const
+	{
+		if (!Class || !Class->IsChildOf(TClass::StaticClass()))
+		{
+			return nullptr;
+		}
+		return Class;
+	}
+	
+	/** Dereference back into a UClass */
+	FORCEINLINE UClass* Get() const
+	{
+		return **this;
+	}
 
-	/**
-	* Dereference back into a UClass
-	* @return	the embedded UClass
-	*/
+	/** Dereference back into a UClass */
 	FORCEINLINE UClass* operator->() const
 	{
 		return **this;
 	}
 
-	/**
-	* Implicit conversion to UClass
-	* @return	the embedded UClass
-	*/
+	/** Implicit conversion to UClass */
 	FORCEINLINE operator UClass* () const
 	{
 		return **this;
 	}
 
 	/**
-	* Get the CDO if we are referencing a valid class
-	* @return the CDO, or NULL
-	*/
-	FORCEINLINE TClass* GetDefaultObject() const;
+	 * Get the CDO if we are referencing a valid class
+	 *
+	 * @return the CDO, or null if class is null
+	 */
+	FORCEINLINE TClass* GetDefaultObject() const
+	{
+		return Class ? Class->GetDefaultObject<TClass>() : nullptr;
+	}
 
 	friend FArchive& operator<<(FArchive& Ar, TSubclassOf& SubclassOf)
 	{
@@ -96,23 +94,3 @@ public:
 private:
 	UClass* Class;
 };
-
-/**
-* Dereference back into a UClass
-* @return	the embedded UClass
-*/
-template<class TClass>
-FORCEINLINE UClass* TSubclassOf<TClass>::operator*() const
-{
-	if (!Class || !Class->IsChildOf(TClass::StaticClass()))
-	{
-		return NULL;
-	}
-	return Class;
-}
-
-template<class TClass>
-FORCEINLINE TClass* TSubclassOf<TClass>::GetDefaultObject() const
-{
-	return Class ? Class->GetDefaultObject<TClass>() : NULL;
-}
diff --git a/Engine/Source/Runtime/CoreUObject/Public/UObject/AssetPtr.h b/Engine/Source/Runtime/CoreUObject/Public/UObject/AssetPtr.h
--- a/Engine/Source/Runtime/CoreUObject/Public/UObject/AssetPtr.h
+++ b/Engine/Source/Runtime/CoreUObject/Public/UObject/AssetPtr.h
@@ -1,84 +1,91 @@
 // Copyright 1998-2016 Epic Games, Inc. All Rights Reserved.
 
 /*=============================================================================
-	AssetPtr.h: Pointer to UObject asset, keeps extra information so that it is robust to missing assets.
+	AssetPtr.h: Pointer to UObject asset, keeps extra information so that it is works even if the asset is not in memory
 =============================================================================*/
 
 #pragma once
 
-#include "AutoPointer.h"
 #include "StringAssetReference.h"
 
-/***
- * FAssetPtr is a type of weak pointer to a UObject.
- * It will change back and forth between being valid or pending as the referenced object loads or unloads
+/**
+ * FAssetPtr is a type of weak pointer to a UObject, that also keeps track of the path to the object on disk.
+ * It will change back and forth between being Valid and Pending as the referenced object loads or unloads.
  * It has no impact on if the object is garbage collected or not.
- * It can't be directly used across a network.
- * Implicit conversion is disabled because that can cause memory leaks by reserving IDs that will never be used.
  *
- * This is useful for cross level references or any place you need to declare a reference to an object that may not be loaded
- **/
+ * This is useful to specify assets that you may want to asynchronously load on demand.
+ */
 class FAssetPtr : public TPersistentObjectPtr<FStringAssetReference>
 {
 public:	
-	/** NULL constructor **/
+	/** Default constructor, will be null */
 	FORCEINLINE FAssetPtr()
 	{
 	}
 
-	/**  
-	 * Construct from another lazy pointer
-	 * @param Other lazy pointer to copy from
-	 */
-	FORCEINLINE FAssetPtr(const FAssetPtr &Other)
+	/** Construct from another asset pointer */
+	FORCEINLINE FAssetPtr(const FAssetPtr& Other)
 	{
 		(*this)=Other;
 	}
 
-	/**  
-	 * Construct from a unique object identifier
-	 * @param ObjectID Object identifier to create a weak pointer to
-	**/
-	explicit FORCEINLINE FAssetPtr(const FStringAssetReference &ObjectID)
-		: TPersistentObjectPtr<FStringAssetReference>(ObjectID)
+	/** Construct from a string asset reference */
+	explicit FORCEINLINE FAssetPtr(const FStringAssetReference& AssetRef)
+		: TPersistentObjectPtr<FStringAssetReference>(AssetRef)
 	{
 	}
 
-	/**  
-	 * Construct from another lazy pointer
-	 * @param Other lazy pointer to copy from
-	 */
-	explicit FORCEINLINE FAssetPtr(const UObject *Object)
+	/** Construct from an object already in memory */
+	explicit FORCEINLINE FAssetPtr(const UObject* Object)
 	{
 		(*this)=Object;
 	}
 
-	/**  
-	 * Synchronously load (if necessary) and return the asset object represented by this asset ptr
-	 */
+	/** Synchronously load (if necessary) and return the asset object represented by this asset ptr */
 	UObject* LoadSynchronous()
 	{
 		UObject* Asset = Get();
 		if (Asset == nullptr && IsPending())
 		{
-			Asset = GetUniqueID().TryLoad();
+			Asset = ToStringReference().TryLoad();
 			*this = Asset;
 		}
 		return Asset;
 	}
 
+	/** Returns the StringAssetReference that is wrapped by this AssetPtr */
+	FORCEINLINE const FStringAssetReference& ToStringReference() const
+	{
+		return GetUniqueID();
+	}
+
+	/** Returns string representation of reference, in form /package/path.assetname */
+	FORCEINLINE const FString& ToString() const
+	{
+		return ToStringReference().ToString();
+	}
+
+	/** Returns /package/path string, leaving off the asset name */
+	FORCEINLINE FString GetLongPackageName() const
+	{
+		return ToStringReference().GetLongPackageName();
+	}
+	
+	/** Returns assetname string, leaving off the /package/path. part */
+	FORCEINLINE FString GetAssetName() const
+	{
+		return ToStringReference().GetAssetName();
+	}
+
 	using TPersistentObjectPtr<FStringAssetReference>::operator=;
 };
 
 template <> struct TIsPODType<FAssetPtr> { enum { Value = TIsPODType<TPersistentObjectPtr<FStringAssetReference> >::Value }; };
 template <> struct TIsWeakPointerType<FAssetPtr> { enum { Value = TIsWeakPointerType<TPersistentObjectPtr<FStringAssetReference> >::Value }; };
 
-/***
-* 
-* TAssetPtr is templatized version of the generic FAssetPtr
-* 
-**/
-
+/**
+ * TAssetPtr is templatized wrapper of the generic FAssetPtr, it can be used in UProperties
+ */
 template<class T=UObject>
 class TAssetPtr
 {
@@ -86,91 +93,71 @@ class TAssetPtr
 	friend class TAssetPtr;
 
 public:
-	/** NULL constructor **/
+	/** Default constructor, will be null */
 	FORCEINLINE TAssetPtr()
 	{
 	}
 	
-	/**  
-	 * Construct from another lazy pointer
-	 * @param Other lazy pointer to copy from
-	 */
+	/** Construct from another asset pointer */
 	template <class U, class = typename TEnableIf<TPointerIsConvertibleFromTo<U, T>::Value>::Type>
 	FORCEINLINE TAssetPtr(const TAssetPtr<U>& Other)
 		: AssetPtr(Other.AssetPtr)
 	{
 	}
 
-	/**  
-	 * Construct from an object pointer
-	 * @param Object object to create a reference to
-	**/
+	/** Construct from an object already in memory */
 	FORCEINLINE TAssetPtr(const T* Object)
 		: AssetPtr(Object)
 	{
 	}
 
-	/**  
-	 * Construct from a unique object identifier
-	 * @param ObjectID Object identifier to create a weak pointer to
-	**/
-	explicit FORCEINLINE TAssetPtr(const FStringAssetReference &ObjectID)
-		: AssetPtr(ObjectID)
+	/** Construct from a string asset reference */
+	explicit FORCEINLINE TAssetPtr(const FStringAssetReference& AssetRef)
+		: AssetPtr(AssetRef)
 	{
 	}
 
-	/**
-	 * Reset the lazy pointer back to the NULL state
-	 */
+	/** Reset the asset pointer back to the null state */
 	FORCEINLINE void Reset()
 	{
 		AssetPtr.Reset();
 	}
 
-	/**  
-	 * Copy from an object pointer
-	 * @param Object Object to create a lazy pointer to
-	 */
-	FORCEINLINE void operator=(const T* Object)
+	/** Copy from an object already in memory */
+	FORCEINLINE TAssetPtr& operator=(const T* Object)
 	{
 		AssetPtr = Object;
+		return *this;
 	}
 
-	/**  
-	 * Copy from a unique object identifier
-	 * @param ObjectID Object identifier to create a weak pointer to
-	 */
-	FORCEINLINE void operator=(const FStringAssetReference &ObjectID)
+	/** Copy from a string asset reference */
+	FORCEINLINE TAssetPtr& operator=(const FStringAssetReference& AssetRef)
 	{
-		AssetPtr = ObjectID;
+		AssetPtr = AssetRef;
+		return *this;
 	}
 
-	/**  
-	 * Copy from a weak pointer
-	 * @param Other Weak pointer to copy from
-	 */
-	template <class U>
-	FORCEINLINE typename TEnableIf<TPointerIsConvertibleFromTo<U, T>::Value>::Type
-		operator=(const TWeakObjectPtr<U>& Other)
+	/** Copy from a weak pointer to an object already in memory */
+	template <class U, class = typename TEnableIf<TPointerIsConvertibleFromTo<U, T>::Value>::Type>
+	FORCEINLINE TAssetPtr& operator=(const TWeakObjectPtr<U>& Other)
 	{
 		AssetPtr = Other;
+		return *this;
 	}
 
-	/**  
-	 * Construct from another lazy pointer
-	 * @param Other lazy pointer to copy from
-	 */
-	template <class U>
-	FORCEINLINE typename TEnableIf<TPointerIsConvertibleFromTo<U, T>::Value>::Type
-		operator=(const TAssetPtr<U>& Other)
+	/** Copy from another asset pointer */
+	template <class U, class = typename TEnableIf<TPointerIsConvertibleFromTo<U, T>::Value>::Type>
+	FORCEINLINE TAssetPtr& operator=(const TAssetPtr<U>& Other)
 	{
 		AssetPtr = Other.AssetPtr;
+		return *this;
 	}
 
 	/**  
-	 * Compare lazy pointers for equality
-	 * @param Other lazy pointer to compare to
-	 * Caution: Two lazy pointers might not be equal to each other, but they both might return NULL.
+	 * Compare asset pointers for equality
+	 * Caution: Two asset pointers might not be equal to each other, but they both might return nullptr
+	 *
+	 * @param Other asset pointer to compare to 
 	 */
 	FORCEINLINE friend bool operator==(const TAssetPtr& Lhs, const TAssetPtr& Rhs)
 	{
@@ -178,73 +165,59 @@ public:
 	}
 
 	/**  
-	 * Compare lazy pointers for inequality
-	 * @param Other lazy pointer to compare to
-	 * Caution: Two lazy pointers might not be equal to each other, but they both might return NULL.
+	 * Compare asset pointers for inequality
+	 * Caution: Two asset pointers might not be equal to each other, but they both might return nullptr
+	 *
+	 * @param Other asset pointer to compare to
 	 */
 	FORCEINLINE friend bool operator!=(const TAssetPtr& Lhs, const TAssetPtr& Rhs)
 	{
 		return Lhs.AssetPtr != Rhs.AssetPtr;
 	}
 
-	/**
-	 * Gets the unique object identifier associated with this lazy pointer. Valid even if pointer is not currently valid
-	 * @return Unique ID for this object, or an invalid FStringAssetReference if this pointer isn't set to anything
-	 */
-	FStringAssetReference GetUniqueID() const
-	{
-		return AssetPtr.GetUniqueID();
-	}
-
 	/**  
-	 * Dereference the lazy pointer.
-	 * @return NULL if this object is gone or the lazy pointer was NULL, otherwise a valid uobject pointer
+	 * Dereference the asset pointer.
+	 *
+	 * @return nullptr if this object is gone or the lazy pointer was null, otherwise a valid UObject pointer
 	 */
 	FORCEINLINE T* Get() const
 	{
 		return dynamic_cast<T*>(AssetPtr.Get());
 	}
 
-	/**  
-	 * Dereference the lazy pointer.
-	 */
+	/** Dereference the asset pointer */
 	FORCEINLINE T& operator*() const
 	{
 		return *Get();
 	}
 
-	/**  
-	 * Dereference the lazy pointer.
-	 */
+	/** Dereference the asset pointer */
 	FORCEINLINE T* operator->() const
 	{
 		return Get();
 	}
 
-	/**  
-	 * Synchronously load (if necessary) and return the asset object represented by this asset ptr
-	 */
+	/** Synchronously load (if necessary) and return the asset object represented by this asset ptr */
 	T* LoadSynchronous()
 	{
-		UObject* Asset = AssetPtr.Get();
-		if (Asset == nullptr && IsPending())
-		{
-			Asset = AssetPtr.LoadSynchronous();
-		}
+		UObject* Asset = AssetPtr.LoadSynchronous();
 		return Cast<T>(Asset);
 	}
 
 	/**  
 	 * Test if this points to a live UObject
+	 *
 	 * @return true if Get() would return a valid non-null pointer
 	 */
 	FORCEINLINE bool IsValid() const
 	{
-		return AssetPtr.IsValid();
+		// This does the runtime type check
+		return Get() != nullptr;
 	}
 
 	/**  
 	 * Test if this does not point to a live UObject, but may in the future
+	 *
 	 * @return true if this does not point to a real object, but could possibly
 	 */
 	FORCEINLINE bool IsPending() const
@@ -254,6 +227,7 @@ public:
 
 	/**  
 	 * Test if this can never point to a live UObject
+	 *
 	 * @return true if this is explicitly pointing to no object
 	 */
 	FORCEINLINE bool IsNull() const
@@ -261,38 +235,43 @@ public:
 		return AssetPtr.IsNull();
 	}
 
-	/**  
-	 * Returns the StringAssetReference that is wrapped by this AssetPtr
-	 */
-	FORCEINLINE const FStringAssetReference& ToStringReference() const
+	/** Returns the StringAssetReference that is wrapped by this AssetPtr */
+	FORCEINLINE const FStringAssetReference& GetUniqueID() const
 	{
 		return AssetPtr.GetUniqueID();
 	}
 
+	/** Returns the StringAssetReference that is wrapped by this AssetPtr */
+	FORCEINLINE const FStringAssetReference& ToStringReference() const
+	{
+		return AssetPtr.ToStringReference();
+	}
+
+	/** Returns string representation of reference, in form /package/path.assetname */
 	FORCEINLINE const FString& ToString() const
 	{
 		return ToStringReference().ToString();
 	}
 
+	/** Returns /package/path string, leaving off the asset name */
 	FORCEINLINE FString GetLongPackageName() const
 	{
 		return ToStringReference().GetLongPackageName();
 	}
-
+	
+	/** Returns assetname string, leaving off the /package/path part */
 	FORCEINLINE FString GetAssetName() const
 	{
 		return ToStringReference().GetAssetName();
 	}
 
-	/**  
-	 * Dereference lazy pointer to see if it points somewhere valid.
-	 */
+	/** Dereference asset pointer to see if it points somewhere valid */
 	FORCEINLINE explicit operator bool() const
 	{
-		return Get() != NULL;
+		return IsValid();
 	}
 
-	/** Hash function. */
+	/** Hash function */
 	FORCEINLINE friend uint32 GetTypeHash(const TAssetPtr<T>& Other)
 	{
 		return GetTypeHash(static_cast<const TPersistentObjectPtr<FStringAssetReference>&>(Other.AssetPtr));
@@ -305,106 +284,85 @@ public:
 	}
 
 private:
-	/** Wrapper of a UObject* */
 	FAssetPtr AssetPtr;
 };
 
 template<class T> struct TIsPODType<TAssetPtr<T> > { enum { Value = TIsPODType<FAssetPtr>::Value }; };
-template <class T> struct TIsWeakPointerType<TAssetPtr<T> > { enum { Value = TIsWeakPointerType<FAssetPtr>::Value }; };
+template<class T> struct TIsWeakPointerType<TAssetPtr<T> > { enum { Value = TIsWeakPointerType<FAssetPtr>::Value }; };
 
+/**
+ * TAssetSubclassOf is a templatized wrapper around FAssetPtr that works like a TSubclassOf, it can be used in UProperties for blueprint subclasses
+ */
 template<class TClass>
 class TAssetSubclassOf
 {
-	template <class U>
+	template <class TClassA>
 	friend class TAssetSubclassOf;
 
 public:
-	/** NULL constructor **/
+	/** Default constructor, will be null */
 	FORCEINLINE TAssetSubclassOf()
 	{
 	}
-	
-	/**  
-	 * Construct from another lazy pointer
-	 * @param Other lazy pointer to copy from
-	 */
-	template<class TClassA>
+		
+	/** Construct from another asset pointer */
+	template <class TClassA, class = typename TEnableIf<TPointerIsConvertibleFromTo<TClassA, TClass>::Value>::Type>
 	FORCEINLINE TAssetSubclassOf(const TAssetSubclassOf<TClassA>& Other)
 		: AssetPtr(Other.AssetPtr)
 	{
-		TClass* Test = (TClassA*)0; // verify that TClassA derives from TClass
 	}
 
-	/**  
-	 * Construct from an class pointer
-	 * @param From Class to create a reference to
-	**/
+	/** Construct from a class already in memory */
 	FORCEINLINE TAssetSubclassOf(const UClass* From)
 		: AssetPtr(From)
 	{
 	}
 
-	/**  
-	 * Construct from a unique object identifier
-	 * @param ObjectID Object identifier to create a weak pointer to
-	**/
-	explicit FORCEINLINE TAssetSubclassOf(const FStringAssetReference &ObjectID)
-		: AssetPtr(ObjectID)
+	/** Construct from a string asset reference */
+	explicit FORCEINLINE TAssetSubclassOf(const FStringAssetReference& AssetRef)
+		: AssetPtr(AssetRef)
 	{
 	}
 
-	/**
-	 * Reset the lazy pointer back to the NULL state
-	 */
+	/** Reset the asset pointer back to the null state */
 	FORCEINLINE void Reset()
 	{
 		AssetPtr.Reset();
 	}
 
-	/**  
-	 * Copy from an object pointer
-	 * @param Object Object to create a lazy pointer to
-	 */
+	/** Copy from a class already in memory */
 	FORCEINLINE void operator=(const UClass* From)
 	{
 		AssetPtr = From;
 	}
 
-	/**  
-	 * Copy from a unique object identifier
-	 * @param ObjectID Object identifier to create a weak pointer to
-	 */
-	FORCEINLINE void operator=(const FStringAssetReference &ObjectID)
+	/** Copy from a string asset reference */
+	FORCEINLINE void operator=(const FStringAssetReference& AssetRef)
 	{
-		AssetPtr = ObjectID;
+		AssetPtr = AssetRef;
 	}
 
-	/**  
-	 * Copy from a weak pointer
-	 * @param Other Weak pointer to copy from
-	 */
-	template<class TClassA>
-	FORCEINLINE void operator=(const TWeakObjectPtr<TClassA> &Other)
+	/** Copy from a weak pointer already in memory */
+	template<class TClassA, class = typename TEnableIf<TPointerIsConvertibleFromTo<TClassA, UClass>::Value>::Type>
+	FORCEINLINE TAssetSubclassOf& operator=(const TWeakObjectPtr<TClassA>& Other)
 	{
-		UClass* Test = (TClassA*)0; //verify that TClassA derives from UClass
 		AssetPtr = Other;
+		return *this;
 	}
 
-	/**  
-	 * Construct from another lazy pointer
-	 * @param Other lazy pointer to copy from
-	 */
-	template<class TClassA>
-	FORCEINLINE void operator=(const TAssetSubclassOf<TClassA> &Other)
+	/** Copy from another asset pointer */
+	template<class TClassA, class = typename TEnableIf<TPointerIsConvertibleFromTo<TClassA, TClass>::Value>::Type>
+	FORCEINLINE TAssetSubclassOf& operator=(const TAssetPtr<TClassA>& Other)
 	{
-		TClass* Test = (TClassA*)0; // verify that TClassA derives from TClass
-		AssetPtr = Other;
+		AssetPtr = Other.AssetPtr;
+		return *this;
 	}
 
 	/**  
-	 * Compare lazy pointers for equality
-	 * @param Other lazy pointer to compare to
-	 * Caution: Two lazy pointers might not be equal to each other, but they both might return NULL.
+	 * Compare asset pointers for equality
+	 * Caution: Two asset pointers might not be equal to each other, but they both might return nullptr
+	 *
+	 * @param Other asset pointer to compare to 
 	 */
 	FORCEINLINE friend bool operator==(const TAssetSubclassOf& Lhs, const TAssetSubclassOf& Rhs)
 	{
@@ -412,49 +370,38 @@ public:
 	}
 
 	/**  
-	 * Compare lazy pointers for inequality
-	 * @param Other lazy pointer to compare to
-	 * Caution: Two lazy pointers might not be equal to each other, but they both might return NULL.
+	 * Compare asset pointers for inequality
+	 * Caution: Two asset pointers might not be equal to each other, but they both might return nullptr
+	 *
+	 * @param Other asset pointer to compare to
 	 */
 	FORCEINLINE friend bool operator!=(const TAssetSubclassOf& Lhs, const TAssetSubclassOf& Rhs)
 	{
 		return Lhs.AssetPtr != Rhs.AssetPtr;
 	}
 
-	/**
-	 * Gets the unique object identifier associated with this lazy pointer. Valid even if pointer is not currently valid
-	 * @return Unique ID for this object, or an invalid FStringAssetReference if this pointer isn't set to anything
-	 */
-	FStringAssetReference GetUniqueID() const
-	{
-		return AssetPtr.GetUniqueID();
-	}
-
 	/**  
-	 * Dereference the lazy pointer.
-	 * @return NULL if this object is gone or the lazy pointer was NULL, otherwise a valid uobject pointer
+	 * Dereference the asset pointer
+	 *
+	 * @return nullptr if this object is gone or the asset pointer was null, otherwise a valid UClass pointer
 	 */
-	FORCEINLINE UClass *Get() const
+	FORCEINLINE UClass* Get() const
 	{
 		UClass* Class = dynamic_cast<UClass*>(AssetPtr.Get());
 		if (!Class || !Class->IsChildOf(TClass::StaticClass()))
 		{
-			return NULL;
+			return nullptr;
 		}
 		return Class;
 	}
 
-	/**  
-	 * Dereference the lazy pointer.
-	 */
+	/** Dereference the asset pointer */
 	FORCEINLINE UClass& operator*() const
 	{
 		return *Get();
 	}
 
-	/**  
-	 * Dereference the lazy pointer.
-	 */
+	/** Dereference the asset pointer */
 	FORCEINLINE UClass* operator->() const
 	{
 		return Get();
@@ -462,15 +409,18 @@ public:
 
 	/**  
 	 * Test if this points to a live UObject
+	 *
 	 * @return true if Get() would return a valid non-null pointer
 	 */
 	FORCEINLINE bool IsValid() const
 	{
-		return (NULL != Get());
+		// This also does the UClass type check
+		return Get() != nullptr;
 	}
 
 	/**  
 	 * Test if this does not point to a live UObject, but may in the future
+	 *
 	 * @return true if this does not point to a real object, but could possibly
 	 */
 	FORCEINLINE bool IsPending() const
@@ -480,6 +430,7 @@ public:
 
 	/**  
 	 * Test if this can never point to a live UObject
+	 *
 	 * @return true if this is explicitly pointing to no object
 	 */
 	FORCEINLINE bool IsNull() const
@@ -487,38 +438,52 @@ public:
 		return AssetPtr.IsNull();
 	}
 
-	/**  
-	 * Returns the StringAssetReference that is wrapped by this AssetPtr
-	 */
+	/** Returns the StringAssetReference that is wrapped by this AssetPtr */
+	FORCEINLINE const FStringAssetReference& GetUniqueID() const
+	{
+		return AssetPtr.GetUniqueID();
+	}
+
+	/** Returns the StringAssetReference that is wrapped by this AssetPtr */
 	FORCEINLINE const FStringAssetReference& ToStringReference() const
 	{
 		return AssetPtr.GetUniqueID();
 	}
 
-	/**  
-	 * Dereference lazy pointer to see if it points somewhere valid.
-	 */
+	/** Returns string representation of reference, in form /package/path.assetname  */
+	FORCEINLINE const FString& ToString() const
+	{
+		return ToStringReference().ToString();
+	}
+
+	/** Returns /package/path string, leaving off the asset name */
+	FORCEINLINE FString GetLongPackageName() const
+	{
+		return ToStringReference().GetLongPackageName();
+	}
+	
+	/** Returns assetname string, leaving off the /package/path part */
+	FORCEINLINE FString GetAssetName() const
+	{
+		return ToStringReference().GetAssetName();
+	}
+
+	/** Dereference asset pointer to see if it points somewhere valid */
 	FORCEINLINE explicit operator bool() const
 	{
-		return Get() != NULL;
+		return IsValid();
 	}
 
-	/** Hash function. */
+	/** Hash function */
 	FORCEINLINE friend uint32 GetTypeHash(const TAssetSubclassOf<TClass>& Other)
 	{
 		return GetTypeHash(static_cast<const TPersistentObjectPtr<FStringAssetReference>&>(Other.AssetPtr));
 	}
 
-	/**  
-	 * Synchronously load (if necessary) and return the asset object represented by this asset ptr
-	 */
+	/** Synchronously load (if necessary) and return the asset object represented by this asset ptr */
 	UClass* LoadSynchronous()
 	{
-		UObject* Asset = AssetPtr.Get();
-		if (Asset == nullptr && IsPending())
-		{
-			Asset = AssetPtr.LoadSynchronous();
-		}
+		UObject* Asset = AssetPtr.LoadSynchronous();
 		UClass* Class = dynamic_cast<UClass*>(Asset);
 		if (!Class || !Class->IsChildOf(TClass::StaticClass()))
 		{
@@ -534,9 +499,8 @@ public:
 	}
 
 private:
-	/** Wrapper of a UObject* */
 	FAssetPtr AssetPtr;
 };
 
-template<class T> struct TIsPODType<TAssetSubclassOf<T> > { enum { Value = TIsPODType<FAssetPtr>::Value }; };
+template <class T> struct TIsPODType<TAssetSubclassOf<T> > { enum { Value = TIsPODType<FAssetPtr>::Value }; };
 template <class T> struct TIsWeakPointerType<TAssetSubclassOf<T> > { enum { Value = TIsWeakPointerType<FAssetPtr>::Value }; };
diff --git a/Engine/Source/Runtime/CoreUObject/Public/UObject/Class.h b/Engine/Source/Runtime/CoreUObject/Public/UObject/Class.h
--- a/Engine/Source/Runtime/CoreUObject/Public/UObject/Class.h
+++ b/Engine/Source/Runtime/CoreUObject/Public/UObject/Class.h
@@ -1414,25 +1414,25 @@ public:
 	FString CppType;
 
 	/** Gets enum name by index in Names array. Returns NAME_None if Index is not valid. */
-	FName GetNameByIndex(uint8 Index) const;
-		
+	FName GetNameByIndex(int32 Index) const;
+
 	/** Gets enum value by index in Names array. */
-	uint8 GetValueByIndex(uint8 Index) const;
+	int64 GetValueByIndex(int32 Index) const;
 
 	/** Gets enum name by value. Returns NAME_None if value is not found. */
-	FName GetNameByValue(uint8 InValue) const;
-		
+	FName GetNameByValue(int64 InValue) const;
+
 	/** Gets enum value by name. Returns INDEX_NONE when name is not found. */
-	int32 GetValueByName(FName InName) const;
+	int64 GetValueByName(FName InName) const;
 
 	/** Gets index of name in enum. Returns INDEX_NONE when name is not found. */
 	int32 GetIndexByName(FName InName) const;
 
 	/** Gets max value of Enum. Defaults to zero if there are no entries. */
-	uint8 GetMaxEnumValue() const;
+	int64 GetMaxEnumValue() const;
 
 	/** Checks if enum has entry with given value. Includes autogenerated _MAX entry. */
-	bool IsValidEnumValue(uint8 InValue) const;
+	bool IsValidEnumValue(int64 InValue) const;
 
 	/** Checks if enum has entry with given name. Includes autogenerated _MAX entry. */
 	bool IsValidEnumName(FName InName) const;
@@ -1440,7 +1440,7 @@ public:
 protected:
 	// Variables.
 	/** List of pairs of all enum names and values. */
-	TArray<TPair<FName, uint8>> Names;
+	TArray<TPair<FName, int64>> Names;
 
 	/** How the enum was originally defined. */
 	ECppForm CppForm;
@@ -1466,7 +1466,7 @@ public:
 	/*
 	 *	Try to update an out-of-date enum index after an enum's change
 	 */
-	virtual int32 ResolveEnumerator(FArchive& Ar, int32 EnumeratorIndex) const;
+	virtual int64 ResolveEnumerator(FArchive& Ar, int64 EnumeratorIndex) const;
 
 	/**
 	 * Returns the type of enum: whether it's a regular enum, namespaced enum or C++11 enum class.
@@ -1509,7 +1509,7 @@ public:
 	/** searches the list of all enum value names for the specified name
 	 * @return the value the specified name represents if found, otherwise INDEX_NONE
 	 */
-	static int32 LookupEnumName(FName TestName, UEnum** FoundEnum = nullptr)
+	static int64 LookupEnumName(FName TestName, UEnum** FoundEnum = nullptr)
 	{
 		UEnum* TheEnum = AllEnumNames.FindRef(TestName);
 		if (FoundEnum != nullptr)
@@ -1522,10 +1522,10 @@ public:
 	/** searches the list of all enum value names for the specified name
 	 * @return the value the specified name represents if found, otherwise INDEX_NONE
 	 */
-	static int32 LookupEnumNameSlow(const TCHAR* InTestShortName, UEnum** FoundEnum = NULL)
+	static int64 LookupEnumNameSlow(const TCHAR* InTestShortName, UEnum** FoundEnum = NULL)
 	{
-		int32 EnumIndex = LookupEnumName(InTestShortName, FoundEnum);
-		if (EnumIndex == INDEX_NONE)
+		int64 Result = LookupEnumName(InTestShortName, FoundEnum);
+		if (Result == INDEX_NONE)
 		{
 			FString TestShortName = FString(TEXT("::")) + InTestShortName;
 			UEnum* TheEnum = NULL;
@@ -1540,16 +1540,16 @@ public:
 			{
 				*FoundEnum = TheEnum;
 			}
-			EnumIndex = (TheEnum != NULL) ? TheEnum->GetValueByName(InTestShortName) : INDEX_NONE;
+			Result = (TheEnum != NULL) ? TheEnum->GetValueByName(InTestShortName) : INDEX_NONE;
 		}
-		return EnumIndex;
+		return Result;
 	}
 
 	/** parses the passed in string for a name, then searches for that name in any Enum (in any package)
 	 * @param Str	pointer to string to parse; if we successfully find an enum, this pointer is advanced past the name found
 	 * @return index of the value the parsed enum name matches, or INDEX_NONE if no matches
 	 */
-	static int32 ParseEnum(const TCHAR*& Str);
+	static int64 ParseEnum(const TCHAR*& Str);
 
 	/**
 	 * Sets the array of enums.
@@ -1559,7 +1559,7 @@ public:
 	 * @param bAddMaxKeyIfMissing Should a default Max item be added.
 	 * @return	true unless the MAX enum already exists and isn't the last enum.
 	 */
-	virtual bool SetEnums(TArray<TPair<FName, uint8>>& InNames, ECppForm InCppForm, bool bAddMaxKeyIfMissing = true);
+	virtual bool SetEnums(TArray<TPair<FName, int64>>& InNames, ECppForm InCppForm, bool bAddMaxKeyIfMissing = true);
 
 	/**
 	 * @return	The enum name at the specified Index.
@@ -1573,7 +1573,7 @@ public:
 		return NAME_None;
 	}
 
-	int32 GetIndexByValue(int32 Value) const
+	int32 GetIndexByValue(int64 Value) const
 	{
 		for (int32 i = 0; i < Names.Num(); ++i)
 		{
@@ -1585,7 +1585,7 @@ public:
 
 		return INDEX_NONE;
 	}
-	FString GetEnumNameStringByValue(int32 Value) const
+	FString GetEnumNameStringByValue(int64 Value) const
 	{
 		int32 Index = GetIndexByValue(Value);
 		return GetEnumName(Index);
@@ -1596,7 +1596,7 @@ public:
 	 */
 	FString GetEnumName(int32 InIndex) const;
 
-	FText GetEnumTextByValue(int32 Value)
+	FText GetEnumTextByValue(int64 Value)
 	{
 		int32 Index = GetIndexByValue(Value);
 		return GetEnumText(Index);
@@ -1637,7 +1637,7 @@ public:
 	 * @return The display name for this object.
 	 */
 	FText GetDisplayNameText(int32 NameIndex=INDEX_NONE) const;
-	FText GetDisplayNameTextByValue(int32 Value = INDEX_NONE) const;
+	FText GetDisplayNameTextByValue(int64 Value = INDEX_NONE) const;
 
 	/**
 	 * Finds the localized tooltip or native tooltip as a fallback.
@@ -1700,55 +1700,61 @@ public:
 
 
 	/**
-	 * @param EnumPath	- Full enum path
-	 * @param EnumValue - Enum value
+	 * @param EnumPath         Full enum path.
+	 * @param EnumeratorIndex  Enumerator index.
+	 *
+	 * Note: Despite the name, this function actually uses the enumerator at a given index.  If the
+	 *       enum is not zero-based or is non-consecutive, this will give surprising results.
 	 *
-	 * @return the string associated with the specified enum value for the enum specified by a path
+	 * @return the string associated with the enumerator at a given index for the enum specified by a path.
 	 */
 	template <typename T>
-	FORCEINLINE static FString GetValueAsString( const TCHAR* EnumPath, const T EnumValue )
+	FORCEINLINE static FString GetValueAsString( const TCHAR* EnumPath, const T EnumeratorIndex )
 	{
 		// For the C++ enum.
 		static_assert(TIsEnum<T>::Value, "Should only call this with enum types");
-		return GetValueAsString_Internal(EnumPath, (int32)EnumValue);
+		return GetIndexAsString_Internal(EnumPath, (int32)EnumeratorIndex);
 	}
 
 	template <typename T>
-	FORCEINLINE static FString GetValueAsString( const TCHAR* EnumPath, const TEnumAsByte<T> EnumValue )
+	FORCEINLINE static FString GetValueAsString( const TCHAR* EnumPath, const TEnumAsByte<T> EnumeratorIndex )
 	{
-		return GetValueAsString_Internal(EnumPath, (int32)EnumValue.GetValue());
+		return GetIndexAsString_Internal(EnumPath, (int32)EnumeratorIndex.GetValue());
 	}
 
 	template< class T >
-	FORCEINLINE static void GetValueAsString( const TCHAR* EnumPath, const T EnumValue, FString& out_StringValue )
+	FORCEINLINE static void GetValueAsString( const TCHAR* EnumPath, const T EnumeratorIndex, FString& out_StringValue )
 	{
-		out_StringValue = GetValueAsString( EnumPath, EnumValue );
+		out_StringValue = GetValueAsString( EnumPath, EnumeratorIndex );
 	}
 
 	/**
-	 * @param EnumPath	- Full enum path
-	 * @param EnumValue - Enum value
+	 * @param EnumPath         Full enum path.
+	 * @param EnumeratorIndex  Enumerator index.
+	 *
+	 * Note: Despite the name, this function actually uses the enumerator at a given index.  If the
+	 *       enum is not zero-based or is non-consecutive, this will give surprising results.
 	 *
 	 * @return the localized display string associated with the specified enum value for the enum specified by a path
 	 */
 	template <typename T>
-	FORCEINLINE static FText GetDisplayValueAsText( const TCHAR* EnumPath, const T EnumValue )
+	FORCEINLINE static FText GetDisplayValueAsText( const TCHAR* EnumPath, const T EnumeratorIndex )
 	{
 		// For the C++ enum.
 		static_assert(TIsEnum<T>::Value, "Should only call this with enum types");
-		return GetDisplayValueAsText_Internal(EnumPath, (int32)EnumValue);
+		return GetDisplayIndexAsText_Internal(EnumPath, (int32)EnumeratorIndex);
 	}
 
 	template <typename T>
-	FORCEINLINE static FText GetDisplayValueAsText( const TCHAR* EnumPath, const TEnumAsByte<T> EnumValue )
+	FORCEINLINE static FText GetDisplayValueAsText( const TCHAR* EnumPath, const TEnumAsByte<T> EnumeratorIndex )
 	{
-		return GetDisplayValueAsText_Internal(EnumPath, (int32)EnumValue.GetValue());
+		return GetDisplayIndexAsText_Internal(EnumPath, (int32)EnumeratorIndex.GetValue());
 	}
 
 	template< class T >
-	FORCEINLINE static void GetDisplayValueAsText( const TCHAR* EnumPath, const T EnumValue, FText& out_TextValue )
+	FORCEINLINE static void GetDisplayValueAsText( const TCHAR* EnumPath, const T EnumeratorIndex, FText& out_TextValue )
 	{
-		out_TextValue = GetDisplayValueAsText( EnumPath, EnumValue );
+		out_TextValue = GetDisplayValueAsText( EnumPath, EnumeratorIndex );
 	}
 
 private:
@@ -1758,18 +1764,18 @@ private:
 	static TMap<FName,TMap<FString,FString> > EnumSubstringRedirects;
 	static void InitEnumRedirectsMap();
 
-	FORCEINLINE static FString GetValueAsString_Internal( const TCHAR* EnumPath, const int32 Value )
+	FORCEINLINE static FString GetIndexAsString_Internal( const TCHAR* EnumPath, const int32 Index )
 	{
 		UEnum* EnumClass = FindObject<UEnum>( nullptr, EnumPath );
 		UE_CLOG( !EnumClass, LogClass, Fatal, TEXT("Couldn't find enum '%s'"), EnumPath );
-		return EnumClass->GetEnumName(Value);
+		return EnumClass->GetEnumName(Index);
 	}
 
-	FORCEINLINE static FText GetDisplayValueAsText_Internal( const TCHAR* EnumPath, const int32 Value )
+	FORCEINLINE static FText GetDisplayIndexAsText_Internal( const TCHAR* EnumPath, const int32 EnumeratorIndex )
 	{
 		UEnum* EnumClass = FindObject<UEnum>(nullptr, EnumPath);
 		UE_CLOG(!EnumClass, LogClass, Fatal, TEXT("Couldn't find enum '%s'"), EnumPath);
-		return EnumClass->GetEnumText(Value);
+		return EnumClass->GetEnumText(EnumeratorIndex);
 	}
 
 	/**
@@ -2029,12 +2035,16 @@ public:
 
 	/** Reference token stream used by realtime garbage collector, finalized in AssembleReferenceTokenStream */
 	FGCReferenceTokenStream ReferenceTokenStream;
+	/** CS for the token stream. Token stream can assemble code can sometimes be called from two threads throuh a web of async loading calls. */
+	FCriticalSection ReferenceTokenStreamCritical;
 
 #if !(UE_BUILD_TEST || UE_BUILD_SHIPPING)
 	/* TokenIndex map to look-up token stream index origin. */
 	FGCDebugReferenceTokenMap DebugTokenMap;
 #endif
 
+	
+
 	/** This class's native functions. */
 	TArray<FNativeFunctionLookup> NativeFunctionLookupTable;
 
diff --git a/Engine/Source/Runtime/CoreUObject/Public/UObject/EnumProperty.h b/Engine/Source/Runtime/CoreUObject/Public/UObject/EnumProperty.h
new file mode 100644
--- /dev/null
+++ b/Engine/Source/Runtime/CoreUObject/Public/UObject/EnumProperty.h
@@ -0,0 +1,72 @@
+// Copyright 1998-2016 Epic Games, Inc. All Rights Reserved.
+
+#pragma once
+
+#include "ObjectMacros.h"
+
+class UEnum;
+class UNumericProperty;
+
+namespace UE4EnumProperty_Private
+{
+	struct FEnumPropertyFriend;
+}
+
+class COREUOBJECT_API UEnumProperty : public UProperty
+{
+	DECLARE_CASTED_CLASS_INTRINSIC(UEnumProperty, UProperty, 0, TEXT("/Script/CoreUObject"), CASTCLASS_UEnumProperty)
+
+public:
+	UEnumProperty(const FObjectInitializer& ObjectInitializer, UEnum* InEnum);
+	UEnumProperty(const FObjectInitializer& ObjectInitializer, ECppProperty, int32 InOffset, uint64 InFlags, UEnum* InEnum);
+
+	// UObject interface
+	virtual void Serialize( FArchive& Ar ) override;
+	static void AddReferencedObjects(UObject* InThis, FReferenceCollector& Collector);
+	virtual void GetPreloadDependencies(TArray<UObject*>& OutDeps) override;
+	// End of UObject interface
+
+	// UField interface
+	virtual void AddCppProperty(UProperty* Property) override;
+	// End of UField interface
+
+	// UProperty interface
+	virtual FString GetCPPMacroType( FString& ExtendedTypeText ) const  override;
+	virtual FString GetCPPType( FString* ExtendedTypeText, uint32 CPPExportFlags ) const override;
+	virtual FString GetCPPTypeForwardDeclaration() const override;
+	virtual void LinkInternal(FArchive& Ar) override;
+	virtual bool Identical( const void* A, const void* B, uint32 PortFlags ) const override;
+	virtual void SerializeItem( FArchive& Ar, void* Value, void const* Defaults ) const override;
+	virtual bool NetSerializeItem( FArchive& Ar, UPackageMap* Map, void* Data, TArray<uint8> * MetaData = NULL ) const override;
+	virtual void ExportTextItem( FString& ValueStr, const void* PropertyValue, const void* DefaultValue, UObject* Parent, int32 PortFlags, UObject* ExportRootScope ) const override;
+	virtual const TCHAR* ImportText_Internal( const TCHAR* Buffer, void* Data, int32 PortFlags, UObject* OwnerObject, FOutputDevice* ErrorText ) const override;
+	virtual int32 GetMinAlignment() const override;
+	virtual bool SameType(const UProperty* Other) const override;
+	virtual bool ConvertFromType(const FPropertyTag& Tag, FArchive& Ar, uint8* Data, UStruct* DefaultsStruct, bool& bOutAdvanceProperty) override;
+	// End of UProperty interface
+
+	/**
+	 * Returns a pointer to the UEnum of this property.
+	 */
+	FORCEINLINE UEnum* GetEnum() const
+	{
+		return Enum;
+	}
+
+	/**
+	 * Returns the numeric property which represents the integral type of the enum.
+	 */
+	FORCEINLINE UNumericProperty* GetUnderlyingProperty() const
+	{
+		return UnderlyingProp;
+	}
+
+private:
+	friend struct UE4EnumProperty_Private::FEnumPropertyFriend;
+
+#if HACK_HEADER_GENERATOR
+public:
+#endif
+	UNumericProperty* UnderlyingProp; // The property which represents the underlying type of the enum
+	UEnum* Enum; // The enum represented by this property
+};
diff --git a/Engine/Source/Runtime/CoreUObject/Public/UObject/LazyObjectPtr.h b/Engine/Source/Runtime/CoreUObject/Public/UObject/LazyObjectPtr.h
--- a/Engine/Source/Runtime/CoreUObject/Public/UObject/LazyObjectPtr.h
+++ b/Engine/Source/Runtime/CoreUObject/Public/UObject/LazyObjectPtr.h
@@ -1,7 +1,7 @@
 // Copyright 1998-2016 Epic Games, Inc. All Rights Reserved.
 
 /*=============================================================================
-	LazyObjectPtr.h: Lazy, guid-based pointer to UObject
+	LazyObjectPtr.h: Lazy, guid-based weak pointer to a UObject, mostly useful for actors
 =============================================================================*/
 
 #pragma once
@@ -20,72 +20,60 @@ struct COREUOBJECT_API FUniqueObjectGuid
 		: Guid(InGuid)
 	{}
 
-	/**
-	 * Reset the guid pointer back to the invalid state
-	 */
+	/** Reset the guid pointer back to the invalid state */
 	FORCEINLINE void Reset()
 	{
 		Guid.Invalidate();
 	}
 
-	/**
-	 * Construct from an existing object
-	 */
+	/** Construct from an existing object */
 	FUniqueObjectGuid(const class UObject* InObject);
 
-	/**
-	 * Converts into a string
-	 */
+	/** Converts into a string */
 	FString ToString() const;
 
-	/**
-	 * Converts from a string
-	 */
+	/** Converts from a string */
 	void FromString(const FString& From);
 
-	/**
-	 * Fixes up this UniqueObjectID to add or remove the PIE prefix depending on what is currently active
-	 */
+	/** Fixes up this UniqueObjectID to add or remove the PIE prefix depending on what is currently active */
 	FUniqueObjectGuid FixupForPIE(int32 PlayInEditorID = GPlayInEditorID) const;
 
 	/**
 	 * Attempts to find a currently loaded object that matches this object ID
-	 * @return Found UObject, or NULL if not currently loaded
+	 *
+	 * @return Found UObject, or nullptr if not currently loaded
 	 */
 	class UObject *ResolveObject() const;
 
-	/**
-	 * Test if this can ever point to a live UObject
-	 */
+	/** Test if this can ever point to a live UObject */
 	FORCEINLINE bool IsValid() const
 	{
 		return Guid.IsValid();
 	}
 
-	FORCEINLINE bool operator==(const FUniqueObjectGuid &Other) const
+	FORCEINLINE bool operator==(const FUniqueObjectGuid& Other) const
 	{
 		return Guid == Other.Guid;
 	}
 
-	FORCEINLINE bool operator!=(const FUniqueObjectGuid &Other) const
+	FORCEINLINE bool operator!=(const FUniqueObjectGuid& Other) const
 	{
 		return Guid != Other.Guid;
 	}
 
-	/**
-	 * Annotation API
-	 * @return true is this is a default.
-	 */
+	/** Returns true is this is the default value */
 	FORCEINLINE bool IsDefault()
 	{
-		return !IsValid(); // an default GUID is 0,0,0,0 and this is "invalid"
+		// A default GUID is 0,0,0,0 and this is "invalid"
+		return !IsValid(); 
 	}
 
 	FORCEINLINE friend uint32 GetTypeHash(const FUniqueObjectGuid& ObjectGuid)
 	{
 		return GetTypeHash(ObjectGuid.Guid);
 	}
-	FORCEINLINE FGuid GetGuid() const
+	/** Returns wrapped Guid */
+	FORCEINLINE const FGuid& GetGuid() const
 	{
 		return Guid;
 	}
@@ -96,6 +84,7 @@ struct COREUOBJECT_API FUniqueObjectGuid
 		return Ar;
 	}
 
+	/** Code needed by FLazyPtr internals */
 	static int32 GetCurrentTag()
 	{
 		return CurrentAnnotationTag.GetValue();
@@ -108,95 +97,74 @@ struct COREUOBJECT_API FUniqueObjectGuid
 	static FUniqueObjectGuid GetOrCreateIDForObject(const class UObject *Object);
 
 private:
+	/** Guid representing the object, should be unique */
 	FGuid Guid;
 
 	/** Global counter that determines when we need to re-search for GUIDs because more objects have been loaded **/
 	static FThreadSafeCounter CurrentAnnotationTag;
 };
 
-
 template<> struct TIsPODType<FUniqueObjectGuid> { enum { Value = true }; };
 
-/***
- * FLazyObjectPtr is a type of weak pointer to a UObject.
- * It will change back and forth between being valid or pending as the referenced object loads or unloads
+/**
+ * FLazyObjectPtr is a type of weak pointer to a UObject that uses a GUID created at save time.
+ * It will change back and forth between being valid or pending as the referenced object loads or unloads.
  * It has no impact on if the object is garbage collected or not.
  * It can't be directly used across a network.
  *
- * This is useful for cross level references or any place you need to declare a reference to an object that may not be loaded
- **/
-
+ * This is useful for cross level references or places where you need to point to an object whose name changes often.
+ */
 class FLazyObjectPtr : public TPersistentObjectPtr<FUniqueObjectGuid>
 {
 public:	
-	/**
-	 * Called by UObject::Serialize so that we can save / load the Guid possibly associated with an object
-	 */
-	COREUOBJECT_API static void PossiblySerializeObjectGuid(UObject *Object, FArchive& Ar);
-
-	/**
-	 * Called when entering PIE to prepare it for PIE-specific fixups
-	 */
-	COREUOBJECT_API static void ResetPIEFixups();
-
-	/** NULL constructor **/
+	/** Default constructor, sets to null */
 	FORCEINLINE FLazyObjectPtr()
 	{
 	}
 
-	/**
-	 * Construct from another lazy pointer
-	 * @param Other lazy pointer to copy from
-	 */
-	FORCEINLINE FLazyObjectPtr(const FLazyObjectPtr &Other)
+	/** Construct from another lazy pointer */
+	FORCEINLINE FLazyObjectPtr(const FLazyObjectPtr& Other)
 	{
 		(*this)=Other;
 	}
 
-	/**
-	 * Construct from another lazy pointer
-	 * @param Other lazy pointer to copy from
-	 */
-	explicit FORCEINLINE FLazyObjectPtr(const UObject *Object)
+	/** Construct from object already in memory */
+	explicit FORCEINLINE FLazyObjectPtr(const UObject* Object)
 	{
 		(*this)=Object;
 	}
-	/**
-	 * Copy from an object pointer
-	 * @param Object object to create a weak pointer to
-	 */
+	
+	/** Copy from an object already in memory */
 	FORCEINLINE void operator=(const UObject *Object)
 	{
 		TPersistentObjectPtr<FUniqueObjectGuid>::operator=(Object);
 	}
 
-	/**
-	 * Copy from a lazy pointer
-	 */
-	FORCEINLINE void operator=(const FLazyObjectPtr &Other)
+	/** Copy from a lazy pointer */
+	FORCEINLINE void operator=(const FLazyObjectPtr& Other)
 	{
 		TPersistentObjectPtr<FUniqueObjectGuid>::operator=(Other);
 	}
 
-	/**
-	 * Copy from a unique object identifier
-	 * @param ObjectID Object identifier to create a weak pointer to
-	 */
-	FORCEINLINE void operator=(const FUniqueObjectGuid &InObjectID)
+	/** Copy from a unique object identifier */
+	FORCEINLINE void operator=(const FUniqueObjectGuid& InObjectID)
 	{
 		TPersistentObjectPtr<FUniqueObjectGuid>::operator=(InObjectID);
 	}
+	
+	/** Called by UObject::Serialize so that we can save / load the Guid possibly associated with an object */
+	COREUOBJECT_API static void PossiblySerializeObjectGuid(UObject* Object, FArchive& Ar);
+
+	/** Called when entering PIE to prepare it for PIE-specific fixups */
+	COREUOBJECT_API static void ResetPIEFixups();
 };
 
 template <> struct TIsPODType<FLazyObjectPtr> { enum { Value = TIsPODType<TPersistentObjectPtr<FUniqueObjectGuid> >::Value }; };
 template <> struct TIsWeakPointerType<FLazyObjectPtr> { enum { Value = TIsWeakPointerType<TPersistentObjectPtr<FUniqueObjectGuid> >::Value }; };
 
-/***
-* 
-* TLazyObjectPtr is templatized version of the generic FLazyObjectPtr
-* 
-**/
-
+/**
+ * TLazyObjectPtr is templatized version of the generic FLazyObjectPtr
+ */
 template<class T=UObject>
 class TLazyObjectPtr : private FLazyObjectPtr
 {
@@ -208,24 +176,18 @@ public:
 	FORCEINLINE TLazyObjectPtr<T>& operator=(const TLazyObjectPtr<T>& Other) = default;
 	FORCEINLINE TLazyObjectPtr<T>& operator=(TLazyObjectPtr<T>&& Other) = default;
 #else
-	/** NULL constructor **/
+	/** Default constructor **/
 	FORCEINLINE TLazyObjectPtr()
 	{
 	}
 	
-	/**
-	 * Construct from another lazy pointer
-	 * @param Other lazy pointer to copy from
-	 */
+	/** Construct from another lazy pointer */
 	FORCEINLINE TLazyObjectPtr(const TLazyObjectPtr<T>& Other) :
 		FLazyObjectPtr(Other)
 	{
 	}
 
-	/**
-	* Assign from another lazy pointer
-	* @param Other lazy pointer to copy from
-	*/
+	/** Copy from another lazy pointer */
 	FORCEINLINE TLazyObjectPtr<T>& operator=(const TLazyObjectPtr<T>& Other)
 	{
 		FLazyObjectPtr::operator=(Other);
@@ -233,20 +195,14 @@ public:
 	}
 #endif
 
-	/**
-	 * Construct from another lazy pointer with implicit upcasting allowed
-	 * @param Other lazy pointer to copy from
-	 */
+	/** Construct from another lazy pointer with implicit upcasting allowed */
 	template<typename U, typename = typename TEnableIf<TPointerIsConvertibleFromTo<U, T>::Value>::Type>
 	FORCEINLINE TLazyObjectPtr(const TLazyObjectPtr<U>& Other) :
 		FLazyObjectPtr((const FLazyObjectPtr&)Other)
 	{
 	}
 	
-	/**
-	 * Assign from another lazy pointer with implicit upcasting allowed
-	 * @param Other lazy pointer to copy from
-	 */
+	/** Assign from another lazy pointer with implicit upcasting allowed */
 	template<typename U>
 	FORCEINLINE typename TEnableIf<TPointerIsConvertibleFromTo<U, T>::Value, TLazyObjectPtr<T>&>::Type
 		operator=(const TLazyObjectPtr<U>& Other)
@@ -255,27 +211,19 @@ public:
 		return *this;
 	}
 
-	/**
-	 * Construct from an object pointer
-	 * @param Object Object to create a lazy pointer to
-	 */
+	/** Construct from an object pointer */
 	FORCEINLINE TLazyObjectPtr(T* Object)
 	{
 		FLazyObjectPtr::operator=(Object);
 	}
 
-	/**
-	 * Reset the lazy pointer back to the NULL state
-	 */
+	/** Reset the lazy pointer back to the null state */
 	FORCEINLINE void Reset()
 	{
 		FLazyObjectPtr::Reset();
 	}
 
-	/**
-	 * Copy from an object pointer
-	 * @param Object Object to create a lazy pointer to
-	 */
+	/** Copy from an object pointer */
 	FORCEINLINE void operator=(T* Object)
 	{
 		FLazyObjectPtr::operator=(Object);
@@ -285,6 +233,7 @@ public:
 	 * Copy from a unique object identifier
 	 * WARNING: this doesn't check the type of the object is correct,
 	 * because the object corresponding to this ID may not even be loaded!
+	 *
 	 * @param ObjectID Object identifier to create a lazy pointer to
 	 */
 	FORCEINLINE void operator=(const FUniqueObjectGuid& InObjectID)
@@ -294,16 +243,18 @@ public:
 
 	/**
 	 * Gets the unique object identifier associated with this lazy pointer. Valid even if pointer is not currently valid
+	 *
 	 * @return Unique ID for this object, or an invalid FUniqueObjectID if this pointer isn't set to anything
 	 */
-	FUniqueObjectGuid GetUniqueID() const
+	FORCEINLINE const FUniqueObjectGuid& GetUniqueID() const
 	{
 		return FLazyObjectPtr::GetUniqueID();
 	}
 
 	/**
 	 * Dereference the lazy pointer.
-	 * @return NULL if this object is gone or the lazy pointer was NULL, otherwise a valid uobject pointer
+	 *
+	 * @return nullptr if this object is gone or the lazy pointer was null, otherwise a valid UObject pointer
 	 */
 	FORCEINLINE T* Get() const
 	{
@@ -313,26 +264,19 @@ public:
 		return dynamic_cast<T*>(FLazyObjectPtr::Get());
 	}
 
-	/**
-	 * Dereference the lazy pointer.
-	 */
+	/** Dereference the lazy pointer */
 	FORCEINLINE T& operator*() const
 	{
 		return *Get();
 	}
 
-	/**
-	 * Dereference the lazy pointer.
-	 */
+	/** Dereference the lazy pointer */
 	FORCEINLINE T* operator->() const
 	{
 		return Get();
 	}
 
-	/**
-	 * Test if this points to a live UObject
-	 * @return true if Get() would return a valid non-null pointer
-	 */
+	/** Test if this points to a live UObject */
 	FORCEINLINE bool IsValid() const
 	{
 		return FLazyObjectPtr::IsValid();
@@ -340,8 +284,9 @@ public:
 
 	/**
 	 * Slightly different than !IsValid(), returns true if this used to point to a UObject, but doesn't any more and has not been assigned or reset in the mean time.
+	 *
 	 * @return true if this used to point at a real object but no longer does.
-	**/
+	 */
 	FORCEINLINE bool IsStale() const
 	{
 		return FLazyObjectPtr::IsStale();
@@ -349,8 +294,9 @@ public:
 
 	/**
 	 * Test if this does not point to a live UObject, but may in the future
+	 *
 	 * @return true if this does not point to a real object, but could possibly
-	**/
+	 */
 	FORCEINLINE bool IsPending() const
 	{
 		return FLazyObjectPtr::IsPending();
@@ -358,16 +304,15 @@ public:
 
 	/**
 	 * Test if this can never point to a live UObject
+	 *
 	 * @return true if this is explicitly pointing to no object
-	**/
+	 */
 	FORCEINLINE bool IsNull() const
 	{
 		return FLazyObjectPtr::IsNull();
 	}
 
-	/**
-	 * Dereference lazy pointer to see if it points somewhere valid.
-	 */
+	/** Dereference lazy pointer to see if it points somewhere valid */
 	FORCEINLINE explicit operator bool() const
 	{
 		return IsValid();
@@ -388,7 +333,7 @@ public:
 
 // The reason these aren't inside the class (above) is because Visual Studio 2012-2013 crashes when compiling them :D
 
-// Compare with another TLazyObjectPtr of related type
+/** Compare with another TLazyObjectPtr of related type */
 template<typename T, typename U>
 FORCEINLINE bool operator==(const TLazyObjectPtr<T>& Lhs, const TLazyObjectPtr<U>& Rhs)
 {
@@ -416,7 +361,7 @@ FORCEINLINE bool operator==(const U* Lhs, const TLazyObjectPtr<T>& Rhs)
 	return Lhs == Rhs.Get();
 }
 
-/** Compare to null **/
+/** Compare to null */
 template<typename T>
 FORCEINLINE bool operator==(const TLazyObjectPtr<T>& Lhs, TYPE_OF_NULLPTR)
 {
@@ -454,6 +399,5 @@ FORCEINLINE bool operator!=(TYPE_OF_NULLPTR, const TLazyObjectPtr<T>& Rhs)
 	return Rhs.IsValid();
 }
 
-
 template<class T> struct TIsPODType<TLazyObjectPtr<T> > { enum { Value = TIsPODType<FLazyObjectPtr>::Value }; };
 template<class T> struct TIsWeakPointerType<TLazyObjectPtr<T> > { enum { Value = TIsWeakPointerType<FLazyObjectPtr>::Value }; };
diff --git a/Engine/Source/Runtime/CoreUObject/Public/UObject/Object.h b/Engine/Source/Runtime/CoreUObject/Public/UObject/Object.h
--- a/Engine/Source/Runtime/CoreUObject/Public/UObject/Object.h
+++ b/Engine/Source/Runtime/CoreUObject/Public/UObject/Object.h
@@ -53,9 +53,8 @@ class COREUOBJECT_API UObject : public UObjectBaseUtility
 	/**
 	* Create a component or subobject only to be used with the editor.
 	* @param	TReturnType					class of return type, all overrides must be of this type
-	* @param	Outer						outer to construct the subobject in
 	* @param	SubobjectName				name of the new component
-	* @param bTransient		true if the component is being assigned to a transient property
+	* @param	bTransient					true if the component is being assigned to a transient property
 	*/
 	template<class TReturnType>
 	TReturnType* CreateEditorOnlyDefaultSubobject(FName SubobjectName, bool bTransient = false)
@@ -67,9 +66,8 @@ class COREUOBJECT_API UObject : public UObjectBaseUtility
 	/**
 	* Create a component or subobject
 	* @param	TReturnType					class of return type, all overrides must be of this type
-	* @param	Outer						outer to construct the subobject in
 	* @param	SubobjectName				name of the new component
-	* @param bTransient		true if the component is being assigned to a transient property
+	* @param	bTransient					true if the component is being assigned to a transient property
 	*/
 	template<class TReturnType>
 	TReturnType* CreateDefaultSubobject(FName SubobjectName, bool bTransient = false)
@@ -82,7 +80,6 @@ class COREUOBJECT_API UObject : public UObjectBaseUtility
 	* Create a component or subobject
 	* @param TReturnType class of return type, all overrides must be of this type
 	* @param TClassToConstructByDefault class to construct by default
-	* @param Outer outer to construct the subobject in
 	* @param SubobjectName name of the new component
 	* @param bTransient		true if the component is being assigned to a transient property
 	*/
@@ -96,9 +93,8 @@ class COREUOBJECT_API UObject : public UObjectBaseUtility
 	* Create optional component or subobject. Optional subobjects may not get created
 	* when a derived class specified DoNotCreateDefaultSubobject with the subobject's name.
 	* @param	TReturnType					class of return type, all overrides must be of this type
-	* @param	Outer						outer to construct the subobject in
 	* @param	SubobjectName				name of the new component
-	* @param bTransient		true if the component is being assigned to a transient property
+	* @param	bTransient					true if the component is being assigned to a transient property
 	*/
 	template<class TReturnType>
 	TReturnType* CreateOptionalDefaultSubobject(FName SubobjectName, bool bTransient = false)
@@ -111,9 +107,8 @@ class COREUOBJECT_API UObject : public UObjectBaseUtility
 	* Create optional component or subobject. Optional subobjects may not get created
 	* when a derived class specified DoNotCreateDefaultSubobject with the subobject's name.
 	* @param	TReturnType					class of return type, all overrides must be of this type
-	* @param	Outer						outer to construct the subobject in
 	* @param	SubobjectName				name of the new component
-	* @param bTransient		true if the component is being assigned to a transient property
+	* @param	bTransient					true if the component is being assigned to a transient property
 	*/
 	template<class TReturnType>
 	TReturnType* CreateAbstractDefaultSubobject(FName SubobjectName, bool bTransient = false)
diff --git a/Engine/Source/Runtime/CoreUObject/Public/UObject/ObjectMacros.h b/Engine/Source/Runtime/CoreUObject/Public/UObject/ObjectMacros.h
--- a/Engine/Source/Runtime/CoreUObject/Public/UObject/ObjectMacros.h
+++ b/Engine/Source/Runtime/CoreUObject/Public/UObject/ObjectMacros.h
@@ -340,6 +340,7 @@ typedef uint64 EClassCastFlags;
 #define CASTCLASS_UStaticMeshComponent			DECLARE_UINT64(0x0000200000000000)
 #define CASTCLASS_UMapProperty					DECLARE_UINT64(0x0000400000000000)
 #define CASTCLASS_USetProperty					DECLARE_UINT64(0x0000800000000000)
+#define CASTCLASS_UEnumProperty					DECLARE_UINT64(0x0001000000000000)
 
 #define CASTCLASS_AllFlags						DECLARE_UINT64(0xFFFFFFFFFFFFFFFF)
 
diff --git a/Engine/Source/Runtime/CoreUObject/Public/UObject/Package.h b/Engine/Source/Runtime/CoreUObject/Public/UObject/Package.h
--- a/Engine/Source/Runtime/CoreUObject/Public/UObject/Package.h
+++ b/Engine/Source/Runtime/CoreUObject/Public/UObject/Package.h
@@ -6,6 +6,7 @@
 #include "WorldCompositionUtility.h"
 #include "GatherableTextData.h"
 #include "PropertyLocalizationDataGathering.h"
+#include "UniquePtr.h"
 
 /**
  * Represents the result of saving a package
@@ -148,13 +149,13 @@ public:
 	int64 FileSize;
 	
 	/** Editor only: Thumbnails stored in this package */
-	TScopedPointer< FThumbnailMap > ThumbnailMap;
+	TUniquePtr< FThumbnailMap > ThumbnailMap;
 
 	// MetaData for the editor, or NULL in the game
 	class UMetaData*	MetaData;
 	
 	// World browser information
-	TScopedPointer< FWorldTileInfo > WorldTileInfo;
+	TUniquePtr< FWorldTileInfo > WorldTileInfo;
 
 	TMap<FName, int32> ClassUniqueNameIndexMap;
 
diff --git a/Engine/Source/Runtime/CoreUObject/Public/UObject/PersistentObjectPtr.h b/Engine/Source/Runtime/CoreUObject/Public/UObject/PersistentObjectPtr.h
--- a/Engine/Source/Runtime/CoreUObject/Public/UObject/PersistentObjectPtr.h
+++ b/Engine/Source/Runtime/CoreUObject/Public/UObject/PersistentObjectPtr.h
@@ -14,15 +14,13 @@ class TPersistentObjectPtr
 {
 public:	
 
-	/** NULL constructor **/
+	/** Default constructor, will be null */
 	FORCEINLINE TPersistentObjectPtr()
 	{
 		Reset();
 	}
 
-	/**
-	 * Reset the lazy pointer back to the NULL state
-	 */
+	/** Reset the lazy pointer back to the null state */
 	FORCEINLINE void Reset()
 	{
 		WeakPtr.Reset();
@@ -30,19 +28,13 @@ public:
 		TagAtLastTest = 0;
 	}
 
-	/**  
-	 * Construct from another lazy pointer
-	 * @param Other lazy pointer to copy from
-	 */
+	/** Construct from another pointer of the same type */
 	FORCEINLINE TPersistentObjectPtr(const TPersistentObjectPtr& Other)
 	{
 		(*this)=Other;
 	}
 
-	/**  
-	 * Construct from a unique object identifier
-	 * @param InObjectID a unique object identifier
-	 */
+	/** Construct from a unique object identifier */
 	explicit FORCEINLINE TPersistentObjectPtr(const TObjectID& InObjectID)
 		: WeakPtr()
 		, TagAtLastTest(0)
@@ -50,10 +42,7 @@ public:
 	{
 	}
 
-	/**  
-	 * Copy from a unique object identifier
-	 * @param ObjectID object identifier to create a weak pointer to
-	 */
+	/** Copy from a unique object identifier */
 	FORCEINLINE void operator=(const TObjectID& InObjectID)
 	{
 		WeakPtr.Reset();
@@ -61,10 +50,7 @@ public:
 		TagAtLastTest = 0;
 	}
 
-	/**  
-	 * Copy from an object pointer
-	 * @param Object object to create a weak pointer to
-	 */
+	/** Copy from an object pointer */
 	FORCEINLINE void operator=(const class UObject* Object)
 	{
 		if (Object)
@@ -79,21 +65,15 @@ public:
 		}
 	}
 	
-	/**  
-	 * Copy from an existing weak pointer, reserve IDs if required
-	 * @param Other weak pointer to copy from
-	 */
+	/** Copy from an existing weak pointer, reserve IDs if required */
 	FORCEINLINE void operator=(const FWeakObjectPtr& Other)
 	{
 		// If object exists need to make sure it gets registered properly in above function, if it doesn't exist just empty it
-		const UObject *Object = Other.Get();
+		const UObject* Object = Other.Get();
 		*this = Object;
 	}
 
-	/**  
-	 * Construct from another lazy pointer
-	 * @param Other lazy pointer to copy from
-	 */
+	/** Construct from another pointer */
 	FORCEINLINE void operator=(const TPersistentObjectPtr<TObjectID>& Other)
 	{
 		WeakPtr = Other.WeakPtr;
@@ -104,6 +84,7 @@ public:
 	
 	/**
 	 * Gets the unique object identifier associated with this lazy pointer. Valid even if pointer is not currently valid
+	 *
 	 * @return Unique ID for this object, or an invalid FUniqueObjectGuid if this pointer isn't set to anything
 	 */
 	FORCEINLINE const TObjectID& GetUniqueID() const
@@ -111,28 +92,27 @@ public:
 		return ObjectID;
 	}
 
-	/**
-	 * Non-const version of the above.
-	 */
+	/** Non-const version of the above */
 	FORCEINLINE TObjectID& GetUniqueID()
 	{
 		return ObjectID;
 	}
 
 	/**
-	 * Dereference the lazy pointer, which may cause it to become valid again. Will not try to load pending outside of game thread
-	 * @return NULL if this object is gone or the lazy pointer was NULL, otherwise a valid UObject pointer
+	 * Dereference the pointer, which may cause it to become valid again. Will not try to load pending outside of game thread
+	 *
+	 * @return nullptr if this object is gone or the pointer was null, otherwise a valid UObject pointer
 	 */
 	FORCEINLINE UObject* Get() const
 	{
-		UObject *Object = WeakPtr.Get();
+		UObject* Object = WeakPtr.Get();
 		if (!Object && TObjectID::GetCurrentTag() != TagAtLastTest && ObjectID.IsValid())
 		{
 			Object = ObjectID.ResolveObject();
 			WeakPtr = Object;
 			TagAtLastTest = TObjectID::GetCurrentTag();
 
-			// If this object is pending kill or otherwise invalid, this will return NULL as expected
+			// If this object is pending kill or otherwise invalid, this will return nullptr as expected
 			Object = WeakPtr.Get();
 		}
 		return Object;
@@ -140,55 +120,44 @@ public:
 
 	/**
 	 * Dereference the lazy pointer, which may cause it to become valid again. Will not try to load pending outside of game thread
+	 *
 	 * @param bEvenIfPendingKill, if this is true, pendingkill objects are considered valid
-	 * @return NULL if this object is gone or the lazy pointer was NULL, otherwise a valid UObject pointer
+	 * @return nullptr if this object is gone or the lazy pointer was null, otherwise a valid UObject pointer
 	 */
 	FORCEINLINE UObject* Get(bool bEvenIfPendingKill) const
 	{
-		UObject *Object = WeakPtr.Get(bEvenIfPendingKill);
+		UObject* Object = WeakPtr.Get(bEvenIfPendingKill);
 		if (!Object && TObjectID::GetCurrentTag() != TagAtLastTest && ObjectID.IsValid())
 		{
 			Object = ObjectID.ResolveObject();
 			WeakPtr = Object;
 			TagAtLastTest = TObjectID::GetCurrentTag();
 
-			// If this object is pending kill or otherwise invalid, this will return NULL as expected
+			// If this object is pending kill or otherwise invalid, this will return nullptr as expected
 			Object = WeakPtr.Get(bEvenIfPendingKill);
 		}
 		return Object;
 	}
 
-	/**  
-	 * Dereference the lazy pointer.
-	 */
+	/** Dereference the pointer */
 	FORCEINLINE class UObject& operator*() const
 	{
 		return *Get();
 	}
 
-	/**  
-	 * Dereference the lazy pointer.
-	 */
+	/** Dereference the pointer */
 	FORCEINLINE class UObject* operator->() const
 	{
 		return Get();
 	}
 
-	/**  
-	 * Compare lazy pointers for equality. Only Serial Number matters for lazy pointers
-	 * @param Other lazy pointer to compare to
-	 * Caution: Two lazy pointers might not be equal to each other, but they both might return NULL.
-	**/
+	/** Compare pointers for equality. Only Serial Number matters for the base implementation */
 	FORCEINLINE bool operator==(const TPersistentObjectPtr& Other) const
 	{
 		return ObjectID == Other.ObjectID;
 	}
 
-	/**  
-	 * Compare lazy pointers for inequality. Only Serial Number matters for lazy pointers
-	 * @param Other lazy pointer to compare to
-	 * Caution: Two lazy pointers might not be equal to each other, but they both might return NULL.
-	**/
+	/** Compare pointers for inequality. Only Serial Number matters for the base implementation */
 	FORCEINLINE bool operator!=(const TPersistentObjectPtr& Other) const
 	{
 		return ObjectID != Other.ObjectID;
@@ -196,18 +165,19 @@ public:
 
 	/**  
 	 * Test if this does not point to a live UObject, but may in the future
+	 * 
 	 * @return true if this does not point to a real object, but could possibly
 	 */
 	FORCEINLINE bool IsPending() const
 	{
-		return Get() == NULL && ObjectID.IsValid();
+		return Get() == nullptr && ObjectID.IsValid();
 	}
 
-
 	/**  
 	 * Test if this points to a live UObject
+	 *
 	 * @return true if Get() would return a valid non-null pointer
-	**/
+	 */
 	FORCEINLINE bool IsValid() const
 	{
 		return !!Get();
@@ -215,22 +185,24 @@ public:
 
 	/**  
 	 * Slightly different than !IsValid(), returns true if this used to point to a UObject, but doesn't any more and has not been assigned or reset in the mean time.
+	 *
 	 * @return true if this used to point at a real object but no longer does.
-	**/
+	 */
 	FORCEINLINE bool IsStale() const
 	{
 		return WeakPtr.IsStale();
 	}
 	/**  
 	 * Test if this can never point to a live UObject
+	 *
 	 * @return true if this is explicitly pointing to no object
 	 */
 	FORCEINLINE bool IsNull() const
 	{
-		return Get() == NULL && !ObjectID.IsValid();
+		return Get() == nullptr && !ObjectID.IsValid();
 	}
 
-	/** Hash function. */
+	/** Hash function */
 	FORCEINLINE friend uint32 GetTypeHash(const TPersistentObjectPtr& Ptr)
 	{
 		return GetTypeHash(Ptr.ObjectID);
@@ -239,11 +211,11 @@ public:
 private:
 
 	/** Once the object has been noticed to be loaded, this is set to the object weak pointer **/
-	mutable FWeakObjectPtr		WeakPtr;
+	mutable FWeakObjectPtr	WeakPtr;
 	/** Compared to CurrentAnnotationTag and if they are not equal, a guid search will be performed **/
-	mutable int32					TagAtLastTest;
+	mutable int32			TagAtLastTest;
 	/** Guid for the object this pointer points to or will point to. **/
-	TObjectID					ObjectID;
+	TObjectID				ObjectID;
 };
 
 template <class TObjectID> struct TIsPODType<TPersistentObjectPtr<TObjectID> > { enum { Value = TIsPODType<TObjectID>::Value }; };
diff --git a/Engine/Source/Runtime/CoreUObject/Public/UObject/PropertyTag.h b/Engine/Source/Runtime/CoreUObject/Public/UObject/PropertyTag.h
--- a/Engine/Source/Runtime/CoreUObject/Public/UObject/PropertyTag.h
+++ b/Engine/Source/Runtime/CoreUObject/Public/UObject/PropertyTag.h
@@ -20,7 +20,7 @@ struct FPropertyTag
 	uint8	BoolVal;	// a boolean property's value (never need to serialize data for bool properties except here)
 	FName	Name;		// Name of property.
 	FName	StructName;	// Struct name if UStructProperty.
-	FName	EnumName;	// Enum name if UByteProperty
+	FName	EnumName;	// Enum name if UByteProperty or UEnumProperty
 	FName	InnerType;	// Inner type if UArrayProperty, USetProperty, or UMapProperty
 	FName	ValueType;	// Value type if UMapPropery
 	int32	Size;       // Property size.
diff --git a/Engine/Source/Runtime/CoreUObject/Public/UObject/ScriptMacros.h b/Engine/Source/Runtime/CoreUObject/Public/UObject/ScriptMacros.h
--- a/Engine/Source/Runtime/CoreUObject/Public/UObject/ScriptMacros.h
+++ b/Engine/Source/Runtime/CoreUObject/Public/UObject/ScriptMacros.h
@@ -72,6 +72,9 @@ enum {MAX_VARIABLE_SIZE = 0x0FFF };
 #define P_GET_ARRAY(ElementType,ParamName)			ElementType ParamName[(MAX_VARIABLE_SIZE/sizeof(ElementType))+1];		Stack.StepCompiledIn<UProperty>(ParamName);
 #define P_GET_ARRAY_REF(ElementType,ParamName)		ElementType ParamName##Temp[(MAX_VARIABLE_SIZE/sizeof(ElementType))+1]; ElementType* ParamName = Stack.StepCompiledInRef<UProperty, ElementType*>(ParamName##Temp);
 
+#define P_GET_ENUM(EnumType,ParamName)				EnumType ParamName = (EnumType)0; Stack.StepCompiledIn<UEnumProperty>(&ParamName);
+#define P_GET_ENUM_REF(EnumType,ParamName)			PARAM_PASSED_BY_REF_ZEROED(ParamName, UEnumProperty, EnumType)
+
 #define P_FINISH									Stack.Code += !!Stack.Code; /* increment the code ptr unless it is null */
 
 #define P_NATIVE_BEGIN { SCOPED_SCRIPT_NATIVE_TIMER(ScopedNativeCallTimer);
diff --git a/Engine/Source/Runtime/CoreUObject/Public/UObject/UnrealType.h b/Engine/Source/Runtime/CoreUObject/Public/UObject/UnrealType.h
--- a/Engine/Source/Runtime/CoreUObject/Public/UObject/UnrealType.h
+++ b/Engine/Source/Runtime/CoreUObject/Public/UObject/UnrealType.h
@@ -1119,6 +1119,23 @@ class COREUOBJECT_API UNumericProperty : public UProperty
 	/** Return true if this property is for a integral or enum type **/
 	virtual bool IsInteger() const;
 
+	template <typename T>
+	bool CanHoldValue(T Value) const
+	{
+		if (!TIsFloatingPoint<T>::Value)
+		{
+			return CanHoldDoubleValueInternal(Value);
+		}
+		else if (TIsSigned<T>::Value)
+		{
+			return CanHoldSignedValueInternal(Value);
+		}
+		else
+		{
+			return CanHoldUnsignedValueInternal(Value);
+		}
+	}
+
 	/** Return true if this property is a UByteProperty with a non-null Enum **/
 	FORCEINLINE bool IsEnum() const
 	{
@@ -1188,6 +1205,11 @@ class COREUOBJECT_API UNumericProperty : public UProperty
 	// End of UNumericProperty interface
 
 	static uint8 ReadEnumAsUint8(FArchive& Ar, UStruct* DefaultsStruct, const FPropertyTag& Tag);
+
+private:
+	virtual bool CanHoldDoubleValueInternal  (double Value) const PURE_VIRTUAL(UNumericProperty::CanHoldDoubleValueInternal,   return false;);
+	virtual bool CanHoldSignedValueInternal  (int64  Value) const PURE_VIRTUAL(UNumericProperty::CanHoldSignedValueInternal,   return false;);
+	virtual bool CanHoldUnsignedValueInternal(uint64 Value) const PURE_VIRTUAL(UNumericProperty::CanHoldUnsignedValueInternal, return false;);
 };
 
 template<typename InTCppType>
@@ -1287,6 +1309,12 @@ public:
 			}
 			return true;
 		}
+		else if (Tag.Type == NAME_EnumProperty)
+		{
+			uint8 PreviousValue = this->ReadEnumAsUint8(Ar, DefaultsStruct, Tag);
+			this->SetPropertyValue_InContainer(Data, PreviousValue, Tag.ArrayIndex);
+			return true;
+		}
 		else if (Tag.Type == NAME_UInt16Property)
 		{
 			ConvertFromInt<uint16>(Ar, Data, Tag);
@@ -1355,6 +1383,22 @@ public:
 		return TTypeFundamentals::GetPropertyValue(Data);
 	}
 	// End of UNumericProperty interface
+
+private:
+	virtual bool CanHoldDoubleValueInternal(double Value) const
+	{
+		return (double)(InTCppType)Value == Value;
+	}
+
+	virtual bool CanHoldSignedValueInternal(int64 Value) const
+	{
+		return (int64)(InTCppType)Value == Value;
+	}
+
+	virtual bool CanHoldUnsignedValueInternal(uint64 Value) const
+	{
+		return (uint64)(InTCppType)Value == Value;
+	}
 };
 
 /*-----------------------------------------------------------------------------
diff --git a/Engine/Source/Runtime/Engine/Classes/Animation/Skeleton.h b/Engine/Source/Runtime/Engine/Classes/Animation/Skeleton.h
--- a/Engine/Source/Runtime/Engine/Classes/Animation/Skeleton.h
+++ b/Engine/Source/Runtime/Engine/Classes/Animation/Skeleton.h
@@ -510,7 +510,7 @@ public:
 	typedef TArray<FBoneNode> FBoneTreeType;
 
 	/** Runtime built mapping table between SkeletalMeshes, and LinkupCache array indices. */
-	TMap<TAutoWeakObjectPtr<class USkeletalMesh>, int32> SkelMesh2LinkupCache;
+	TMap<TWeakObjectPtr<USkeletalMesh>, int32> SkelMesh2LinkupCache;
 
 #if WITH_EDITORONLY_DATA
 
diff --git a/Engine/Source/Runtime/Engine/Classes/Components/ReflectionCaptureComponent.h b/Engine/Source/Runtime/Engine/Classes/Components/ReflectionCaptureComponent.h
--- a/Engine/Source/Runtime/Engine/Classes/Components/ReflectionCaptureComponent.h
+++ b/Engine/Source/Runtime/Engine/Classes/Components/ReflectionCaptureComponent.h
@@ -5,6 +5,55 @@
 
 class FReflectionCaptureProxy;
 
+/*
+* Refcounted class to pass around uncompressed cubemap data and track memory, for use with FReflectionCaptureFullHDR::GetUncompressedData
+*/
+class FReflectionCaptureUncompressedData : public FRefCountedObject
+{
+public:
+	FReflectionCaptureUncompressedData() :
+		TrackedBufferSize(0)
+	{}
+
+	FReflectionCaptureUncompressedData( uint32 InSizeBytes ) :
+		TrackedBufferSize(0)
+	{
+		CubemapDataArray.Empty(InSizeBytes);
+		CubemapDataArray.AddUninitialized(InSizeBytes);
+		UpdateMemoryTracking();
+	}
+
+	~FReflectionCaptureUncompressedData()
+	{
+		DEC_MEMORY_STAT_BY(STAT_ReflectionCaptureMemory, CubemapDataArray.Num() );
+	}
+
+	uint8* GetData( int Offset = 0 )
+	{
+		return &CubemapDataArray[Offset];
+	}
+
+	const uint32 Size() const 
+	{ 
+		return CubemapDataArray.Num();
+	}
+
+	TArray<uint8>& GetArray() 
+	{
+		return CubemapDataArray;
+	}
+
+	void UpdateMemoryTracking()
+	{
+		INC_MEMORY_STAT_BY( STAT_ReflectionCaptureMemory, CubemapDataArray.Num() - TrackedBufferSize);
+		TrackedBufferSize = CubemapDataArray.Num();
+	}
+private:
+	int32 TrackedBufferSize;
+	TArray<uint8> CubemapDataArray;
+};
+
+
 class FReflectionCaptureFullHDR
 {
 public:
@@ -17,6 +66,11 @@ public:
 	TArray<uint8> CompressedCapturedData;
 	int32 CubemapSize;
 
+	/**
+	* Generated at cook time. Avoids decompression overhead in GetUncompressedData
+	*/
+	TRefCountPtr<FReflectionCaptureUncompressedData> StoredUncompressedData;
+
 	/** Destructor. */
 	~FReflectionCaptureFullHDR();
 
@@ -24,30 +78,35 @@ public:
 	ENGINE_API void InitializeFromUncompressedData(const TArray<uint8>& UncompressedData, int32 CubmapSize);
 
 	/** Decompresses the capture data. */
-	ENGINE_API void GetUncompressedData(TArray<uint8>& UncompressedData) const;
+	ENGINE_API TRefCountPtr<FReflectionCaptureUncompressedData> GetUncompressedData() const;
 
-	ENGINE_API TArray<uint8>& GetCapturedDataForSM4Load() 
+	ENGINE_API TRefCountPtr<FReflectionCaptureUncompressedData> GetCapturedDataForSM4Load()
 	{
-		GetUncompressedData(CapturedDataForSM4Load);
+		CapturedDataForSM4Load = GetUncompressedData();
 		return CapturedDataForSM4Load;
 	}
 
-private:
+	ENGINE_API bool HasValidData() const { return StoredUncompressedData != nullptr || CompressedCapturedData.Num() > 0; }
 
+private:
 	/** 
 	 * This is used to pass the uncompressed capture data to the RT on load for SM4. 
 	 * The game thread must not modify it again after sending read commands to the RT.
 	 */
-	TArray<uint8> CapturedDataForSM4Load;
+	TRefCountPtr<FReflectionCaptureUncompressedData> CapturedDataForSM4Load;
 };
 
 struct FReflectionCaptureEncodedHDRDerivedData : FRefCountedObject
 {
+	FReflectionCaptureEncodedHDRDerivedData()
+	{
+		CapturedData = new FReflectionCaptureUncompressedData;
+	}
 	/** 
 	 * The uncompressed encoded HDR capture data, with all mips and faces laid out linearly. 
 	 * This is read and written from the rendering thread directly after load, so the game thread must not access it again.
 	 */
-	TArray<uint8> CapturedData;
+	TRefCountPtr<FReflectionCaptureUncompressedData> CapturedData;
 
 	/** Destructor. */
 	virtual ~FReflectionCaptureEncodedHDRDerivedData();
@@ -59,7 +118,7 @@ struct FReflectionCaptureEncodedHDRDerivedData : FRefCountedObject
 		// Data / (6 cubemaps * sizeof(FColor)) = (4*topMip - lastMip)/3
 		// when lastMip = 1 simplifies to the following:
 		// see https://en.wikipedia.org/wiki/1/4_%2B_1/16_%2B_1/64_%2B_1/256_%2B_%E2%8B%AF for maths
-		return (int32)sqrt((float)(2 * sizeof(FColor) + CapturedData.Num()) / (float)(8 * sizeof(FColor)));
+		return (int32)sqrt((float)(2 * sizeof(FColor) + CapturedData->Size()) / (float)(8 * sizeof(FColor)));
 	}
 
 	/** Generates encoded HDR data from full HDR data and saves it in the DDC, or loads an already generated version from the DDC. */
@@ -152,6 +211,9 @@ class UReflectionCaptureComponent : public USceneComponent
 	inline float* GetAverageBrightnessPtr() { return &AverageBrightness; }
 	inline const float* GetAverageBrightnessPtr() const { return &AverageBrightness; }
 
+	/** Issues a renderthread command to release the data, and NULLs the pointer on the gamethread */
+	ENGINE_API void ReleaseHDRData();
+
 	ENGINE_API static int32 GetReflectionCaptureSize_GameThread();
 	ENGINE_API static int32 GetReflectionCaptureSize_RenderThread();
 
diff --git a/Engine/Source/Runtime/Engine/Classes/Components/StaticMeshComponent.h b/Engine/Source/Runtime/Engine/Classes/Components/StaticMeshComponent.h
--- a/Engine/Source/Runtime/Engine/Classes/Components/StaticMeshComponent.h
+++ b/Engine/Source/Runtime/Engine/Classes/Components/StaticMeshComponent.h
@@ -6,6 +6,7 @@
 #include "Components/MeshComponent.h"
 #include "Runtime/RenderCore/Public/PackedNormal.h"
 #include "RawIndexBuffer.h"
+#include "UniquePtr.h"
 #include "StaticMeshComponent.generated.h"
 
 class FColorVertexBuffer;
@@ -68,7 +69,7 @@ struct FStaticMeshComponentLODInfo
 	FMeshMapBuildData* LegacyMapBuildData;
 
 	/** Transient override lightmap data, used by landscape grass. */
-	TScopedPointer<FMeshMapBuildData> OverrideMapBuildData;
+	TUniquePtr<FMeshMapBuildData> OverrideMapBuildData;
 
 	/** Vertex data cached at the time this LOD was painted, if any */
 	UPROPERTY()
diff --git a/Engine/Source/Runtime/Engine/Classes/Components/TimelineComponent.h b/Engine/Source/Runtime/Engine/Classes/Components/TimelineComponent.h
--- a/Engine/Source/Runtime/Engine/Classes/Components/TimelineComponent.h
+++ b/Engine/Source/Runtime/Engine/Classes/Components/TimelineComponent.h
@@ -249,7 +249,7 @@ private:
 
 	/** Cached property pointer for setting timeline direction */
 	UPROPERTY(Transient, NotReplicated)
-	UByteProperty* DirectionProperty;
+	UProperty* DirectionProperty;
 
 public:
 	FTimeline()
diff --git a/Engine/Source/Runtime/Engine/Classes/EditorFramework/AssetImportData.h b/Engine/Source/Runtime/Engine/Classes/EditorFramework/AssetImportData.h
--- a/Engine/Source/Runtime/Engine/Classes/EditorFramework/AssetImportData.h
+++ b/Engine/Source/Runtime/Engine/Classes/EditorFramework/AssetImportData.h
@@ -62,6 +62,9 @@ class ENGINE_API UAssetImportData : public UObject
 public:
 	GENERATED_UCLASS_BODY()
 
+	/** Only valid in the editor */
+	virtual bool IsEditorOnly() const override { return true; }
+
 #if WITH_EDITORONLY_DATA
 
 	/** Path to the resource used to construct this static mesh. Relative to the object's package, BaseDir() or absolute */
diff --git a/Engine/Source/Runtime/Engine/Classes/Engine/BlueprintGeneratedClass.h b/Engine/Source/Runtime/Engine/Classes/Engine/BlueprintGeneratedClass.h
--- a/Engine/Source/Runtime/Engine/Classes/Engine/BlueprintGeneratedClass.h
+++ b/Engine/Source/Runtime/Engine/Classes/Engine/BlueprintGeneratedClass.h
@@ -765,4 +765,6 @@ protected:
 private:
 	/** List of native class-owned properties that differ from defaults. This is used to optimize property initialization during post-construction by minimizing the number of native class-owned property values that get copied to the new instance. */
 	TIndirectArray<FCustomPropertyListNode> CustomPropertyListForPostConstruction;
+	/** In some cases UObject::ConditionalPostLoad() code calls PostLoadDefaultObject() on a class that's still being serialized. */
+	FCriticalSection SerializeAndPostLoadCritical;
 };
diff --git a/Engine/Source/Runtime/Engine/Classes/Engine/CoreSettings.h b/Engine/Source/Runtime/Engine/Classes/Engine/CoreSettings.h
--- a/Engine/Source/Runtime/Engine/Classes/Engine/CoreSettings.h
+++ b/Engine/Source/Runtime/Engine/Classes/Engine/CoreSettings.h
@@ -111,6 +111,10 @@ extern ENGINE_API float GPriorityAsyncLoadingExtraTime;
 extern ENGINE_API float GLevelStreamingActorsUpdateTimeLimit;
 /** Batching granularity used to register actor components during level streaming. */
 extern ENGINE_API int32 GLevelStreamingComponentsRegistrationGranularity;
+/** Batching granularity used to unregister actor components during level streaming.  */
+extern ENGINE_API int32 GLevelStreamingComponentsUnregistrationGranularity;
+/** Maximum allowed time to spend for actor unregistration steps during level streaming (ms per frame). If this is 0.0 then we don't timeslice.*/
+extern ENGINE_API float GLevelStreamingUnregisterComponentsTimeLimit;
 
 /**
 * Implements the settings for garbage collection.
diff --git a/Engine/Source/Runtime/Engine/Classes/Engine/Level.h b/Engine/Source/Runtime/Engine/Classes/Engine/Level.h
--- a/Engine/Source/Runtime/Engine/Classes/Engine/Level.h
+++ b/Engine/Source/Runtime/Engine/Classes/Engine/Level.h
@@ -510,10 +510,17 @@ public:
 	uint8										bIsBeingRemoved:1;
 	/** Current index into actors array for updating components.							*/
 	int32										CurrentActorIndexForUpdateComponents;
+	/** Current index into actors array for updating components.							*/
+	int32										CurrentActorIndexForUnregisterComponents;
+
 
 	/** Whether the level is currently pending being made visible.							*/
+	DEPRECATED(4.15, "Use HasVisibilityChangeRequestPending")
 	bool HasVisibilityRequestPending() const;
 
+	/** Whether the level is currently pending being made invisible or visible.				*/
+	bool HasVisibilityChangeRequestPending() const;
+
 	// Event on level transform changes
 	DECLARE_MULTICAST_DELEGATE_OneParam(FLevelTransformEvent, const FTransform&);
 	FLevelTransformEvent OnApplyLevelTransform;
@@ -629,6 +636,15 @@ public:
 	 */
 	void IncrementalUpdateComponents( int32 NumComponentsToUpdate, bool bRerunConstructionScripts );
 
+	/**
+	* Incrementally unregisters all components of actors associated with this level.
+	* This is done at the granularity of actors (individual actors have all of their components unregistered)
+    *
+	* @param NumComponentsToUnregister		Minimum number of components to unregister in this run, 0 for all
+	*/
+	bool IncrementalUnregisterComponents(int32 NumComponentsToUnregister);
+
+
 	/**
 	 * Invalidates the cached data used to render the level's UModel.
 	 */
diff --git a/Engine/Source/Runtime/Engine/Classes/Engine/StaticMesh.h b/Engine/Source/Runtime/Engine/Classes/Engine/StaticMesh.h
--- a/Engine/Source/Runtime/Engine/Classes/Engine/StaticMesh.h
+++ b/Engine/Source/Runtime/Engine/Classes/Engine/StaticMesh.h
@@ -4,6 +4,7 @@
 #include "Interfaces/Interface_CollisionDataProvider.h"
 #include "Interfaces/Interface_AssetUserData.h"
 #include "MeshMerging.h"
+#include "UniquePtr.h"
 #include "StaticMesh.generated.h"
 
 /** The maximum number of static mesh LODs allowed. */
@@ -398,7 +399,7 @@ class UStaticMesh : public UObject, public IInterface_CollisionDataProvider, pub
 	GENERATED_UCLASS_BODY()
 
 	/** Pointer to the data used to render this static mesh. */
-	TScopedPointer<class FStaticMeshRenderData> RenderData;
+	TUniquePtr<class FStaticMeshRenderData> RenderData;
 
 #if WITH_EDITORONLY_DATA
 	static const float MinimumAutoLODPixelError;
diff --git a/Engine/Source/Runtime/Engine/Classes/Engine/UserDefinedEnum.h b/Engine/Source/Runtime/Engine/Classes/Engine/UserDefinedEnum.h
--- a/Engine/Source/Runtime/Engine/Classes/Engine/UserDefinedEnum.h
+++ b/Engine/Source/Runtime/Engine/Classes/Engine/UserDefinedEnum.h
@@ -52,14 +52,14 @@ public:
 	 *	@param EnumeratorIndex	old index
 	 *	@return	new index
 	 */
-	virtual int32 ResolveEnumerator(FArchive& Ar, int32 EnumeratorValue) const override;
+	virtual int64 ResolveEnumerator(FArchive& Ar, int64 EnumeratorValue) const override;
 
 	/**
 	 * @return	The enum string at the specified index.
 	 */
 	virtual FText GetEnumText(int32 InIndex) const override;
 
-	virtual bool SetEnums(TArray<TPair<FName, uint8>>& InNames, ECppForm InCppForm, bool bAddMaxKeyIfMissing = true) override;
+	virtual bool SetEnums(TArray<TPair<FName, int64>>& InNames, ECppForm InCppForm, bool bAddMaxKeyIfMissing = true) override;
 
 #if WITH_EDITOR
 	//~ Begin UObject Interface
diff --git a/Engine/Source/Runtime/Engine/Classes/Engine/World.h b/Engine/Source/Runtime/Engine/Classes/Engine/World.h
--- a/Engine/Source/Runtime/Engine/Classes/Engine/World.h
+++ b/Engine/Source/Runtime/Engine/Classes/Engine/World.h
@@ -734,9 +734,13 @@ class ENGINE_API UWorld final : public UObject, public FNetworkNotify
 	UPROPERTY()
 	FString										StreamingLevelsPrefix;
 	
-	/** Pointer to the current level in the queue to be made visible, NULL if none are pending.									*/
+	/** Pointer to the current level in the queue to be made visible, NULL if none are pending.					*/
 	UPROPERTY(Transient)
 	class ULevel*								CurrentLevelPendingVisibility;
+
+	/** Pointer to the current level in the queue to be made invisible, NULL if none are pending.					*/
+	UPROPERTY(Transient)
+	class ULevel*								CurrentLevelPendingInvisibility;
 	
 	/** Fake NetDriver for capturing network traffic to record demos															*/
 	UPROPERTY()
@@ -889,19 +893,19 @@ public:
 private:
 
 	/** List of all the controllers in the world. */
-	TArray<TAutoWeakObjectPtr<class AController> > ControllerList;
+	TArray<TWeakObjectPtr<class AController> > ControllerList;
 
 	/** List of all the player controllers in the world. */
-	TArray<TAutoWeakObjectPtr<class APlayerController> > PlayerControllerList;
+	TArray<TWeakObjectPtr<class APlayerController> > PlayerControllerList;
 
 	/** List of all the pawns in the world. */
-	TArray<TAutoWeakObjectPtr<class APawn> > PawnList;	
+	TArray<TWeakObjectPtr<class APawn> > PawnList;
 
 	/** List of all the cameras in the world that auto-activate for players. */
-	TArray<TAutoWeakObjectPtr<ACameraActor> > AutoCameraActorList;
+	TArray<TWeakObjectPtr<ACameraActor> > AutoCameraActorList;
 
 	/** List of all physics volumes in the world. Does not include the DefaultPhysicsVolume. */
-	TArray<TAutoWeakObjectPtr<APhysicsVolume> > NonDefaultPhysicsVolumeList;
+	TArray<TWeakObjectPtr<APhysicsVolume> > NonDefaultPhysicsVolumeList;
 
 	/** Physics scene for this world. */
 	FPhysScene*									PhysicsScene;
@@ -2143,7 +2147,7 @@ public:
 	 *
 	 * @param Level			Level object we should remove
 	 */
-	void RemoveFromWorld( ULevel* Level );
+	void RemoveFromWorld( ULevel* Level, bool bAllowIncrementalRemoval = false );
 
 	/**
 	 * Updates sub-levels (load/unload/show/hide) using streaming levels current state
diff --git a/Engine/Source/Runtime/Engine/Classes/Materials/Material.h b/Engine/Source/Runtime/Engine/Classes/Materials/Material.h
--- a/Engine/Source/Runtime/Engine/Classes/Materials/Material.h
+++ b/Engine/Source/Runtime/Engine/Classes/Materials/Material.h
@@ -8,6 +8,7 @@
 #include "Materials/MaterialInterface.h"
 #include "Materials/MaterialExpressionMaterialFunctionCall.h"
 #include "Materials/MaterialFunction.h"
+#include "UniquePtr.h"
 
 #include "Material.generated.h"
 
@@ -1243,12 +1244,12 @@ public:
 	 * Backs up all material shaders to memory through serialization, organized by FMaterialShaderMap. 
 	 * This will also clear all FMaterialShaderMap references to FShaders.
 	 */
-	ENGINE_API static void BackupMaterialShadersToMemory(TMap<class FMaterialShaderMap*, TScopedPointer<TArray<uint8> > >& ShaderMapToSerializedShaderData);
+	ENGINE_API static void BackupMaterialShadersToMemory(TMap<class FMaterialShaderMap*, TUniquePtr<TArray<uint8> > >& ShaderMapToSerializedShaderData);
 
 	/** 
 	 * Recreates FShaders for FMaterialShaderMap's from the serialized data.  Shader maps may not be complete after this due to changes in the shader keys.
 	 */
-	ENGINE_API static void RestoreMaterialShadersFromMemory(const TMap<class FMaterialShaderMap*, TScopedPointer<TArray<uint8> > >& ShaderMapToSerializedShaderData);
+	ENGINE_API static void RestoreMaterialShadersFromMemory(const TMap<class FMaterialShaderMap*, TUniquePtr<TArray<uint8> > >& ShaderMapToSerializedShaderData);
 
 	/** Builds a map from UMaterialInterface name to the shader maps that are needed for rendering on the given platform. */
 	ENGINE_API static void CompileMaterialsForRemoteRecompile(
diff --git a/Engine/Source/Runtime/Engine/Classes/Materials/MaterialParameterCollection.h b/Engine/Source/Runtime/Engine/Classes/Materials/MaterialParameterCollection.h
--- a/Engine/Source/Runtime/Engine/Classes/Materials/MaterialParameterCollection.h
+++ b/Engine/Source/Runtime/Engine/Classes/Materials/MaterialParameterCollection.h
@@ -5,6 +5,7 @@
  */
 
 #pragma once
+#include "UniquePtr.h"
 #include "MaterialParameterCollection.generated.h"
 
 /** Base struct for collection parameters */
@@ -116,7 +117,7 @@ private:
 	/** Default resource used when no instance is available. */
 	class FMaterialParameterCollectionInstanceResource* DefaultResource;
 
-	TScopedPointer<FUniformBufferStruct> UniformBufferStruct;
+	TUniquePtr<FUniformBufferStruct> UniformBufferStruct;
 
 	void CreateBufferStruct();
 
diff --git a/Engine/Source/Runtime/Engine/Private/AI/Navigation/RecastNavMeshGenerator.cpp b/Engine/Source/Runtime/Engine/Private/AI/Navigation/RecastNavMeshGenerator.cpp
--- a/Engine/Source/Runtime/Engine/Private/AI/Navigation/RecastNavMeshGenerator.cpp
+++ b/Engine/Source/Runtime/Engine/Private/AI/Navigation/RecastNavMeshGenerator.cpp
@@ -3958,7 +3958,7 @@ void FRecastNavMeshGenerator::SortPendingBuildTiles()
 	// Collect players positions
 	for (FConstPlayerControllerIterator PlayerIt = CurWorld->GetPlayerControllerIterator(); PlayerIt; ++PlayerIt)
 	{
-		if (*PlayerIt && (*PlayerIt)->GetPawn() != NULL)
+		if (PlayerIt->IsValid() && PlayerIt->Get()->GetPawn() != NULL)
 		{
 			const FVector2D SeedLoc((*PlayerIt)->GetPawn()->GetActorLocation());
 			SeedLocations.Add(SeedLoc);
