Copying //UE4/Dev-Core to //UE4/Dev-Main (Source: //UE4/Dev-Core @ 3228803)

diff --git a/Engine/Source/Runtime/CoreUObject/Private/Serialization/AsyncLoading.cpp b/Engine/Source/Runtime/CoreUObject/Private/Serialization/AsyncLoading.cpp
--- a/Engine/Source/Runtime/CoreUObject/Private/Serialization/AsyncLoading.cpp
+++ b/Engine/Source/Runtime/CoreUObject/Private/Serialization/AsyncLoading.cpp
@@ -36,13 +36,10 @@
 #include "Serialization/AsyncLoadingPrivate.h"
 #include "UObject/UObjectHash.h"
 #include "UniquePtr.h"
+#include "Serialization/BufferReader.h"
 
 #define FIND_MEMORY_STOMPS (1 && (PLATFORM_WINDOWS || PLATFORM_LINUX) && !WITH_EDITORONLY_DATA)
 
-#ifndef USE_EVENT_DRIVEN_ASYNC_LOAD
-#error "USE_EVENT_DRIVEN_ASYNC_LOAD must be defined"
-#endif
-
 DEFINE_LOG_CATEGORY(LogLoadingDev);
 
 //#pragma clang optimize off
@@ -146,7 +143,6 @@ static FAutoConsoleVariableRef CVarTimeLimitExceededMinTime(
 	);
 
 static int32 GPreloadPackageDependencies = 0;
-#if !USE_EVENT_DRIVEN_ASYNC_LOAD // we don't want preloading of dependencies with event driven loading, we would rather find the dependencies naturally.
 static FAutoConsoleVariableRef CVarPreloadPackageDependencies(
 	TEXT("s.PreloadPackageDependencies"),
 	GPreloadPackageDependencies,
@@ -155,12 +151,11 @@ static FAutoConsoleVariableRef CVarPreloadPackageDependencies(
 	TEXT("1 - Preload package dependencies. Faster but requires asset registry data to be loaded into memory\n"),
 	ECVF_Default
 	);
-#endif
 
-static int32 GEventDrivenLoaderEnabled = 0;
+static int32 GEventDrivenLoaderEnabledInCookedBuilds = 0;
 static FAutoConsoleVariableRef CVarEventDrivenLoaderEnabled(
 	TEXT("s.EventDrivenLoaderEnabled"),
-	GEventDrivenLoaderEnabled,
+	GEventDrivenLoaderEnabledInCookedBuilds,
 	TEXT("Placeholder console variable, currently not used in runtime."),
 	ECVF_Default
 );
@@ -211,8 +206,6 @@ FORCEINLINE bool FAsyncPackage::IsTimeLimitExceeded()
 	return AsyncLoadingThread.IsAsyncLoadingSuspended() || ::IsTimeLimitExceeded(TickStartTime, bUseTimeLimit, TimeLimit, LastTypeOfWorkPerformed, LastObjectWorkWasPerformedOn);
 }
 
-#if USE_NEW_ASYNC_IO
-
 DEFINE_LOG_CATEGORY_STATIC(LogAsyncArchive, Display, All);
 DECLARE_MEMORY_STAT(TEXT("FArchiveAsync2 Buffers"), STAT_FArchiveAsync2Mem, STATGROUP_Memory);
 
@@ -255,7 +248,6 @@ FORCEINLINE void FArchiveAsync2::LogItem(const TCHAR* Item, int64 Offset, int64
 			*FileName);
 	}
 }
-#endif
 
 
 #if LOOKING_FOR_PERF_ISSUES
@@ -590,7 +582,9 @@ int32 FAsyncLoadingThread::CreateAsyncPackagesFromQueue(bool bUseTimeLimit, bool
 		{
 			IAssetRegistryInterface* AssetRegistry = nullptr;
 
-			if (GPreloadPackageDependencies && IsPlatformFileCompatibleWithDependencyPreloading())
+			if (GPreloadPackageDependencies && IsPlatformFileCompatibleWithDependencyPreloading() && 
+					// we don't want preloading of dependencies with event driven loading, we would rather find the dependencies naturally.
+				 (!IsEventDrivenLoaderEnabledInCookedBuilds() || !FPlatformProperties::RequiresCookedData()))
 			{
 				AssetRegistry = IAssetRegistryInterface::GetPtr();
 			}
@@ -613,8 +607,6 @@ int32 FAsyncLoadingThread::CreateAsyncPackagesFromQueue(bool bUseTimeLimit, bool
 	return NumCreated;
 }
 
-#if USE_EVENT_DRIVEN_ASYNC_LOAD
-
 static FThreadSafeCounter AsyncPackageSerialNumber;
 
 FName FUnsafeWeakAsyncPackagePtr::HumanReadableStringForDebugging() const
@@ -1028,9 +1020,6 @@ void FEventLoadGraph::CheckForCycles()
 	}
 }
 
-//@todoio this won't work with UT etc, need *_API or a real api to work with the pak precacher.
-extern bool GPakCache_AcceptPrecacheRequests;
-
 struct FPrecacheCallbackHandler
 {
 	FAsyncFileCallBack PrecacheCallBack;
@@ -2399,19 +2388,20 @@ void FAsyncPackage::DumpDependencies(const TCHAR* Label, UObject* Obj)
 	{
 		UE_LOG(LogStreaming, Error, TEXT("    %s has no linker"), *Package->GetFullName());
 	}
-	for (int32 LocalExportIndex = 0; LocalExportIndex < Package->LinkerLoad->ExportMap.Num(); LocalExportIndex++)
+	else
 	{
-		FObjectExport& Export = Package->LinkerLoad->ExportMap[LocalExportIndex];
-		if (Export.Object == Obj || Export.Object == nullptr)
+		for (int32 LocalExportIndex = 0; LocalExportIndex < Package->LinkerLoad->ExportMap.Num(); LocalExportIndex++)
 		{
-			if (Export.ObjectName == Obj->GetFName())
+			FObjectExport& Export = Package->LinkerLoad->ExportMap[LocalExportIndex];
+			if (Export.Object == Obj || Export.Object == nullptr)
 			{
-				DumpDependencies(TEXT(""), Package->LinkerLoad, FPackageIndex::FromExport(LocalExportIndex));
+				if (Export.ObjectName == Obj->GetFName())
+				{
+					DumpDependencies(TEXT(""), Package->LinkerLoad, FPackageIndex::FromExport(LocalExportIndex));
+				}
 			}
 		}
 	}
-
-
 }
 
 void FAsyncPackage::DumpDependencies(const TCHAR* Label, FLinkerLoad* DumpLinker, FPackageIndex DumpExportIndex)
@@ -2537,6 +2527,11 @@ UObject* FAsyncPackage::EventDrivenIndexToObject(FPackageIndex Index, bool bChec
 			UE_LOG(LogStreaming, Fatal, TEXT("Missing Dependency, request for %s but it was still has RF_NeedLoad."), *Linker->ImpExp(Index).ObjectName.ToString());
 		}
 	}
+	if (Result)
+	{
+		UE_CLOG(Result->HasAnyInternalFlags(EInternalObjectFlags::Unreachable), LogStreaming, Fatal, TEXT("Returning an object  (%s) from EventDrivenIndexToObject that is unreachable."), *Result->GetFullName());
+		checkSlow(ReferencedObjects.Contains(Result));
+	}
 	return Result;
 }
 
@@ -2620,18 +2615,6 @@ void FAsyncPackage::EventDrivenCreateExport(int32 LocalExportIndex)
 				Export.bExportLoadFailed = true;
 				return;
 			}
-			UStruct* SuperStruct = nullptr;
-			if (!Export.SuperIndex.IsNull())
-			{
-				SuperStruct = CastEventDrivenIndexToObject<UStruct>(Export.SuperIndex, true, FPackageIndex::FromExport(LocalExportIndex));
-				if (!SuperStruct)
-				{
-					// see FLinkerLoad::CreateExport, there may be some more we can do here
-					UE_LOG(LogStreaming, Error, TEXT("Could not find SuperStruct %s to create %s"), *Linker->ImpExp(Export.SuperIndex).ObjectName.ToString(), *Export.ObjectName.ToString());
-					Export.bExportLoadFailed = true;
-					return;
-				}
-			}
 			UObject* ThisParent = nullptr;
 			if (!Export.OuterIndex.IsNull())
 			{
@@ -2810,21 +2793,6 @@ void FAsyncPackage::EventDrivenCreateExport(int32 LocalExportIndex)
 				Export.Object->SetLinker(Linker, LocalExportIndex);
 				check(Export.Object->GetClass() == LoadClass);
 				check(NewName == Export.ObjectName);
-
-				// If it's a struct or class, set its parent.
-				if (UStruct* Struct = dynamic_cast<UStruct*>(Export.Object))
-				{
-					if (SuperStruct)
-					{
-						Struct->SetSuperStruct(SuperStruct);
-					}
-
-					// If it's a class, bind it to C++.
-					if (UClass* ClassObject = dynamic_cast<UClass*>(Export.Object))
-					{
-						ClassObject->Bind();
-					}
-				}
 			}
 		}
 	}
@@ -2867,14 +2835,26 @@ void FAsyncPackage::EventDrivenSerializeExport(int32 LocalExportIndex)
 		// If this is a struct, make sure that its parent struct is completely loaded
 		if (UStruct* Struct = dynamic_cast<UStruct*>(Object))
 		{
-			Cls = dynamic_cast<UClass*>(Object);
-			UStruct* SuperStruct = Struct->GetSuperStruct();
-			if (SuperStruct && SuperStruct->HasAnyFlags(RF_NeedLoad))
+			UStruct* SuperStruct = nullptr;
+			if (!Export.SuperIndex.IsNull())
 			{
-				UE_LOG(LogStreaming, Fatal, TEXT("SuperStruct %s had RF_NeedLoad while creating %s"), *SuperStruct->GetFullName(), *Export.ObjectName.ToString());
-				Export.bExportLoadFailed = true;
-				return;
-			}		
+				SuperStruct = CastEventDrivenIndexToObject<UStruct>(Export.SuperIndex, true, FPackageIndex::FromExport(LocalExportIndex));
+				if (!SuperStruct)
+				{
+					// see FLinkerLoad::CreateExport, there may be some more we can do here
+					UE_LOG(LogStreaming, Fatal, TEXT("Could not find SuperStruct %s to create %s"), *Linker->ImpExp(Export.SuperIndex).ObjectName.ToString(), *Export.ObjectName.ToString());
+					Export.bExportLoadFailed = true;
+					return;
+				}
+			}
+			if (SuperStruct)
+			{
+				Struct->SetSuperStruct(SuperStruct);
+				if (UClass* ClassObject = dynamic_cast<UClass*>(Object))
+				{
+					ClassObject->Bind();
+				}
+			}
 		}
 		check(Export.SerialOffset >= CurrentBlockOffset && Export.SerialOffset + Export.SerialSize <= CurrentBlockOffset + CurrentBlockBytes);
 
@@ -3112,8 +3092,8 @@ void FAsyncPackage::MakeNextPrecacheRequestCurrent()
 		RemoveNode(EEventLoadNode::Export_StartIO, FPackageIndex::FromExport(LocalExportIndex));
 	}
 	Read->WaitCompletion();
-	delete Read;
 	PrecacheRequests.Remove(Read);
+	delete Read;	
 }
 
 void FAsyncPackage::FlushPrecacheBuffer()
@@ -3347,8 +3327,12 @@ void FAsyncPackage::Event_ProcessPostloadWait()
 		}
 		AlreadyHandled.Add(PoppedPtr);
 		FAsyncPackage* TestPkg(AsyncLoadingThread.GetPackage(PoppedPtr));
+		if (!TestPkg)
+		{
+			continue;
+		}
 		check(TestPkg != this);
-		if (!TestPkg || int32(TestPkg->AsyncPackageLoadingState) > int32(EAsyncPackageLoadingState::WaitingForPostLoad))
+		if (int32(TestPkg->AsyncPackageLoadingState) > int32(EAsyncPackageLoadingState::WaitingForPostLoad))
 		{
 			continue;
 		}
@@ -3370,8 +3354,12 @@ void FAsyncPackage::Event_ProcessPostloadWait()
 			if (!AlreadyHandled.Contains(MaybeRecursePtr))
 			{
 				FAsyncPackage* MaybeRecursePkg(AsyncLoadingThread.GetPackage(MaybeRecursePtr));
+				if (!MaybeRecursePkg)
+				{
+					continue;
+				}
 				check(MaybeRecursePkg != this);
-				if (!MaybeRecursePkg || int32(MaybeRecursePkg->AsyncPackageLoadingState) > int32(EAsyncPackageLoadingState::WaitingForPostLoad))
+				if (int32(MaybeRecursePkg->AsyncPackageLoadingState) > int32(EAsyncPackageLoadingState::WaitingForPostLoad))
 				{
 					continue;
 				}
@@ -3416,7 +3404,8 @@ void FAsyncPackage::Event_StartPostload()
 		{
 			FObjectExport& Export = Linker->ExportMap[LocalExportIndex];
 			UObject* Object = Export.Object;
-			if (Object && Object->HasAnyFlags(RF_NeedPostLoad))
+			checkSlow(!(Object && !ReferencedObjects.Contains(Object)));
+			if (Object && (Object->HasAnyFlags(RF_NeedPostLoad) || Object->HasAnyInternalFlags(EInternalObjectFlags::AsyncLoading)))
 			{
 				ObjLoaded.Add(Object);
 			}
@@ -3567,17 +3556,18 @@ void FAsyncPackage::FireNode(FEventLoadNodePtr& NodeToFire)
 }
 
 
-
-#endif
-
 void FAsyncLoadingThread::InsertPackage(FAsyncPackage* Package, bool bReinsert, EAsyncPackageInsertMode InsertMode)
 {
 	checkSlow(IsInAsyncLoadThread());
 	check(!IsInGameThread() || !IsMultithreaded());
 
-#if USE_EVENT_DRIVEN_ASYNC_LOAD && DO_CHECK
-	FWeakAsyncPackagePtr WeakPtr(Package);
-	check(Package);
+#if DO_CHECK
+	FWeakAsyncPackagePtr WeakPtr;
+	if (GEventDrivenLoaderEnabled)
+	{
+		WeakPtr = Package;
+		check(Package);
+	}
 #endif
 
 	if (!bReinsert)
@@ -3628,16 +3618,14 @@ void FAsyncLoadingThread::InsertPackage(FAsyncPackage* Package, bool bReinsert,
 		if (!bReinsert)
 		{
 			AsyncPackageNameLookup.Add(Package->GetPackageName(), Package);
-#if USE_EVENT_DRIVEN_ASYNC_LOAD
-			// @todo If this is a reinsert for some priority thing, well we don't go back and retract the stuff in flight to adjust the priority of events
-			QueueEvent_CreateLinker(Package);
-#endif
+			if (GEventDrivenLoaderEnabled)
+			{
+				// @todo If this is a reinsert for some priority thing, well we don't go back and retract the stuff in flight to adjust the priority of events
+				QueueEvent_CreateLinker(Package);
+			}
 		}
-
 	}
-#if USE_EVENT_DRIVEN_ASYNC_LOAD
-	check(GetPackage(WeakPtr) == Package);
-#endif
+	check(!GEventDrivenLoaderEnabled || GetPackage(WeakPtr) == Package);
 }
 
 void FAsyncLoadingThread::AddToLoadedPackages(FAsyncPackage* Package)
@@ -3652,7 +3640,6 @@ void FAsyncLoadingThread::AddToLoadedPackages(FAsyncPackage* Package)
 	}
 }
 
-#if USE_EVENT_DRIVEN_ASYNC_LOAD
 #if !UE_BUILD_SHIPPING && !UE_BUILD_TEST
 static FThreadSafeCounter RecursionNotAllowed;
 struct FScopedRecursionNotAllowed
@@ -3667,7 +3654,6 @@ struct FScopedRecursionNotAllowed
 	}
 };
 #endif
-#endif
 
 EAsyncPackageState::Type FAsyncLoadingThread::ProcessAsyncLoading(int32& OutPackagesProcessed, bool bUseTimeLimit /*= false*/, bool bUseFullTimeLimit /*= false*/, float TimeLimit /*= 0.0f*/, FFlushTree* FlushTree)
 {
@@ -3677,263 +3663,253 @@ EAsyncPackageState::Type FAsyncLoadingThread::ProcessAsyncLoading(int32& OutPack
 
 	// If we're not multithreaded and flushing async loading, update the thread heartbeat
 	const bool bNeedsHeartbeatTick = !bUseTimeLimit && !FAsyncLoadingThread::IsMultithreaded();
-
 	EAsyncPackageState::Type LoadingState = EAsyncPackageState::Complete;
 	OutPackagesProcessed = 0;
 
 	double TickStartTime = FPlatformTime::Seconds();
 
-#if USE_EVENT_DRIVEN_ASYNC_LOAD
+	if (GEventDrivenLoaderEnabled)
+	{
 #if !UE_BUILD_SHIPPING && !UE_BUILD_TEST
-	FScopedRecursionNotAllowed RecursionGuard;
+		FScopedRecursionNotAllowed RecursionGuard;
 #endif
-	FAsyncLoadingTickScope InAsyncLoadingTick;
-	uint32 LoopIterations = 0;
-
-	while (true)
-	{
-		if (bNeedsHeartbeatTick && (++LoopIterations) % 32 == 31)
-		{
-			// Update heartbeat after 32 events
-			FThreadHeartBeat::Get().HeartBeat();
-		}
-
-		if (FileOpenLogActive() && GPrecacheCallbackHandler.AnyIOOutstanding())
-		{
-			GPrecacheCallbackHandler.WaitForIO(10.0f); // We don't want the ideal load order to be affected by the behavior of the disk...so we always just wait for IO
-		}
-		bool bDidSomething = GPrecacheCallbackHandler.ProcessIncoming();
-		OutPackagesProcessed += (bDidSomething ? 1 : 0);
+		FAsyncLoadingTickScope InAsyncLoadingTick;
+		uint32 LoopIterations = 0;
 
-		if (IsTimeLimitExceeded(TickStartTime, bUseTimeLimit, TimeLimit, TEXT("ProcessIncoming"), nullptr))
+		while (true)
 		{
-			return EAsyncPackageState::TimeOut;
-		}
+			if (bNeedsHeartbeatTick && (++LoopIterations) % 32 == 31)
+			{
+				// Update heartbeat after 32 events
+				FThreadHeartBeat::Get().HeartBeat();
+			}
 
-		if (IsAsyncLoadingSuspended())
-		{
-			return EAsyncPackageState::TimeOut;
-		}
+			if (FileOpenLogActive() && GPrecacheCallbackHandler.AnyIOOutstanding())
+			{
+				GPrecacheCallbackHandler.WaitForIO(10.0f); // We don't want the ideal load order to be affected by the behavior of the disk...so we always just wait for IO
+			}
+			bool bDidSomething = GPrecacheCallbackHandler.ProcessIncoming();
+			OutPackagesProcessed += (bDidSomething ? 1 : 0);
 
-		{
-			const float RemainingTimeLimit = FMath::Max(0.0f, TimeLimit - (float)(FPlatformTime::Seconds() - TickStartTime));
-			int32 NumCreated = CreateAsyncPackagesFromQueue(bUseTimeLimit, bUseFullTimeLimit, RemainingTimeLimit);
-			OutPackagesProcessed += NumCreated;
-			bDidSomething = NumCreated > 0 || bDidSomething;
-			if (IsTimeLimitExceeded(TickStartTime, bUseTimeLimit, TimeLimit, TEXT("CreateAsyncPackagesFromQueue")))
+			if (IsTimeLimitExceeded(TickStartTime, bUseTimeLimit, TimeLimit, TEXT("ProcessIncoming"), nullptr))
 			{
 				return EAsyncPackageState::TimeOut;
 			}
-		}
-		if (bDidSomething)
-		{
-			continue;
-		}
 
-		{
-			FAsyncLoadEventArgs Args;
-			Args.bUseTimeLimit = bUseTimeLimit;
-			Args.TickStartTime = TickStartTime;
-			Args.TimeLimit = TimeLimit;
+			if (IsAsyncLoadingSuspended())
+			{
+				return EAsyncPackageState::TimeOut;
+			}
 
-			Args.OutLastTypeOfWorkPerformed = nullptr;
-			Args.OutLastObjectWorkWasPerformedOn = nullptr;
-			if (EventQueue.PopAndExecute(Args))
 			{
-				OutPackagesProcessed++;
-				if (IsTimeLimitExceeded(Args.TickStartTime, Args.bUseTimeLimit, Args.TimeLimit, Args.OutLastTypeOfWorkPerformed, Args.OutLastObjectWorkWasPerformedOn))
+				const float RemainingTimeLimit = FMath::Max(0.0f, TimeLimit - (float)(FPlatformTime::Seconds() - TickStartTime));
+				int32 NumCreated = CreateAsyncPackagesFromQueue(bUseTimeLimit, bUseFullTimeLimit, RemainingTimeLimit);
+				OutPackagesProcessed += NumCreated;
+				bDidSomething = NumCreated > 0 || bDidSomething;
+				if (IsTimeLimitExceeded(TickStartTime, bUseTimeLimit, TimeLimit, TEXT("CreateAsyncPackagesFromQueue")))
 				{
 					return EAsyncPackageState::TimeOut;
 				}
-				bDidSomething = true;
 			}
-		}
-		if (bDidSomething)
-		{
-			continue;
-		}
-		if (AsyncPackagesReadyForTick.Num())
-		{
-			OutPackagesProcessed++;
-			bDidSomething = true;
-			FAsyncPackage* Package = AsyncPackagesReadyForTick[0];
-			check(Package->AsyncPackageLoadingState == EAsyncPackageLoadingState::PostLoad_Etc);
-			SCOPED_LOADTIMER(ProcessAsyncLoadingTime);
+			if (bDidSomething)
+			{
+				continue;
+			}
 
-			EAsyncPackageState::Type LocalLoadingState = EAsyncPackageState::Complete;
-			if (Package->HasFinishedLoading() == false)
 			{
-				float RemainingTimeLimit = FMath::Max(0.0f, TimeLimit - (float)(FPlatformTime::Seconds() - TickStartTime));
-				LocalLoadingState = Package->TickAsyncPackage(bUseTimeLimit, bUseFullTimeLimit, RemainingTimeLimit, FlushTree);
-				if (LocalLoadingState == EAsyncPackageState::TimeOut)
+				FAsyncLoadEventArgs Args;
+				Args.bUseTimeLimit = bUseTimeLimit;
+				Args.TickStartTime = TickStartTime;
+				Args.TimeLimit = TimeLimit;
+
+				Args.OutLastTypeOfWorkPerformed = nullptr;
+				Args.OutLastObjectWorkWasPerformedOn = nullptr;
+				if (EventQueue.PopAndExecute(Args))
 				{
-					if (IsTimeLimitExceeded(TickStartTime, bUseTimeLimit, TimeLimit, TEXT("TickAsyncPackage")))
+					OutPackagesProcessed++;
+					if (IsTimeLimitExceeded(Args.TickStartTime, Args.bUseTimeLimit, Args.TimeLimit, Args.OutLastTypeOfWorkPerformed, Args.OutLastObjectWorkWasPerformedOn))
 					{
 						return EAsyncPackageState::TimeOut;
 					}
-					UE_LOG(LogStreaming, Error, TEXT("Should not have a timeout when the time limit is not exceeded."));
-					continue;
+					bDidSomething = true;
 				}
 			}
-			else
+			if (bDidSomething)
 			{
-				// This package has finished loading but some other package is still holding
-				// a reference to it because it has this package in its dependency list.
-				// move it to the back of the list
-				AsyncPackagesReadyForTick.RemoveAt(0, 1, false); //@todoio this should maybe be a heap or something to avoid the removal cost
-				AsyncPackagesReadyForTick.Add(Package);
+				continue;
 			}
-			if (LocalLoadingState == EAsyncPackageState::Complete)
+			if (AsyncPackagesReadyForTick.Num())
 			{
+				OutPackagesProcessed++;
+				bDidSomething = true;
+				FAsyncPackage* Package = AsyncPackagesReadyForTick[0];
+				check(Package->AsyncPackageLoadingState == EAsyncPackageLoadingState::PostLoad_Etc);
+				SCOPED_LOADTIMER(ProcessAsyncLoadingTime);
 
-				// We're done, at least on this thread, so we can remove the package now.
-				AddToLoadedPackages(Package);
+				EAsyncPackageState::Type LocalLoadingState = EAsyncPackageState::Complete;
+				if (Package->HasFinishedLoading() == false)
+				{
+					float RemainingTimeLimit = FMath::Max(0.0f, TimeLimit - (float)(FPlatformTime::Seconds() - TickStartTime));
+					LocalLoadingState = Package->TickAsyncPackage(bUseTimeLimit, bUseFullTimeLimit, RemainingTimeLimit, FlushTree);
+					if (LocalLoadingState == EAsyncPackageState::TimeOut)
+					{
+						if (IsTimeLimitExceeded(TickStartTime, bUseTimeLimit, TimeLimit, TEXT("TickAsyncPackage")))
+						{
+							return EAsyncPackageState::TimeOut;
+						}
+						UE_LOG(LogStreaming, Error, TEXT("Should not have a timeout when the time limit is not exceeded."));
+						continue;
+					}
+				}
+				else
 				{
+					check(0); // if it has finished loading, it should not be in AsyncPackagesReadyForTick
+				}
+				if (LocalLoadingState == EAsyncPackageState::Complete)
+				{
+
+					// We're done, at least on this thread, so we can remove the package now.
+					AddToLoadedPackages(Package);
+					{
 #if THREADSAFE_UOBJECTS
-					FScopeLock LockAsyncPackages(&AsyncPackagesCritical);
+						FScopeLock LockAsyncPackages(&AsyncPackagesCritical);
 #endif
-					AsyncPackageNameLookup.Remove(Package->GetPackageName());
-					int32 PackageIndex = AsyncPackages.Find(Package);
-					AsyncPackages.RemoveAt(PackageIndex);
-					AsyncPackagesReadyForTick.RemoveAt(0, 1, false); //@todoio this should maybe be a heap or something to avoid the removal cost
+						AsyncPackageNameLookup.Remove(Package->GetPackageName());
+						int32 PackageIndex = AsyncPackages.Find(Package);
+						AsyncPackages.RemoveAt(PackageIndex);
+						AsyncPackagesReadyForTick.RemoveAt(0, 1, false); //@todoio this should maybe be a heap or something to avoid the removal cost
+					}
+				}
+				if (IsTimeLimitExceeded(TickStartTime, bUseTimeLimit, TimeLimit, TEXT("TickAsyncPackage")))
+				{
+					return EAsyncPackageState::TimeOut;
 				}
 			}
-			if (IsTimeLimitExceeded(TickStartTime, bUseTimeLimit, TimeLimit, TEXT("TickAsyncPackage")))
+			if (bDidSomething)
 			{
-				return EAsyncPackageState::TimeOut;
+				continue;
 			}
-		}
-		if (bDidSomething)
-		{
-			continue;
-		}
-		bool bAnyIOOutstanding = GPrecacheCallbackHandler.AnyIOOutstanding();
-		if (bAnyIOOutstanding)
-		{
-			SCOPED_LOADTIMER(Package_EventIOWait);
-			double StartTime = FPlatformTime::Seconds();
-			if (bUseTimeLimit)
+			bool bAnyIOOutstanding = GPrecacheCallbackHandler.AnyIOOutstanding();
+			if (bAnyIOOutstanding)
 			{
-				if (bUseFullTimeLimit)
+				SCOPED_LOADTIMER(Package_EventIOWait);
+				double StartTime = FPlatformTime::Seconds();
+				if (bUseTimeLimit)
 				{
-					const float RemainingTimeLimit = FMath::Max(0.0f, TimeLimit - (float)(FPlatformTime::Seconds() - TickStartTime));
-					if (RemainingTimeLimit > 0.0f)
+					if (bUseFullTimeLimit)
 					{
-						bool bGotIO = GPrecacheCallbackHandler.WaitForIO(RemainingTimeLimit);
-						if (bGotIO)
-						{
-							OutPackagesProcessed++;
-							continue; // we got some IO, so start processing at the top
-						}
+						const float RemainingTimeLimit = FMath::Max(0.0f, TimeLimit - (float)(FPlatformTime::Seconds() - TickStartTime));
+						if (RemainingTimeLimit > 0.0f)
 						{
-							float ThisTime = float(FPlatformTime::Seconds() - StartTime);
+							bool bGotIO = GPrecacheCallbackHandler.WaitForIO(RemainingTimeLimit);
+							if (bGotIO)
+							{
+								OutPackagesProcessed++;
+								continue; // we got some IO, so start processing at the top
+							}
+							{
+								float ThisTime = float(FPlatformTime::Seconds() - StartTime);
+							}
 						}
 					}
+					//UE_LOG(LogStreaming, Error, TEXT("Not using full time limit, waiting for IO..."));
+					return EAsyncPackageState::TimeOut;
+				}
+				else
+				{
+					bool bGotIO = GPrecacheCallbackHandler.WaitForIO(10.0f); // wait "forever"
+					if (!bGotIO)
+					{
+						UE_LOG(LogStreaming, Error, TEXT("Waited for 10 seconds on IO...."));
+					}
+					{
+						OutPackagesProcessed++;
+					}
 				}
-				//UE_LOG(LogStreaming, Error, TEXT("Not using full time limit, waiting for IO..."));
-				return EAsyncPackageState::TimeOut;
 			}
 			else
 			{
-				bool bGotIO = GPrecacheCallbackHandler.WaitForIO(10.0f); // wait "forever"
-				if (!bGotIO)
-				{
-					UE_LOG(LogStreaming, Error, TEXT("Waited for 10 seconds on IO...."));
-				}
-				{
-					OutPackagesProcessed++;
-				}
+				LoadingState = EAsyncPackageState::Complete;
+				break;
 			}
 		}
-		else
-		{
-			LoadingState = EAsyncPackageState::Complete;
-			break;
-		}
-	}
-
-
-#else
-#if 0 && USE_NEW_ASYNC_IO
-bool bDepthFirst = true;
-#else
-bool bDepthFirst = false;
-#endif
-
-	// We need to loop as the function has to handle finish loading everything given no time limit
-	// like e.g. when called from FlushAsyncLoading.
-	for (int32 PackageIndex = 0; ((bDepthFirst && LoadingState == EAsyncPackageState::Complete) || (!bDepthFirst && LoadingState != EAsyncPackageState::TimeOut)) && PackageIndex < AsyncPackages.Num(); ++PackageIndex)
+	} // GEventDrivenLoaderEnabled
+	else
 	{
-		SCOPED_LOADTIMER(ProcessAsyncLoadingTime);
-		OutPackagesProcessed++;
+		bool bDepthFirst = false; // GNewAsyncIO
 
-		// Package to be loaded.
-		FAsyncPackage* Package = AsyncPackages[PackageIndex];
-		if (FlushTree && !FlushTree->Contains(Package->GetPackageName()))
-		{
-			LoadingState = EAsyncPackageState::PendingImports;
-		}
-		else if (Package->HasFinishedLoading() == false)
+		// We need to loop as the function has to handle finish loading everything given no time limit
+		// like e.g. when called from FlushAsyncLoading.
+		for (int32 PackageIndex = 0; ((bDepthFirst && LoadingState == EAsyncPackageState::Complete) || (!bDepthFirst && LoadingState != EAsyncPackageState::TimeOut)) && PackageIndex < AsyncPackages.Num(); ++PackageIndex)
 		{
-#if USE_EVENT_DRIVEN_ASYNC_LOAD
-			LoadingState = EAsyncPackageState::PendingImports;
-#else
-			// @todo: Guard against recursively re-entering?
-			//if (!Package->IsCurrentlyBeingProcessed())
+			SCOPED_LOADTIMER(ProcessAsyncLoadingTime);
+			OutPackagesProcessed++;
+
+			// Package to be loaded.
+			FAsyncPackage* Package = AsyncPackages[PackageIndex];
+			if (FlushTree && !FlushTree->Contains(Package->GetPackageName()))
 			{
-				// Package tick returns EAsyncPackageState::Complete on completion.
-				// We only tick packages that have not yet been loaded.
-				LoadingState = Package->TickAsyncPackage(bUseTimeLimit, bUseFullTimeLimit, TimeLimit, FlushTree);
+				LoadingState = EAsyncPackageState::PendingImports;
 			}
-#endif
-		}
-		else
-		{
-			// This package has finished loading but some other package is still holding
-			// a reference to it because it has this package in its dependency list.
-			LoadingState = EAsyncPackageState::Complete;
-		}
-		if (LoadingState == EAsyncPackageState::Complete)
-		{
-			// We're done, at least on this thread, so we can remove the package now.
-			if (!Package->HasThreadedLoadingFinished())
+			else if (Package->HasFinishedLoading() == false)
+			{
+				if (GEventDrivenLoaderEnabled)
+				{
+					LoadingState = EAsyncPackageState::PendingImports;
+				}
+				else
+				// @todo: Guard against recursively re-entering?
+				{
+					// Package tick returns EAsyncPackageState::Complete on completion.
+					// We only tick packages that have not yet been loaded.
+					LoadingState = Package->TickAsyncPackage(bUseTimeLimit, bUseFullTimeLimit, TimeLimit, FlushTree);
+				}
+			}
+			else
+			{
+				// This package has finished loading but some other package is still holding
+				// a reference to it because it has this package in its dependency list.
+				LoadingState = EAsyncPackageState::Complete;
+			}
+			if (LoadingState == EAsyncPackageState::Complete)
 			{
-				Package->ThreadedLoadingHasFinished();
-			AddToLoadedPackages(Package);
+				// We're done, at least on this thread, so we can remove the package now.
+				if (!Package->HasThreadedLoadingFinished())
+				{
+					Package->ThreadedLoadingHasFinished();
+					AddToLoadedPackages(Package);
 #if THREADSAFE_UOBJECTS
-				FScopeLock LockAsyncPackages(&AsyncPackagesCritical);
+					FScopeLock LockAsyncPackages(&AsyncPackagesCritical);
 #endif
-				AsyncPackageNameLookup.Remove(Package->GetPackageName());
-				AsyncPackages.Remove(Package);
+					AsyncPackageNameLookup.Remove(Package->GetPackageName());
+					AsyncPackages.Remove(Package);
 
-				// Need to process this index again as we just removed an item
-				PackageIndex--;
-			}
+					// Need to process this index again as we just removed an item
+					PackageIndex--;
+				}
 
-			check(!AsyncPackages.Contains(Package));
-					
-		}
-		else if (!bUseTimeLimit && !FPlatformProcess::SupportsMultithreading())
-		{
-#if !USE_NEW_ASYNC_IO
-			// Tick async loading when multithreading is disabled.
-			FIOSystem::Get().TickSingleThreaded();
-#endif
-		}
+				check(!AsyncPackages.Contains(Package));
 
-		{ // maybe we shouldn't do this if we are already out of time?
-			// Check if there's any new packages in the queue.
-			const float RemainingTimeLimit = FMath::Max(0.0f, TimeLimit - (float)(FPlatformTime::Seconds() - TickStartTime));
-			CreateAsyncPackagesFromQueue(bUseTimeLimit, bUseFullTimeLimit, RemainingTimeLimit);
-		}
+			}
+			else if (!bUseTimeLimit && !FPlatformProcess::SupportsMultithreading() && !GNewAsyncIO)
+			{
+				// Tick async loading when multithreading is disabled.
+				FIOSystem::Get().TickSingleThreaded();
+			}
 
-		if (bNeedsHeartbeatTick)
-		{
-			// Update heartbeat after each package has been processed
-			FThreadHeartBeat::Get().HeartBeat();
+			{ // maybe we shouldn't do this if we are already out of time?
+				// Check if there's any new packages in the queue.
+				const float RemainingTimeLimit = FMath::Max(0.0f, TimeLimit - (float)(FPlatformTime::Seconds() - TickStartTime));
+				CreateAsyncPackagesFromQueue(bUseTimeLimit, bUseFullTimeLimit, RemainingTimeLimit);
+			}
+
+			if (bNeedsHeartbeatTick)
+			{
+				// Update heartbeat after each package has been processed
+				FThreadHeartBeat::Get().HeartBeat();
+			}
 		}
-	}
-#endif
+	} // !GEventDrivenLoaderEnabled
 	return LoadingState;
 }
 
@@ -4062,11 +4038,11 @@ EAsyncPackageState::Type FAsyncLoadingThread::TickAsyncLoading(bool bUseTimeLimi
 	check(IsInGameThread());
 
 	const bool bLoadingSuspended = IsAsyncLoadingSuspended();
-	const bool bIsMultithreaded = FAsyncLoadingThread::IsMultithreaded();
 	EAsyncPackageState::Type Result = bLoadingSuspended ? EAsyncPackageState::PendingImports : EAsyncPackageState::Complete;
 
 	if (!bLoadingSuspended)
 	{
+		const bool bIsMultithreaded = FAsyncLoadingThread::IsMultithreaded();
 		double TickStartTime = FPlatformTime::Seconds();
 		double TimeLimitUsedForProcessLoaded = 0;
 
@@ -4098,13 +4074,11 @@ EAsyncPackageState::Type FAsyncLoadingThread::TickAsyncLoading(bool bUseTimeLimi
 					IsTimeLimitExceeded(TickStartTime, bUseTimeLimit, TimeLimit, TEXT("FDeferredMessageLog::Flush()"));
 				}
 			}
-#if USE_EVENT_DRIVEN_ASYNC_LOAD
-			if (!bDidSomething && !bIsMultithreaded)
+			if (!bDidSomething && !bIsMultithreaded && GEventDrivenLoaderEnabled)
 			{
 				CheckForCycles();
 				IsTimeLimitExceeded(TickStartTime, bUseTimeLimit, TimeLimit, TEXT("CheckForCycles (non-shipping)"));
 			}
-#endif
 		}
 	}
 
@@ -4113,13 +4087,20 @@ EAsyncPackageState::Type FAsyncLoadingThread::TickAsyncLoading(bool bUseTimeLimi
 
 FAsyncLoadingThread::FAsyncLoadingThread()
 {
-#if USE_EVENT_DRIVEN_ASYNC_LOAD
-	UE_CLOG(!IsEventDrivenLoaderEnabledInCookedBuilds(), LogStreaming, Fatal,
-		TEXT("Event driven async loader is being used but it does NOT seem to be enabled in project settings."));
-#elif !WITH_EDITORONLY_DATA
-	UE_CLOG(IsEventDrivenLoaderEnabledInCookedBuilds(), LogStreaming, Fatal,
-		TEXT("Event driven async loader is NOT being used but it seems to be enabled in project settings."));
-#endif
+	// Current these two vars are always on or off together but can be made separate
+	GNewAsyncIO = IsEventDrivenLoaderEnabled();
+	GEventDrivenLoaderEnabled = GNewAsyncIO;
+
+	if (IsEventDrivenLoaderEnabled())
+	{
+		UE_CLOG(!IsEventDrivenLoaderEnabledInCookedBuilds(), LogStreaming, Fatal,
+			TEXT("Event driven async loader is being used but it does NOT seem to be enabled in project settings."));
+	}
+	else if (FPlatformProperties::RequiresCookedData())
+	{
+		UE_CLOG(IsEventDrivenLoaderEnabledInCookedBuilds(), LogStreaming, Fatal,
+			TEXT("Event driven async loader is NOT being used but it seems to be enabled in project settings."));
+	}
 
 	QueuedRequestsEvent = FPlatformProcess::GetSynchEventFromPool();
 	CancelLoadingEvent = FPlatformProcess::GetSynchEventFromPool();
@@ -4127,29 +4108,34 @@ FAsyncLoadingThread::FAsyncLoadingThread()
 	ThreadResumedEvent = FPlatformProcess::GetSynchEventFromPool();
 	if (FAsyncLoadingThread::IsMultithreaded())
 	{
-		UE_LOG(LogStreaming, Log, TEXT("Async loading is multithreaded."));
 		Thread = FRunnableThread::Create(this, TEXT("FAsyncLoadingThread"), 0, TPri_Normal);
 	}
 	else
 	{
-		UE_LOG(LogStreaming, Log, TEXT("Async loading is time-sliced."));
 		Thread = nullptr;
 		Init();
 	}
 	AsyncLoadingTickCounter = 0;
+
+#if !IS_PROGRAM && !WITH_EDITORONLY_DATA
+	UE_LOG(LogStreaming, Display, TEXT("Async Loading initialized: New Async IO: %s, Event Driven Loader: %s, Async Loading Thread: %s"),
+		GNewAsyncIO ? TEXT("true") : TEXT("false"),
+		GEventDrivenLoaderEnabled ? TEXT("true") : TEXT("false"),
+		FAsyncLoadingThread::IsMultithreaded() ? TEXT("true") : TEXT("false"));
+#endif
 }
 
 FAsyncLoadingThread::~FAsyncLoadingThread()
 {
-#if USE_EVENT_DRIVEN_ASYNC_LOAD
-	// check that event queue is empty
-	const TCHAR* LastTypeOfWorkPerformed = nullptr;
-	UObject* LastObjectWorkWasPerformedOn = nullptr;
-
-	FAsyncLoadEventArgs Args;
-	check(!EventQueue.PopAndExecute(Args));
+	if (IsEventDrivenLoaderEnabled())
+	{
+		// check that event queue is empty
+		const TCHAR* LastTypeOfWorkPerformed = nullptr;
+		UObject* LastObjectWorkWasPerformedOn = nullptr;
 
-#endif
+		FAsyncLoadEventArgs Args;
+		check(!EventQueue.PopAndExecute(Args));
+	}
 
 	delete Thread;
 	Thread = nullptr;
@@ -4198,7 +4184,7 @@ uint32 FAsyncLoadingThread::Run()
 	return 0;
 }
 
-#if USE_EVENT_DRIVEN_ASYNC_LOAD
+
 void FAsyncLoadingThread::CheckForCycles()
 {
 	// no outstanding IO, nothing was done in this iteration, we are done
@@ -4222,8 +4208,6 @@ void FAsyncLoadingThread::CheckForCycles()
 	}
 #endif
 }
-#endif
-
 
 
 EAsyncPackageState::Type  FAsyncLoadingThread::TickAsyncThread(bool bUseTimeLimit, bool bUseFullTimeLimit, float TimeLimit, bool& bDidSomething, FFlushTree* FlushTree)
@@ -4251,10 +4235,11 @@ EAsyncPackageState::Type  FAsyncLoadingThread::TickAsyncThread(bool bUseTimeLimi
 		}
 		if (ProcessedRequests == 0 && IsMultithreaded())
 		{
-#if USE_EVENT_DRIVEN_ASYNC_LOAD
-			CheckForCycles();
-			IsTimeLimitExceeded(TickStartTime, bUseTimeLimit, TimeLimit, TEXT("CheckForCycles (non-shipping)"));
-#endif
+			if (IsEventDrivenLoaderEnabled())
+			{
+				CheckForCycles();
+				IsTimeLimitExceeded(TickStartTime, bUseTimeLimit, TimeLimit, TEXT("CheckForCycles (non-shipping)"));
+			}
 			const bool bIgnoreThreadIdleStats = true;
 			SCOPED_LOADTIMER(Package_Temp3);
 			QueuedRequestsEvent->Wait(30, bIgnoreThreadIdleStats);
@@ -4420,7 +4405,7 @@ FAsyncPackage::FAsyncPackage(const FAsyncPackageDesc& InDesc)
 , LoadPercentage(0)
 , ReentryCount(0)
 , AsyncLoadingThread(FAsyncLoadingThread::Get())
-#if USE_EVENT_DRIVEN_ASYNC_LOAD
+// Begin EDL specific properties
 , AsyncPackageLoadingState(EAsyncPackageLoadingState::NewPackage)
 , SerialNumber(AsyncPackageSerialNumber.Increment())
 , CurrentBlockOffset(-1)
@@ -4430,7 +4415,7 @@ FAsyncPackage::FAsyncPackage(const FAsyncPackageDesc& InDesc)
 , bProcessImportsAndExportsInFlight(false)
 , bProcessPostloadWaitInFlight(false)
 , bAllExportsSerialized(false)
-#endif
+// End EDL specific properties
 
 #if PERF_TRACK_DETAILED_ASYNC_STATS
 , TickCount(0)
@@ -4457,23 +4442,27 @@ FAsyncPackage::FAsyncPackage(const FAsyncPackageDesc& InDesc)
 
 FAsyncPackage::~FAsyncPackage()
 {
-#if DO_CHECK && USE_EVENT_DRIVEN_ASYNC_LOAD
-	for (FCompletionCallback& CompletionCallback : CompletionCallbacks)
+#if DO_CHECK
+	if (GEventDrivenLoaderEnabled)
 	{
-		if (!CompletionCallback.bCalled)
+		for (FCompletionCallback& CompletionCallback : CompletionCallbacks)
 		{
-			check(0);
+			if (!CompletionCallback.bCalled)
+			{
+				check(0);
+			}
 		}
 	}
 #endif
 
 	MarkRequestIDsAsComplete();
 	DetachLinker();
-#if USE_EVENT_DRIVEN_ASYNC_LOAD
-	SerialNumber = 0; // the weak pointer will always fail now
-	check(!EventNodeArray.Array && !EventNodeArray.TotalNumberOfNodesAdded);
-	RemoveAllNodes();
-#endif
+	if (GEventDrivenLoaderEnabled)
+	{
+		SerialNumber = 0; // the weak pointer will always fail now
+		check(!EventNodeArray.Array && !EventNodeArray.TotalNumberOfNodesAdded);
+		RemoveAllNodes();
+	}
 
 	EmptyReferencedObjects();
 }
@@ -4582,13 +4571,14 @@ void FAsyncPackage::DetachLinker()
  */
 bool FAsyncPackage::GiveUpTimeSlice()
 {
-#if !USE_NEW_ASYNC_IO
-	static const bool bPlatformIsSingleThreaded = !FPlatformProcess::SupportsMultithreading();
-	if (bPlatformIsSingleThreaded)
+	if (!GNewAsyncIO)
 	{
-		FIOSystem::Get().TickSingleThreaded();
+		static const bool bPlatformIsSingleThreaded = !FPlatformProcess::SupportsMultithreading();
+		if (bPlatformIsSingleThreaded)
+		{
+			FIOSystem::Get().TickSingleThreaded();
+		}
 	}
-#endif
 
 	if (bUseTimeLimit && !bUseFullTimeLimit)
 	{
@@ -4648,9 +4638,10 @@ void FAsyncPackage::EndAsyncLoad()
 
 EAsyncPackageState::Type FAsyncPackage::TickAsyncPackage(bool InbUseTimeLimit, bool InbUseFullTimeLimit, float& InOutTimeLimit, FFlushTree* FlushTree)
 {
-#if USE_EVENT_DRIVEN_ASYNC_LOAD
-	check(int32(AsyncPackageLoadingState) > int32(EAsyncPackageLoadingState::ProcessNewImportsAndExports));
-#endif
+
+	// We want this check only with EDL enabled
+	check(!GEventDrivenLoaderEnabled || int32(AsyncPackageLoadingState) > int32(EAsyncPackageLoadingState::ProcessNewImportsAndExports));
+
 	ReentryCount++;
 
 	SCOPE_CYCLE_COUNTER(STAT_FAsyncPackage_Tick);
@@ -4698,58 +4689,67 @@ EAsyncPackageState::Type FAsyncPackage::TickAsyncPackage(bool InbUseTimeLimit, b
 		// We have begun loading a package that we know the name of. Let the package time tracker know.
 		FExclusiveLoadPackageTimeTracker::PushLoadPackage(Desc.NameToLoad);
 
-#if !USE_EVENT_DRIVEN_ASYNC_LOAD
-		// Create raw linker. Needs to be async created via ticking before it can be used.
-		if (LoadingState == EAsyncPackageState::Complete)
+		if (!GEventDrivenLoaderEnabled)
 		{
-			SCOPED_LOADTIMER(Package_CreateLinker);
-			LoadingState = CreateLinker();
-		}
+			// Create raw linker. Needs to be async created via ticking before it can be used.
+			if (LoadingState == EAsyncPackageState::Complete)
+			{
+				SCOPED_LOADTIMER(Package_CreateLinker);
+				LoadingState = CreateLinker();
+			}
 
-		// Async create linker.
-		if (LoadingState == EAsyncPackageState::Complete)
-		{
-			SCOPED_LOADTIMER(Package_FinishLinker);
-			LoadingState = FinishLinker();
-		}
+			// Async create linker.
+			if (LoadingState == EAsyncPackageState::Complete)
+			{
+				SCOPED_LOADTIMER(Package_FinishLinker);
+				LoadingState = FinishLinker();
+			}
 
-		// Load imports from linker import table asynchronously.
-		if (LoadingState == EAsyncPackageState::Complete)
-		{
-			SCOPED_LOADTIMER(Package_LoadImports);
-			LoadingState = LoadImports(FlushTree);
-		}
+			// Load imports from linker import table asynchronously.
+			if (LoadingState == EAsyncPackageState::Complete)
+			{
+				SCOPED_LOADTIMER(Package_LoadImports);
+				LoadingState = LoadImports(FlushTree);
+			}
 
-		// Create imports from linker import table.
-		if (LoadingState == EAsyncPackageState::Complete)
-		{
-			SCOPED_LOADTIMER(Package_CreateImports);
-			LoadingState = CreateImports();
-		}
+			// Create imports from linker import table.
+			if (LoadingState == EAsyncPackageState::Complete)
+			{
+				SCOPED_LOADTIMER(Package_CreateImports);
+				LoadingState = CreateImports();
+			}
 
 #if WITH_EDITORONLY_DATA
-		// Create and preload the package meta-data
-		if (LoadingState == EAsyncPackageState::Complete)
-		{
-			SCOPED_LOADTIMER(Package_CreateMetaData);
-			LoadingState = CreateMetaData();
-		}
+			// Create and preload the package meta-data
+			if (LoadingState == EAsyncPackageState::Complete)
+			{
+				SCOPED_LOADTIMER(Package_CreateMetaData);
+				LoadingState = CreateMetaData();
+			}
 #endif // WITH_EDITORONLY_DATA
 
-		// Create exports from linker export table and also preload them.
-		if (LoadingState == EAsyncPackageState::Complete)
-		{
-			SCOPED_LOADTIMER(Package_CreateExports);
-			LoadingState = CreateExports();
-		}
+			// Create exports from linker export table and also preload them.
+			if (LoadingState == EAsyncPackageState::Complete)
+			{
+				SCOPED_LOADTIMER(Package_CreateExports);
+				LoadingState = CreateExports();
+			}
+
+			// Call Preload on the linker for all loaded objects which causes actual serialization.
+			if (LoadingState == EAsyncPackageState::Complete)
+			{
+				SCOPED_LOADTIMER(Package_PreLoadObjects);
+				if (GNewAsyncIO)
+				{
+					LoadingState = PreLoadObjectsForNewAsyncIO();
+				}
+				else
+				{
+					LoadingState = PreLoadObjects();
+				}
+			}
+		} // !GEventDrivenLoaderEnabled
 
-		// Call Preload on the linker for all loaded objects which causes actual serialization.
-		if (LoadingState == EAsyncPackageState::Complete)
-		{
-			SCOPED_LOADTIMER(Package_PreLoadObjects);
-			LoadingState = PreLoadObjects();
-		}
-#endif
 		// Call PostLoad on objects, this could cause new objects to be loaded that require
 		// another iteration of the PreLoad loop.
 		if (LoadingState == EAsyncPackageState::Complete && !bLoadHasFailed)
@@ -4786,13 +4786,13 @@ EAsyncPackageState::Type FAsyncPackage::TickAsyncPackage(bool InbUseTimeLimit, b
 	LastTypeOfWorkPerformed = nullptr;
 	// Mark this package as loaded if everything completed.
 	bLoadHasFinished = (LoadingState == EAsyncPackageState::Complete);
-#if USE_EVENT_DRIVEN_ASYNC_LOAD
-	if (bLoadHasFinished)
+
+	if (bLoadHasFinished && GEventDrivenLoaderEnabled)
 	{
 		check(AsyncPackageLoadingState == EAsyncPackageLoadingState::PostLoad_Etc);
 		AsyncPackageLoadingState = EAsyncPackageLoadingState::PackageComplete;
 	}
-#endif
+
 	// Subtract the time it took to load this package from the global limit.
 	InOutTimeLimit = (float)FMath::Max(0.0, InOutTimeLimit - (FPlatformTime::Seconds() - TickStartTime));
 
@@ -4815,8 +4815,8 @@ EAsyncPackageState::Type FAsyncPackage::CreateLinker()
 	{
 		SCOPE_CYCLE_COUNTER(STAT_FAsyncPackage_CreateLinker);
 
-		LastObjectWorkWasPerformedOn	= nullptr;
-		LastTypeOfWorkPerformed			= TEXT("creating Linker");
+		LastObjectWorkWasPerformedOn = nullptr;
+		LastTypeOfWorkPerformed = TEXT("creating Linker");
 
 		// Try to find existing package or create it if not already present.
 		UPackage* Package = nullptr;
@@ -4843,17 +4843,15 @@ EAsyncPackageState::Type FAsyncPackage::CreateLinker()
 		// if the linker already exists, we don't need to lookup the file (it may have been pre-created with
 		// a different filename)
 		Linker = FLinkerLoad::FindExistingLinkerForPackage(Package);
-#if USE_EVENT_DRIVEN_ASYNC_LOAD
-		if (Linker)
+		if (Linker && GEventDrivenLoaderEnabled)
 		{
 			// this almost works, but the EDL does not tolerate redoing steps it already did
-			UE_LOG(LogStreaming, Fatal , TEXT("Package %s was reloaded before it even closed the linker from a previous load. Seems like a waste of time eh?"), *Desc.Name.ToString());
+			UE_LOG(LogStreaming, Fatal, TEXT("Package %s was reloaded before it even closed the linker from a previous load. Seems like a waste of time eh?"), *Desc.Name.ToString());
 			check(Package);
 			FWeakAsyncPackagePtr WeakPtr(this);
 			GPrecacheCallbackHandler.RegisterNewSummaryRequest(this);
 			GPrecacheCallbackHandler.SummaryComplete(WeakPtr);
 		}
-#endif
 
 		if (!Linker)
 		{
@@ -4873,7 +4871,7 @@ EAsyncPackageState::Type FAsyncPackage::CreateLinker()
 
 			if (Desc.NameToLoad == NAME_None ||
 				(!GetConvertedDynamicPackageNameToTypeName().Contains(Desc.Name) &&
-				!DoesPackageExist))
+					!DoesPackageExist))
 			{
 				FName FailedLoadName = FName(*NameToLoad);
 
@@ -4908,32 +4906,36 @@ EAsyncPackageState::Type FAsyncPackage::CreateLinker()
 				LinkerFlags |= LOAD_PackageForPIE;
 			}
 #endif
-#if USE_EVENT_DRIVEN_ASYNC_LOAD
-			FWeakAsyncPackagePtr WeakPtr(this);
-			check(Package);
-			Linker = FLinkerLoad::CreateLinkerAsync(Package, *PackageFileName, LinkerFlags
-				, TFunction<void()>(
-					[WeakPtr]()
-					{
-						GPrecacheCallbackHandler.SummaryComplete(WeakPtr);
-					}
-				));
-			if (Linker)
+			if (GEventDrivenLoaderEnabled)
 			{
-				GPrecacheCallbackHandler.RegisterNewSummaryRequest(this);
-				if (Linker->bDynamicClassLinker)
+				FWeakAsyncPackagePtr WeakPtr(this);
+				check(Package);
+				Linker = FLinkerLoad::CreateLinkerAsync(Package, *PackageFileName, LinkerFlags
+					, TFunction<void()>(
+						[WeakPtr]()
 				{
-//native blueprint 
-					check(!Linker->GetFArchiveAsync2Loader());
 					GPrecacheCallbackHandler.SummaryComplete(WeakPtr);
 				}
+				));
+				if (Linker)
+				{
+					GPrecacheCallbackHandler.RegisterNewSummaryRequest(this);
+					if (Linker->bDynamicClassLinker)
+					{
+						//native blueprint 
+						check(!Linker->GetFArchiveAsync2Loader());
+						GPrecacheCallbackHandler.SummaryComplete(WeakPtr);
+					}
+				}
+			}
+			else
+			{
+				Linker = FLinkerLoad::CreateLinkerAsync(Package, *PackageFileName, LinkerFlags, TFunction<void()>([]() {}));
 			}
-#else
-			Linker = FLinkerLoad::CreateLinkerAsync(Package, *PackageFileName, LinkerFlags);
-#endif
 		}
 
 		// Associate this async package with the linker
+		check(Linker);
 		check(Linker->AsyncRoot == nullptr || Linker->AsyncRoot == this);
 		Linker->AsyncRoot = this;
 
@@ -5019,11 +5021,8 @@ void FAsyncPackage::AddImportDependency(const FName& PendingImport, FFlushTree*
 			PackageToStream->Desc.Priority = Desc.Priority;
 		}
 	}
-#if USE_NEW_ASYNC_IO
-	const bool bDepthFirst = true;
-#else
-	const bool bDepthFirst = false;
-#endif
+
+	const bool bDepthFirst = GNewAsyncIO;
 
 	if (!bReinsert || bDepthFirst)
 	{
@@ -5178,10 +5177,6 @@ EAsyncPackageState::Type FAsyncPackage::LoadImports(FFlushTree* FlushTree)
 			{
 				UE_LOG(LogStreaming, Verbose, TEXT("FAsyncPackage::LoadImports for %s: Loading %s"), *Desc.NameToLoad.ToString(), *ImportPackageName);
 				AddImportDependency(ImportPackageFName, FlushTree);
-#if 0 && USE_NEW_ASYNC_IO
-				LoadImportIndex--;
-				return EAsyncPackageState::TimeOut; // lets go load that package now so we can do it in order.
-#endif
 			}
 			else
 			{
@@ -5295,14 +5290,17 @@ EAsyncPackageState::Type FAsyncPackage::CreateExports()
 		const FObjectExport& Export = Linker->ExportMap[ExportIndex];
 		// Precache data and see whether it's already finished.
 		bool bReady;
-#if USE_NEW_ASYNC_IO
-		FArchiveAsync2* FAA2 = Linker->GetFArchiveAsync2Loader();
+
+		FArchiveAsync2* FAA2 = nullptr;
+		if (GNewAsyncIO)
+		{
+			FAA2 = Linker->GetFArchiveAsync2Loader();
+		}
 		if (FAA2)
 		{
 			bReady = FAA2->Precache(Export.SerialOffset, Export.SerialSize, bUseTimeLimit, bUseFullTimeLimit, TickStartTime, TimeLimit);
 		}
 		else
-#endif
 		{
 			bReady = Linker->Precache(Export.SerialOffset, Export.SerialSize);
 		}
@@ -5355,49 +5353,7 @@ void FAsyncPackage::FreeReferencedImports()
 	ReferencedImports.Empty();
 }
 
-/**
- * Preloads aka serializes all loaded objects.
- *
- * @return true if we finished serializing all loaded objects, false otherwise.
- */
-
-#if 0 && USE_NEW_ASYNC_IO
-
-// this can sort exports before preloading them
-struct FPreloadSort
-{
-	UObject* Object;
-	FLinkerLoad* Linker;
-	int32 SerialOffset;
-
-	FPreloadSort(UObject* InObject)
-		: Object(InObject)
-	{
-		if (Object)
-		{
-			Linker = Object->GetLinker();
-			if (Linker)
-			{
-				SerialOffset = Linker->ExportMap[Object->GetLinkerIndex()].SerialOffset;
-			}
-			else
-			{
-				SerialOffset = 0;
-			}
-		}
-		else
-		{
-			Linker = nullptr;
-			SerialOffset = 0;
-		}
-	}
-};
-
-TArray<UObject*>* GilObjLoaded = nullptr;
-
-FString GilFullName;
-
-EAsyncPackageState::Type FAsyncPackage::PreLoadObjects()
+EAsyncPackageState::Type FAsyncPackage::PreLoadObjectsForNewAsyncIO()
 {
 	SCOPED_LOADTIMER(PreLoadObjectsTime);
 	SCOPE_CYCLE_COUNTER(STAT_FAsyncPackage_PreLoadObjects);
@@ -5405,114 +5361,7 @@ EAsyncPackageState::Type FAsyncPackage::PreLoadObjects()
 	// GC can't run in here
 	FGCScopeGuard GCGuard;
 
-	TArray<UObject*>& ObjLoaded = FUObjectThreadContext::Get().ObjLoaded;
-	GilObjLoaded = &ObjLoaded;
-
-	// Preload (aka serialize) the objects.
-	while (PreLoadIndex < ObjLoaded.Num() && !IsTimeLimitExceeded())
-	{
-		if ( //PreLoadSortIndex <= PreLoadIndex && 
-			ObjLoaded.Num() > PreLoadSortIndex)
-		{
-			int32 SortFrom = PreLoadIndex;
-			FLinkerLoad* LastLinker = PreLoadIndex > 0 ? ObjLoaded[PreLoadIndex - 1]->GetLinker() : nullptr;
-
-			auto SortPredicate =
-				[LastLinker](const FPreloadSort& LHS, const FPreloadSort& RHS) -> bool
-			{
-				if (LastLinker)
-				{
-					// if we have a last linker, then those are always first
-					if (LHS.Linker == LastLinker && RHS.Linker != LastLinker)
-					{
-						return true;
-					}
-					if (LHS.Linker != LastLinker && RHS.Linker == LastLinker)
-					{
-						return false;
-					}
-				}
-				if (LHS.Linker && !RHS.Linker)
-				{
-					return true; // null linkers are last, maybe they will be non-null later
-				}
-				if (!LHS.Linker && RHS.Linker)
-				{
-					return false; // null linkers are last, maybe they will be non-null later
-				}
-				if (LHS.Linker != RHS.Linker)
-				{
-					return UPTRINT(LHS.Linker) < UPTRINT(RHS.Linker); // arbitrary order, based on the pointer
-				}
-				check(LHS.Linker == RHS.Linker && LHS.Linker);
-				return LHS.SerialOffset < RHS.SerialOffset; // in disk order
-			};
-
-			TArray<FPreloadSort> SortedArrayTail;
-			SortedArrayTail.Reserve(ObjLoaded.Num() - SortFrom);
-			for (int32 Index = SortFrom; Index < ObjLoaded.Num(); Index++)
-			{
-				SortedArrayTail.Emplace(ObjLoaded[Index]);
-			}
-			SortedArrayTail.Sort(SortPredicate);
-			{
-				int32 NumTail = SortedArrayTail.Num();
-				for (int32 Index = 0; Index < NumTail; Index++)
-				{
-					FPreloadSort& Item = SortedArrayTail[Index];
-					ObjLoaded[SortFrom + Index] = Item.Object;
-				}
-			}
-			PreLoadSortIndex = ObjLoaded.Num();
-		}
-		UObject* Object = ObjLoaded[PreLoadIndex];
-		if (Object && Object->HasAnyFlags(RF_NeedLoad))
-		{
-			FLinkerLoad* ObjectLinker = Object->GetLinker();
-			if (ObjectLinker)
-			{
-				const FObjectExport& Export = ObjectLinker->ExportMap[Object->GetLinkerIndex()];
-				GilFullName = Object->GetPathName();
-				FArchiveAsync2* FAA2 = ObjectLinker->GetFArchiveAsync2Loader();
-				bool bReady;
-				if (FAA2)
-				{
-					bReady = FAA2->Precache(Export.SerialOffset, Export.SerialSize, bUseTimeLimit, bUseFullTimeLimit, TickStartTime, TimeLimit);
-				}
-				else
-				{
-					bReady = Linker->Precache(Export.SerialOffset, Export.SerialSize);
-				}
-				if (bReady)
-				{
-					ObjectLinker->Preload(Object);
-					LastObjectWorkWasPerformedOn = Object;
-					LastTypeOfWorkPerformed = TEXT("preloading");
-				}
-				// Data isn't ready yet. Give up remainder of time slice if we're not using a time limit.
-				else if (GiveUpTimeSlice())
-				{
-					INC_FLOAT_STAT_BY(STAT_AsyncIO_AsyncPackagePrecacheWaitTime, (float)FApp::GetDeltaTime());
-					return EAsyncPackageState::TimeOut;
-				}
-			}
-		}
-		PreLoadIndex++;
-	}
-	return PreLoadIndex == ObjLoaded.Num() ? EAsyncPackageState::Complete : EAsyncPackageState::TimeOut;
-}
-
-#elif USE_NEW_ASYNC_IO
-
-EAsyncPackageState::Type FAsyncPackage::PreLoadObjects()
-{
-	SCOPED_LOADTIMER(PreLoadObjectsTime);
-	SCOPE_CYCLE_COUNTER(STAT_FAsyncPackage_PreLoadObjects);
-
-	// GC can't run in here
-	FGCScopeGuard GCGuard;
-
-	TArray<UObject*>& ThreadObjLoaded = FUObjectThreadContext::Get().ObjLoaded;	
+	TArray<UObject*>& ThreadObjLoaded = FUObjectThreadContext::Get().ObjLoaded;
 	PackageObjLoaded.Append(ThreadObjLoaded);
 	ThreadObjLoaded.Reset();
 
@@ -5557,8 +5406,6 @@ EAsyncPackageState::Type FAsyncPackage::PreLoadObjects()
 	return PreLoadIndex == PackageObjLoaded.Num() ? EAsyncPackageState::Complete : EAsyncPackageState::TimeOut;
 }
 
-#else
-
 EAsyncPackageState::Type FAsyncPackage::PreLoadObjects()
 {
 	SCOPED_LOADTIMER(PreLoadObjectsTime);
@@ -5589,7 +5436,6 @@ EAsyncPackageState::Type FAsyncPackage::PreLoadObjects()
 
 	return PreLoadIndex == PackageObjLoaded.Num() ? EAsyncPackageState::Complete : EAsyncPackageState::TimeOut;
 }
-#endif
 
 /**
  * Route PostLoad to all loaded objects. This might load further objects!
@@ -5599,9 +5445,9 @@ EAsyncPackageState::Type FAsyncPackage::PreLoadObjects()
 EAsyncPackageState::Type FAsyncPackage::PostLoadObjects()
 {
 	SCOPE_CYCLE_COUNTER(STAT_FAsyncPackage_PostLoadObjects);
-	
+
 	SCOPED_LOADTIMER(PostLoadObjectsTime);
-	
+
 	// GC can't run in here
 	FGCScopeGuard GCGuard;
 
@@ -5616,9 +5462,10 @@ EAsyncPackageState::Type FAsyncPackage::PostLoadObjects()
 		return EAsyncPackageState::TimeOut;
 	}
 
-#if USE_EVENT_DRIVEN_ASYNC_LOAD
-	PreLoadIndex = PackageObjLoaded.Num(); // we did preloading in a different way and never incremented this
-#endif
+	if (GEventDrivenLoaderEnabled)
+	{
+		PreLoadIndex = PackageObjLoaded.Num(); // we did preloading in a different way and never incremented this
+	}
 
 	// PostLoad objects.
 	while (PostLoadIndex < PackageObjLoaded.Num() && PostLoadIndex < PreLoadIndex && !IsTimeLimitExceeded())
@@ -5630,9 +5477,9 @@ EAsyncPackageState::Type FAsyncPackage::PostLoadObjects()
 			{
 				FScopeCycleCounterUObject ConstructorScope(Object, GET_STATID(STAT_FAsyncPackage_PostLoadObjects));
 
-#if USE_EVENT_DRIVEN_ASYNC_LOAD
-				check(!Object->HasAnyFlags(RF_NeedLoad));
-#endif
+				// We want this check only with EDL enabled
+				check(!GEventDrivenLoaderEnabled || !Object->HasAnyFlags(RF_NeedLoad));
+
 				Object->ConditionalPostLoad();
 
 				LastObjectWorkWasPerformedOn = Object;
@@ -5756,24 +5603,28 @@ EAsyncPackageState::Type FAsyncPackage::PostLoadDeferredObjects(double InTickSta
 			if (UDynamicClass* DynamicClass = Cast<UDynamicClass>(Object))
 			{
 				check((DynamicClass->ClassFlags & CLASS_Constructed) != 0);
-#if USE_EVENT_DRIVEN_ASYNC_LOAD
-//native blueprint 
 
-				check(DynamicClass->HasAnyClassFlags(CLASS_TokenStreamAssembled));
-				// this block should be removed entirely when and if we add the CDO to the fake export table
-				CDOToHandle = DynamicClass->GetDefaultObject(false);
-				UE_CLOG(CDOToHandle, LogStreaming, Fatal, TEXT("EDL did not create the CDO for %s before it finished loading."), *DynamicClass->GetFullName());
-				CDOToHandle->AtomicallyClearInternalFlags(EInternalObjectFlags::AsyncLoading);
-#else
-				UObject* OldCDO = DynamicClass->GetDefaultObject(false);
-				UObject* NewCDO = DynamicClass->GetDefaultObject(true);
-				const bool bCDOWasJustCreated = (OldCDO != NewCDO);
-				if (bCDOWasJustCreated && (NewCDO != nullptr))
+				if (GEventDrivenLoaderEnabled)
 				{
-					NewCDO->AtomicallyClearInternalFlags(EInternalObjectFlags::AsyncLoading);
-					CDOToHandle = NewCDO;
+					//native blueprint 
+
+					check(DynamicClass->HasAnyClassFlags(CLASS_TokenStreamAssembled));
+					// this block should be removed entirely when and if we add the CDO to the fake export table
+					CDOToHandle = DynamicClass->GetDefaultObject(false);
+					UE_CLOG(!CDOToHandle, LogStreaming, Fatal, TEXT("EDL did not create the CDO for %s before it finished loading."), *DynamicClass->GetFullName());
+					CDOToHandle->AtomicallyClearInternalFlags(EInternalObjectFlags::AsyncLoading);
+				}
+				else
+				{
+					UObject* OldCDO = DynamicClass->GetDefaultObject(false);
+					UObject* NewCDO = DynamicClass->GetDefaultObject(true);
+					const bool bCDOWasJustCreated = (OldCDO != NewCDO);
+					if (bCDOWasJustCreated && (NewCDO != nullptr))
+					{
+						NewCDO->AtomicallyClearInternalFlags(EInternalObjectFlags::AsyncLoading);
+						CDOToHandle = NewCDO;
+					}
 				}
-#endif
 			}
 			else
 			{
@@ -5809,10 +5660,11 @@ EAsyncPackageState::Type FAsyncPackage::PostLoadDeferredObjects(double InTickSta
 			if (Linker)
 			{
 				CreateClustersFromPackage(Linker);
-#if !USE_NEW_ASYNC_IO
-				// give a hint to the IO system that we are done with this file for now
-				FIOSystem::Get().HintDoneWithFile(Linker->Filename);
-#endif
+				if (!GNewAsyncIO)
+				{
+					// give a hint to the IO system that we are done with this file for now
+					FIOSystem::Get().HintDoneWithFile(Linker->Filename);
+				}
 			}
 			::IsTimeLimitExceeded(InTickStartTime, bInUseTimeLimit, InOutTimeLimit, LastTypeOfWorkPerformed, LastObjectWorkWasPerformedOn);
 		}
@@ -5871,8 +5723,8 @@ EAsyncPackageState::Type FAsyncPackage::FinishObjects()
 		// Flush linker cache now to reduce peak memory usage (5.5-10x)
 		// We shouldn't need it anyway at this point and even if something attempts to read in PostLoad, 
 		// we're just going to re-cache then.
-#if USE_NEW_ASYNC_IO && !UE_BUILD_SHIPPING
-		if (FPlatformProperties::RequiresCookedData())
+#if !UE_BUILD_SHIPPING
+		if (FPlatformProperties::RequiresCookedData() && GNewAsyncIO)
 		{
 			int32 TotalJunk = 0;
 			for (int32 Index = 0; Index < Linker->ExportMap.Num(); Index++)
@@ -5911,15 +5763,16 @@ void FAsyncPackage::CloseDelayedLinkers()
 	{
 		check(LinkerToClose);
 		check(LinkerToClose->LinkerRoot);
-	#if USE_EVENT_DRIVEN_ASYNC_LOAD
-		FLinkerLoad* LinkerToReset = FLinkerLoad::FindExistingLinkerForPackage(CastChecked<UPackage>(LinkerToClose->LinkerRoot));
-		check(LinkerToReset == LinkerToClose);
-		if (LinkerToReset && LinkerToReset->AsyncRoot)
+		if (GEventDrivenLoaderEnabled)
 		{
-			UE_LOG(LogStreaming, Error, TEXT("Linker cannot be reset right now...leaking %s"), *LinkerToReset->GetArchiveName());
+			FLinkerLoad* LinkerToReset = FLinkerLoad::FindExistingLinkerForPackage(CastChecked<UPackage>(LinkerToClose->LinkerRoot));
+			check(LinkerToReset == LinkerToClose);
+			if (LinkerToReset && LinkerToReset->AsyncRoot)
+			{
+				UE_LOG(LogStreaming, Error, TEXT("Linker cannot be reset right now...leaking %s"), *LinkerToReset->GetArchiveName());
 				continue;
+			}
 		}
-	#endif
 		FLinkerManager::Get().ResetLoaders(LinkerToClose->LinkerRoot);
 		check(LinkerToClose->LinkerRoot == nullptr);
 		check(LinkerToClose->AsyncRoot == nullptr);
@@ -5941,9 +5794,8 @@ void FAsyncPackage::CallCompletionCallbacks(bool bInternal, EAsyncLoadingResult:
 
 void FAsyncPackage::Cancel()
 {
-#if USE_NEW_ASYNC_IO
-	UE_LOG(LogStreaming, Fatal, TEXT("FAsyncPackage::Cancel is not supported with the new loader"));
-#endif
+	UE_CLOG(GNewAsyncIO, LogStreaming, Fatal, TEXT("FAsyncPackage::Cancel is not supported with the new loader"));
+
 	// Call any completion callbacks specified.
 	const EAsyncLoadingResult::Type Result = EAsyncLoadingResult::Canceled;
 	for (int32 CallbackIndex = 0; CallbackIndex < CompletionCallbacks.Num(); CallbackIndex++)
@@ -5970,10 +5822,11 @@ void FAsyncPackage::Cancel()
 	{
 		if (Linker)
 		{
-#if !USE_NEW_ASYNC_IO
-			// give a hint to the IO system that we are done with this file for now
-			FIOSystem::Get().HintDoneWithFile(Linker->Filename);
-#endif
+			if (!GNewAsyncIO)
+			{
+				// give a hint to the IO system that we are done with this file for now
+				FIOSystem::Get().HintDoneWithFile(Linker->Filename);
+			}
 			Linker->FlushCache();
 		}
 		LinkerRoot->ClearFlags(RF_WasLoaded);
@@ -6119,11 +5972,12 @@ void FlushAsyncLoading(int32 PackageID /* = INDEX_NONE */)
 
 		FCoreDelegates::OnAsyncLoadingFlush.Broadcast();
 
-#if !USE_NEW_ASYNC_IO
-		// Disallow low priority requests like texture streaming while we are flushing streaming
-		// in order to avoid excessive seeking.
-		FIOSystem::Get().SetMinPriority( AIOP_Normal );
-#endif
+		if (!GNewAsyncIO)
+		{
+			// Disallow low priority requests like texture streaming while we are flushing streaming
+			// in order to avoid excessive seeking.
+			FIOSystem::Get().SetMinPriority(AIOP_Normal);
+		}
 
 		// Flush async loaders by not using a time limit. Needed for e.g. garbage collection.
 		UE_LOG(LogStreaming, Log,  TEXT("Flushing async loaders.") );
@@ -6153,10 +6007,11 @@ void FlushAsyncLoading(int32 PackageID /* = INDEX_NONE */)
 
 		check(PackageID != INDEX_NONE || !IsAsyncLoading());
 
-#if !USE_NEW_ASYNC_IO
-		// Reset min priority again.
-		FIOSystem::Get().SetMinPriority( AIOP_MIN );
-#endif
+		if (!GNewAsyncIO)
+		{
+			// Reset min priority again.
+			FIOSystem::Get().SetMinPriority(AIOP_MIN);
+		}
 	}
 }
 
@@ -6208,21 +6063,34 @@ void ResumeAsyncLoadingInternal()
 
 bool IsEventDrivenLoaderEnabledInCookedBuilds()
 {
-#if WITH_EDITORONLY_DATA
-	check(GConfig); // Otherwise there's no way we have the correct value of GEventDrivenLoaderEnabled
-	return !!GEventDrivenLoaderEnabled && !FApp::IsEngineInstalled();
-#elif USE_EVENT_DRIVEN_ASYNC_LOAD
-	return true;
-#else
-	return false;
-#endif
+	static struct FEventDrivenLoaderEnabledInCookedBuildsInit
+	{
+		bool bEventDrivenLoaderEnabled;
+		FEventDrivenLoaderEnabledInCookedBuildsInit()
+			: bEventDrivenLoaderEnabled(false)
+		{
+			check(GConfig);
+			GConfig->GetBool(TEXT("/Script/Engine.StreamingSettings"), TEXT("s.EventDrivenLoaderEnabled"), bEventDrivenLoaderEnabled, GEngineIni);
+		}
+	} EventDrivenLoaderEnabledInCookedBuilds;
+	return EventDrivenLoaderEnabledInCookedBuilds.bEventDrivenLoaderEnabled;
+}
+
+bool IsEventDrivenLoaderEnabled()
+{
+	static struct FEventDrivenLoaderEnabledInit
+	{
+		FEventDrivenLoaderEnabledInit()
+		{
+			GEventDrivenLoaderEnabled = IsEventDrivenLoaderEnabledInCookedBuilds() && FPlatformProperties::RequiresCookedData();
+		}
+	} EventDrivenLoaderEnabledInit;
+	return GEventDrivenLoaderEnabled;
 }
 
 /*----------------------------------------------------------------------------
 	FArchiveAsync.
 ----------------------------------------------------------------------------*/
-#if !USE_NEW_ASYNC_IO
-
 static uint8* MallocAsyncBuffer(const SIZE_T Size, SIZE_T& OutAllocatedSize)
 {
 	uint8* Result = nullptr;
@@ -6766,26 +6634,20 @@ void FArchiveAsync::Seek( int64 InPos )
 	CurrentPos = InPos;
 }
 
-
-#endif // !USE_NEW_ASYNC_IO
 /*----------------------------------------------------------------------------
 	End.
 ----------------------------------------------------------------------------*/
 
-#if USE_NEW_ASYNC_IO
 //@todoio Also all of the leaf-first loader can go away
-
 FArchiveAsync2::FArchiveAsync2(const TCHAR* InFileName
-#if USE_EVENT_DRIVEN_ASYNC_LOAD
+
 	, TFunction<void()>&& InSummaryReadyCallback
-#endif
+
 	)
 	: Handle(nullptr)
 	, SizeRequestPtr(nullptr)
-#if !(SPLIT_COOKED_FILES && USE_EVENT_DRIVEN_ASYNC_LOAD)
 	, SummaryRequestPtr(nullptr)
 	, SummaryPrecacheRequestPtr(nullptr)
-#endif
 	, ReadRequestPtr(nullptr)
 	, CanceledReadRequestPtr(nullptr)
 	, PrecacheBuffer(nullptr)
@@ -6802,10 +6664,7 @@ FArchiveAsync2::FArchiveAsync2(const TCHAR* InFileName
 	, OpenTime(FPlatformTime::Seconds())
 	, SummaryReadTime(0.0)
 	, ExportReadTime(0.0)
-#if USE_EVENT_DRIVEN_ASYNC_LOAD
 	, SummaryReadyCallback(Forward<TFunction<void()>>(InSummaryReadyCallback))
-#endif
-
 {
 	LogItem(TEXT("Open"));
 	Handle = FPlatformFileManager::Get().GetPlatformFile().OpenAsyncRead(InFileName);
@@ -6815,13 +6674,16 @@ FArchiveAsync2::FArchiveAsync2(const TCHAR* InFileName
 	{
 		ReadCallback(bWasCancelled, Request);
 	};
-#if USE_EVENT_DRIVEN_ASYNC_LOAD
-	check(SummaryReadyCallback);
-	ReadCallbackFunctionForLinkerLoad = [this](bool bWasCancelled, IAsyncReadRequest* Request)
+
+	if (GEventDrivenLoaderEnabled)
 	{
-		SummaryReadyCallback();
-	};
-#endif
+		check(SummaryReadyCallback);
+		ReadCallbackFunctionForLinkerLoad = [this](bool bWasCancelled, IAsyncReadRequest* Request)
+		{
+			SummaryReadyCallback();
+		};
+	}
+
 	SizeRequestPtr = Handle->SizeRequest(&ReadCallbackFunction);
 
 }
@@ -6855,27 +6717,28 @@ void FArchiveAsync2::ReadCallback(bool bWasCancelled, IAsyncReadRequest* Request
 		}
 		else
 		{
-#if SPLIT_COOKED_FILES && USE_EVENT_DRIVEN_ASYNC_LOAD
-			// in this case we don't need to serialize the summary because we know the header is the whole file
-			HeaderSize = FileSize;
-			LogItem(TEXT("Starting Split Header"), 0, FileSize);
-			PrecacheInternal(0, HeaderSize);
-			FPlatformMisc::MemoryBarrier();
-			LoadPhase = ELoadPhase::WaitingForHeader;
-#else
-			int64 Size = FMath::Min<int64>(MAX_SUMMARY_SIZE, FileSize);
-			LogItem(TEXT("Starting Summary"), 0, Size);
-			SummaryRequestPtr = Handle->ReadRequest(0, Size, AIOP_Normal, &ReadCallbackFunction);
-			// I need a precache request here to keep the memory alive until I submit the header request
-			SummaryPrecacheRequestPtr = Handle->ReadRequest(0, Size, AIOP_Precache);
-#endif
+			if (GEventDrivenLoaderEnabled)
+			{
+				// in this case we don't need to serialize the summary because we know the header is the whole file
+				HeaderSize = FileSize;
+				LogItem(TEXT("Starting Split Header"), 0, FileSize);
+				PrecacheInternal(0, HeaderSize);
+				FPlatformMisc::MemoryBarrier();
+				LoadPhase = ELoadPhase::WaitingForHeader;
+			}
+			else
+			{
+				int64 Size = FMath::Min<int64>(MAX_SUMMARY_SIZE, FileSize);
+				LogItem(TEXT("Starting Summary"), 0, Size);
+				SummaryRequestPtr = Handle->ReadRequest(0, Size, AIOP_Normal, &ReadCallbackFunction);
+				// I need a precache request here to keep the memory alive until I submit the header request
+				SummaryPrecacheRequestPtr = Handle->ReadRequest(0, Size, AIOP_Precache);
+			}
 		}
 	}
 	else if (LoadPhase == ELoadPhase::WaitingForSummary)
 	{
-#if SPLIT_COOKED_FILES && USE_EVENT_DRIVEN_ASYNC_LOAD
-		check(0);
-#else
+		check(!GEventDrivenLoaderEnabled);
 		uint8* Mem = Request->GetReadResults();
 		if (!Mem)
 		{
@@ -6903,7 +6766,6 @@ void FArchiveAsync2::ReadCallback(bool bWasCancelled, IAsyncReadRequest* Request
 		}
 		FPlatformMisc::MemoryBarrier();
 		LoadPhase = ELoadPhase::WaitingForHeader;
-#endif
 	}
 	else
 	{
@@ -6913,7 +6775,9 @@ void FArchiveAsync2::ReadCallback(bool bWasCancelled, IAsyncReadRequest* Request
 
 void FArchiveAsync2::FlushPrecacheBlock()
 {
+#if DEVIRTUALIZE_FLinkerLoad_Serialize
 	DiscardInlineBufferAndUpdateCurrentPos();
+#endif
 	if (PrecacheBuffer)
 	{
 		DEC_MEMORY_STAT_BY(STAT_FArchiveAsync2Mem, PrecacheEndPos - PrecacheStartPos);
@@ -6987,12 +6851,10 @@ int64 FArchiveAsync2::TotalSize()
 	{
 		QUICK_SCOPE_CYCLE_COUNTER(STAT_FArchiveAsync2_TotalSize);
 		SizeRequestPtr->WaitCompletion();
-#if SPLIT_COOKED_FILES && USE_EVENT_DRIVEN_ASYNC_LOAD
-		if (HeaderSizeWhenReadingExportsFromSplitFile)
+		if (GEventDrivenLoaderEnabled && HeaderSizeWhenReadingExportsFromSplitFile)
 		{
 			FileSize = SizeRequestPtr->GetSizeResults();
 		}
-#endif
 		delete SizeRequestPtr;
 		SizeRequestPtr = nullptr;
 	}
@@ -7022,8 +6884,7 @@ FORCEINLINE void FArchiveAsync2::SetPosAndUpdatePrecacheBuffer(int64 Pos)
 
 void FArchiveAsync2::Seek(int64 InPos)
 {
-#if SPLIT_COOKED_FILES && USE_EVENT_DRIVEN_ASYNC_LOAD
-	if (LoadPhase < ELoadPhase::ProcessingExports)
+	if (GEventDrivenLoaderEnabled && LoadPhase < ELoadPhase::ProcessingExports)
 	{
 		check(!HeaderSizeWhenReadingExportsFromSplitFile && HeaderSize && TotalSize() == HeaderSize);
 		if (InPos >= HeaderSize)
@@ -7031,7 +6892,6 @@ void FArchiveAsync2::Seek(int64 InPos)
 			FirstExportStarting();
 		}
 	}
-#endif
 	check(InPos >= 0 && InPos <= TotalSizeOrMaxInt64IfNotReady());
 #if DEVIRTUALIZE_FLinkerLoad_Serialize
 	SetPosAndUpdatePrecacheBuffer(InPos);
@@ -7119,9 +6979,7 @@ void FArchiveAsync2::CancelRead()
 bool FArchiveAsync2::WaitForIntialPhases(float InTimeLimit)
 {
 	if (SizeRequestPtr 
-#if !(SPLIT_COOKED_FILES && USE_EVENT_DRIVEN_ASYNC_LOAD)
-		|| SummaryRequestPtr || SummaryPrecacheRequestPtr
-#endif
+		|| GEventDrivenLoaderEnabled || SummaryRequestPtr || SummaryPrecacheRequestPtr
 		)
 	{
 		QUICK_SCOPE_CYCLE_COUNTER(STAT_FArchiveAsync2_WaitForIntialPhases);
@@ -7139,52 +6997,53 @@ bool FArchiveAsync2::WaitForIntialPhases(float InTimeLimit)
 				return false;
 			}
 		}
-#if !(SPLIT_COOKED_FILES && USE_EVENT_DRIVEN_ASYNC_LOAD)
-		if (SummaryRequestPtr)
+		if (!GEventDrivenLoaderEnabled)
 		{
-			float TimeLimit = 0.0f;
-			if (InTimeLimit > 0.0f)
+			if (SummaryRequestPtr)
 			{
-				TimeLimit = InTimeLimit - (FPlatformTime::Seconds() - StartTime);
-				if (TimeLimit < MIN_REMAIN_TIME)
+				float TimeLimit = 0.0f;
+				if (InTimeLimit > 0.0f)
 				{
+					TimeLimit = InTimeLimit - (FPlatformTime::Seconds() - StartTime);
+					if (TimeLimit < MIN_REMAIN_TIME)
+					{
+						return false;
+					}
+				}
+				if (SummaryRequestPtr->WaitCompletion(TimeLimit))
+				{
+					delete SummaryRequestPtr;
+					SummaryRequestPtr = nullptr;
+				}
+				else
+				{
+					check(InTimeLimit > 0.0f);
 					return false;
 				}
 			}
-			if (SummaryRequestPtr->WaitCompletion(TimeLimit))
-			{
-				delete SummaryRequestPtr;
-				SummaryRequestPtr = nullptr;
-			}
-			else
-			{
-				check(InTimeLimit > 0.0f);
-				return false;
-			}
-		}
-		if (SummaryPrecacheRequestPtr)
-		{
-			float TimeLimit = 0.0f;
-			if (InTimeLimit > 0.0f)
+			if (SummaryPrecacheRequestPtr)
 			{
-				TimeLimit = InTimeLimit - (FPlatformTime::Seconds() - StartTime);
-				if (TimeLimit < MIN_REMAIN_TIME)
+				float TimeLimit = 0.0f;
+				if (InTimeLimit > 0.0f)
 				{
+					TimeLimit = InTimeLimit - (FPlatformTime::Seconds() - StartTime);
+					if (TimeLimit < MIN_REMAIN_TIME)
+					{
+						return false;
+					}
+				}
+				if (SummaryPrecacheRequestPtr->WaitCompletion(TimeLimit))
+				{
+					delete SummaryPrecacheRequestPtr;
+					SummaryPrecacheRequestPtr = nullptr;
+				}
+				else
+				{
+					check(InTimeLimit > 0.0f);
 					return false;
 				}
 			}
-			if (SummaryPrecacheRequestPtr->WaitCompletion(TimeLimit))
-			{
-				delete SummaryPrecacheRequestPtr;
-				SummaryPrecacheRequestPtr = nullptr;
-			}
-			else
-			{
-				check(InTimeLimit > 0.0f);
-				return false;
-			}
 		}
-#endif
 		LogItem(TEXT("Wait Summary"), 0, HeaderSize, StartTime);
 	}
 	return true;
@@ -7252,9 +7111,7 @@ bool FArchiveAsync2::PrecacheInternal(int64 RequestOffset, int64 RequestSize, bo
 #endif
 	check(ReadRequestOffset - HeaderSizeWhenReadingExportsFromSplitFile >= 0);
 	ReadRequestPtr = Handle->ReadRequest(ReadRequestOffset - HeaderSizeWhenReadingExportsFromSplitFile, ReadRequestSize, AIOP_Normal
-#if USE_EVENT_DRIVEN_ASYNC_LOAD
-		, (LoadPhase == ELoadPhase::WaitingForSummary) ? &ReadCallbackFunctionForLinkerLoad : nullptr
-#endif
+		, (GEventDrivenLoaderEnabled && LoadPhase == ELoadPhase::WaitingForSummary) ? &ReadCallbackFunctionForLinkerLoad : nullptr
 		);
 	if (LoadPhase != ELoadPhase::WaitingForSummary && ReadRequestPtr->PollCompletion())
 	{
@@ -7280,33 +7137,33 @@ void FArchiveAsync2::FirstExportStarting()
 	LogItem(TEXT("Exports"));
 	LoadPhase = ELoadPhase::ProcessingExports;
 
-#if SPLIT_COOKED_FILES && USE_EVENT_DRIVEN_ASYNC_LOAD
-	FlushCache();
-	if (Handle)
+	if (GEventDrivenLoaderEnabled)
 	{
-		delete Handle;
-		Handle = nullptr;
-	}
+		FlushCache();
+		if (Handle)
+		{
+			delete Handle;
+			Handle = nullptr;
+		}
 
-	HeaderSizeWhenReadingExportsFromSplitFile = HeaderSize;
-	FileName = FPaths::GetBaseFilename(FileName, false) + TEXT(".uexp");
+		HeaderSizeWhenReadingExportsFromSplitFile = HeaderSize;
+		FileName = FPaths::GetBaseFilename(FileName, false) + TEXT(".uexp");
 
-	Handle = FPlatformFileManager::Get().GetPlatformFile().OpenAsyncRead(*FileName);
-	check(Handle); // this generally cannot fail because it is async
+		Handle = FPlatformFileManager::Get().GetPlatformFile().OpenAsyncRead(*FileName);
+		check(Handle); // this generally cannot fail because it is async
 
-	check(!SizeRequestPtr);
-	SizeRequestPtr = Handle->SizeRequest();
-	if (SizeRequestPtr->PollCompletion())
-	{
-		TotalSize(); // complete the request
+		check(!SizeRequestPtr);
+		SizeRequestPtr = Handle->SizeRequest();
+		if (SizeRequestPtr->PollCompletion())
+		{
+			TotalSize(); // complete the request
+		}
 	}
-
-#endif
 }
 
-#if USE_EVENT_DRIVEN_ASYNC_LOAD
 IAsyncReadRequest* FArchiveAsync2::MakeEventDrivenPrecacheRequest(int64 Offset, int64 BytesToRead, FAsyncFileCallBack* CompleteCallback)
 {
+	check(GEventDrivenLoaderEnabled);
 	if (LoadPhase == ELoadPhase::WaitingForFirstExport)
 	{
 		FirstExportStarting();
@@ -7318,7 +7175,6 @@ IAsyncReadRequest* FArchiveAsync2::MakeEventDrivenPrecacheRequest(int64 Offset,
 	LogItem(TEXT("Event Precache"), Offset - HeaderSizeWhenReadingExportsFromSplitFile, BytesToRead, StartTime);
 	return Precache;
 }
-#endif
 
 
 bool FArchiveAsync2::Precache(int64 RequestOffset, int64 RequestSize, bool bUseTimeLimit, bool bUseFullTimeLimit, double TickStartTime, float TimeLimit)
@@ -7664,6 +7520,4 @@ static FAutoConsoleCommand DumpSerializeCmd(
 	);
 #endif
 
-#endif // USE_NEW_ASYNC_IO
-
 //#pragma clang optimize on
diff --git a/Engine/Source/Runtime/CoreUObject/Private/Serialization/AsyncLoadingPrivate.h b/Engine/Source/Runtime/CoreUObject/Private/Serialization/AsyncLoadingPrivate.h
--- a/Engine/Source/Runtime/CoreUObject/Private/Serialization/AsyncLoadingPrivate.h
+++ b/Engine/Source/Runtime/CoreUObject/Private/Serialization/AsyncLoadingPrivate.h
@@ -8,8 +8,6 @@
 
 DECLARE_LOG_CATEGORY_EXTERN(LogLoadingDev, Fatal, All);
 
-#if USE_NEW_ASYNC_IO
-
 class COREUOBJECT_API FArchiveAsync2 final: public FArchive
 {
 public:
@@ -29,9 +27,7 @@ public:
 	};
 
 	FArchiveAsync2(const TCHAR* InFileName
-#if USE_EVENT_DRIVEN_ASYNC_LOAD
 		, TFunction<void()>&& InSummaryReadyCallback
-#endif
 		);
 	virtual ~FArchiveAsync2();
 	virtual bool Close() override;
@@ -40,7 +36,12 @@ public:
 	virtual void Serialize(void* Data, int64 Num) override;
 	FORCEINLINE virtual int64 Tell() override
 	{
+#if DEVIRTUALIZE_FLinkerLoad_Serialize
 		return CurrentPos + (ActiveFPLB->StartFastPathLoadBuffer - ActiveFPLB->OriginalFastPathLoadBuffer);
+#else
+		check(false);
+		return 0;
+#endif
 	}
 	virtual int64 TotalSize() override;
 	virtual void Seek(int64 InPos) override;
@@ -53,9 +54,8 @@ public:
 	void StartReadingHeader();
 	void EndReadingHeader();
 	void FirstExportStarting();
-#if USE_EVENT_DRIVEN_ASYNC_LOAD
+
 	IAsyncReadRequest* MakeEventDrivenPrecacheRequest(int64 Offset, int64 BytesToRead, FAsyncFileCallBack* CompleteCallback);
-#endif
 
 	void LogItem(const TCHAR* Item, int64 Offset = 0, int64 Size = 0, double StartTime = 0.0);
 
@@ -86,10 +86,10 @@ private:
 
 	IAsyncReadFileHandle* Handle;
 	IAsyncReadRequest* SizeRequestPtr;
-#if !(SPLIT_COOKED_FILES && USE_EVENT_DRIVEN_ASYNC_LOAD)
+
 	IAsyncReadRequest* SummaryRequestPtr;
 	IAsyncReadRequest* SummaryPrecacheRequestPtr;
-#endif
+
 	IAsyncReadRequest* ReadRequestPtr;
 	IAsyncReadRequest* CanceledReadRequestPtr;
 
@@ -118,11 +118,8 @@ private:
 	double OpenTime;
 	double SummaryReadTime;
 	double ExportReadTime;
-#if USE_EVENT_DRIVEN_ASYNC_LOAD
+
 	TFunction<void()> SummaryReadyCallback;
 	FAsyncFileCallBack ReadCallbackFunctionForLinkerLoad;
-#endif
 
 };
-
-#endif
diff --git a/Engine/Source/Runtime/CoreUObject/Private/Serialization/AsyncLoadingThread.h b/Engine/Source/Runtime/CoreUObject/Private/Serialization/AsyncLoadingThread.h
--- a/Engine/Source/Runtime/CoreUObject/Private/Serialization/AsyncLoadingThread.h
+++ b/Engine/Source/Runtime/CoreUObject/Private/Serialization/AsyncLoadingThread.h
@@ -18,12 +18,7 @@
 
 class IAssetRegistryInterface;
 
-#ifndef USE_EVENT_DRIVEN_ASYNC_LOAD
-#error "USE_EVENT_DRIVEN_ASYNC_LOAD must be defined"
-#endif
-
-#if USE_EVENT_DRIVEN_ASYNC_LOAD
-
+/** [EDL] Event Driven loader event */
 struct FAsyncLoadEvent
 {
 	enum
@@ -71,6 +66,7 @@ struct FAsyncLoadEvent
 	}
 };
 
+/** [EDL] Event queue for the event driven loader */
 struct FAsyncLoadEventQueue
 {
 	//FCriticalSection AsyncLoadEventQueueCritical; //@todoio maybe this doesn't need to be protected by a critical section. 
@@ -109,8 +105,6 @@ struct FAsyncLoadEventQueue
 	}
 };
 
-#endif
-
 /** Package dependency tree used for flushing specific packages */
 struct FFlushTree
 {
@@ -189,11 +183,10 @@ class FAsyncLoadingThread : public FRunnable
 	/** [ASYNC THREAD] Array of packages that are being preloaded */
 	TArray<FAsyncPackage*> AsyncPackages;
 	TMap<FName, FAsyncPackage*> AsyncPackageNameLookup;
-#if USE_EVENT_DRIVEN_ASYNC_LOAD
 public:
+	/** [EDL] Async Packages that are ready for tick */
 	TArray<FAsyncPackage*> AsyncPackagesReadyForTick;
 private:
-#endif 
 
 #if THREADSAFE_UOBJECTS
 	/** We only lock AsyncPackages array to make GetAsyncLoadPercentage thread safe, so we only care about locking Add/Remove operations on the async thread */
@@ -243,9 +236,10 @@ public:
 	/** Returns the async loading thread singleton */
 	static FAsyncLoadingThread& Get();
 
-#if USE_EVENT_DRIVEN_ASYNC_LOAD
+	/** [EDL] Event queue */
 	FAsyncLoadEventQueue EventQueue;
 
+	/** [EDL] Gets a package from a weak pointer */
 	FORCEINLINE FAsyncPackage* GetPackage(FWeakAsyncPackagePtr Ptr)
 	{
 		if (Ptr.PackageName != NAME_None && Ptr.SerialNumber)
@@ -259,16 +253,24 @@ public:
 		return nullptr;
 	}
 
+	/** [EDL] Queues CreateLinker event */
 	void QueueEvent_CreateLinker(FAsyncPackage* Pkg, int32 EventSystemPriority = 0);
+	/** [EDL] Queues FinishLinker event */
 	void QueueEvent_FinishLinker(FWeakAsyncPackagePtr WeakPtr, int32 EventSystemPriority = 0);
+	/** [EDL] Queues StartImportPackages event */
 	void QueueEvent_StartImportPackages(FAsyncPackage* Pkg, int32 EventSystemPriority = 0);
+	/** [EDL] Queues SetupImports event */
 	void QueueEvent_SetupImports(FAsyncPackage* Pkg, int32 EventSystemPriority = 0);
+	/** [EDL] Queues SetupExports event */
 	void QueueEvent_SetupExports(FAsyncPackage* Pkg, int32 EventSystemPriority = 0);
+	/** [EDL] Queues ProcessImportsAndExports event */
 	void QueueEvent_ProcessImportsAndExports(FAsyncPackage* Pkg, int32 EventSystemPriority = 0);
+	/** [EDL] Queues ExportsDone event */
 	void QueueEvent_ExportsDone(FAsyncPackage* Pkg, int32 EventSystemPriority = 0);
+	/** [EDL] Queues ProcessPostload event */
 	void QueueEvent_ProcessPostloadWait(FAsyncPackage* Pkg, int32 EventSystemPriority = 0);
+	/** [EDL] Queues StartPostLoad event */
 	void QueueEvent_StartPostLoad(FAsyncPackage* Pkg, int32 EventSystemPriority = 0);
-#endif
 
 	/** True if multithreaded async loading should be used. */
 	static FORCEINLINE bool IsMultithreaded()
@@ -439,12 +441,10 @@ public:
 	*/
 	EAsyncPackageState::Type ProcessAsyncLoading(int32& OutPackagesProcessed, bool bUseTimeLimit = false, bool bUseFullTimeLimit = false, float TimeLimit = 0.0f, FFlushTree* FlushTree = nullptr);
 
-#if USE_EVENT_DRIVEN_ASYNC_LOAD
 	/**
-	* [ASYNC* THREAD] Checks fopr cycles in the event driven loader and does fatal errors in that case
+	* [EDL] [ASYNC* THREAD] Checks fopr cycles in the event driven loader and does fatal errors in that case
 	*/
 	void CheckForCycles();
-#endif
 
 	/**
 	* [GAME THREAD] Ticks game thread side of async loading.
diff --git a/Engine/Source/Runtime/CoreUObject/Private/Serialization/BulkData.cpp b/Engine/Source/Runtime/CoreUObject/Private/Serialization/BulkData.cpp
--- a/Engine/Source/Runtime/CoreUObject/Private/Serialization/BulkData.cpp
+++ b/Engine/Source/Runtime/CoreUObject/Private/Serialization/BulkData.cpp
@@ -667,9 +667,8 @@ void FUntypedBulkData::StartSerializingBulkData(FArchive& Ar, UObject* Owner, in
 	{
 		BulkDataAsync.Reallocate(GetBulkDataSize(), BulkDataAlignment);
 
-#if USE_NEW_ASYNC_IO
-		UE_LOG(LogSerialization, Error, TEXT("Attempt to stream bulk data with EDL enabled. This is not desireable. File %s"), *Filename);
-#endif
+		UE_CLOG(GNewAsyncIO, LogSerialization, Error, TEXT("Attempt to stream bulk data with EDL enabled. This is not desireable. File %s"), *Filename);
+
 		FArchive* FileReaderAr = IFileManager::Get().CreateFileReader(*Filename, FILEREAD_Silent);
 		checkf(FileReaderAr != NULL, TEXT("Attempted to load bulk data from an invalid filename '%s'."), *Filename);
 
@@ -696,30 +695,24 @@ static FAutoConsoleVariableRef CVarMinimumBulkDataSizeForAsyncLoading(
 	ECVF_Default
 	);
 
-#if !defined(USE_NEW_ASYNC_IO) || !defined(SPLIT_COOKED_FILES) || !defined(USE_EVENT_DRIVEN_ASYNC_LOAD)
-	#error "USE_NEW_ASYNC_IO and SPLIT_COOKED_FILES and USE_EVENT_DRIVEN_ASYNC_LOAD must be defined."
-#endif
-
 bool FUntypedBulkData::ShouldStreamBulkData()
 {
-#if USE_NEW_ASYNC_IO
-	const bool bPayloadInline = !(BulkDataFlags&BULKDATA_PayloadAtEndOfFile);
-	if (bPayloadInline)
+	if (GNewAsyncIO && !(BulkDataFlags&BULKDATA_PayloadAtEndOfFile))
 	{
 		return false; // if it is inline, it is already precached, so use it
 	}
 
-#endif
-#if SPLIT_COOKED_FILES && USE_EVENT_DRIVEN_ASYNC_LOAD
-	const bool bSeperateFile = !!(BulkDataFlags&BULKDATA_PayloadInSeperateFile);
-
-	if (!bSeperateFile)
+	if (GEventDrivenLoaderEnabled)
 	{
-		check(!"Bulk data should either be inline or stored in a separate file for the new uobject loader.");
-		return false; // if it is not in a separate file, then we can't easily find the correct offset in the uexp file; we don't want this case anyway!
+		const bool bSeperateFile = !!(BulkDataFlags&BULKDATA_PayloadInSeperateFile);
+
+		if (!bSeperateFile)
+		{
+			check(!"Bulk data should either be inline or stored in a separate file for the new uobject loader.");
+			return false; // if it is not in a separate file, then we can't easily find the correct offset in the uexp file; we don't want this case anyway!
+		}
 	}
 
-#endif
 	const bool bForceStream = !!(BulkDataFlags & BULKDATA_ForceStreamPayload);
 
 	return (FPlatformProperties::RequiresCookedData() && !Filename.IsEmpty() &&
@@ -892,17 +885,15 @@ void FUntypedBulkData::Serialize( FArchive& Ar, UObject* Owner, int32 Idx )
 						if (BulkDataFlags & BULKDATA_PayloadInSeperateFile)
 						{
 							// open seperate bulk data file
-#if USE_NEW_ASYNC_IO
-							UE_LOG(LogSerialization, Error, TEXT("Attempt to sync load bulk data with EDL enabled (separate file). This is not desireable. File %s"), *Filename);
-#endif
-#if SPLIT_COOKED_FILES && USE_EVENT_DRIVEN_ASYNC_LOAD
-							if (Filename.EndsWith(TEXT(".uasset")) || Filename.EndsWith(TEXT(".umap")))
+							UE_CLOG(GNewAsyncIO, LogSerialization, Error, TEXT("Attempt to sync load bulk data with EDL enabled (separate file). This is not desireable. File %s"), *Filename);
+
+							if (GEventDrivenLoaderEnabled && (Filename.EndsWith(TEXT(".uasset")) || Filename.EndsWith(TEXT(".umap"))))
 							{
 								BulkDataOffsetInFile -= IFileManager::Get().FileSize(*Filename);
 								check(BulkDataOffsetInFile >= 0);
 								Filename = FPaths::GetBaseFilename(Filename, false) + TEXT(".uexp");
 							}
-#endif
+
 							FArchive* TargetArchive = IFileManager::Get().CreateFileReader(*Filename);
 							// seek to the location in the file where the payload is stored
 							TargetArchive->Seek(BulkDataOffsetInFile);
@@ -913,9 +904,8 @@ void FUntypedBulkData::Serialize( FArchive& Ar, UObject* Owner, int32 Idx )
 						}
 						else
 						{
-#if USE_NEW_ASYNC_IO
-							UE_LOG(LogSerialization, Error, TEXT("Attempt to sync load bulk data with EDL enabled. This is not desireable. File %s"), *Filename);
-#endif
+							UE_CLOG(GNewAsyncIO, LogSerialization, Error, TEXT("Attempt to sync load bulk data with EDL enabled. This is not desireable. File %s"), *Filename);
+
 							// store the current file offset
 							int64 CurOffset = Ar.Tell();
 							// seek to the location in the file where the payload is stored
@@ -926,7 +916,7 @@ void FUntypedBulkData::Serialize( FArchive& Ar, UObject* Owner, int32 Idx )
 							Ar.Seek(CurOffset);
 						}
 					}
-  				}
+  			}
 			}
 		}
 		// We're saving to the persistent archive.
@@ -1383,17 +1373,15 @@ void FUntypedBulkData::LoadDataIntoMemory( void* Dest )
 		// load from the specied filename when the linker has been cleared
 		checkf( Filename != TEXT(""), TEXT( "Attempted to load bulk data without a proper filename." ) );
 	
-#if USE_NEW_ASYNC_IO
-		UE_CLOG(!(IsInGameThread() || IsInAsyncLoadingThread()), LogSerialization, Error, TEXT("Attempt to sync load bulk data with EDL enabled (LoadDataIntoMemory). This is not desireable. File %s"), *Filename);
-#endif
-#if SPLIT_COOKED_FILES && USE_EVENT_DRIVEN_ASYNC_LOAD
-		if (Filename.EndsWith(TEXT(".uasset")) || Filename.EndsWith(TEXT(".umap")))
+		UE_CLOG(GNewAsyncIO && !(IsInGameThread() || IsInAsyncLoadingThread()), LogSerialization, Error, TEXT("Attempt to sync load bulk data with EDL enabled (LoadDataIntoMemory). This is not desireable. File %s"), *Filename);
+
+		if (GEventDrivenLoaderEnabled && (Filename.EndsWith(TEXT(".uasset")) || Filename.EndsWith(TEXT(".umap"))))
 		{
 			BulkDataOffsetInFile -= IFileManager::Get().FileSize(*Filename);
 			check(BulkDataOffsetInFile >= 0);
 			Filename = FPaths::GetBaseFilename(Filename, false) + TEXT(".uexp");
 		}
-#endif
+
 		FArchive* Ar = IFileManager::Get().CreateFileReader(*Filename, FILEREAD_Silent);
 		checkf( Ar != NULL, TEXT( "Attempted to load bulk data from an invalid filename '%s'." ), *Filename );
 	
diff --git a/Engine/Source/Runtime/CoreUObject/Private/UObject/Class.cpp b/Engine/Source/Runtime/CoreUObject/Private/UObject/Class.cpp
--- a/Engine/Source/Runtime/CoreUObject/Private/UObject/Class.cpp
+++ b/Engine/Source/Runtime/CoreUObject/Private/UObject/Class.cpp
@@ -48,10 +48,6 @@ DEFINE_LOG_CATEGORY(LogClass);
 	#endif
 #endif
 
-#if !defined(USE_EVENT_DRIVEN_ASYNC_LOAD)
-#error "USE_EVENT_DRIVEN_ASYNC_LOAD must be defined"
-#endif
-
 //////////////////////////////////////////////////////////////////////////
 
 FThreadSafeBool& InternalSafeGetTokenStreamDirtyFlag()
@@ -529,8 +525,7 @@ void UStruct::StaticLink(bool bRelinkExistingProperties)
 void UStruct::GetPreloadDependencies(TArray<UObject*>& OutDeps)
 {
 	Super::GetPreloadDependencies(OutDeps);
-	UStruct* InheritanceSuper = GetInheritanceSuper();
-	OutDeps.Add(InheritanceSuper);
+	OutDeps.Add(SuperStruct);
 
 	for (UField* Field = Children; Field; Field = Field->Next)
 	{
@@ -2664,12 +2659,11 @@ UObject* UClass::CreateDefaultObject()
 		{
 			UObjectForceRegistration(ParentClass);
 			ParentDefaultObject = ParentClass->GetDefaultObject(); // Force the default object to be constructed if it isn't already
-#if USE_EVENT_DRIVEN_ASYNC_LOAD
-			if (!GIsInitialLoad)
+			check(GConfig);
+			if (GEventDrivenLoaderEnabled && EVENT_DRIVEN_ASYNC_LOAD_ACTIVE_AT_RUNTIME)
 			{ 
 				check(ParentDefaultObject && !ParentDefaultObject->HasAnyFlags(RF_NeedLoad));
 			}
-#endif
 		}
 
 		if ( (ParentDefaultObject != NULL) || (this == UObject::StaticClass()) )
@@ -3638,17 +3632,20 @@ void UClass::Serialize( FArchive& Ar )
 	{
 		if (ClassDefaultObject == NULL)
 		{
-
-#if USE_EVENT_DRIVEN_ASYNC_LOAD
+			check(GConfig);
+			if (GEventDrivenLoaderEnabled)
+			{
 			ClassDefaultObject = GetDefaultObject();
 			// we do this later anyway, once we find it and set it in the export table. 
 			// ClassDefaultObject->SetFlags(RF_NeedLoad | RF_NeedPostLoad | RF_NeedPostLoadSubobjects);
-#else
+			}
+			else
+			{
 			UE_LOG(LogClass, Error, TEXT("CDO for class %s did not load!"), *GetPathName());
 			ensure(ClassDefaultObject != NULL);
 			ClassDefaultObject = GetDefaultObject();
 			Ar.ForceBlueprintFinalization();
-#endif
+			}
 		}
 	}
 }
diff --git a/Engine/Source/Runtime/CoreUObject/Private/UObject/CoreNative.cpp b/Engine/Source/Runtime/CoreUObject/Private/UObject/CoreNative.cpp
--- a/Engine/Source/Runtime/CoreUObject/Private/UObject/CoreNative.cpp
+++ b/Engine/Source/Runtime/CoreUObject/Private/UObject/CoreNative.cpp
@@ -269,15 +269,10 @@ UObject* FObjectInstancingGraph::InstancePropertyValue( class UObject* Component
 	return NewValue;
 }
 
-#if !defined(USE_EVENT_DRIVEN_ASYNC_LOAD)
-#error "USE_EVENT_DRIVEN_ASYNC_LOAD must be defined"
-#endif
-
 void FObjectInstancingGraph::AddNewObject(class UObject* ObjectInstance, UObject* InArchetype /*= nullptr*/)
 {
-#if USE_EVENT_DRIVEN_ASYNC_LOAD
-	check(!InArchetype || !InArchetype->HasAnyFlags(RF_NeedLoad));
-#endif
+	check(!GEventDrivenLoaderEnabled || !InArchetype || !InArchetype->HasAnyFlags(RF_NeedLoad));
+
 	if (HasDestinationRoot())
 	{
 		AddNewInstance(ObjectInstance, InArchetype);
diff --git a/Engine/Source/Runtime/CoreUObject/Private/UObject/EnumProperty.cpp b/Engine/Source/Runtime/CoreUObject/Private/UObject/EnumProperty.cpp
--- a/Engine/Source/Runtime/CoreUObject/Private/UObject/EnumProperty.cpp
+++ b/Engine/Source/Runtime/CoreUObject/Private/UObject/EnumProperty.cpp
@@ -23,7 +23,7 @@ namespace UE4EnumProperty_Private
 				LogClass,
 				Warning,
 				TEXT("Failed to find valid enum value '%s' for enum type '%s' when converting property '%s' during property loading - setting to '%s'"),
-				*LexicalConversion::ToString(OldValue),
+				*Lex::ToString(OldValue),
 				*Enum->GetName(),
 				*EnumProp->GetName(),
 				*Enum->GetNameByValue(Enum->GetMaxEnumValue()).ToString()
diff --git a/Engine/Source/Runtime/CoreUObject/Private/UObject/GarbageCollection.cpp b/Engine/Source/Runtime/CoreUObject/Private/UObject/GarbageCollection.cpp
--- a/Engine/Source/Runtime/CoreUObject/Private/UObject/GarbageCollection.cpp
+++ b/Engine/Source/Runtime/CoreUObject/Private/UObject/GarbageCollection.cpp
@@ -353,7 +353,7 @@ public:
 
 	/** Marks all objects that can't be directly in a cluster but are referenced by it as reachable */
 	template <bool bParallel>
-	FORCEINLINE void MarkClusterMutableObjectsAsReachable(FUObjectCluster* Cluster, TArray<UObject*>& ObjectsToSerialize)
+	static FORCEINLINE void MarkClusterMutableObjectsAsReachable(FUObjectCluster* Cluster, TArray<UObject*>& ObjectsToSerialize)
 	{
 		for (int32 ReferencedMutableObjectIndex : Cluster->MutableObjects)
 		{
@@ -376,7 +376,7 @@ public:
 
 	/** Marks all clusters referenced by another cluster as rechable */
 	template <bool bParallel>
-	FORCEINLINE void MarkReferencedClustersAsReachable(int32 ClusterRootIndex, TArray<UObject*>& ObjectsToSerialize)
+	static FORCEINLINE void MarkReferencedClustersAsReachable(int32 ClusterRootIndex, TArray<UObject*>& ObjectsToSerialize)
 	{
 		FUObjectCluster* Cluster = GUObjectClusters.FindChecked(ClusterRootIndex);
 		// Also mark all referenced objects from outside of the cluster as reachable
@@ -427,6 +427,8 @@ public:
 		// Remove references to pending kill objects if we're allowed to do so.
 		if (ObjectItem->IsPendingKill() && bAllowReferenceElimination)
 		{
+			checkSlow(ObjectItem->HasAnyFlags(EInternalObjectFlags::ClusterRoot) == false);
+			checkSlow(ObjectItem->GetOwnerIndex() == 0)
 			// Null out reference.
 			Object = NULL;
 		}
@@ -487,23 +489,41 @@ public:
 		}
 		else if (ObjectItem->GetOwnerIndex() && !ObjectItem->HasAnyFlags(EInternalObjectFlags::ReachableInCluster))
 		{
-			ObjectItem->SetFlags(EInternalObjectFlags::ReachableInCluster);
-			// Make sure cluster root object is reachable too
-			const int32 OwnerIndex = ObjectItem->GetOwnerIndex();
-			FUObjectItem* RootObjectItem = GUObjectArray.IndexToObjectUnsafeForGC(OwnerIndex);
-			checkSlow(RootObjectItem->HasAnyFlags(EInternalObjectFlags::ClusterRoot));
+			bool bNeedsDoing = true;
 			if (GIsRunningParallelReachability)
 			{
-				if (RootObjectItem->ThisThreadAtomicallyClearedRFUnreachable())
-				{
-					RootObjectItem->ThisThreadAtomicallyClearedFlag(EInternalObjectFlags::NoStrongReference);
-					// Make sure all referenced clusters are marked as reachable too
-					MarkReferencedClustersAsReachable<true>(OwnerIndex, ObjectsToSerialize);
-				}
+				bNeedsDoing = ObjectItem->ThisThreadAtomicallySetFlag(EInternalObjectFlags::ReachableInCluster);
+			}
+			else
+			{
+				ObjectItem->SetFlags(EInternalObjectFlags::ReachableInCluster);
 			}
-			else if (RootObjectItem->IsUnreachable())
+			if (bNeedsDoing)
 			{
-				RootObjectItem->ClearFlags(EInternalObjectFlags::Unreachable | EInternalObjectFlags::NoStrongReference);
+			    // Make sure cluster root object is reachable too
+			    const int32 OwnerIndex = ObjectItem->GetOwnerIndex();
+			    FUObjectItem* RootObjectItem = GUObjectArray.IndexToObjectUnsafeForGC(OwnerIndex);
+				checkSlow(RootObjectItem->HasAnyFlags(EInternalObjectFlags::ClusterRoot));
+			    if (GIsRunningParallelReachability)
+			    {
+				    if (RootObjectItem->ThisThreadAtomicallyClearedRFUnreachable())
+				    {
+						if (bStrongReference)
+						{
+					    	RootObjectItem->ThisThreadAtomicallyClearedFlag(EInternalObjectFlags::NoStrongReference);
+						}
+					    // Make sure all referenced clusters are marked as reachable too
+					    MarkReferencedClustersAsReachable<true>(OwnerIndex, ObjectsToSerialize);
+				    }
+			    }
+			    else if (RootObjectItem->IsUnreachable())
+			    {
+					RootObjectItem->ClearFlags(EInternalObjectFlags::Unreachable);
+					if (bStrongReference)
+					{
+						RootObjectItem->ClearFlags(EInternalObjectFlags::NoStrongReference);
+					}
+				}
 				// Make sure all referenced clusters are marked as reachable too
 				MarkReferencedClustersAsReachable<false>(OwnerIndex, ObjectsToSerialize);
 			}
@@ -723,6 +743,7 @@ public:
 
 		// Iterate over all objects. Note that we iterate over the UObjectArray and usually check only internal flags which
 		// are part of the array so we don't suffer from cache misses as much as we would if we were to check ObjectFlags.
+		TArray<UObjectBase*> KeepClusterRefs;
 		for (FRawObjectIterator It(true); It; ++It)
 		{
 			FUObjectItem* ObjectItem = *It;
@@ -768,6 +789,11 @@ public:
 					// IsValidLowLevel is extremely slow in this loop so only do it in debug
 					checkSlow(Object->IsValidLowLevel());
 					ObjectsToSerialize.Add(Object);
+
+					if (ObjectItem->HasAnyFlags(EInternalObjectFlags::ClusterRoot) || ObjectItem->GetOwnerIndex())
+					{
+						KeepClusterRefs.Add(Object);
+					}
 				}
 				else
 				{
@@ -775,6 +801,38 @@ public:
 				}
 			}
 		}
+		for (UObjectBase* Object : KeepClusterRefs)
+		{
+			const int32 ObjectIndex = GUObjectArray.ObjectToIndex(Object);
+			FUObjectItem* ObjectItem = GUObjectArray.IndexToObjectUnsafeForGC(ObjectIndex);
+			if (ObjectItem->GetOwnerIndex())
+			{
+				check(!ObjectItem->HasAnyFlags(EInternalObjectFlags::ClusterRoot));
+				bool bNeedsDoing = !ObjectItem->HasAnyFlags(EInternalObjectFlags::ReachableInCluster);
+				if (bNeedsDoing)
+				{
+					ObjectItem->SetFlags(EInternalObjectFlags::ReachableInCluster);
+					// Make sure cluster root object is reachable too
+					const int32 OwnerIndex = ObjectItem->GetOwnerIndex();
+					FUObjectItem* RootObjectItem = GUObjectArray.IndexToObjectUnsafeForGC(OwnerIndex);
+					check(RootObjectItem->HasAnyFlags(EInternalObjectFlags::ClusterRoot));
+					// if it is reachable via keep flags we will do this below (or maybe already have)
+					if (RootObjectItem->IsUnreachable()) 
+					{
+						RootObjectItem->ClearFlags(EInternalObjectFlags::Unreachable | EInternalObjectFlags::NoStrongReference);
+						// Make sure all referenced clusters are marked as reachable too
+						FGCReferenceProcessor::MarkReferencedClustersAsReachable<false>(OwnerIndex, ObjectsToSerialize);
+					}
+				}
+			}
+			else
+			{
+				check(ObjectItem->HasAnyFlags(EInternalObjectFlags::ClusterRoot));
+				// this thing is definitely not marked unreachable, so don't test it here
+				// Make sure all referenced clusters are marked as reachable too
+				FGCReferenceProcessor::MarkReferencedClustersAsReachable<false>(ObjectIndex, ObjectsToSerialize);
+			}
+		}
 	}
 
 	/**
@@ -842,12 +900,32 @@ void IncrementalPurgeGarbage( bool bUseTimeLimit, float TimeLimit )
 		return;
 	}
 
+	bool bCompleted = false;
+
+	struct FResetPurgeProgress
+	{
+		bool& bCompletedRef;
+		FResetPurgeProgress(bool& bInCompletedRef)
+			: bCompletedRef(bInCompletedRef)
+		{
+			// Incremental purge is now in progress.
+			GObjIncrementalPurgeIsInProgress = true;
+			FPlatformMisc::MemoryBarrier();
+		}
+		~FResetPurgeProgress()
+		{
+			if (bCompletedRef)
+			{
+				GObjIncrementalPurgeIsInProgress = false;
+				FPlatformMisc::MemoryBarrier();
+			}
+		}
+
+	} ResetPurgeProgress(bCompleted);
+
 	// Set 'I'm garbage collecting' flag - might be checked inside UObject::Destroy etc.
 	FGCScopeLock GCLock;
 
-	// Incremental purge is now in progress.
-	GObjIncrementalPurgeIsInProgress						= true;
-
 	// Keep track of start time to enforce time limit unless bForceFullPurge is true;
 	const double		StartTime							= FPlatformTime::Seconds();
 	bool		bTimeLimitReached							= false;
@@ -1057,8 +1135,8 @@ void IncrementalPurgeGarbage( bool bUseTimeLimit, float TimeLimit )
 
 		if( !GObjCurrentPurgeObjectIndex )
 		{
+			bCompleted = true;
 			// Incremental purge is finished, time to reset variables.
-			GObjIncrementalPurgeIsInProgress				= false;
 			GObjFinishDestroyHasBeenRoutedToAllObjects		= false;
 			GObjPurgeIsRequired								= false;
 			GObjCurrentPurgeObjectIndexNeedsReset			= true;
@@ -1834,11 +1912,28 @@ void UClass::EmitFixedArrayEnd()
 
 void UClass::AssembleReferenceTokenStream(bool bForce)
 {
-	const bool bNotNative = !(ClassFlags & CLASS_Native);
-	if (bNotNative)
+	// Lock for non-native classes
+	struct FScopeLockIfNotNative
 	{
-		ReferenceTokenStreamCritical.Lock();
-	}
+		FCriticalSection& ScopeCritical;
+		const bool bNotNative;
+		FScopeLockIfNotNative(FCriticalSection& InScopeCritical, bool bIsNotNative)
+			: ScopeCritical(InScopeCritical)
+			, bNotNative(bIsNotNative)
+		{
+			if (bNotNative)
+			{
+				ScopeCritical.Lock();
+			}
+		}
+		~FScopeLockIfNotNative()
+		{
+			if (bNotNative)
+			{
+				ScopeCritical.Unlock();
+			}
+		}
+	} ReferenceTokenStreamLock(ReferenceTokenStreamCritical, !(ClassFlags & CLASS_Native));
 
 	if (!HasAnyClassFlags(CLASS_TokenStreamAssembled) || bForce)
 	{
@@ -1898,10 +1993,6 @@ void UClass::AssembleReferenceTokenStream(bool bForce)
 		check(!HasAnyClassFlags(CLASS_TokenStreamAssembled)); // recursion here is probably bad
 		ClassFlags |= CLASS_TokenStreamAssembled;
 	}
-	if (bNotNative)
-	{
-		ReferenceTokenStreamCritical.Unlock();
-	}
 }
 
 
