Copying //UE4/Dev-Core to //UE4/Dev-Main (Source: //UE4/Dev-Core @ 3228803)

diff --git a/Engine/Source/Runtime/CoreUObject/Private/UObject/LinkerLoad.cpp b/Engine/Source/Runtime/CoreUObject/Private/UObject/LinkerLoad.cpp
--- a/Engine/Source/Runtime/CoreUObject/Private/UObject/LinkerLoad.cpp
+++ b/Engine/Source/Runtime/CoreUObject/Private/UObject/LinkerLoad.cpp
@@ -60,15 +60,12 @@ static FAutoConsoleVariableRef CVarLinkerAllowDynamicClasses(
 	);
 #endif
 
-#if !USE_NEW_ASYNC_IO
-/** Map that keeps track of any precached full package reads															*/
-TMap<FString, FLinkerLoad::FPackagePrecacheInfo> FLinkerLoad::PackagePrecacheMap;
-#endif
-
 UClass* FLinkerLoad::UTexture2DStaticClass = NULL;
 
 FName FLinkerLoad::NAME_LoadErrors("LoadErrors");
 
+TMap<FString, FLinkerLoad::FPackagePrecacheInfo> FLinkerLoad::PackagePrecacheMap;
+
 /**
  * Here is the format for the ClassRedirection:
  * 
@@ -358,18 +355,15 @@ static FORCEINLINE bool IsCoreUObjectPackage(const FName& PackageName)
 	FLinkerLoad.
 ----------------------------------------------------------------------------*/
 
-#if !USE_NEW_ASYNC_IO
-
 void FLinkerLoad::GetListOfPackagesInPackagePrecacheMap( TArray<FString>& ListOfPackages )
 {
+	check(!GNewAsyncIO);
 	for ( TMap<FString, FLinkerLoad::FPackagePrecacheInfo>::TIterator It(PackagePrecacheMap); It; ++It )
 	{
 		ListOfPackages.Add( It.Key() );
 	}
 }
 
-#endif
-
 void FLinkerLoad::StaticInit(UClass* InUTexture2DStaticClass)
 {
 	UTexture2DStaticClass = InUTexture2DStaticClass;
@@ -401,9 +395,7 @@ FLinkerLoad* FLinkerLoad::CreateLinker(UPackage* Parent, const TCHAR* Filename,
 #endif // USE_CIRCULAR_DEPENDENCY_LOAD_DEFERRING
 
 	FLinkerLoad* Linker = CreateLinkerAsync(Parent, Filename, LoadFlags
-#if USE_EVENT_DRIVEN_ASYNC_LOAD
 		, TFunction<void()>([](){})
-#endif
 		);
 	{
 #if USE_CIRCULAR_DEPENDENCY_LOAD_DEFERRING
@@ -546,9 +538,7 @@ FName FLinkerLoad::FindSubobjectRedirectName(const FName& Name)
  * @return	new FLinkerLoad object for Parent/ Filename
  */
 FLinkerLoad* FLinkerLoad::CreateLinkerAsync( UPackage* Parent, const TCHAR* Filename, uint32 LoadFlags 
-#if USE_EVENT_DRIVEN_ASYNC_LOAD
 	, TFunction<void()>&& InSummaryReadyCallback
-#endif
 	)
 {
 	check(Parent);
@@ -557,31 +547,29 @@ FLinkerLoad* FLinkerLoad::CreateLinkerAsync( UPackage* Parent, const TCHAR* File
 	FLinkerLoad* Linker = FindExistingLinkerForPackage(Parent);
 	if (Linker)
 	{
-#if USE_EVENT_DRIVEN_ASYNC_LOAD
+		if (GEventDrivenLoaderEnabled)
+		{
 		UE_LOG(LogStreaming, Fatal, TEXT("FLinkerLoad::CreateLinkerAsync: Found existing linker for '%s'"), *Parent->GetName());
-#else
+		}
+		else
+		{
 		UE_LOG(LogStreaming, Log, TEXT("FLinkerLoad::CreateLinkerAsync: Found existing linker for '%s'"), *Parent->GetName());
-#endif
+		}		
 	}
 
 	// Create a new linker if there isn't an existing one.
 	if( Linker == NULL )
 	{
-#if USE_NEW_ASYNC_IO
-		if (FApp::IsGame() && !GIsEditor)
+		if (GNewAsyncIO && FApp::IsGame() && !GIsEditor)
 		{
 			LoadFlags |= LOAD_Async;
 		}
-#endif
 		Linker = new FLinkerLoad(Parent, Filename, LoadFlags );
 		Parent->LinkerLoad = Linker;
-#if USE_EVENT_DRIVEN_ASYNC_LOAD
-		if (Linker)
+		if (GEventDrivenLoaderEnabled && Linker)
 		{
 			Linker->CreateLoader(Forward<TFunction<void()>>(InSummaryReadyCallback));
 		}
-#endif
-
 	}
 	
 	check(Parent->LinkerLoad == Linker);
@@ -614,21 +602,24 @@ FLinkerLoad::ELinkerStatus FLinkerLoad::Tick( float InTimeLimit, bool bInUseTime
 
 		do
 		{
-#if USE_EVENT_DRIVEN_ASYNC_LOAD
+			bool bCanSerializePackageFileSummary = false;
+			if (GEventDrivenLoaderEnabled)
+			{
 			check(Loader || bDynamicClassLinker);
-			if (true)
-#else
+				bCanSerializePackageFileSummary = true;
+			}
+			else
+			{
 			// Create loader, aka FArchive used for serialization and also precache the package file summary.
 			// false is returned until any precaching is complete.
-			if( true )
-			{
 				SCOPED_LOADTIMER(LinkerLoad_CreateLoader);
-				Status = CreateLoader();
+				Status = CreateLoader(TFunction<void()>([]() {}));
+
+				bCanSerializePackageFileSummary = (Status == LINKER_Loaded);
 			}
 
 			// Serialize the package file summary and presize the various arrays (name, import & export map)
-			if( Status == LINKER_Loaded )
-#endif
+			if (bCanSerializePackageFileSummary)
 			{
 				SCOPED_LOADTIMER(LinkerLoad_SerializePackageFileSummary);
 				Status = SerializePackageFileSummary();
@@ -748,14 +739,10 @@ FLinkerLoad::FLinkerLoad(UPackage* InParent, const TCHAR* InFilename, uint32 InL
 , LoadFlags(InLoadFlags)
 , bHaveImportsBeenVerified(false)
 , bDynamicClassLinker(false)
-#if USE_EVENT_DRIVEN_ASYNC_LOAD
 , TemplateForGetArchetypeFromLoader(nullptr)
 , bForceSimpleIndexToObject(false)
 , bLockoutLegacyOperations(false)
-#endif
-#if USE_NEW_ASYNC_IO
 , bLoaderIsFArchiveAsync2(false)
-#endif
 , Loader(nullptr)
 , AsyncRoot(nullptr)
 , NameMapIndex(0)
@@ -851,9 +838,7 @@ bool FLinkerLoad::IsTimeLimitExceeded( const TCHAR* CurrentTask, int32 Granulari
  * Creates loader used to serialize content.
  */
 FLinkerLoad::ELinkerStatus FLinkerLoad::CreateLoader(
-#if USE_EVENT_DRIVEN_ASYNC_LOAD
 	TFunction<void()>&& InSummaryReadyCallback
-#endif
 	)
 {
 	//DECLARE_SCOPE_CYCLE_COUNTER( TEXT( "FLinkerLoad::CreateLoader" ), STAT_LinkerLoad_CreateLoader, STATGROUP_LinkerLoad );
@@ -892,8 +877,9 @@ FLinkerLoad::ELinkerStatus FLinkerLoad::CreateLoader(
 			// In this case we can skip serializing PackageFileSummary and fill all the required info here
 			CreateDynamicTypeLoader();
 		}
-#if USE_NEW_ASYNC_IO
-		else if (!bIsAsyncLoad)		
+		else if (GNewAsyncIO)
+		{
+			if (!bIsAsyncLoad)
 		{
 			check(!FPlatformProperties::RequiresCookedData() && !FSHA1::GetFileSHAHash(*Filename, NULL));
 			Loader = IFileManager::Get().CreateFileReader(*Filename, 0);
@@ -906,9 +892,7 @@ FLinkerLoad::ELinkerStatus FLinkerLoad::CreateLoader(
 		else
 		{
 			Loader = new FArchiveAsync2(*Filename
-#if USE_EVENT_DRIVEN_ASYNC_LOAD
-				, Forward<TFunction<void()>>(InSummaryReadyCallback)
-#endif				
+					, GEventDrivenLoaderEnabled ? Forward<TFunction<void()>>(InSummaryReadyCallback) : TFunction<void()>([]() {})
 				);
 
 			if (!Loader)
@@ -924,7 +908,11 @@ FLinkerLoad::ELinkerStatus FLinkerLoad::CreateLoader(
 				UE_LOG(LogLinker, Warning, TEXT("Error opening file '%s'."), *Filename);
 				return LINKER_Failed;
 			}
+#if DEVIRTUALIZE_FLinkerLoad_Serialize
 			ActiveFPLB = Loader->ActiveFPLB; // make sure my fast past loading is using the FAA2 fast path buffer
+#else
+				check(false);
+#endif
 
 			bool bHasHashEntry = FSHA1::GetFileSHAHash(*Filename, NULL);
 			if ((LoadFlags & LOAD_MemoryReader) || bHasHashEntry)
@@ -952,10 +940,12 @@ FLinkerLoad::ELinkerStatus FLinkerLoad::CreateLoader(
 				bLoaderIsFArchiveAsync2 = true;
 			}
 		}
-#else
+		} // GNewAsyncIO
+		else
+		{
 		// NOTE: Precached memory read gets highest priority, then memory reader, then seek free, then normal
 		// check to see if there is was an async preload request for this file
-		else if (FPackagePrecacheInfo* PrecacheInfo = PackagePrecacheMap.Find(*Filename))
+			if (FPackagePrecacheInfo* PrecacheInfo = PackagePrecacheMap.Find(*Filename))
 		{
 			// if so, serialize from memory (note this will have uncompressed a fully compressed package)
 			// block until the async read is complete
@@ -1031,7 +1021,7 @@ FLinkerLoad::ELinkerStatus FLinkerLoad::CreateLoader(
 				return LINKER_Failed;
 			}
 		}
-#endif // USE_NEW_ASYNC_IO
+		}
 
 		check(bDynamicClassLinker || Loader);
 		check(bDynamicClassLinker || !Loader->IsError());
@@ -1052,23 +1042,24 @@ FLinkerLoad::ELinkerStatus FLinkerLoad::CreateLoader(
 		// Reset all custom versions
 		ResetCustomVersions();
 	}
-#if USE_EVENT_DRIVEN_ASYNC_LOAD
-	else
+	else if (GEventDrivenLoaderEnabled)
 	{
 		check(0);
 	}
+	if (GEventDrivenLoaderEnabled)
+	{
 	return LINKER_TimedOut;
-#else
+	}
+	else
+	{
 	bool bExecuteNextStep = true;
 	if( bHasSerializedPackageFileSummary == false )
 	{
-#if USE_NEW_ASYNC_IO
-		if (bLoaderIsFArchiveAsync2)
+			if (GNewAsyncIO && bLoaderIsFArchiveAsync2)
 		{
 			bExecuteNextStep = GetFArchiveAsync2Loader()->ReadyToStartReadingHeader(bUseTimeLimit, bUseFullTimeLimit, TickStartTime, TimeLimit);
 		}
 		else
-#endif
 		{
 			int64 Size = Loader->TotalSize();
 			if (Size <= 0)
@@ -1092,7 +1083,7 @@ FLinkerLoad::ELinkerStatus FLinkerLoad::CreateLoader(
 	}
 
 	return (bExecuteNextStep && !IsTimeLimitExceeded( TEXT("creating loader") )) ? LINKER_Loaded : LINKER_TimedOut;
-#endif
+	}
 }
 
 /**
@@ -1109,12 +1100,10 @@ FLinkerLoad::ELinkerStatus FLinkerLoad::SerializePackageFileSummary()
 			UE_LOG(LogLinker, Warning, TEXT("The file '%s' contains unrecognizable data, check that it is of the expected type."), *Filename);
 			return LINKER_Failed;
 		}
-#if USE_NEW_ASYNC_IO
-		if (bLoaderIsFArchiveAsync2)
+		if (GNewAsyncIO && bLoaderIsFArchiveAsync2)
 		{
 			GetFArchiveAsync2Loader()->StartReadingHeader();
 		}
-#endif
 
 #if WITH_EDITOR
 		LoadProgressScope->EnterProgressFrame(1);
@@ -1275,9 +1264,8 @@ FLinkerLoad::ELinkerStatus FLinkerLoad::SerializePackageFileSummary()
 		if( Summary.PackageFlags & PKG_StoreCompressed )
 #endif
 		{
-#if USE_NEW_ASYNC_IO
-			check(!"Package level compression cannot be used with the async io scheme.");
-#else
+			checkf(!GNewAsyncIO, TEXT("Package level compression cannot be used with the async io scheme."));
+
 			// Set compression mapping. Failure means Loader doesn't support package compression.
 			check( Summary.CompressedChunks.Num() );
 			if( !Loader->SetCompressionMap( &Summary.CompressedChunks, (ECompressionFlags) Summary.CompressionFlags ) )
@@ -1304,7 +1292,6 @@ FLinkerLoad::ELinkerStatus FLinkerLoad::SerializePackageFileSummary()
 				// Set the compression map and verify it won't fail this time.
 				verify( Loader->SetCompressionMap( &Summary.CompressedChunks, (ECompressionFlags) Summary.CompressionFlags ) );
 			}
-#endif // USE_NEW_ASYNC_IO
 		}
 
 		UPackage* LinkerRootPackage = LinkerRoot;
@@ -1386,16 +1373,12 @@ FLinkerLoad::ELinkerStatus FLinkerLoad::SerializeNameMap()
 		if( Summary.TotalHeaderSize > 0 )
 		{
 			// Precache name, import and export map.
-#if USE_NEW_ASYNC_IO
-			if (bLoaderIsFArchiveAsync2)
+			if (GNewAsyncIO && bLoaderIsFArchiveAsync2)
 			{
 				bFinishedPrecaching = GetFArchiveAsync2Loader()->ReadyToStartReadingHeader(bUseTimeLimit, bUseFullTimeLimit, TickStartTime, TimeLimit);
-#if USE_EVENT_DRIVEN_ASYNC_LOAD
-				check(bFinishedPrecaching);
-#endif
+				check(!GEventDrivenLoaderEnabled || bFinishedPrecaching);
 			}
 			else
-#endif
 			{
 				bFinishedPrecaching = Loader->Precache(Summary.NameOffset, Summary.TotalHeaderSize - Summary.NameOffset);
 			}
@@ -2013,8 +1996,7 @@ FLinkerLoad::ELinkerStatus FLinkerLoad::FinalizeCreation()
 			Verify();
 		}
 
-#if USE_EVENT_DRIVEN_ASYNC_LOAD
-		if (AsyncRoot)
+		if (GEventDrivenLoaderEnabled && AsyncRoot)
 		{
 			for (int32 ImportIndex = 0; ImportIndex < ImportMap.Num(); ++ImportIndex)
 			{
@@ -2027,13 +2009,11 @@ FLinkerLoad::ELinkerStatus FLinkerLoad::FinalizeCreation()
 				AsyncRoot->ObjectNameToImportOrExport.Add(Exp(Index).ObjectName, Index);
 			}
 		}
-#endif
-#if USE_NEW_ASYNC_IO
-		if (bLoaderIsFArchiveAsync2)
+
+		if (GNewAsyncIO && bLoaderIsFArchiveAsync2)
 		{
 			GetFArchiveAsync2Loader()->EndReadingHeader();
 		}
-#endif
 
 		// Avoid duplicate work in the case of async linker creation.
 		bHasFinishedInitialization = true;
@@ -2338,9 +2318,8 @@ void FLinkerLoad::GatherImportDependencies(int32 ImportIndex, TSet<FDependencyRe
 
 FLinkerLoad::EVerifyResult FLinkerLoad::VerifyImport(int32 ImportIndex)
 {
-#if USE_EVENT_DRIVEN_ASYNC_LOAD
-	check(GIsInitialLoad);
-#endif
+	check(!GEventDrivenLoaderEnabled || !EVENT_DRIVEN_ASYNC_LOAD_ACTIVE_AT_RUNTIME);
+
 	FObjectImport& Import = ImportMap[ImportIndex];
 
 	// keep a string of modifiers to add to the Editor Warning dialog
@@ -2511,9 +2490,7 @@ UPackage* LoadPackageInternal(UPackage* InOuter, const TCHAR* InLongPackageName,
  */
 bool FLinkerLoad::VerifyImportInner(const int32 ImportIndex, FString& WarningSuffix)
 {
-#if USE_EVENT_DRIVEN_ASYNC_LOAD
-	check(GIsInitialLoad);
-#endif
+	check(!GEventDrivenLoaderEnabled || !EVENT_DRIVEN_ASYNC_LOAD_ACTIVE_AT_RUNTIME);
 
 	check(IsLoading());
 
@@ -3046,6 +3023,7 @@ void FLinkerLoad::LoadAllObjects( bool bForcePreload )
 
 		UObject* LoadedObject = CreateExportAndPreload(ExportIndex, bForcePreload);
 		// DynamicClass could be created without calling CreateImport. The imported objects will be required later when a CDO is created.
+		// intentional merge conflict, talk to Gil
 		if (UDynamicClass* DynamicClass = Cast<UDynamicClass>(LoadedObject))
 		{
 			for (int32 ImportIndex = 0; ImportIndex < ImportMap.Num(); ++ImportIndex)
@@ -3053,6 +3031,7 @@ void FLinkerLoad::LoadAllObjects( bool bForcePreload )
 				CreateImport(ImportIndex);
 			}
 		}
+		// intentional merge conflict, talk to Gil
 
 		// If needed send a heartbeat, but no need to do it too often
 		if (bShouldTickHeartBeat && (ExportIndex % 10) == 0)
@@ -3266,9 +3245,7 @@ void FLinkerLoad::Preload( UObject* Object )
 	{
 		if (Object->GetLinker() == this)
 		{
-#if USE_EVENT_DRIVEN_ASYNC_LOAD
-			check(!bLockoutLegacyOperations);
-#endif
+			check(!GEventDrivenLoaderEnabled || !bLockoutLegacyOperations);
 #if USE_CIRCULAR_DEPENDENCY_LOAD_DEFERRING
 			bool const bIsNonNativeObject = !Object->GetOutermost()->HasAnyPackageFlags(PKG_CompiledIn);
 			// we can determine that this is a blueprint class/struct by checking if it 
@@ -3365,20 +3342,22 @@ void FLinkerLoad::Preload( UObject* Object )
 				// move to the position in the file where this object's data
 				// is stored
 				Loader->Seek(Export.SerialOffset);
-#if USE_NEW_ASYNC_IO
-				FArchiveAsync2* FAA2 = GetFArchiveAsync2Loader();
-#endif
+
+				FArchiveAsync2* FAA2 = nullptr;
+				if (GNewAsyncIO)
+				{
+					FAA2 = GetFArchiveAsync2Loader();
+				}
+
 				{
 					SCOPE_CYCLE_COUNTER(STAT_LinkerPrecache);
 					// tell the file reader to read the raw data from disk
-#if USE_NEW_ASYNC_IO
 					if (FAA2)
 					{
 						bool bReady = FAA2->Precache(Export.SerialOffset, Export.SerialSize, bUseTimeLimit, bUseFullTimeLimit, TickStartTime, TimeLimit);
 						UE_CLOG(!(bReady || !bUseTimeLimit || !FPlatformProperties::RequiresCookedData()), LogLinker, Warning, TEXT("Hitch on async loading of %s; this export was not properly precached."), *Object->GetFullName());
 					}
 					else
-#endif
 					{
 						Loader->Precache(Export.SerialOffset, Export.SerialSize);
 					}
@@ -3630,10 +3609,7 @@ UObject* FLinkerLoad::CreateExport( int32 Index )
 	// Check whether we already loaded the object and if not whether the context flags allow loading it.
 	if( !Export.Object && !FilterExport(Export) ) // for some acceptable position, it was not "not for" 
 	{
-#if USE_EVENT_DRIVEN_ASYNC_LOAD
-		check(!bLockoutLegacyOperations);
-#endif
-
+		check(!GEventDrivenLoaderEnabled || !bLockoutLegacyOperations);
 		check(Export.ObjectName!=NAME_None || !(Export.ObjectFlags&RF_Public));
 		check(IsLoading());
 
@@ -4192,9 +4168,8 @@ bool FLinkerLoad::IsImportNative(const int32 Index) const
 // Return the loaded object corresponding to an import index; any errors are fatal.
 UObject* FLinkerLoad::CreateImport( int32 Index )
 {
-#if USE_EVENT_DRIVEN_ASYNC_LOAD
-	check(!bLockoutLegacyOperations);
-#endif
+	check(!GEventDrivenLoaderEnabled || !bLockoutLegacyOperations);
+
 	FScopedCreateImportCounter ScopedCounter( this, Index );
 	FObjectImport& Import = ImportMap[ Index ];
 	
@@ -4571,14 +4546,13 @@ FArchive& FLinkerLoad::operator<<( UObject*& Object )
 	FPackageIndex Index;
 	FArchive& Ar = *this;
 	Ar << Index;
-#if USE_EVENT_DRIVEN_ASYNC_LOAD
-	if (bForceSimpleIndexToObject)
+
+	if (GEventDrivenLoaderEnabled && bForceSimpleIndexToObject)
 	{
 		check(Ar.IsLoading() && AsyncRoot);
 		Object = AsyncRoot->EventDrivenIndexToObject(Index, false);
 		return *this;
 	}
-#endif
 
 	UObject* Temporary = NULL;
 	Temporary = IndexToObject( Index );
@@ -4622,9 +4596,10 @@ void FLinkerLoad::BadNameIndexError(NAME_INDEX NameIndex)
 	UE_LOG(LogLinker, Error, TEXT("Bad name index %i/%i"), NameIndex, NameMap.Num());
 }
 
-#if !USE_NEW_ASYNC_IO
 void FLinkerLoad::AsyncPreloadPackage(const TCHAR* PackageName)
 {
+	check(!GNewAsyncIO);
+
 	// get package filename
 	FString PackageFilename;
 	if (!FPackageName::DoesPackageExist(PackageName, NULL, &PackageFilename))
@@ -4691,7 +4666,6 @@ void FLinkerLoad::AsyncPreloadPackage(const TCHAR* PackageName)
 
 	check(RequestId);
 }
-#endif
 
 /**
  * Called when an object begins serializing property data using script serialization.
@@ -4735,13 +4709,19 @@ bool FLinkerLoad::FindImportClassAndPackage( FName ClassName, FPackageIndex &Cla
 	return false;
 }
 
-#if USE_EVENT_DRIVEN_ASYNC_LOAD
+
 UObject* FLinkerLoad::GetArchetypeFromLoader(const UObject* Obj)
 {
+	if (GEventDrivenLoaderEnabled)
+	{
 	check(!TemplateForGetArchetypeFromLoader || FUObjectThreadContext::Get().SerializedObject == Obj);
 	return TemplateForGetArchetypeFromLoader;
 }
-#endif
+	else
+	{
+		return FArchiveUObject::GetArchetypeFromLoader(Obj);
+	}
+}
 
 
 /**
diff --git a/Engine/Source/Runtime/CoreUObject/Private/UObject/Obj.cpp b/Engine/Source/Runtime/CoreUObject/Private/UObject/Obj.cpp
--- a/Engine/Source/Runtime/CoreUObject/Private/UObject/Obj.cpp
+++ b/Engine/Source/Runtime/CoreUObject/Private/UObject/Obj.cpp
@@ -821,15 +821,9 @@ bool UObject::ConditionalFinishDestroy()
 	}
 }
 
-#if !defined(USE_EVENT_DRIVEN_ASYNC_LOAD)
-#error "USE_EVENT_DRIVEN_ASYNC_LOAD must be defined"
-#endif
-
 void UObject::ConditionalPostLoad()
 {
-#if USE_EVENT_DRIVEN_ASYNC_LOAD
-	check(!HasAnyFlags(RF_NeedLoad)); //@todoio Added this as "nicks rule"
-#endif
+	check(!GEventDrivenLoaderEnabled || !HasAnyFlags(RF_NeedLoad)); //@todoio Added this as "nicks rule"
 									  // PostLoad only if the object needs it and has already been serialized
 	//@todoio note this logic should be unchanged compared to main
 	if (HasAnyFlags(RF_NeedPostLoad))
@@ -875,25 +869,18 @@ void UObject::ConditionalPostLoad()
 	}
 }
 
-#if !defined(USE_EVENT_DRIVEN_ASYNC_LOAD)
-#error "USE_EVENT_DRIVEN_ASYNC_LOAD must be defined."
-#endif
-
-
 void UObject::PostLoadSubobjects( FObjectInstancingGraph* OuterInstanceGraph/*=NULL*/ )
 {
-#if USE_EVENT_DRIVEN_ASYNC_LOAD
-	check(!HasAnyFlags(RF_NeedLoad));
-#endif
+	check(!GEventDrivenLoaderEnabled || !HasAnyFlags(RF_NeedLoad));
+
 	if( GetClass()->HasAnyClassFlags(CLASS_HasInstancedReference) )
 	{
 		UObject* ObjOuter = GetOuter();
 		// make sure our Outer has already called ConditionalPostLoadSubobjects
 		if (ObjOuter != NULL && ObjOuter->HasAnyFlags(RF_NeedPostLoadSubobjects) )
 		{
-#if USE_EVENT_DRIVEN_ASYNC_LOAD
-			check(!ObjOuter->HasAnyFlags(RF_NeedLoad));
-#endif
+			check(!GEventDrivenLoaderEnabled || !ObjOuter->HasAnyFlags(RF_NeedLoad));
+
 			if (ObjOuter->HasAnyFlags(RF_NeedPostLoad) )
 			{
 				ObjOuter->ConditionalPostLoad();
diff --git a/Engine/Source/Runtime/CoreUObject/Private/UObject/ObjectResource.cpp b/Engine/Source/Runtime/CoreUObject/Private/UObject/ObjectResource.cpp
--- a/Engine/Source/Runtime/CoreUObject/Private/UObject/ObjectResource.cpp
+++ b/Engine/Source/Runtime/CoreUObject/Private/UObject/ObjectResource.cpp
@@ -147,12 +147,9 @@ FArchive& operator<<( FArchive& Ar, FObjectExport& E )
 
 FObjectImport::FObjectImport()
 	: FObjectResource()
-#if USE_EVENT_DRIVEN_ASYNC_LOAD
 	, bImportPackageHandled(false)
 	, bImportSearchedFor(false)
 	, bImportFailed(false)
-#endif
-
 {
 }
 
@@ -163,11 +160,9 @@ FObjectImport::FObjectImport(UObject* InObject)
 	, XObject(InObject)
 	, SourceLinker(NULL)
 	, SourceIndex(INDEX_NONE)
-#if USE_EVENT_DRIVEN_ASYNC_LOAD
 	, bImportPackageHandled(false)
 	, bImportSearchedFor(false)
 	, bImportFailed(false)
-#endif
 {
 }
 
@@ -178,11 +173,9 @@ FObjectImport::FObjectImport(UObject* InObject, UClass* InClass)
 	, XObject(InObject)
 	, SourceLinker(NULL)
 	, SourceIndex(INDEX_NONE)
-#if USE_EVENT_DRIVEN_ASYNC_LOAD
 	, bImportPackageHandled(false)
 	, bImportSearchedFor(false)
 	, bImportFailed(false)
-#endif
 {
 }
 
diff --git a/Engine/Source/Runtime/CoreUObject/Private/UObject/SavePackage.cpp b/Engine/Source/Runtime/CoreUObject/Private/UObject/SavePackage.cpp
--- a/Engine/Source/Runtime/CoreUObject/Private/UObject/SavePackage.cpp
+++ b/Engine/Source/Runtime/CoreUObject/Private/UObject/SavePackage.cpp
@@ -91,7 +91,7 @@ namespace SavePackageStats
 		// Don't use FCookStatsManager::CreateKeyValueArray because there's just too many arguments. Don't need to overburden the compiler here.
 		TArray<FCookStatsManager::StringKeyValue> StatsList;
 		StatsList.Empty(15);
-		#define ADD_COOK_STAT(Name) StatsList.Emplace(TEXT(#Name), FCookStatsManager::ToString(Name))
+		#define ADD_COOK_STAT(Name) StatsList.Emplace(TEXT(#Name), Lex::ToString(Name))
 		ADD_COOK_STAT(NumPackagesSaved);
 		ADD_COOK_STAT(SavePackageTimeSec);
 		ADD_COOK_STAT(TagPackageExportsPresaveTimeSec);
@@ -459,7 +459,6 @@ void AsyncWriteFile(FLargeMemoryPtr Data, const int64 DataSize, const TCHAR* Fil
 	(new FAutoDeleteAsyncTask<FAsyncWriteWorker>(Filename, MoveTemp(Data), DataSize, TimeStamp, bUseTempFilename))->StartBackgroundTask();
 }
 
-#if SPLIT_COOKED_FILES
 void AsyncWriteFileWithSplitExports(FLargeMemoryPtr Data, const int64 DataSize, const int64 HeaderSize, const TCHAR* Filename, const FDateTime& TimeStamp, bool bUseTempFilename = true)
 {
 	class FAsyncWriteWorkerWithSplitExports : public FNonAbandonableTask
@@ -609,7 +608,7 @@ void AsyncWriteFileWithSplitExports(FLargeMemoryPtr Data, const int64 DataSize,
 	OutstandingAsyncWrites.Increment();
 	(new FAutoDeleteAsyncTask<FAsyncWriteWorkerWithSplitExports>(Filename, MoveTemp(Data), DataSize, HeaderSize, TimeStamp, bUseTempFilename))->StartBackgroundTask();
 }
-#endif
+
 /** 
  * Helper object for all of the cross object state that is needed only while saving a package.
  */
@@ -931,6 +930,39 @@ FArchive& FArchiveSaveTagExports::operator<<( UObject*& Obj )
 			{
 				// In EDL we don't always load this object, so mark it appropriately client/server only
 				ConditionallyExcludeObjectForTarget(Obj, *this);
+
+				UObject* CDO = Obj;
+
+				// Gets all subobjects defined in a class, including the CDO, CDO components and blueprint-created components
+				TArray<UObject*> ObjectTemplates;
+				ObjectTemplates.Add(CDO);
+
+				TArray<UObject*> CurrentSubobjects;
+				TArray<UObject*> NextSubobjects;
+
+				// Recursively search for subobjects. Only care about ones that have a full subobject chain as some nested objects are set wrong
+				GetObjectsWithOuter(CDO->GetClass(), NextSubobjects, false);
+				GetObjectsWithOuter(CDO, NextSubobjects, false);
+
+				while (NextSubobjects.Num() > 0)
+				{
+					CurrentSubobjects = NextSubobjects;
+					NextSubobjects.Empty();
+					for (UObject* SubObj : CurrentSubobjects)
+					{
+						if (SubObj->HasAnyFlags(RF_DefaultSubObject | RF_ArchetypeObject))
+						{
+							ObjectTemplates.Add(SubObj);
+							GetObjectsWithOuter(SubObj, NextSubobjects, false);
+						}
+					}
+				}
+
+				for (UObject* ObjTemplate : ObjectTemplates)
+				{
+					*this << ObjTemplate;
+				}
+
 			}
 		}
 		else
@@ -3660,6 +3692,140 @@ public:
 
 #endif
 
+
+struct FEDLCookChecker
+{
+	bool bIsActive;
+	TMultiMap<FString, FString> ImportToImportingPackage;
+	TSet<FString> Exports;
+	TMultiMap<FString, FString> NodePrereqs;
+
+	FEDLCookChecker()
+		: bIsActive(false)
+	{
+	}
+
+	void Reset()
+	{
+		ImportToImportingPackage.Empty();
+		Exports.Empty();
+		NodePrereqs.Empty();
+		bIsActive = false;
+	}
+
+	void AddImport(UObject* Import, UPackage* ImportingPackage)
+	{
+		if (bIsActive)
+		{
+			if (!Import->GetOutermost()->HasAnyPackageFlags(PKG_CompiledIn))
+			{
+				ImportToImportingPackage.Add(Import->GetFullName(), ImportingPackage->GetPathName());
+			}
+		}
+	}
+	void AddExport(UObject* Export)
+	{
+		if (bIsActive)
+		{
+			Exports.Add(Export->GetFullName());
+			AddArc(Export, false, Export, true); // every export must be created before it can be serialize...these arcs are implicit and not listed in any table.
+		}
+	}
+
+	void AddArc(UObject* DepObject, bool bDepIsSerialize, UObject* Export, bool bExportIsSerialize)
+	{
+		if (bIsActive)
+		{
+			NodePrereqs.Add(NodeName(Export, bExportIsSerialize), NodeName(DepObject, bDepIsSerialize));
+		}
+	}
+
+	bool CheckForCyclesInner(TSet<FString>& Visited, TSet<FString>& Stack, const FString& Visit, FString& FailNode)
+	{
+		bool bResult = false;
+		if (Stack.Contains(Visit))
+		{
+			FailNode = Visit;
+			bResult = true;
+		}
+		else
+		{
+			bool bWasAlreadyTested = false;
+			Visited.Add(Visit, &bWasAlreadyTested);
+			if (!bWasAlreadyTested)
+			{
+				Stack.Add(Visit);
+				for (auto It = NodePrereqs.CreateConstKeyIterator(Visit); !bResult && It; ++It)
+				{
+					bResult = CheckForCyclesInner(Visited, Stack, It.Value(), FailNode);
+				}
+				Stack.Remove(Visit);
+			}
+		}
+		UE_CLOG(bResult && Stack.Contains(FailNode), LogSavePackage, Error, TEXT("Cycle Node %s"), *Visit);
+		return bResult;
+	}
+
+	void Verify()
+	{
+		if (bIsActive && Exports.Num())
+		{
+			double StartTime = FPlatformTime::Seconds();
+			// imports to things that are not exports...
+			for (auto& Pair : ImportToImportingPackage)
+			{
+				if (!Exports.Contains(Pair.Key))
+				{
+					UE_LOG(LogSavePackage, Warning, TEXT("%s imported %s, but it was never saved as an export."), *Pair.Value, *Pair.Key);
+				}
+			}
+			// cycles in the dep graph
+			TSet<FString> Visited;
+			TSet<FString> Stack;
+			bool bHadCycle = false;
+			for (FString& Export : Exports)
+			{
+				FString FailNode;
+				if (CheckForCyclesInner(Visited, Stack, Prefix(true) + Export, FailNode))
+				{
+					UE_LOG(LogSavePackage, Error, TEXT("----- %s contained a cycle (listed above)."), *FailNode);
+					bHadCycle = true;
+				}
+			}
+			if (bHadCycle)
+			{
+				UE_LOG(LogSavePackage, Fatal, TEXT("EDL dep graph contained a cycle (see errors, above). This is fatal at runtime so it is fatal at cook time."));
+			}
+			UE_LOG(LogSavePackage, Display, TEXT("Took %fs to verify the EDL loading graph."), float(FPlatformTime::Seconds() - StartTime));
+		}
+		Reset();
+	}
+
+	static FString Prefix(bool bDepIsSerialize)
+	{
+		return FString(bDepIsSerialize ? TEXT("Serialize:") : TEXT("Create:"));
+	}
+
+	static FString NodeName(UObject* DepObject, bool bDepIsSerialize)
+	{
+		return Prefix(bDepIsSerialize) + DepObject->GetFullName();
+	}
+};
+
+static FEDLCookChecker GEDLCookChecker;
+
+void StartSavingEDLCookInfoForVerification()
+{
+	GEDLCookChecker.Reset();
+	GEDLCookChecker.bIsActive = IsEventDrivenLoaderEnabledInCookedBuilds();
+}
+
+void VerifyEDLCookInfo()
+{
+	GEDLCookChecker.Verify();
+}
+
+
 extern FGCCSyncObject GGarbageCollectionGuardCritical;
 
 ESavePackageResult UPackage::Save(UPackage* InOuter, UObject* Base, EObjectFlags TopLevelFlags, const TCHAR* Filename,
@@ -3742,9 +3908,10 @@ ESavePackageResult UPackage::Save(UPackage* InOuter, UObject* Base, EObjectFlags
 
 		(*GFlushStreamingFunc)();
 
-#if !USE_NEW_ASYNC_IO
-		FIOSystem::Get().BlockTillAllRequestsFinishedAndFlushHandles();
-#endif
+		if (!GNewAsyncIO)
+		{
+			FIOSystem::Get().BlockTillAllRequestsFinishedAndFlushHandles();
+		}
 
 		uint32 Time = 0; CLOCK_CYCLES(Time);
 
@@ -4799,6 +4966,18 @@ ESavePackageResult UPackage::Save(UPackage* InOuter, UObject* Base, EObjectFlags
 						checkf(Conform != NULL, TEXT("NULL XObject for import %i - Object: %s Class: %s"), i, *Linker->ImportMap[i].ObjectName.ToString(), *Linker->ImportMap[i].ClassName.ToString());
 					}
 				}
+				if (IsEventDrivenLoaderEnabledInCookedBuilds() && TargetPlatform)
+				{
+					GEDLCookChecker.AddExport(InOuter); // the package isn't actually in the export map, but that is ok, we add it as export anyway for error checking
+					for (int32 i = 0; i < Linker->ImportMap.Num(); i++)
+					{
+						UObject* Object = Linker->ImportMap[i].XObject;
+						if (Object != NULL)
+						{
+							GEDLCookChecker.AddImport(Object, InOuter);
+						}
+					}
+				}
 
 				// Convert the searchable names map from UObject to packageindex
 				for (const TPair<const UObject *, TArray<FName>>& SearchableNamePair : Linker->SearchableNamesObjectMap)
@@ -5061,10 +5240,10 @@ ESavePackageResult UPackage::Save(UPackage* InOuter, UObject* Base, EObjectFlags
 						FObjectExport& Export = Linker->ExportMap[i];
 						if (Export.Object)
 						{
+							GEDLCookChecker.AddExport(Export.Object);
 							TSet<FPackageIndex> SerializationBeforeCreateDependencies;
 							{
 								IncludeIndexAsDependency(SerializationBeforeCreateDependencies, Export.ClassIndex);
-								IncludeIndexAsDependency(SerializationBeforeCreateDependencies, Export.SuperIndex);
 								UObject* CDO = Export.Object->GetArchetype();
 								IncludeObjectAsDependency(SerializationBeforeCreateDependencies, CDO, Export.Object, true);
 								Subobjects.Reset();
@@ -5082,10 +5261,15 @@ ESavePackageResult UPackage::Save(UPackage* InOuter, UObject* Base, EObjectFlags
 							{
 								TArray<UObject*> Deps;
 								Export.Object->GetPreloadDependencies(Deps);
+
 								for (UObject* Obj : Deps)
 								{
 									IncludeObjectAsDependency(SerializationBeforeSerializationDependencies, Obj, Export.Object, false);
 								}
+								if (Export.Object->HasAnyFlags(RF_ArchetypeObject | RF_ClassDefaultObject) && !Export.Object->GetOuter()->IsA(UPackage::StaticClass()))
+								{
+									IncludeObjectAsDependency(SerializationBeforeSerializationDependencies, Export.Object->GetOuter(), Export.Object, true);
+								}
 							}
 
 							TSet<FPackageIndex> CreateBeforeSerializationDependencies;
@@ -5105,8 +5289,21 @@ ESavePackageResult UPackage::Save(UPackage* InOuter, UObject* Base, EObjectFlags
 							TSet<FPackageIndex> CreateBeforeCreateDependencies;
 							{
 								IncludeIndexAsDependency(CreateBeforeCreateDependencies, Export.OuterIndex);
+								IncludeIndexAsDependency(CreateBeforeCreateDependencies, Export.SuperIndex);
 							}
 
+							auto AddArcForDepChecking = [Linker, &Export](bool bExportIsSerialize, FPackageIndex Dep, bool bDepIsSerialize)
+							{
+								check(Export.Object);
+								check(!Dep.IsNull());
+								UObject* DepObject = Dep.IsExport() ? Linker->Exp(Dep).Object : Linker->Imp(Dep).XObject;
+								check(DepObject);
+
+								Linker->DepListForErrorChecking.Add(Dep);
+
+								GEDLCookChecker.AddArc(DepObject, bDepIsSerialize, Export.Object, bExportIsSerialize);
+							};
+
 							for (FPackageIndex Index : SerializationBeforeSerializationDependencies)
 							{
 								if (SerializationBeforeCreateDependencies.Contains(Index))
@@ -5121,7 +5318,7 @@ ESavePackageResult UPackage::Save(UPackage* InOuter, UObject* Base, EObjectFlags
 								Linker->Summary.PreloadDependencyCount++;
 								Export.SerializationBeforeSerializationDependencies++;
 								*Linker << Index;
-								Linker->DepListForErrorChecking.Add(Index);
+								AddArcForDepChecking(true, Index, true);
 							}
 							for (FPackageIndex Index : CreateBeforeSerializationDependencies)
 							{
@@ -5133,6 +5330,10 @@ ESavePackageResult UPackage::Save(UPackage* InOuter, UObject* Base, EObjectFlags
 								{
 									continue; // if the other thing must be serialized before we serialize, then this is a redundant dep
 								}
+								if (CreateBeforeCreateDependencies.Contains(Index))
+								{
+									continue; // if the other thing must be created before we are created, then this is a redundant dep
+								}
 								if (Export.FirstExportDependency == -1)
 								{
 									Export.FirstExportDependency = Linker->Summary.PreloadDependencyCount;
@@ -5141,7 +5342,7 @@ ESavePackageResult UPackage::Save(UPackage* InOuter, UObject* Base, EObjectFlags
 								Linker->Summary.PreloadDependencyCount++;
 								Export.CreateBeforeSerializationDependencies++;
 								*Linker << Index;
-								Linker->DepListForErrorChecking.Add(Index);
+								AddArcForDepChecking(true, Index, false);
 							}
 							for (FPackageIndex Index : SerializationBeforeCreateDependencies)
 							{
@@ -5153,7 +5354,7 @@ ESavePackageResult UPackage::Save(UPackage* InOuter, UObject* Base, EObjectFlags
 								Linker->Summary.PreloadDependencyCount++;
 								Export.SerializationBeforeCreateDependencies++;
 								*Linker << Index;
-								Linker->DepListForErrorChecking.Add(Index);
+								AddArcForDepChecking(false, Index, true);
 							}
 							for (FPackageIndex Index : CreateBeforeCreateDependencies)
 							{
@@ -5165,7 +5366,7 @@ ESavePackageResult UPackage::Save(UPackage* InOuter, UObject* Base, EObjectFlags
 								Linker->Summary.PreloadDependencyCount++;
 								Export.CreateBeforeCreateDependencies++;
 								*Linker << Index;
-								Linker->DepListForErrorChecking.Add(Index);
+								AddArcForDepChecking(false, Index, false);
 							}
 						}
 					}
@@ -5287,13 +5488,11 @@ ESavePackageResult UPackage::Save(UPackage* InOuter, UObject* Base, EObjectFlags
 							{
 								bEnable = false;
 							}
-#if SPLIT_COOKED_FILES
 							if (IsEventDrivenLoaderEnabledInCookedBuilds())
 							{
 								// Always split bulk data when splitting cooked files
 								bEnable = true;
 							}
-#endif
 						}
 					} ShouldUseSeperateBulkDataFiles;
 
@@ -5578,13 +5777,11 @@ ESavePackageResult UPackage::Save(UPackage* InOuter, UObject* Base, EObjectFlags
 							int64 DataSize = Writer->TotalSize();
 							FLargeMemoryPtr DataPtr(Writer->GetData());
 							Writer->ReleaseOwnership();
-#if SPLIT_COOKED_FILES
 							if (IsEventDrivenLoaderEnabledInCookedBuilds() && Linker->IsCooking())
 							{
 								AsyncWriteFileWithSplitExports(MoveTemp(DataPtr), DataSize, Linker->Summary.TotalHeaderSize, *NewPath, FinalTimeStamp);
 							}
 							else
-#endif
 							{
 								AsyncWriteFile(MoveTemp(DataPtr), DataSize, *NewPath, FinalTimeStamp);
 							}
diff --git a/Engine/Source/Runtime/CoreUObject/Private/UObject/UObjectArchetype.cpp b/Engine/Source/Runtime/CoreUObject/Private/UObject/UObjectArchetype.cpp
--- a/Engine/Source/Runtime/CoreUObject/Private/UObject/UObjectArchetype.cpp
+++ b/Engine/Source/Runtime/CoreUObject/Private/UObject/UObjectArchetype.cpp
@@ -10,13 +10,6 @@
 #include "UObject/Class.h"
 #include "UObject/Package.h"
 
-#if !defined(USE_EVENT_DRIVEN_ASYNC_LOAD)
-#error "USE_EVENT_DRIVEN_ASYNC_LOAD must be defined"
-#endif
-
-#if USE_EVENT_DRIVEN_ASYNC_LOAD
-COREUOBJECT_API bool GIgnoreGetArchetypeFromRequiredInfo_RF_NeedsLoad = false;
-#endif
 
 UObject* UObject::GetArchetypeFromRequiredInfo(UClass* Class, UObject* Outer, FName Name, EObjectFlags ObjectFlags)
 {
@@ -70,9 +63,8 @@ UObject* UObject::GetArchetypeFromRequiredInfo(UClass* Class, UObject* Outer, FN
 			Result = Class->GetDefaultObject();
 		}
 	}
-#if USE_EVENT_DRIVEN_ASYNC_LOAD
-	//if (!GIgnoreGetArchetypeFromRequiredInfo_RF_NeedsLoad)
-	if (!GIsInitialLoad)
+
+	if (GEventDrivenLoaderEnabled && EVENT_DRIVEN_ASYNC_LOAD_ACTIVE_AT_RUNTIME)
 	{
 		if (Result && Result->HasAnyFlags(RF_NeedLoad))
 		{
@@ -80,7 +72,6 @@ UObject* UObject::GetArchetypeFromRequiredInfo(UClass* Class, UObject* Outer, FN
 		}
 		check(!Result || !Result->HasAnyFlags(RF_NeedLoad));
 	}
-#endif
 
 	return Result;
 }
diff --git a/Engine/Source/Runtime/CoreUObject/Private/UObject/UObjectBase.cpp b/Engine/Source/Runtime/CoreUObject/Private/UObject/UObjectBase.cpp
--- a/Engine/Source/Runtime/CoreUObject/Private/UObject/UObjectBase.cpp
+++ b/Engine/Source/Runtime/CoreUObject/Private/UObject/UObjectBase.cpp
@@ -1229,9 +1229,10 @@ UPackage* FindOrConstructDynamicTypePackage(const TCHAR* PackageName)
 	if (!Package)
 	{
 		Package = CreatePackage(nullptr, PackageName);
-#if !USE_EVENT_DRIVEN_ASYNC_LOAD
-		Package->SetPackageFlags(PKG_CompiledIn);
-#endif
+		if (!GEventDrivenLoaderEnabled)
+		{
+			Package->SetPackageFlags(PKG_CompiledIn);
+		}
 	}
 	check(Package);
 	return Package;
diff --git a/Engine/Source/Runtime/CoreUObject/Private/UObject/UObjectClusters.cpp b/Engine/Source/Runtime/CoreUObject/Private/UObject/UObjectClusters.cpp
--- a/Engine/Source/Runtime/CoreUObject/Private/UObject/UObjectClusters.cpp
+++ b/Engine/Source/Runtime/CoreUObject/Private/UObject/UObjectClusters.cpp
@@ -63,12 +63,20 @@ static void ListClusters(const TArray<FString>& Args)
 				UObject* Object = static_cast<UObject*>(ObjectItem->Object);
 				UE_LOG(LogObj, Display, TEXT("    [%.4d]: %s (Index: %d)"), Index++, *Object->GetFullName(), ObjectIndex);
 			}
+			UE_LOG(LogObj, Display, TEXT("  Referenced clusters: %d"), Pair.Value->ReferencedClusters.Num());
 			for (int32 ClusterIndex : Pair.Value->ReferencedClusters)
 			{
 				FUObjectItem* ClusterRootItem = GUObjectArray.IndexToObjectUnsafeForGC(ClusterIndex);
 				UObject* ClusterRootObject = static_cast<UObject*>(ClusterRootItem->Object);
 				UE_LOG(LogObj, Display, TEXT("    -> %s (Index: %d)"), *ClusterRootObject->GetFullName(), ClusterIndex);
 			}
+			UE_LOG(LogObj, Display, TEXT("  External (mutable) objects: %d"), Pair.Value->MutableObjects.Num());
+			for (int32 ObjectIndex : Pair.Value->MutableObjects)
+			{
+				FUObjectItem* ObjectItem = GUObjectArray.IndexToObjectUnsafeForGC(ObjectIndex);
+				UObject* Object = static_cast<UObject*>(ObjectItem->Object);
+				UE_LOG(LogObj, Display, TEXT("    -> %s (Index: %d)"), *Object->GetFullName(), ObjectIndex);
+			}
 		}
 	}
 	UE_LOG(LogObj, Display, TEXT("Number of clusters: %d"), GUObjectClusters.Num());
@@ -287,6 +295,8 @@ public:
 	 */
 	void AddObjectToCluster(int32 ObjectIndex, FUObjectItem* ObjectItem, UObject* Obj, TArray<UObject*>& ObjectsToSerialize, bool bOuterAndClass)
 	{
+		// If we haven't finished loading, we can't be sure we know all the references
+		check(!Obj->HasAnyFlags(RF_NeedLoad | RF_NeedPostLoad));
 		if (ObjectIndex != ClusterRootIndex && ObjectItem->GetOwnerIndex() == 0 && !ObjectItem->IsRootSet() && !GUObjectArray.IsDisregardForGC(Obj) && Obj->CanBeInCluster())
 		{
 			ObjectsToSerialize.Add(Obj);
@@ -367,6 +377,9 @@ public:
 	{
 		if (Object)
 		{
+			// If we haven't finished loading, we can't be sure we know all the references
+			check(!Object->HasAnyFlags(RF_NeedLoad | RF_NeedPostLoad));
+
 			FUObjectItem* ObjectItem = GUObjectArray.ObjectToObjectItem(Object);
 
 			// Add encountered object reference to list of to be serialized objects if it hasn't already been added.
@@ -616,71 +629,54 @@ public:
 				!IsPossiblyAllocatedUObjectPointer(Object) ||
 #endif
 				!Object->IsValidLowLevelFast())
+		{
+			FString TokenDebugInfo;
+			if (UClass *Class = (ReferencingObject ? ReferencingObject->GetClass() : nullptr))
 			{
-				FString TokenDebugInfo;
-				if (UClass *Class = (ReferencingObject ? ReferencingObject->GetClass() : nullptr))
-				{
-					const FTokenInfo& TokenInfo = Class->DebugTokenMap.GetTokenInfo(TokenIndex);
-					TokenDebugInfo = FString::Printf(TEXT("ReferencingObjectClass: %s, Property Name: %s, Offset: %d"),
-						*Class->GetFullName(), *TokenInfo.Name.GetPlainNameString(), TokenInfo.Offset);
-				}
-				else
-				{
-					// This means this objects is most likely being referenced by AddReferencedObjects
-					TokenDebugInfo = TEXT("Native Reference");
-				}
-
-				UE_LOG(LogObj, Fatal, TEXT("Invalid object while verifying cluster assumptions: 0x%016llx, ReferencingObject: %s, %s, TokenIndex: %d"),
-					(int64)(PTRINT)Object,
-					ReferencingObject ? *ReferencingObject->GetFullName() : TEXT("NULL"),
-					*TokenDebugInfo, TokenIndex);
+				const FTokenInfo& TokenInfo = Class->DebugTokenMap.GetTokenInfo(TokenIndex);
+				TokenDebugInfo = FString::Printf(TEXT("ReferencingObjectClass: %s, Property Name: %s, Offset: %d"),
+					*Class->GetFullName(), *TokenInfo.Name.GetPlainNameString(), TokenInfo.Offset);
+			}
+			else
+			{
+				// This means this objects is most likely being referenced by AddReferencedObjects
+				TokenDebugInfo = TEXT("Native Reference");
 			}
+
+			TArray<FString> Args;
+			Args.Add(TEXT("Hierarchy"));
+			ListClusters(Args);
+
+			UE_LOG(LogObj, Fatal, TEXT("Invalid object while verifying cluster assumptions: 0x%016llx, ReferencingObject: %s, %s, TokenIndex: %d"),
+				(int64)(PTRINT)Object,
+				ReferencingObject ? *ReferencingObject->GetFullName() : TEXT("NULL"),
+				*TokenDebugInfo, TokenIndex);
+		}
 #endif
 			if (!ProcessedObjects.Contains(Object))
+		{
+			ProcessedObjects.Add(Object);
+			FUObjectItem* ObjectItem = GUObjectArray.ObjectToObjectItem(Object);
+			if (ObjectItem->GetOwnerIndex() == 0)
 			{
-				ProcessedObjects.Add(Object);
-				FUObjectItem* ObjectItem = GUObjectArray.ObjectToObjectItem(Object);
-				if (ObjectItem->GetOwnerIndex() == 0)
+				// We are allowed to reference other clusters, root set objects and objects from diregard for GC pool
+				if (!ObjectItem->HasAnyFlags(EInternalObjectFlags::ClusterRoot|EInternalObjectFlags::RootSet) && !GUObjectArray.IsDisregardForGC(Object) && Object->CanBeInCluster())
 				{
-					// We are allowed to reference other clusters, root set objects and objects from diregard for GC pool
-					if (!ObjectItem->HasAnyFlags(EInternalObjectFlags::ClusterRoot|EInternalObjectFlags::RootSet) && !GUObjectArray.IsDisregardForGC(Object) && Object->CanBeInCluster())
-					{
-						UE_LOG(LogObj, Warning, TEXT("Object %s from cluster %s is referencing 0x%016llx %s which is not part of root set or cluster."),
-							*ReferencingObject->GetFullName(),
-							*ClusterRootObject->GetFullName(),
-							(int64)(PTRINT)Object,
-							*Object->GetFullName());
-						bFailed = true;
+					UE_LOG(LogObj, Warning, TEXT("Object %s from cluster %s is referencing 0x%016llx %s which is not part of root set or cluster."),
+						*ReferencingObject->GetFullName(),
+						*ClusterRootObject->GetFullName(),
+						(int64)(PTRINT)Object,
+						*Object->GetFullName());
+					bFailed = true;
 #if UE_BUILD_DEBUG
-						FReferenceChainSearch RefChainSearch(Object, FReferenceChainSearch::ESearchMode::Shortest | FReferenceChainSearch::ESearchMode::PrintResults);
+					FReferenceChainSearch RefChainSearch(Object, FReferenceChainSearch::ESearchMode::Shortest | FReferenceChainSearch::ESearchMode::PrintResults);
 #endif
-					}
-					else if (ObjectItem->HasAnyFlags(EInternalObjectFlags::ClusterRoot))
-					{
-						// However, clusters need to be referenced by the current cluster otherwise they can also get GC'd too early.
-						const int32 OtherClusterRootIndex = GUObjectArray.ObjectToIndex(Object);
-						const FUObjectItem* OtherClusterRootItem = GUObjectArray.IndexToObjectUnsafeForGC(OtherClusterRootIndex);
-						check(OtherClusterRootItem && OtherClusterRootItem->Object);
-						UObject* OtherClusterRootObject = static_cast<UObject*>(OtherClusterRootItem->Object);
-						UE_CLOG(OtherClusterRootIndex != ClusterRootIndex && !Cluster.ReferencedClusters.Contains(OtherClusterRootIndex), LogObj, Fatal,
-							TEXT("Object %s from source cluster %s is referencing object %s (0x%016llx) from cluster %s which is not referenced by the source cluster."),
-							*ReferencingObject->GetFullName(),
-							*ClusterRootObject->GetFullName(),
-							*Object->GetFullName(),
-							(int64)(PTRINT)Object,
-							*OtherClusterRootObject->GetFullName());
-					}
 				}
-				else if (ObjectItem->GetOwnerIndex() == ClusterRootIndex)
+				else if (ObjectItem->HasAnyFlags(EInternalObjectFlags::ClusterRoot))
 				{
-					// If this object belongs to the current cluster, keep processing its references. Otherwise ignore it as it will be processed by its cluster
-					ObjectsToSerialize.Add(Object);
-				}
-				else
-				{
-					// If we're referencing an object from another cluster, make sure the other cluster is actually referenced by this cluster
-					const int32 OtherClusterRootIndex = ObjectItem->GetOwnerIndex();
-					const FUObjectItem* OtherClusterRootItem = GUObjectArray.IndexToObjectUnsafeForGC(OtherClusterRootIndex);				
+					// However, clusters need to be referenced by the current cluster otherwise they can also get GC'd too early.
+					const int32 OtherClusterRootIndex = GUObjectArray.ObjectToIndex(Object);
+					const FUObjectItem* OtherClusterRootItem = GUObjectArray.IndexToObjectUnsafeForGC(OtherClusterRootIndex);
 					check(OtherClusterRootItem && OtherClusterRootItem->Object);
 					UObject* OtherClusterRootObject = static_cast<UObject*>(OtherClusterRootItem->Object);
 					UE_CLOG(OtherClusterRootIndex != ClusterRootIndex && !Cluster.ReferencedClusters.Contains(OtherClusterRootIndex), LogObj, Fatal,
@@ -692,8 +688,29 @@ public:
 						*OtherClusterRootObject->GetFullName());
 				}
 			}
+			else if (ObjectItem->GetOwnerIndex() == ClusterRootIndex)
+			{
+				// If this object belongs to the current cluster, keep processing its references. Otherwise ignore it as it will be processed by its cluster
+				ObjectsToSerialize.Add(Object);
+			}
+			else
+			{
+				// If we're referencing an object from another cluster, make sure the other cluster is actually referenced by this cluster
+				const int32 OtherClusterRootIndex = ObjectItem->GetOwnerIndex();
+				const FUObjectItem* OtherClusterRootItem = GUObjectArray.IndexToObjectUnsafeForGC(OtherClusterRootIndex);				
+				check(OtherClusterRootItem && OtherClusterRootItem->Object);
+				UObject* OtherClusterRootObject = static_cast<UObject*>(OtherClusterRootItem->Object);
+				UE_CLOG(OtherClusterRootIndex != ClusterRootIndex && !Cluster.ReferencedClusters.Contains(OtherClusterRootIndex), LogObj, Fatal,
+					TEXT("Object %s from source cluster %s is referencing object %s (0x%016llx) from cluster %s which is not referenced by the source cluster."),
+					*ReferencingObject->GetFullName(),
+					*ClusterRootObject->GetFullName(),
+					*Object->GetFullName(),
+					(int64)(PTRINT)Object,
+					*OtherClusterRootObject->GetFullName());
+			}
 		}
 	}
+	}
 };
 
 bool VerifyClusterAssumptions(UObject* ClusterRootObject)
diff --git a/Engine/Source/Runtime/CoreUObject/Private/UObject/UObjectGlobals.cpp b/Engine/Source/Runtime/CoreUObject/Private/UObject/UObjectGlobals.cpp
--- a/Engine/Source/Runtime/CoreUObject/Private/UObject/UObjectGlobals.cpp
+++ b/Engine/Source/Runtime/CoreUObject/Private/UObject/UObjectGlobals.cpp
@@ -49,12 +49,6 @@
 
 DEFINE_LOG_CATEGORY(LogUObjectGlobals);
 
-#if !defined(USE_EVENT_DRIVEN_ASYNC_LOAD)
-#error "USE_EVENT_DRIVEN_ASYNC_LOAD must be defined"
-#endif
-
-#define UE_USE_ASYNCPATH_FOR_LOADPACKAGE (USE_EVENT_DRIVEN_ASYNC_LOAD)
-
 #if USE_MALLOC_PROFILER
 #include "MallocProfiler.h"
 #endif
@@ -1022,7 +1016,7 @@ public:
 	{
 		Package->LinkerLoad = this;
 
-		while (CreateLoader() == FLinkerLoad::LINKER_TimedOut)
+		while (CreateLoader(TFunction<void()>([]() {})) == FLinkerLoad::LINKER_TimedOut)
 		{
 		}
 
@@ -1095,9 +1089,8 @@ static UPackage* LoadPackageInternalInner(UPackage* InOuter, const TCHAR* InLong
 
 	UPackage* Result = nullptr;
 
-#if UE_USE_ASYNCPATH_FOR_LOADPACKAGE
-	if (FPlatformProperties::RequiresCookedData()
-		&& !GIsInitialLoad //@todoio fix this so we can async load during compiled in object init		
+	if (FPlatformProperties::RequiresCookedData() && GEventDrivenLoaderEnabled
+		&& EVENT_DRIVEN_ASYNC_LOAD_ACTIVE_AT_RUNTIME
 		)
 	{
 		FString InName;
@@ -1124,7 +1117,6 @@ static UPackage* LoadPackageInternalInner(UPackage* InOuter, const TCHAR* InLong
 		FName PackageFName(*InPackageName);
 
 		Result = FindObjectFast<UPackage>(nullptr, PackageFName);
-#if USE_EVENT_DRIVEN_ASYNC_LOAD
 		if (!Result || !Result->IsFullyLoaded())
 		{
 			if (FLinkerLoad::IsKnownMissingPackage(PackageFName) || (Result && Result->IsPendingKill()))
@@ -1146,22 +1138,9 @@ static UPackage* LoadPackageInternalInner(UPackage* InOuter, const TCHAR* InLong
 				}
 			}
 		}
-#else
-		if (!Result || !Result->IsFullyLoaded())
-		{
-			int32 RequestID = LoadPackageAsync(InName, nullptr, *InPackageName);
-			FlushAsyncLoading(RequestID);
-		}
 
-		if (InOuter)
-		{
-			return InOuter;
-		}
-		Result = FindObjectFast<UPackage>(nullptr, PackageFName);
-#endif
 		return Result;
 }
-#endif	
 
 	FString FileToLoad;
 #if WITH_EDITOR
@@ -1316,11 +1295,14 @@ static UPackage* LoadPackageInternalInner(UPackage* InOuter, const TCHAR* InLong
 			// Make sure we pass the property that's currently being serialized by the linker that owns the import 
 			// that triggered this LoadPackage call
 			FSerializedPropertyScope SerializedProperty(*Linker, ImportLinker ? ImportLinker->GetSerializedProperty() : Linker->GetSerializedProperty());
-#if USE_NEW_ASYNC_IO 
+			if (GNewAsyncIO)
+			{
 			Linker->LoadAllObjects(true);
-#else
+			}
+			else
+			{
 			Linker->LoadAllObjects();
-#endif		
+			}	
 		}
 		else
 		{
@@ -1355,16 +1337,16 @@ static UPackage* LoadPackageInternalInner(UPackage* InOuter, const TCHAR* InLong
 			Result->SetLoadTime( FPlatformTime::Seconds() - StartTime );
 		}
 
-#if USE_NEW_ASYNC_IO
+		if (GNewAsyncIO)
+		{
 		Linker->Flush();
-#else
+		}
 		// @todo: the next two conditions should check the file limit
-		if (FPlatformProperties::RequiresCookedData())
+		else if (FPlatformProperties::RequiresCookedData())
 		{
 			// give a hint to the IO system that we are done with this file for now
 			FIOSystem::Get().HintDoneWithFile(*Linker->Filename);
 		}
-#endif
 
 		// With UE4 and single asset per package, we load so many packages that some platforms will run out
 		// of file handles. So, this will close the package, but just things like bulk data loading will
@@ -1655,8 +1637,7 @@ void EndLoad()
 				}
 			}
 
-#if USE_EVENT_DRIVEN_ASYNC_LOAD
-			if (!GIsInitialLoad)
+			if (GEventDrivenLoaderEnabled && EVENT_DRIVEN_ASYNC_LOAD_ACTIVE_AT_RUNTIME)
 			{
 #if DO_CHECK
 				for (UObject* Obj : ObjLoaded)
@@ -1670,7 +1651,6 @@ void EndLoad()
 #endif
 			}
 			else
-#endif
 			{
 				// Dynamic Class doesn't require/use pre-loading (or post-loading). 
 				// The CDO is created at this point, because now it's safe to solve cyclic dependencies.
@@ -3053,10 +3033,8 @@ void FSubobjectPtr::Set(UObject* InObject)
 // Binary initialize object properties to zero or defaults.
 void FObjectInitializer::InitProperties(UObject* Obj, UClass* DefaultsClass, UObject* DefaultData, bool bCopyTransientsFromClassDefaults)
 {
-#if USE_EVENT_DRIVEN_ASYNC_LOAD
-	check(!DefaultsClass || !DefaultsClass->HasAnyFlags(RF_NeedLoad));
-	check(!DefaultData || !DefaultData->HasAnyFlags(RF_NeedLoad));
-#endif
+	check(!GEventDrivenLoaderEnabled || !DefaultsClass || !DefaultsClass->HasAnyFlags(RF_NeedLoad));
+	check(!GEventDrivenLoaderEnabled || !DefaultData || !DefaultData->HasAnyFlags(RF_NeedLoad));
 
 	SCOPE_CYCLE_COUNTER(STAT_InitProperties);
 
@@ -3109,9 +3087,7 @@ void FObjectInitializer::InitProperties(UObject* Obj, UClass* DefaultsClass, UOb
 		bCanUsePostConstructLink &= (DefaultData == Class->GetDefaultObject(false));
 
 		UObject* ClassDefaults = bCopyTransientsFromClassDefaults ? DefaultsClass->GetDefaultObject() : NULL;	
-#if USE_EVENT_DRIVEN_ASYNC_LOAD
-		check(!bCopyTransientsFromClassDefaults || !DefaultsClass->GetDefaultObject()->HasAnyFlags(RF_NeedLoad));
-#endif
+		check(!GEventDrivenLoaderEnabled || !bCopyTransientsFromClassDefaults || !DefaultsClass->GetDefaultObject()->HasAnyFlags(RF_NeedLoad));
 
 		for (UProperty* P = bCanUsePostConstructLink ? Class->PostConstructLink : Class->PropertyLink; P; P = bCanUsePostConstructLink ? P->PostConstructLinkNext : P->PropertyLinkNext)
 		{
@@ -3688,11 +3664,6 @@ UScriptStruct* GetFallbackStruct()
 	return TBaseStructure<FFallbackStruct>::Get();
 }
 
-#if USE_EVENT_DRIVEN_ASYNC_LOAD
-extern COREUOBJECT_API bool GIgnoreGetArchetypeFromRequiredInfo_RF_NeedsLoad;
-#endif
-
-
 UObject* FObjectInitializer::CreateDefaultSubobject(UObject* Outer, FName SubobjectFName, UClass* ReturnType, UClass* ClassToCreateByDefault, bool bIsRequired, bool bAbstract, bool bIsTransient) const
 {
 	UE_CLOG(!FUObjectThreadContext::Get().IsInConstructor, LogClass, Fatal, TEXT("Subobjects cannot be created outside of UObject constructors. UObject constructing subobjects cannot be created using new or placement new operator."));
@@ -3725,13 +3696,9 @@ UObject* FObjectInitializer::CreateDefaultSubobject(UObject* Outer, FName Subobj
 				SubobjectFlags |= RF_Transient;
 			}
 
-			{
-#if USE_EVENT_DRIVEN_ASYNC_LOAD
-				TGuardValue<bool> Guard(GIgnoreGetArchetypeFromRequiredInfo_RF_NeedsLoad, true);
-#endif
 				OuterArchetypeClass = Outer->GetArchetype()->GetClass();
 				bOwnerArchetypeIsNotNative = !OuterArchetypeClass->HasAnyClassFlags(CLASS_Native | CLASS_Intrinsic);
-			}
+
 			const bool bOwnerTemplateIsNotCDO = ObjectArchetype != nullptr && ObjectArchetype != Outer->GetClass()->GetDefaultObject(false) && !Outer->HasAnyFlags(RF_ClassDefaultObject);
 #if !UE_BUILD_SHIPPING
 			// Guard against constructing the same subobject multiple times.
diff --git a/Engine/Source/Runtime/CoreUObject/Public/Serialization/ArchiveAsync.h b/Engine/Source/Runtime/CoreUObject/Public/Serialization/ArchiveAsync.h
--- a/Engine/Source/Runtime/CoreUObject/Public/Serialization/ArchiveAsync.h
+++ b/Engine/Source/Runtime/CoreUObject/Public/Serialization/ArchiveAsync.h
@@ -7,7 +7,6 @@
 
 struct FCompressedChunk;
 
-#if !USE_NEW_ASYNC_IO
 /*----------------------------------------------------------------------------
 	FArchiveAsync.
 ----------------------------------------------------------------------------*/
@@ -175,7 +174,6 @@ private:
 	bool PlatformIsSinglethreaded;
 };
 
-#else
 
 COREUOBJECT_API FArchive* NewFArchiveAsync2(const TCHAR* InFileName);
 
@@ -183,5 +181,4 @@ COREUOBJECT_API void HintFutureReadDone(const TCHAR * FileName);
 
 COREUOBJECT_API void HintFutureRead(const TCHAR * FileName);
 
-#endif // USE_NEW_ASYNC_IO
 
diff --git a/Engine/Source/Runtime/CoreUObject/Public/Serialization/AsyncLoading.h b/Engine/Source/Runtime/CoreUObject/Public/Serialization/AsyncLoading.h
--- a/Engine/Source/Runtime/CoreUObject/Public/Serialization/AsyncLoading.h
+++ b/Engine/Source/Runtime/CoreUObject/Public/Serialization/AsyncLoading.h
@@ -21,18 +21,8 @@ struct FFlushTree;
 
 #define PERF_TRACK_DETAILED_ASYNC_STATS (0)
 
-#ifndef USE_EVENT_DRIVEN_ASYNC_LOAD
-#error "USE_EVENT_DRIVEN_ASYNC_LOAD must be defined"
-#endif
-
-
-#if USE_EVENT_DRIVEN_ASYNC_LOAD
-
-#if !USE_NEW_ASYNC_IO
-#error "USE_EVENT_DRIVEN_ASYNC_LOAD requires USE_NEW_ASYNC_IO"
-#endif
-
 struct FAsyncPackage;
+/** [EDL] Async Package Loading State */
 enum class EAsyncPackageLoadingState : uint8
 {
 	NewPackage,
@@ -48,7 +38,7 @@ enum class EAsyncPackageLoadingState : uint8
 	PackageComplete,
 };
 
-// This version is an ordinary pointer. We can swap in the safer version to verify assumptions.
+/** [EDL] This version is an ordinary pointer. We can swap in the safer version to verify assumptions */
 struct FUnsafeWeakAsyncPackagePtr
 {
 	FAsyncPackage* Package;
@@ -72,6 +62,7 @@ struct FUnsafeWeakAsyncPackagePtr
 	FName HumanReadableStringForDebugging() const;
 };
 
+/**  [EDL] Weak pointer to the async package */
 struct FWeakAsyncPackagePtr
 {
 	FName PackageName;
@@ -100,6 +91,7 @@ typedef FWeakAsyncPackagePtr FCheckedWeakAsyncPackagePtr;
 typedef FUnsafeWeakAsyncPackagePtr FCheckedWeakAsyncPackagePtr;
 #endif
 
+/** [EDL] Event Load Node */
 enum class EEventLoadNode
 {
 	Package_LoadSummary,
@@ -127,7 +119,7 @@ enum class EEventLoadNode
 // Export: EEventLoadNode::Export_StartIO -> EEventLoadNode::ImportOrExport_Serialize: can't serialize until the IO request is ready
 // Import: EEventLoadNode::ImportOrExport_Serialize -> EEventLoadNode::Package_ExportsSerialized: can't consider the package done with event driven loading until all exports are serialized
 
-
+/** [EDL] Event Load Node Pointer */
 struct FEventLoadNodePtr
 {
 	FCheckedWeakAsyncPackagePtr WaitingPackage;
@@ -191,6 +183,7 @@ struct FEventLoadNodePtr
 	FString HumanReadableStringForDebugging() const;
 };
 
+/** [EDL] Event Load Node */
 struct FEventLoadNode
 {
 #if USE_IMPLICIT_ARCS
@@ -211,6 +204,7 @@ struct FEventLoadNode
 	}
 };
 
+/** [EDL] Event Load Node Array */
 struct FEventLoadNodeArray
 {
 	FEventLoadNode PackageNodes[int(EEventLoadNode::Package_NumPhases)];
@@ -287,7 +281,7 @@ private:
 
 };
 
-
+/** [EDL] Event Load Graph */
 struct FEventLoadGraph
 {
 	TSet<FCheckedWeakAsyncPackagePtr> PackagesWithNodes;
@@ -307,6 +301,7 @@ struct FEventLoadGraph
 #endif
 };
 
+/** [EDL] Event Load Graph */
 struct FAsyncLoadEventArgs
 {
 	double TickStartTime;
@@ -327,15 +322,13 @@ struct FAsyncLoadEventArgs
 	}
 };
 
-// this is a little wrapper that does random pops for further randomization
+/** [EDL] this is a little wrapper that does random pops for further randomization */
 class FImportOrImportIndexArray : public TArray<int32>
 {
 public:
 	void HeapPop(int32& OutItem, bool bAllowShrinking = true);
 };
 
-#endif
-
 /**
 * Structure containing intermediate data required for async loading of all imports and exports of a
 * FLinkerLoad.
@@ -589,7 +582,8 @@ private:
 
 public:
 
-#if USE_EVENT_DRIVEN_ASYNC_LOAD
+	/** [EDL] Begin Event driven loader specific stuff */
+
 	EAsyncPackageLoadingState AsyncPackageLoadingState;
 	int32 SerialNumber;
 
@@ -715,7 +709,7 @@ public:
 		bUseFullTimeLimit = Args.bUseFullTimeLimit;
 	}
 
-#endif
+	/** [EDL] End Event driven loader specific stuff */
 
 
 #if PERF_TRACK_DETAILED_ASYNC_STATS
@@ -841,6 +835,12 @@ private:
 	 * @return true if we finished serializing all loaded objects, false otherwise.
 	 */
 	EAsyncPackageState::Type PreLoadObjects();
+	/**
+	* Preloads aka serializes all loaded objects for the new async IO.
+	*
+	* @return true if we finished serializing all loaded objects, false otherwise.
+	*/
+	EAsyncPackageState::Type PreLoadObjectsForNewAsyncIO();
 	/**
 	 * Route PostLoad to all loaded objects. This might load further objects!
 	 *
@@ -898,7 +898,6 @@ private:
 #endif // PERF_TRACK_DETAILED_ASYNC_STATS
 };
 
-#if USE_EVENT_DRIVEN_ASYNC_LOAD
 struct FScopedAsyncPackageEvent
 {
 	/** Current scope package */
@@ -909,7 +908,6 @@ struct FScopedAsyncPackageEvent
 	FScopedAsyncPackageEvent(FAsyncPackage* InPackage);
 	~FScopedAsyncPackageEvent();
 };
-#endif
 
 
 
diff --git a/Engine/Source/Runtime/CoreUObject/Public/Serialization/BulkData.h b/Engine/Source/Runtime/CoreUObject/Public/Serialization/BulkData.h
--- a/Engine/Source/Runtime/CoreUObject/Public/Serialization/BulkData.h
+++ b/Engine/Source/Runtime/CoreUObject/Public/Serialization/BulkData.h
@@ -33,7 +33,7 @@ enum EBulkDataFlags
 	BULKDATA_PayloadInSeperateFile				= 1 << 8,
 	/** If set, payload is compressed using platform specific bit window			*/
 	BULKDATA_SerializeCompressedBitWindow		= 1<<9,
-	/** Only used with USE_NEW_ASYNC_IO, which changes the default to inline unless you opt out */
+	/** Only used with GNewAsyncIO, which changes the default to inline unless you opt out */
 	BULKDATA_Force_NOT_InlinePayload = 1 << 10,
 };
 
diff --git a/Engine/Source/Runtime/CoreUObject/Public/UObject/Linker.h b/Engine/Source/Runtime/CoreUObject/Public/UObject/Linker.h
--- a/Engine/Source/Runtime/CoreUObject/Public/UObject/Linker.h
+++ b/Engine/Source/Runtime/CoreUObject/Public/UObject/Linker.h
@@ -11,10 +11,6 @@ class FReferenceCollector;
 
 DECLARE_LOG_CATEGORY_EXTERN(LogLinker, Log, All);
 
-#if !defined(USE_NEW_ASYNC_IO) || !defined(SPLIT_COOKED_FILES)
-#error "USE_NEW_ASYNC_IO and SPLIT_COOKED_FILES must be defined"
-#endif
-
 /**
  * Information about a compressed chunk in a file.
  */
diff --git a/Engine/Source/Runtime/CoreUObject/Public/UObject/LinkerLoad.h b/Engine/Source/Runtime/CoreUObject/Public/UObject/LinkerLoad.h
--- a/Engine/Source/Runtime/CoreUObject/Public/UObject/LinkerLoad.h
+++ b/Engine/Source/Runtime/CoreUObject/Public/UObject/LinkerLoad.h
@@ -68,12 +68,11 @@ struct FScopedCreateImportCounter
 /**
  * Handles loading Unreal package files, including reading UObject data from disk.
  */
-#if USE_NEW_ASYNC_IO
 class FArchiveAsync2;
-#endif
-class FLinkerLoad
-#if USE_NEW_ASYNC_IO
-	final
+
+class FLinkerLoad 
+#if !WITH_EDITOR
+	final 
 #endif
 	: public FLinker, public FArchiveUObject
 {
@@ -120,19 +119,19 @@ public:
 	bool					bHaveImportsBeenVerified;
 	/** Indicates that this linker was created for a dynamic class package and will not use Loader */
 	bool					bDynamicClassLinker;
-#if USE_EVENT_DRIVEN_ASYNC_LOAD
+
 	UObject*				TemplateForGetArchetypeFromLoader;
 	bool					bForceSimpleIndexToObject;
 	bool					bLockoutLegacyOperations;
-#endif
-#if USE_NEW_ASYNC_IO
+
 	/** True if Loader is FArchiveAsync2  */
 	bool					bLoaderIsFArchiveAsync2;
 	FORCEINLINE FArchiveAsync2* GetFArchiveAsync2Loader()
 	{
+		check(GNewAsyncIO);
 		return bLoaderIsFArchiveAsync2 ? (FArchiveAsync2*)Loader : nullptr;
 	}
-#endif
+
 	/** The archive that actually reads the raw data from disk.																*/
 	FArchive*				Loader;
 	/** The async package associated with this linker */
@@ -278,7 +277,6 @@ private:
 	/** Id of the thread that created this linker. This is to guard against using this linker on other threads than the one it was created on **/
 	int32					OwnerThread;
 
-#if !USE_NEW_ASYNC_IO
 	/**
 	 * Helper struct to keep track of background file reads
 	 */
@@ -323,8 +321,6 @@ public:
 	COREUOBJECT_API static void GetListOfPackagesInPackagePrecacheMap(TArray<FString>& ListOfPackages);
 private:
 
-#endif
-
 	/** Allows access to UTexture2D::StaticClass() without linking Core with Engine											*/
 	static UClass* UTexture2DStaticClass;
 
@@ -520,14 +516,12 @@ public:
 	 */
 	COREUOBJECT_API bool WillTextureBeLoaded( UClass* Class, int32 ExportIndex );
 
-#if !USE_NEW_ASYNC_IO
 	/**
 	 * Kick off an async load of a package file into memory
 	 * 
 	 * @param PackageName Name of package to read in. Must be the same name as passed into LoadPackage/CreateLinker
 	 */
 	COREUOBJECT_API static void AsyncPreloadPackage(const TCHAR* PackageName);
-#endif
 
 	/**
 	 * Called when an object begins serializing property data using script serialization.
@@ -539,9 +533,7 @@ public:
 	 */
 	virtual void MarkScriptSerializationEnd( const UObject* Obj ) override;
 
-#if USE_EVENT_DRIVEN_ASYNC_LOAD
 	virtual UObject* GetArchetypeFromLoader(const UObject* Obj) override;
-#endif
 
 	/**
 	 * Looks for an existing linker for the given package, without trying to make one if it doesn't exist
@@ -806,9 +798,7 @@ private:
 	 * @return	new FLinkerLoad object for Parent/ Filename
 	 */
 	COREUOBJECT_API static FLinkerLoad* CreateLinkerAsync(UPackage* Parent, const TCHAR* Filename, uint32 LoadFlags
-#if USE_EVENT_DRIVEN_ASYNC_LOAD
-		, TFunction<void()>&& InSummaryReadyCallback
-#endif		
+		, TFunction<void()>&& InSummaryReadyCallback	
 	);
 
 	/**
@@ -847,9 +837,7 @@ protected: // Daniel L: Made this protected so I can override the constructor an
 	 * Creates loader used to serialize content.
 	 */
 	ELinkerStatus CreateLoader(
-#if USE_EVENT_DRIVEN_ASYNC_LOAD
 		TFunction<void()>&& InSummaryReadyCallback
-#endif
 	);
 private:
 	/**
diff --git a/Engine/Source/Runtime/CoreUObject/Public/UObject/ObjectResource.h b/Engine/Source/Runtime/CoreUObject/Public/UObject/ObjectResource.h
--- a/Engine/Source/Runtime/CoreUObject/Public/UObject/ObjectResource.h
+++ b/Engine/Source/Runtime/CoreUObject/Public/UObject/ObjectResource.h
@@ -6,10 +6,6 @@
 #include "UObject/ObjectMacros.h"
 #include "Misc/Guid.h"
 
-#if !defined(USE_NEW_ASYNC_IO) || !defined(SPLIT_COOKED_FILES)
-#error "USE_NEW_ASYNC_IO and SPLIT_COOKED_FILES must be defined"
-#endif
-
 /**
  * Wrapper for index into a ULnker's ImportMap or ExportMap.
  * Values greater than zero indicate that this is an index into the ExportMap.  The
@@ -379,11 +375,9 @@ struct FObjectImport : public FObjectResource
 	 */
 	int32             SourceIndex;
 
-#if USE_EVENT_DRIVEN_ASYNC_LOAD
 	bool			bImportPackageHandled;
 	bool			bImportSearchedFor;
 	bool			bImportFailed;
-#endif
 
 	/**
 	 * Constructors
diff --git a/Engine/Source/Runtime/CoreUObject/Public/UObject/Package.h b/Engine/Source/Runtime/CoreUObject/Public/UObject/Package.h
--- a/Engine/Source/Runtime/CoreUObject/Public/UObject/Package.h
+++ b/Engine/Source/Runtime/CoreUObject/Public/UObject/Package.h
@@ -37,6 +37,9 @@ enum class ESavePackageResult
 	GenerateStub
 };
 
+COREUOBJECT_API void StartSavingEDLCookInfoForVerification();
+COREUOBJECT_API void VerifyEDLCookInfo();
+
 /**
  * A package.
  */
diff --git a/Engine/Source/Runtime/CoreUObject/Public/UObject/UObjectArray.h b/Engine/Source/Runtime/CoreUObject/Public/UObject/UObjectArray.h
--- a/Engine/Source/Runtime/CoreUObject/Public/UObject/UObjectArray.h
+++ b/Engine/Source/Runtime/CoreUObject/Public/UObject/UObjectArray.h
@@ -80,21 +80,36 @@ struct FUObjectItem
 			{
 				break;
 			}
-			int32 OldValue = (int32)FPlatformAtomics::InterlockedCompareExchange((int32*)&Flags, StartValue & ~int32(FlagToClear), StartValue);
-			// We know the flag was set when we entered this iteration,
-			// so if the old value returned by atomics had the flag set, we must have cleared it.
-			// (there is always a chance that another thread cleared some other flag and the above function did nothing)
-			// But we only care about the flags we want to clear
-			if (!(Flags & int32(FlagToClear)) && (OldValue & int32(FlagToClear)) == (StartValue & int32(FlagToClear)))
+			int32 NewValue = StartValue & ~int32(FlagToClear);
+			checkSlow(NewValue != StartValue);
+			if ((int32)FPlatformAtomics::InterlockedCompareExchange((int32*)&Flags, NewValue, StartValue) == StartValue)
+			{
+				bIChangedIt = true;
+				break;
+			}
+		}
+		return bIChangedIt;
+	}
+
+	FORCEINLINE bool ThisThreadAtomicallySetFlag(EInternalObjectFlags FlagToSet)
+	{
+		static_assert(sizeof(int32) == sizeof(Flags), "Flags must be 32-bit for atomics.");
+		bool bIChangedIt = false;
+		while (1)
+		{
+			int32 StartValue = int32(Flags);
+			if (StartValue & int32(FlagToSet))
+			{
+				break;
+			}
+			int32 NewValue = StartValue | int32(FlagToSet);
+			checkSlow(NewValue != StartValue);
+			if ((int32)FPlatformAtomics::InterlockedCompareExchange((int32*)&Flags, NewValue, StartValue) == StartValue)
 			{
-				// if (the flag has actually been cleared) && (the previous value had the flag set) we must have cleared it
 				bIChangedIt = true;
 				break;
 			}
-			// We didn't clear the flag, probably because some other thread changed flags in the meantime (either the one we want to clear or some other). Try again.
 		}
-		// Make sure the flag was actually cleared
-		checkSlow((Flags & int32(FlagToClear)) == 0);
 		return bIChangedIt;
 	}
 
diff --git a/Engine/Source/Runtime/CoreUObject/Public/UObject/UObjectGlobals.h b/Engine/Source/Runtime/CoreUObject/Public/UObject/UObjectGlobals.h
--- a/Engine/Source/Runtime/CoreUObject/Public/UObject/UObjectGlobals.h
+++ b/Engine/Source/Runtime/CoreUObject/Public/UObject/UObjectGlobals.h
@@ -387,6 +387,11 @@ COREUOBJECT_API void CancelAsyncLoading();
 */
 COREUOBJECT_API bool IsEventDrivenLoaderEnabledInCookedBuilds();
 
+/**
+* Returns true if the event driven loader is enabled in the current build
+*/
+COREUOBJECT_API bool IsEventDrivenLoaderEnabled();
+
 /**
  * Returns the async load percentage for a package in flight with the passed in name or -1 if there isn't one.
  * THIS IS SLOW. MAY BLOCK ASYNC LOADING.
diff --git a/Engine/Source/Runtime/CoreUObject/Public/UObject/UnrealType.h b/Engine/Source/Runtime/CoreUObject/Public/UObject/UnrealType.h
--- a/Engine/Source/Runtime/CoreUObject/Public/UObject/UnrealType.h
+++ b/Engine/Source/Runtime/CoreUObject/Public/UObject/UnrealType.h
@@ -1280,8 +1280,8 @@ protected:
 			TEXT("Potential data loss during conversion of integer property %s of %s - was (%s) now (%s) - for package: %s"),
 			*this->GetName(),
 			*Ar.GetArchiveName(),
-			*LexicalConversion::ToString(OldValue),
-			*LexicalConversion::ToString(NewValue),
+			*Lex::ToString(OldValue),
+			*Lex::ToString(NewValue),
 			*Ar.GetArchiveName()
 			);
 	}
@@ -1374,11 +1374,11 @@ public:
 	}
 	virtual void SetNumericPropertyValueFromString(void* Data, TCHAR const* Value) const override
 	{
-		LexicalConversion::FromString(*TTypeFundamentals::GetPropertyValuePtr(Data), Value);
+		Lex::FromString(*TTypeFundamentals::GetPropertyValuePtr(Data), Value);
 	}
 	virtual FString GetNumericPropertyValueToString(void const* Data) const override
 	{
-		return LexicalConversion::ToString(TTypeFundamentals::GetPropertyValue(Data));
+		return Lex::ToString(TTypeFundamentals::GetPropertyValue(Data));
 	}
 	virtual int64 GetSignedIntPropertyValue(void const* Data) const override
 	{
diff --git a/Engine/Source/Runtime/Engine/Classes/Engine/Engine.h b/Engine/Source/Runtime/Engine/Classes/Engine/Engine.h
--- a/Engine/Source/Runtime/Engine/Classes/Engine/Engine.h
+++ b/Engine/Source/Runtime/Engine/Classes/Engine/Engine.h
@@ -1711,10 +1711,8 @@ public:
 	bool HandleFreezeStreamingCommand( const TCHAR* Cmd, FOutputDevice& Ar, UWorld* InWorld );		// Smedis
 	bool HandleFreezeAllCommand( const TCHAR* Cmd, FOutputDevice& Ar, UWorld* InWorld );			// Smedis
 
-#if !USE_NEW_ASYNC_IO
 	bool HandleFlushIOManagerCommand( const TCHAR* Cmd, FOutputDevice& Ar );						// Smedis
 	bool HandleListPreCacheMapPackagesCommand(const TCHAR* Cmd, FOutputDevice& Ar);
-#endif
 
 	bool HandleToggleRenderingThreadCommand( const TCHAR* Cmd, FOutputDevice& Ar );	
 	bool HandleToggleAsyncComputeCommand( const TCHAR* Cmd, FOutputDevice& Ar );
diff --git a/Engine/Source/Runtime/Engine/Private/AudioStreaming.cpp b/Engine/Source/Runtime/Engine/Private/AudioStreaming.cpp
--- a/Engine/Source/Runtime/Engine/Private/AudioStreaming.cpp
+++ b/Engine/Source/Runtime/Engine/Private/AudioStreaming.cpp
@@ -13,10 +13,8 @@ AudioStreaming.cpp: Implementation of audio streaming classes.
 #include "Serialization/MemoryReader.h"
 #include "HAL/PlatformFile.h"
 #include "HAL/PlatformFilemanager.h"
-
-#if USE_NEW_ASYNC_IO
 #include "AsyncFileHandle.h"
-#endif
+
 /*------------------------------------------------------------------------------
 	Streaming chunks from the derived data cache.
 ------------------------------------------------------------------------------*/
@@ -73,43 +71,42 @@ void FAsyncStreamDerivedChunkWorker::DoWork()
 
 FStreamingWaveData::FStreamingWaveData()
 	: SoundWave(NULL)
-#if USE_NEW_ASYNC_IO
 	, IORequestHandle(nullptr)
-#endif
 {
-#if USE_NEW_ASYNC_IO
-	AsyncFileCallBack =
-		[this](bool bWasCancelled, IAsyncReadRequest* Req)
+	if (GNewAsyncIO)
 	{
-		uint8* Mem = Req->GetReadResults();
-		if (Mem)
+		AsyncFileCallBack =
+			[this](bool bWasCancelled, IAsyncReadRequest* Req)
 		{
-			bool bFound = false;
-			for (FLoadedAudioChunk& LoadedChunk : LoadedChunks)
+			uint8* Mem = Req->GetReadResults();
+			if (Mem)
 			{
-				if (LoadedChunk.IORequest == Req)
+				bool bFound = false;
+				for (FLoadedAudioChunk& LoadedChunk : LoadedChunks)
 				{
-					check(!LoadedChunk.Data);
-					LoadedChunk.Data = Mem;
-					DEC_MEMORY_STAT_BY(STAT_AsyncFileMemory, LoadedChunk.DataSize);
-					INC_DWORD_STAT_BY(STAT_AudioMemorySize, LoadedChunk.DataSize);
-					INC_DWORD_STAT_BY(STAT_AudioMemory, LoadedChunk.DataSize);
-					bFound = true;
-					break;
+					if (LoadedChunk.IORequest == Req)
+					{
+						check(!LoadedChunk.Data);
+						LoadedChunk.Data = Mem;
+						DEC_MEMORY_STAT_BY(STAT_AsyncFileMemory, LoadedChunk.DataSize);
+						INC_DWORD_STAT_BY(STAT_AudioMemorySize, LoadedChunk.DataSize);
+						INC_DWORD_STAT_BY(STAT_AudioMemory, LoadedChunk.DataSize);
+						bFound = true;
+						break;
+					}
 				}
+				check(bFound);
 			}
-			check(bFound);
-		}
-		// else was canceled
-		PendingChunkChangeRequestStatus.Decrement();
-	};
-#endif
+			// else was canceled
+			PendingChunkChangeRequestStatus.Decrement();
+		};
+	}
 }
 
 FStreamingWaveData::~FStreamingWaveData()
 {
 	// Make sure there are no pending requests in flight.
-#if USE_NEW_ASYNC_IO
+	if (GNewAsyncIO)
 	{
 		for (int32 Pass = 0; Pass < 3; Pass++)
 		{
@@ -121,36 +118,32 @@ FStreamingWaveData::~FStreamingWaveData()
 			check(Pass < 2); // we should be done after two passes. Pass 0 will start anything we need and pass 1 will complete those requests
 		}
 	}
-#else
+	else
 	{
 		QUICK_SCOPE_CYCLE_COUNTER(FStreamingWaveData_Destructor_Spin);
 		while (UpdateStreamingStatus() == true)
 		{
 			// Give up timeslice.
-			check(!USE_NEW_ASYNC_IO); // should be done because we block on it above
+			check(!GNewAsyncIO); // should be done because we block on it above
 			FPlatformProcess::SleepNoStats(0.001f);
 		}
 	}
-#endif
 
 	for (FLoadedAudioChunk& LoadedChunk : LoadedChunks)
 	{
 		FreeLoadedChunk(LoadedChunk);
 	}
-#if USE_NEW_ASYNC_IO
-	if (IORequestHandle)
+	if (GNewAsyncIO && IORequestHandle)
 	{
 		delete IORequestHandle;
 		IORequestHandle = nullptr;
 	}
-#endif
 }
 
 void FStreamingWaveData::Initialize(USoundWave* InSoundWave)
 {
-#if USE_NEW_ASYNC_IO
-	check(!IORequestHandle);
-#endif
+	check(!GNewAsyncIO || !IORequestHandle);
+
 	SoundWave = InSoundWave;
 
 	// Always get the first chunk of data so we can play immediately
@@ -185,7 +178,7 @@ bool FStreamingWaveData::UpdateStreamingStatus()
 			if (UE_LOG_ACTIVE(LogAudio, Log) && IndicesToLoad.Num() > 0)
 			{
 				FString LogString = FString::Printf(TEXT("Finalised loading of chunk(s) %d"), IndicesToLoad[0]);
-				for (int32 Index = 1; Index < IndicesToLoad.Num(); ++ Index)
+				for (int32 Index = 1; Index < IndicesToLoad.Num(); ++Index)
 				{
 					LogString += FString::Printf(TEXT(", %d"), IndicesToLoad[Index]);
 				}
@@ -198,18 +191,19 @@ bool FStreamingWaveData::UpdateStreamingStatus()
 			bFailedRequests = FinishDDCRequests();
 #endif //WITH_EDITORONLY_DATA
 
-#if USE_NEW_ASYNC_IO
-			// could maybe iterate over the things we know are done, but I couldn't tell if that was IndicesToLoad or not.
-			for (FLoadedAudioChunk& LoadedChunk : LoadedChunks)
+			if (GNewAsyncIO)
 			{
-				if (LoadedChunk.IORequest && LoadedChunk.IORequest->PollCompletion())
+				// could maybe iterate over the things we know are done, but I couldn't tell if that was IndicesToLoad or not.
+				for (FLoadedAudioChunk& LoadedChunk : LoadedChunks)
 				{
-					LoadedChunk.IORequest->WaitCompletion();
-					delete LoadedChunk.IORequest;
-					LoadedChunk.IORequest = nullptr;
+					if (LoadedChunk.IORequest && LoadedChunk.IORequest->PollCompletion())
+					{
+						LoadedChunk.IORequest->WaitCompletion();
+						delete LoadedChunk.IORequest;
+						LoadedChunk.IORequest = nullptr;
+					}
 				}
 			}
-#endif
 
 			PendingChunkChangeRequestStatus.Decrement();
 			bHasPendingRequestInFlight = false;
@@ -232,10 +226,11 @@ void FStreamingWaveData::UpdateChunkRequests(FWaveRequest& InWaveRequest)
 
 	CurrentRequest = InWaveRequest;
 
-#if !USE_NEW_ASYNC_IO
-	// Clear last batch of request indices
-	IORequestIndices.Empty();
-#endif
+	if (!GNewAsyncIO)
+	{
+		// Clear last batch of request indices
+		IORequestIndices.Empty();
+	}
 }
 
 bool FStreamingWaveData::HasPendingRequests(TArray<uint32>& IndicesToLoad, TArray<uint32>& IndicesToFree) const
@@ -277,35 +272,39 @@ void FStreamingWaveData::BeginPendingRequests(const TArray<uint32>& IndicesToLoa
 		UE_LOG(LogAudio, Log, TEXT("%s"), *LogString);
 	}
 
+	TArray<uint32> FreeChunkIndices;
+
 	// Mark Chunks for removal in case they can be reused
-#if USE_NEW_ASYNC_IO
-	for (auto Index : IndicesToFree)
+	if (GNewAsyncIO)
 	{
-		for (int32 ChunkIndex = 0; ChunkIndex < LoadedChunks.Num(); ++ChunkIndex)
+		for (auto Index : IndicesToFree)
 		{
-			if (LoadedChunks[ChunkIndex].Index == Index)
+			for (int32 ChunkIndex = 0; ChunkIndex < LoadedChunks.Num(); ++ChunkIndex)
 			{
-				FreeLoadedChunk(LoadedChunks[ChunkIndex]);
-				LoadedChunks.RemoveAt(ChunkIndex);
-				break;
+				if (LoadedChunks[ChunkIndex].Index == Index)
+				{
+					FreeLoadedChunk(LoadedChunks[ChunkIndex]);
+					LoadedChunks.RemoveAt(ChunkIndex);
+					break;
+				}
 			}
 		}
 	}
-#else
-	TArray<uint32> FreeChunkIndices;
-	for (auto Index : IndicesToFree)
-	{
-		for (int32 ChunkIndex = 0; ChunkIndex < LoadedChunks.Num(); ++ChunkIndex)
+	else
+	{		
+		for (auto Index : IndicesToFree)
 		{
-			if (LoadedChunks[ChunkIndex].Index == Index)
+			for (int32 ChunkIndex = 0; ChunkIndex < LoadedChunks.Num(); ++ChunkIndex)
 			{
-				check(!FreeChunkIndices.Contains(ChunkIndex));
-				FreeChunkIndices.Add(ChunkIndex);
-				break;
+				if (LoadedChunks[ChunkIndex].Index == Index)
+				{
+					check(!FreeChunkIndices.Contains(ChunkIndex));
+					FreeChunkIndices.Add(ChunkIndex);
+					break;
+				}
 			}
 		}
 	}
-#endif
 
 	if (IndicesToLoad.Num() > 0)
 	{
@@ -319,19 +318,21 @@ void FStreamingWaveData::BeginPendingRequests(const TArray<uint32>& IndicesToLoa
 
 			FLoadedAudioChunk* ChunkStorage = NULL;
 
-#if !USE_NEW_ASYNC_IO
-			for (auto FreeIndex : FreeChunkIndices)
+			if (!GNewAsyncIO)
 			{
-				if (LoadedChunks[FreeIndex].MemorySize >= ChunkSize)
+				for (auto FreeIndex : FreeChunkIndices)
 				{
-					FreeChunkIndices.Remove(FreeIndex);
-					ChunkStorage = &LoadedChunks[FreeIndex];
-					ChunkStorage->DataSize = ChunkSize;
-					ChunkStorage->Index = Index;
-					break;
+					if (LoadedChunks[FreeIndex].MemorySize >= ChunkSize)
+					{
+						FreeChunkIndices.Remove(FreeIndex);
+						ChunkStorage = &LoadedChunks[FreeIndex];
+						ChunkStorage->DataSize = ChunkSize;
+						ChunkStorage->Index = Index;
+						break;
+					}
 				}
 			}
-#endif
+
 			if (ChunkStorage == NULL)
 			{
 				ChunkStorage = AddNewLoadedChunk(ChunkSize);
@@ -362,13 +363,9 @@ void FStreamingWaveData::BeginPendingRequests(const TArray<uint32>& IndicesToLoa
 #endif // #if WITH_EDITORONLY_DATA
 			{
 				check(Chunk.BulkData.GetFilename().Len());
-#if USE_NEW_ASYNC_IO
-				if (Chunk.BulkData.IsStoredCompressedOnDisk())
-				{
-					check(!"USE_NEW_ASYNC_IO does not support compression at the package level.");
-				}
-				else
+				if (GNewAsyncIO)
 				{
+					UE_CLOG(Chunk.BulkData.IsStoredCompressedOnDisk(), LogAudio, Fatal, TEXT("GNewAsyncIO does not support compression at the package level."));
 					check(!ChunkStorage->IORequest);
 					if (!IORequestHandle)
 					{
@@ -382,38 +379,38 @@ void FStreamingWaveData::BeginPendingRequests(const TArray<uint32>& IndicesToLoa
 						// we failed for some reason; file not found I guess.
 						PendingChunkChangeRequestStatus.Decrement();
 					}
-
 				}
-#else
-				if (Chunk.BulkData.IsStoredCompressedOnDisk())
+				else
 				{
-					IORequestIndices.AddUnique(FIOSystem::Get().LoadCompressedData(
-						Chunk.BulkData.GetFilename(),						// filename
-						Chunk.BulkData.GetBulkDataOffsetInFile(),			// offset
-						Chunk.BulkData.GetBulkDataSizeOnDisk(),				// compressed size
-						Chunk.BulkData.GetBulkDataSize(),					// uncompressed size
-						ChunkStorage->Data,									// dest pointer
-						Chunk.BulkData.GetDecompressionFlags(),				// compressed data format
-						&PendingChunkChangeRequestStatus,					// counter to decrement
-						AsyncIOPriority										// priority
+					if (Chunk.BulkData.IsStoredCompressedOnDisk())
+					{
+						IORequestIndices.AddUnique(FIOSystem::Get().LoadCompressedData(
+							Chunk.BulkData.GetFilename(),						// filename
+							Chunk.BulkData.GetBulkDataOffsetInFile(),			// offset
+							Chunk.BulkData.GetBulkDataSizeOnDisk(),				// compressed size
+							Chunk.BulkData.GetBulkDataSize(),					// uncompressed size
+							ChunkStorage->Data,									// dest pointer
+							Chunk.BulkData.GetDecompressionFlags(),				// compressed data format
+							&PendingChunkChangeRequestStatus,					// counter to decrement
+							AsyncIOPriority										// priority
 						)
 						);
-				}
-				// Load async.
-				else
-				{
-					IORequestIndices.AddUnique(FIOSystem::Get().LoadData(
-						Chunk.BulkData.GetFilename(),						// filename
-						Chunk.BulkData.GetBulkDataOffsetInFile(),			// offset
-						Chunk.BulkData.GetBulkDataSize(),					// size
-						ChunkStorage->Data,									// dest pointer
-						&PendingChunkChangeRequestStatus,					// counter to decrement
-						AsyncIOPriority										// priority
+					}
+					// Load async.
+					else
+					{
+						IORequestIndices.AddUnique(FIOSystem::Get().LoadData(
+							Chunk.BulkData.GetFilename(),						// filename
+							Chunk.BulkData.GetBulkDataOffsetInFile(),			// offset
+							Chunk.BulkData.GetBulkDataSize(),					// size
+							ChunkStorage->Data,									// dest pointer
+							&PendingChunkChangeRequestStatus,					// counter to decrement
+							AsyncIOPriority										// priority
 						)
 						);
+					}
+					check(IORequestIndices[IORequestIndices.Num() - 1]);
 				}
-				check(IORequestIndices[IORequestIndices.Num() - 1]);
-#endif
 			}
 		}
 
@@ -426,22 +423,23 @@ void FStreamingWaveData::BeginPendingRequests(const TArray<uint32>& IndicesToLoa
 		PendingChunkChangeRequestStatus.Set(AudioState_ReadyFor_Finalization);
 	}
 
-#if !USE_NEW_ASYNC_IO
-	// Ensure indices are in order so we can step through backwards
-	FreeChunkIndices.Sort();
-
-	for (int32 FreeIndex = FreeChunkIndices.Num() - 1; FreeIndex >= 0; --FreeIndex)
+	if (!GNewAsyncIO)
 	{
-		FreeLoadedChunk(LoadedChunks[FreeChunkIndices[FreeIndex]]);
-		LoadedChunks.RemoveAt(FreeChunkIndices[FreeIndex]);
+		// Ensure indices are in order so we can step through backwards
+		FreeChunkIndices.Sort();
+
+		for (int32 FreeIndex = FreeChunkIndices.Num() - 1; FreeIndex >= 0; --FreeIndex)
+		{
+			FreeLoadedChunk(LoadedChunks[FreeChunkIndices[FreeIndex]]);
+			LoadedChunks.RemoveAt(FreeChunkIndices[FreeIndex]);
+		}
 	}
-#endif
 }
 
-#if USE_NEW_ASYNC_IO
-
 bool FStreamingWaveData::BlockTillAllRequestsFinished(float TimeLimit)
 {
+	check(GNewAsyncIO);
+
 	QUICK_SCOPE_CYCLE_COUNTER(FStreamingWaveData_BlockTillAllRequestsFinished);
 	if (TimeLimit == 0.0f)
 	{
@@ -475,7 +473,6 @@ bool FStreamingWaveData::BlockTillAllRequestsFinished(float TimeLimit)
 	}
 	return true;
 }
-#endif
 
 #if WITH_EDITORONLY_DATA
 bool FStreamingWaveData::FinishDDCRequests()
@@ -499,12 +496,15 @@ FLoadedAudioChunk* FStreamingWaveData::AddNewLoadedChunk(int32 ChunkSize)
 {
 	int32 NewIndex = LoadedChunks.Num();
 	LoadedChunks.AddDefaulted();
-#if !USE_NEW_ASYNC_IO
-	LoadedChunks[NewIndex].MemorySize = ChunkSize;
-	LoadedChunks[NewIndex].Data = static_cast<uint8*>(FMemory::Malloc(ChunkSize));
-	INC_DWORD_STAT_BY(STAT_AudioMemorySize, ChunkSize);
-	INC_DWORD_STAT_BY(STAT_AudioMemory, ChunkSize);
-#endif
+
+	if (!GNewAsyncIO)
+	{
+		LoadedChunks[NewIndex].MemorySize = ChunkSize;
+		LoadedChunks[NewIndex].Data = static_cast<uint8*>(FMemory::Malloc(ChunkSize));
+		INC_DWORD_STAT_BY(STAT_AudioMemorySize, ChunkSize);
+		INC_DWORD_STAT_BY(STAT_AudioMemory, ChunkSize);
+	}
+
 	LoadedChunks[NewIndex].DataSize = ChunkSize;
 
 	return &LoadedChunks[NewIndex];
@@ -512,34 +512,33 @@ FLoadedAudioChunk* FStreamingWaveData::AddNewLoadedChunk(int32 ChunkSize)
 
 void FStreamingWaveData::FreeLoadedChunk(FLoadedAudioChunk& LoadedChunk)
 {
-#if USE_NEW_ASYNC_IO
-	if (LoadedChunk.IORequest)
+	if (GNewAsyncIO && LoadedChunk.IORequest)
 	{
 		LoadedChunk.IORequest->Cancel();
 		LoadedChunk.IORequest->WaitCompletion();
 		delete LoadedChunk.IORequest;
 		LoadedChunk.IORequest = nullptr;
 	}
-#endif
+
 	if (LoadedChunk.Data != NULL)
 	{
 		FMemory::Free(LoadedChunk.Data);
 
 		// Stat housekeeping
-#if USE_NEW_ASYNC_IO
-		DEC_DWORD_STAT_BY(STAT_AudioMemorySize, LoadedChunk.DataSize);
-		DEC_DWORD_STAT_BY(STAT_AudioMemory, LoadedChunk.DataSize);
-#else
-		DEC_DWORD_STAT_BY(STAT_AudioMemorySize, LoadedChunk.MemorySize);
-		DEC_DWORD_STAT_BY(STAT_AudioMemory, LoadedChunk.MemorySize);
-#endif
-
+		if (GNewAsyncIO)
+		{
+			DEC_DWORD_STAT_BY(STAT_AudioMemorySize, LoadedChunk.DataSize);
+			DEC_DWORD_STAT_BY(STAT_AudioMemory, LoadedChunk.DataSize);
+		}
+		else
+		{
+			DEC_DWORD_STAT_BY(STAT_AudioMemorySize, LoadedChunk.MemorySize);
+			DEC_DWORD_STAT_BY(STAT_AudioMemory, LoadedChunk.MemorySize);
+		}
 	}
 	LoadedChunk.Data = NULL;
 	LoadedChunk.DataSize = 0;
-#if !USE_NEW_ASYNC_IO
 	LoadedChunk.MemorySize = 0;
-#endif
 	LoadedChunk.Index = 0;
 }
 
@@ -610,7 +609,7 @@ void FAudioStreamingManager::UpdateResourceStreaming(float DeltaTime, bool bProc
 
 int32 FAudioStreamingManager::BlockTillAllRequestsFinished(float TimeLimit, bool)
 {
-#if USE_NEW_ASYNC_IO
+	if (GNewAsyncIO)
 	{
 		QUICK_SCOPE_CYCLE_COUNTER(FAudioStreamingManager_BlockTillAllRequestsFinished);
 		if (TimeLimit == 0.0f)
@@ -635,7 +634,6 @@ int32 FAudioStreamingManager::BlockTillAllRequestsFinished(float TimeLimit, bool
 		}
 	}
 
-#endif
 	// Not sure yet whether this will work the same as textures - aside from just before destroying
 	return 0;
 }
diff --git a/Engine/Source/Runtime/Engine/Private/BlueprintGeneratedClass.cpp b/Engine/Source/Runtime/Engine/Private/BlueprintGeneratedClass.cpp
--- a/Engine/Source/Runtime/Engine/Private/BlueprintGeneratedClass.cpp
+++ b/Engine/Source/Runtime/Engine/Private/BlueprintGeneratedClass.cpp
@@ -1052,10 +1052,6 @@ void UBlueprintGeneratedClass::GetPreloadDependencies(TArray<UObject*>& OutDeps)
 	{
 		OutDeps.Add(UberGraphFunction);
 	}
-	if (UInheritableComponentHandler* Handler = GetInheritableComponentHandler())
-	{
-		OutDeps.Add(Handler);
-	}
 	
 	UObject *CDO = GetDefaultObject();
 	if (CDO)
diff --git a/Engine/Source/Runtime/Engine/Private/CoreSettings.cpp b/Engine/Source/Runtime/Engine/Private/CoreSettings.cpp
--- a/Engine/Source/Runtime/Engine/Private/CoreSettings.cpp
+++ b/Engine/Source/Runtime/Engine/Private/CoreSettings.cpp
@@ -100,12 +100,6 @@ void UStreamingSettings::PostInitProperties()
 	if (IsTemplate())
 	{
 		ImportConsoleVariableValues();
-
-		// EDL can only be enabled in source code builds
-		if (FApp::IsEngineInstalled())
-		{
-			EventDrivenLoaderEnabled = false;
-		}
 	}
 #endif // #if WITH_EDITOR
 }
@@ -116,16 +110,6 @@ void UStreamingSettings::PostEditChangeProperty(FPropertyChangedEvent& PropertyC
 	static FName NAME_EventDrivenLoaderEnabled(TEXT("EventDrivenLoaderEnabled"));
 
 	Super::PostEditChangeProperty(PropertyChangedEvent);
-
-	if (Cast<UBoolProperty>(PropertyChangedEvent.Property) && PropertyChangedEvent.Property->GetFName() == NAME_EventDrivenLoaderEnabled)
-	{
-		if (FApp::IsEngineInstalled())
-		{
-			UBoolProperty* EventDrivenLoaderEnabledProperty = CastChecked<UBoolProperty>(PropertyChangedEvent.Property);
-			EventDrivenLoaderEnabledProperty->SetPropertyValue_InContainer(this, false);
-			UE_LOG(LogCoreSettings, Warning, TEXT("Event Driven Loader can only be enabled in source code distributions."))
-		}
-	}
 	
 	if (PropertyChangedEvent.Property)
 	{
diff --git a/Engine/Source/Runtime/Engine/Private/CurveTable.cpp b/Engine/Source/Runtime/Engine/Private/CurveTable.cpp
--- a/Engine/Source/Runtime/Engine/Private/CurveTable.cpp
+++ b/Engine/Source/Runtime/Engine/Private/CurveTable.cpp
@@ -465,7 +465,7 @@ TArray<FString> UCurveTable::CreateTableFromJSONString(const FString& InString,
 
 			// Make sure we have a valid float key
 			float EntryKey = 0.0f;
-			if (!LexicalConversion::TryParseString(EntryKey, *ParsedTableRowEntry.Key))
+			if (!Lex::TryParseString(EntryKey, *ParsedTableRowEntry.Key))
 			{
 				OutProblems.Add(FString::Printf(TEXT("Key '%s' on row '%s' is not a float and cannot be parsed."), *ParsedTableRowEntry.Key, *RowName.ToString()));
 				continue;
diff --git a/Engine/Source/Runtime/Engine/Private/EditorFramework/AssetImportData.cpp b/Engine/Source/Runtime/Engine/Private/EditorFramework/AssetImportData.cpp
--- a/Engine/Source/Runtime/Engine/Private/EditorFramework/AssetImportData.cpp
+++ b/Engine/Source/Runtime/Engine/Private/EditorFramework/AssetImportData.cpp
@@ -30,7 +30,7 @@ FString FAssetImportInfo::ToJson() const
 		Json += FString::Printf(TEXT("{ \"RelativeFilename\" : \"%s\", \"Timestamp\" : \"%d\", \"FileMD5\" : \"%s\" }"),
 			*SourceFiles[Index].RelativeFilename,
 			SourceFiles[Index].Timestamp.ToUnixTimestamp(),
-			*LexicalConversion::ToString(SourceFiles[Index].FileHash)
+			*Lex::ToString(SourceFiles[Index].FileHash)
 			);
 
 		if (Index != SourceFiles.Num() - 1)
@@ -75,10 +75,10 @@ TOptional<FAssetImportInfo> FAssetImportInfo::FromJson(FString InJsonString)
 		}
 
 		int64 UnixTimestamp = 0;
-		LexicalConversion::FromString(UnixTimestamp, *TimestampString);
+		Lex::FromString(UnixTimestamp, *TimestampString);
 
 		FMD5Hash FileHash;
-		LexicalConversion::FromString(FileHash, *MD5String);
+		Lex::FromString(FileHash, *MD5String);
 
 		Info.SourceFiles.Emplace(MoveTemp(RelativeFilename), FDateTime::FromUnixTimestamp(UnixTimestamp), FileHash);
 	}
diff --git a/Engine/Source/Runtime/Engine/Private/EngineSessionManager.cpp b/Engine/Source/Runtime/Engine/Private/EngineSessionManager.cpp
--- a/Engine/Source/Runtime/Engine/Private/EngineSessionManager.cpp
+++ b/Engine/Source/Runtime/Engine/Private/EngineSessionManager.cpp
@@ -65,13 +65,13 @@ namespace
 {
 	FString TimestampToString(FDateTime InTimestamp)
 	{
-		return LexicalConversion::ToString(InTimestamp.ToUnixTimestamp());
+		return Lex::ToString(InTimestamp.ToUnixTimestamp());
 	}
 
 	FDateTime StringToTimestamp(FString InString)
 	{
 		int64 TimestampUnix;
-		if (LexicalConversion::TryParseString(TimestampUnix, *InString))
+		if (Lex::TryParseString(TimestampUnix, *InString))
 		{
 			return FDateTime::FromUnixTimestamp(TimestampUnix);
 		}
diff --git a/Engine/Source/Runtime/Engine/Private/Particles/ParticleComponents.cpp b/Engine/Source/Runtime/Engine/Private/Particles/ParticleComponents.cpp
--- a/Engine/Source/Runtime/Engine/Private/Particles/ParticleComponents.cpp
+++ b/Engine/Source/Runtime/Engine/Private/Particles/ParticleComponents.cpp
@@ -2434,11 +2434,11 @@ void UParticleSystem::GetAssetRegistryTags(TArray<FAssetRegistryTag>& OutTags) c
 	const float BoundsSize = FixedRelativeBoundingBox.GetSize().GetMax();
 	OutTags.Add(FAssetRegistryTag("FixedBoundsSize", bUseFixedRelativeBoundingBox ? FString::Printf(TEXT("%.2f"), BoundsSize) : FString(TEXT("None")), FAssetRegistryTag::TT_Numerical));
 
-	OutTags.Add(FAssetRegistryTag("NumEmitters", LexicalConversion::ToString(Emitters.Num()), FAssetRegistryTag::TT_Numerical));
+	OutTags.Add(FAssetRegistryTag("NumEmitters", Lex::ToString(Emitters.Num()), FAssetRegistryTag::TT_Numerical));
 
-	OutTags.Add(FAssetRegistryTag("NumLODs", LexicalConversion::ToString(LODDistances.Num()), FAssetRegistryTag::TT_Numerical));
+	OutTags.Add(FAssetRegistryTag("NumLODs", Lex::ToString(LODDistances.Num()), FAssetRegistryTag::TT_Numerical));
 
-	OutTags.Add(FAssetRegistryTag("WarmupTime", LexicalConversion::ToString(WarmupTime), FAssetRegistryTag::TT_Numerical));
+	OutTags.Add(FAssetRegistryTag("WarmupTime", Lex::ToString(WarmupTime), FAssetRegistryTag::TT_Numerical));
 
 	// Done here instead of as an AssetRegistrySearchable string to avoid the long prefix on the enum value string
 	FString LODMethodString = TEXT("Unknown");
@@ -2473,10 +2473,10 @@ void UParticleSystem::GetAssetRegistryTags(TArray<FAssetRegistryTag>& OutTags) c
 			++NumEmittersAtEachSig[(int32)Emitter->SignificanceLevel];			
 		}
 	}
-	OutTags.Add(FAssetRegistryTag("Critical Emitters", LexicalConversion::ToString(NumEmittersAtEachSig[(int32)EParticleSignificanceLevel::Critical]), FAssetRegistryTag::TT_Numerical));
-	OutTags.Add(FAssetRegistryTag("High Emitters", LexicalConversion::ToString(NumEmittersAtEachSig[(int32)EParticleSignificanceLevel::High]), FAssetRegistryTag::TT_Numerical));
-	OutTags.Add(FAssetRegistryTag("Medium Emitters", LexicalConversion::ToString(NumEmittersAtEachSig[(int32)EParticleSignificanceLevel::Medium]), FAssetRegistryTag::TT_Numerical));
-	OutTags.Add(FAssetRegistryTag("Low Emitters", LexicalConversion::ToString(NumEmittersAtEachSig[(int32)EParticleSignificanceLevel::Low]), FAssetRegistryTag::TT_Numerical));
+	OutTags.Add(FAssetRegistryTag("Critical Emitters", Lex::ToString(NumEmittersAtEachSig[(int32)EParticleSignificanceLevel::Critical]), FAssetRegistryTag::TT_Numerical));
+	OutTags.Add(FAssetRegistryTag("High Emitters", Lex::ToString(NumEmittersAtEachSig[(int32)EParticleSignificanceLevel::High]), FAssetRegistryTag::TT_Numerical));
+	OutTags.Add(FAssetRegistryTag("Medium Emitters", Lex::ToString(NumEmittersAtEachSig[(int32)EParticleSignificanceLevel::Medium]), FAssetRegistryTag::TT_Numerical));
+	OutTags.Add(FAssetRegistryTag("Low Emitters", Lex::ToString(NumEmittersAtEachSig[(int32)EParticleSignificanceLevel::Low]), FAssetRegistryTag::TT_Numerical));
 
 	Super::GetAssetRegistryTags(OutTags);
 }
