Copying //UE4/Dev-Editor to //UE4/Dev-Main (Source: //UE4/Dev-Editor @ 3228984)

diff --git a/Engine/Source/Editor/UnrealEd/Private/UnrealEdEngine.cpp b/Engine/Source/Editor/UnrealEd/Private/UnrealEdEngine.cpp
--- a/Engine/Source/Editor/UnrealEd/Private/UnrealEdEngine.cpp
+++ b/Engine/Source/Editor/UnrealEd/Private/UnrealEdEngine.cpp
@@ -535,6 +535,11 @@ void UUnrealEdEngine::AttemptModifiedPackageNotification()
 {
 	bool bIsCooking = CookServer && CookServer->IsCookingInEditor() && CookServer->IsCookByTheBookRunning();
 
+	if (bShowPackageNotification && !bIsCooking)
+	{
+		ShowPackageNotification();
+	}
+
 	if (PackagesDirtiedThisTick.Num() > 0 && !bIsCooking)
 	{
 		// Force source control state to be updated
@@ -561,8 +566,6 @@ void UUnrealEdEngine::OnSourceControlStateUpdated(const FSourceControlOperationR
 {
 	if (ResultType == ECommandResult::Succeeded)
 	{
-		bool bShowNotification = false;
-
 		// Get the source control state of the package
 		ISourceControlProvider& SourceControlProvider = ISourceControlModule::Get().GetProvider();
 
@@ -589,13 +592,13 @@ void UUnrealEdEngine::OnSourceControlStateUpdated(const FSourceControlOperationR
 						else
 						{
 							PackageToNotifyState.Add(PackagePtr, NS_PendingPrompt);
-							bShowNotification = true;
+							bShowPackageNotification = true;
 						}
 					}
 					else if (!SourceControlState->IsCurrent() || SourceControlState->IsCheckedOutOther())
 					{
 						PackageToNotifyState.Add(PackagePtr, NS_PendingWarning);
-						bShowNotification = true;
+						bShowPackageNotification = true;
 					}
 				}
 			}
@@ -605,11 +608,6 @@ void UUnrealEdEngine::OnSourceControlStateUpdated(const FSourceControlOperationR
 		{
 			SourceControlProvider.Execute(ISourceControlOperation::Create<FCheckOut>(), SourceControlHelpers::AbsoluteFilenames(FilesToAutomaticallyCheckOut), EConcurrency::Asynchronous, FSourceControlOperationComplete::CreateUObject(this, &UUnrealEdEngine::OnPackagesCheckedOut, PackagesToAutomaticallyCheckOut));
 		}
-
-		if (bShowNotification)
-		{
-			ShowPackageNotification();
-		}
 	}
 }
 
@@ -643,8 +641,6 @@ void UUnrealEdEngine::OnPackagesCheckedOut(const FSourceControlOperationRef& Sou
 		{
 			PackageToNotifyState.Add(Package, NS_PendingPrompt);
 		}
-
-		ShowPackageNotification();
 	}
 }
 
diff --git a/Engine/Source/Editor/UnrealEd/Private/UnrealEdSrv.cpp b/Engine/Source/Editor/UnrealEd/Private/UnrealEdSrv.cpp
--- a/Engine/Source/Editor/UnrealEd/Private/UnrealEdSrv.cpp
+++ b/Engine/Source/Editor/UnrealEd/Private/UnrealEdSrv.cpp
@@ -1364,6 +1364,7 @@ void UUnrealEdEngine::ShowPackageNotification()
 
 		if( bCanPrompt )
 		{
+			bShowPackageNotification = false;
 			bool bNeedWarningDialog = false;
 			for (const auto& Entry : PackageToNotifyState)
 			{
diff --git a/Engine/Source/Editor/UnrealEd/Private/UnrealWidget.cpp b/Engine/Source/Editor/UnrealEd/Private/UnrealWidget.cpp
--- a/Engine/Source/Editor/UnrealEd/Private/UnrealWidget.cpp
+++ b/Engine/Source/Editor/UnrealEd/Private/UnrealWidget.cpp
@@ -1303,7 +1303,9 @@ void FWidget::AbsoluteTranslationConvertMouseMovementToAxisMovement(FSceneView*
 				Params.XAxis = InView->ViewMatrices.GetViewMatrix().GetColumn(0);
 				Params.YAxis = InView->ViewMatrices.GetViewMatrix().GetColumn(1);
 				Params.ZAxis = InView->ViewMatrices.GetViewMatrix().GetColumn(2);
-				GetPlaneNormalAndMask(Params.ZAxis, Params.PlaneNormal, Params.NormalToRemove); break;
+				GetPlaneNormalAndMask(Params.ZAxis, Params.PlaneNormal, Params.NormalToRemove);
+				//do not damp the movement in this case, we also want to snap
+				Params.bMovementLockedToCamera = false;
 				break;
 			}
 
diff --git a/Engine/Source/Editor/UnrealEd/Public/AssetViewerSettings.h b/Engine/Source/Editor/UnrealEd/Public/AssetViewerSettings.h
--- a/Engine/Source/Editor/UnrealEd/Public/AssetViewerSettings.h
+++ b/Engine/Source/Editor/UnrealEd/Public/AssetViewerSettings.h
@@ -9,7 +9,7 @@
 #include "UObject/ObjectRedirector.h"
 #include "Engine/Scene.h"
 #include "Engine/TextureCube.h"
-
+#include "EditorUndoClient.h"
 #include "AssetViewerSettings.generated.h"
 
 /**
@@ -118,11 +118,13 @@ struct FPreviewSceneProfile
 * Default asset viewer settings.
 */
 UCLASS(config=Editor, defaultconfig, meta=(DisplayName = "Asset Viewer"))
-class UNREALED_API UAssetViewerSettings : public UObject
+class UNREALED_API UAssetViewerSettings : public UObject, public FEditorUndoClient
 {
 	GENERATED_BODY()
 public:
 	UAssetViewerSettings();
+	virtual ~UAssetViewerSettings();
+
 	static UAssetViewerSettings* Get();
 	
 	/** Saves the config data out to the ini files */
@@ -134,11 +136,19 @@ public:
 	DECLARE_EVENT(UAssetViewerSettings, FOnAssetViewerProfileAddRemovedEvent);
 	FOnAssetViewerProfileAddRemovedEvent& OnAssetViewerProfileAddRemoved() { return OnAssetViewerProfileAddRemovedEvent; }
 
+	DECLARE_EVENT(UAssetViewerSettings, FOnAssetViewerSettingsPostUndo);
+	FOnAssetViewerSettingsPostUndo& OnAssetViewerSettingsPostUndo() { return OnAssetViewerSettingsPostUndoEvent; }
+
 	/** Begin UObject */
 	virtual void PostEditChangeProperty(struct FPropertyChangedEvent& PropertyChangedEvent) override;
 	virtual void PostInitProperties() override;
 	/** End UObject */
 
+	/** Begin FEditorUndoClient */
+	virtual void PostUndo(bool bSuccess) override;
+	virtual void PostRedo(bool bSuccess) override;
+	/** End FEditorUndoClient */
+
 	/** Collection of scene profiles */
 	UPROPERTY(EditAnywhere, config, Category = Settings, meta=(ShowOnlyInnerProperties))
 	TArray<FPreviewSceneProfile> Profiles;
@@ -150,4 +160,6 @@ protected:
 	FOnAssetViewerSettingsChangedEvent OnAssetViewerSettingsChangedEvent;
 
 	FOnAssetViewerProfileAddRemovedEvent OnAssetViewerProfileAddRemovedEvent;
+
+	FOnAssetViewerSettingsPostUndo OnAssetViewerSettingsPostUndoEvent;
 };
diff --git a/Engine/Source/Editor/UnrealEd/Public/AutomatedAssetImportData.h b/Engine/Source/Editor/UnrealEd/Public/AutomatedAssetImportData.h
--- a/Engine/Source/Editor/UnrealEd/Public/AutomatedAssetImportData.h
+++ b/Engine/Source/Editor/UnrealEd/Public/AutomatedAssetImportData.h
@@ -26,7 +26,7 @@ public:
 	bool IsValid() const;
 
 	/** Initalizes the group */
-	void Initialize();
+	void Initialize(TSharedPtr<FJsonObject> InImportGroupJsonData);
 
 	/** @return the display name of the group */
 	FString GetDisplayName() const; 
@@ -59,4 +59,7 @@ public:
 	UPROPERTY()
 	UFactory* Factory;
 
+	/** Json data to be read when importing this group */
+	TSharedPtr<FJsonObject> ImportGroupJsonData;
+
 };
diff --git a/Engine/Source/Editor/UnrealEd/Public/ComponentVisualizerManager.h b/Engine/Source/Editor/UnrealEd/Public/ComponentVisualizerManager.h
--- a/Engine/Source/Editor/UnrealEd/Public/ComponentVisualizerManager.h
+++ b/Engine/Source/Editor/UnrealEd/Public/ComponentVisualizerManager.h
@@ -15,7 +15,7 @@ struct FViewportClick;
 class UNREALED_API FComponentVisualizerManager
 {
 public:
-	FComponentVisualizerManager() {}
+	FComponentVisualizerManager();
 	virtual ~FComponentVisualizerManager() {}
 
 
@@ -52,4 +52,7 @@ public:
 private:
 	/** Currently 'active' visualizer that we should pass input to etc */
 	TWeakPtr<class FComponentVisualizer> EditedVisualizerPtr;
-};
+
+	/** The viewport client for the currently active visualizer */
+	FEditorViewportClient* EditedVisualizerViewportClient;
+};
\ No newline at end of file
diff --git a/Engine/Source/Editor/UnrealEd/Public/CookerSettings.h b/Engine/Source/Editor/UnrealEd/Public/CookerSettings.h
--- a/Engine/Source/Editor/UnrealEd/Public/CookerSettings.h
+++ b/Engine/Source/Editor/UnrealEd/Public/CookerSettings.h
@@ -56,7 +56,7 @@ public:
 	UPROPERTY(GlobalConfig, EditAnywhere, Category = Textures, meta = (DisplayName = "PVRTC Compression Quality (0-4, 0 is fastest)"))
 	int32 DefaultPVRTCQuality;
 
-	/** Quality of 0 means fastest, 3 means best quality */
+	/** Quality of 0 means fastest, 4 means best quality */
 	UPROPERTY(GlobalConfig, EditAnywhere, Category = Textures, meta = (DisplayName = "ASTC Compression Quality vs Speed (0-4, 0 is fastest)"))
 	int32 DefaultASTCQualityBySpeed;
 
diff --git a/Engine/Source/Editor/UnrealEd/Public/EditorReimportHandler.h b/Engine/Source/Editor/UnrealEd/Public/EditorReimportHandler.h
--- a/Engine/Source/Editor/UnrealEd/Public/EditorReimportHandler.h
+++ b/Engine/Source/Editor/UnrealEd/Public/EditorReimportHandler.h
@@ -9,7 +9,7 @@
 class FReimportHandler;
 
 /** Reimport manager for package resources with associated source files on disk. */
-class FReimportManager
+class FReimportManager : FGCObject
 {
 public:
 	/**
@@ -107,7 +107,9 @@ public:
 
 	/** Opens a file dialog to request a new reimport path */
 	UNREALED_API void GetNewReimportPath(UObject* Obj, TArray<FString>& InOutFilenames);
-
+	
+	/** FGCObject interface */
+	virtual void AddReferencedObjects( FReferenceCollector& Collector ) override;
 private:
 	/** Reimport handlers registered with this manager */
 	TArray<FReimportHandler*> Handlers;
@@ -192,4 +194,8 @@ public:
 	 * Import handlers with higher priority values will take precedent over lower priorities.
 	 */
 	UNREALED_API virtual int32 GetPriority() const;
+
+	/** Returns the UFactory object associated with this reimport handler */
+	virtual const UObject* GetFactoryObject() const { return nullptr; }
+
 };
diff --git a/Engine/Source/Editor/UnrealEd/Public/EditorViewportClient.h b/Engine/Source/Editor/UnrealEd/Public/EditorViewportClient.h
--- a/Engine/Source/Editor/UnrealEd/Public/EditorViewportClient.h
+++ b/Engine/Source/Editor/UnrealEd/Public/EditorViewportClient.h
@@ -1513,6 +1513,14 @@ private:
 
 	/** If true, the viewport widget should be invalidated on the next tick (needed to ensure thread safety) */
 	bool bShouldInvalidateViewportWidget;
+
+	/*
+	 * When we drag the manipulator(no camera movement involve) in absolute we need a view to compute the delta
+	 * We want to use the view when the user start dragging and use the same view until the displacement is done.
+	 * Using the same view allow us to move the camera and not disrupt the movement continuity.
+	 */
+	FSceneView* DragStartView;
+	FSceneViewFamily *DragStartViewFamily;
 };
 
 
diff --git a/Engine/Source/Editor/UnrealEd/Public/Kismet2/BlueprintEditorUtils.h b/Engine/Source/Editor/UnrealEd/Public/Kismet2/BlueprintEditorUtils.h
--- a/Engine/Source/Editor/UnrealEd/Public/Kismet2/BlueprintEditorUtils.h
+++ b/Engine/Source/Editor/UnrealEd/Public/Kismet2/BlueprintEditorUtils.h
@@ -1216,6 +1216,14 @@ public:
 	 */
 	static void FindActorsThatReferenceActor( AActor* InActor, TArray<UClass*>& InClassesToIgnore, TArray<AActor*>& OutReferencingActors );
 
+	/**
+	 * Go through the world and build a map of all actors that are referenced by other actors.
+	 * @param InWorld The world to scan for Actors.
+	 * @param InClassesToIgnore  An array of class types to ignore, even if there is an instance of one that references another Actor
+	 * @param OutReferencingActors A map of Actors that are referenced by a list of other Actors.
+	*/
+	static void GetActorReferenceMap(UWorld* InWorld, TArray<UClass*>& InClassesToIgnore, TMap<AActor*, TArray<AActor*> >& OutReferencingActors);
+
 	//////////////////////////////////////////////////////////////////////////
 	// Diagnostics
 
diff --git a/Engine/Source/Editor/UnrealEd/Public/Kismet2/EnumEditorUtils.h b/Engine/Source/Editor/UnrealEd/Public/Kismet2/EnumEditorUtils.h
--- a/Engine/Source/Editor/UnrealEd/Public/Kismet2/EnumEditorUtils.h
+++ b/Engine/Source/Editor/UnrealEd/Public/Kismet2/EnumEditorUtils.h
@@ -78,10 +78,11 @@ public:
 	static int64 ResolveEnumerator(const UEnum* Enum, FArchive& Ar, int64 EnumeratorValue);
 
 	//DISPLAY NAME
-	static FString GetEnumeratorDisplayName(const UUserDefinedEnum* Enum, int32 EnumeratorIndex);
-	static bool SetEnumeratorDisplayName(UUserDefinedEnum* Enum, int32 EnumeratorIndex, FString NewDisplayName);
-	static bool IsEnumeratorDisplayNameValid(const UUserDefinedEnum* Enum, FString NewDisplayName);
+	static FText GetEnumeratorDisplayName(const UUserDefinedEnum* Enum, int32 EnumeratorIndex);
+	static bool SetEnumeratorDisplayName(UUserDefinedEnum* Enum, int32 EnumeratorIndex, FText NewDisplayName);
+	static bool IsEnumeratorDisplayNameValid(const UUserDefinedEnum* Enum, int32 EnumeratorIndex, FText NewDisplayName);
 	static void EnsureAllDisplayNamesExist(class UUserDefinedEnum* Enum);
+	static void UpgradeDisplayNamesFromMetaData(class UUserDefinedEnum* Enum);
 };
 
 #endif // __EnumEditorUtils_h__
diff --git a/Engine/Source/Editor/UnrealEd/Public/PackageTools.h b/Engine/Source/Editor/UnrealEd/Public/PackageTools.h
--- a/Engine/Source/Editor/UnrealEd/Public/PackageTools.h
+++ b/Engine/Source/Editor/UnrealEd/Public/PackageTools.h
@@ -69,6 +69,26 @@ namespace PackageTools
 	 */
 	UNREALED_API bool UnloadPackages( const TArray<UPackage*>& PackagesToUnload, FText& OutErrorMessage );
 
+	/**
+	 * Helper function that attempts to reload the specified top-level packages.
+	 *
+	 * @param	PackagesToReload	The list of packages that should be reloaded
+	 *
+	 * @return	true if the set of loaded packages was changed
+	 */
+	UNREALED_API bool ReloadPackages( const TArray<UPackage*>& PackagesToReload );
+
+	/**
+	 * Helper function that attempts to reload the specified top-level packages.
+	 *
+	 * @param	PackagesToReload	The list of packages that should be reloaded
+	 * @param	OutErrorMessage		An error message specifying any problems with reloading packages
+	 * @param	bInteractive		Whether the function is allowed to ask the user questions (such as whether to reload dirty packages)
+	 *
+	 * @return	true if the set of loaded packages was changed
+	 */
+	UNREALED_API bool ReloadPackages( const TArray<UPackage*>& PackagesToReload, FText& OutErrorMessage, const bool bInteractive = true );
+
 	/**
 	 *	Exports the given packages to files.
 	 *
diff --git a/Engine/Source/Editor/UnrealEd/Public/SAdvancedPreviewDetailsTab.h b/Engine/Source/Editor/UnrealEd/Public/SAdvancedPreviewDetailsTab.h
--- a/Engine/Source/Editor/UnrealEd/Public/SAdvancedPreviewDetailsTab.h
+++ b/Engine/Source/Editor/UnrealEd/Public/SAdvancedPreviewDetailsTab.h
@@ -39,6 +39,7 @@ protected:
 	FReply RemoveProfileButtonClick();
 protected:
 	void OnAssetViewerSettingsRefresh(const FName& InPropertyName);
+	void OnAssetViewerSettingsPostUndo();
 protected:
 	/** Property viewing widget */
 	TSharedPtr<IDetailsView> SettingsView;
@@ -51,6 +52,7 @@ protected:
 
 	FDelegateHandle RefreshDelegate;
 	FDelegateHandle AddRemoveProfileDelegate;
+	FDelegateHandle PostUndoDelegate;
 
 	UEditorPerProjectUserSettings* PerProjectSettings;
 };
diff --git a/Engine/Source/Editor/UnrealEd/Public/SCurveEditor.h b/Engine/Source/Editor/UnrealEd/Public/SCurveEditor.h
--- a/Engine/Source/Editor/UnrealEd/Public/SCurveEditor.h
+++ b/Engine/Source/Editor/UnrealEd/Public/SCurveEditor.h
@@ -589,6 +589,8 @@ private:
 
 	TSharedPtr<FCurveViewModel> GetViewModelForCurve(FRichCurve* InCurve);
 
+	void OnObjectPropertyChanged(UObject* Object, FPropertyChangedEvent& PropertyChangedEvent);
+
 protected:
 
 	/** Set Default output values when range is too small **/
diff --git a/Engine/Source/Editor/UnrealEd/Public/Settings/EditorSettings.h b/Engine/Source/Editor/UnrealEd/Public/Settings/EditorSettings.h
--- a/Engine/Source/Editor/UnrealEd/Public/Settings/EditorSettings.h
+++ b/Engine/Source/Editor/UnrealEd/Public/Settings/EditorSettings.h
@@ -15,13 +15,17 @@ class UNREALED_API UEditorSettings : public UObject
 	GENERATED_UCLASS_BODY()
 
 	/** When checked, the most recently loaded project will be auto-loaded at editor startup if no other project was specified on the command line */
-	UPROPERTY(EditAnywhere, Category=General)
+	UPROPERTY()
 	bool bLoadTheMostRecentlyLoadedProjectAtStartup; // Note that this property is NOT config since it is not necessary to save the value to ini. It is determined at startup in UEditorEngine::InitEditor().
 
 	/** Can the editor report usage analytics (types of assets being spawned, etc...) back to Epic in order for us to improve the editor user experience?  Note: The editor must be restarted for changes to take effect. */
 	UPROPERTY()
 	bool bEditorAnalyticsEnabled_DEPRECATED;
 
+	/** Whether the Content Browser should open the Sources Panel by default */
+	UPROPERTY(EditAnywhere, config, Category=ContentBrowser)
+	bool bOpenSourcesPanelByDefault;
+
 	// =====================================================================
 	// The following options are NOT exposed in the preferences Editor
 	// (usually because there is a different way to set them interactively!)
diff --git a/Engine/Source/Editor/UnrealEd/Public/Tests/FbxAutomationCommon.h b/Engine/Source/Editor/UnrealEd/Public/Tests/FbxAutomationCommon.h
--- a/Engine/Source/Editor/UnrealEd/Public/Tests/FbxAutomationCommon.h
+++ b/Engine/Source/Editor/UnrealEd/Public/Tests/FbxAutomationCommon.h
@@ -127,9 +127,9 @@ class UFbxTestPlan : public UObject
 	UPROPERTY(EditAnywhere, Category = ExpectedResult)
 	TArray< FFbxTestPlanExpectedResult > ExpectedResult;
 
-	/* Options use for this test plan. No UProperty since it will be serialize to json manually */
-	UPROPERTY(EditAnywhere, Instanced, Category = Options)
-	UFbxImportUI *ImportUI;
+	/* Options use for this test plan, Transient because we manually serialize the options. */
+	UPROPERTY(EditAnywhere, Transient, Instanced, Category = Options)
+	UFbxImportUI* ImportUI;
 };
 
 namespace FbxAutomationTestsAPI
diff --git a/Engine/Source/Editor/UnrealEd/Public/Toolkits/AssetEditorManager.h b/Engine/Source/Editor/UnrealEd/Public/Toolkits/AssetEditorManager.h
--- a/Engine/Source/Editor/UnrealEd/Public/Toolkits/AssetEditorManager.h
+++ b/Engine/Source/Editor/UnrealEd/Public/Toolkits/AssetEditorManager.h
@@ -141,6 +141,9 @@ private:
 	/** Restore the assets that were previously open when the editor was last closed */
 	void RestorePreviouslyOpenAssets();
 
+	/** Handles a package being reloaded */
+	void HandlePackageReloaded(const EPackageReloadPhase InPackageReloadPhase, FPackageReloadedEvent* InPackageReloadedEvent);
+
 private:
 
 	/** struct used by OpenedEditorTimes map to store editor names and times */
diff --git a/Engine/Source/Editor/UnrealEd/Public/UnrealEd.h b/Engine/Source/Editor/UnrealEd/Public/UnrealEd.h
--- a/Engine/Source/Editor/UnrealEd/Public/UnrealEd.h
+++ b/Engine/Source/Editor/UnrealEd/Public/UnrealEd.h
@@ -61,6 +61,7 @@ MONOLITHIC_HEADER_BOILERPLATE()
 #include "Classes/Exporters/AnimSequenceExporterFBX.h"
 #include "Classes/Commandlets/AudioTestCommandlet.h"
 #include "Classes/Preferences/CascadeOptions.h"
+#include "Classes/Settings/ClassViewerSettings.h"
 #include "Classes/Commandlets/CompressAnimationsCommandlet.h"
 #include "Classes/Settings/ContentBrowserSettings.h"
 #include "Classes/Commandlets/CookCommandlet.h"
diff --git a/Engine/Source/Editor/UnrealEd/UnrealEd.Build.cs b/Engine/Source/Editor/UnrealEd/UnrealEd.Build.cs
--- a/Engine/Source/Editor/UnrealEd/UnrealEd.Build.cs
+++ b/Engine/Source/Editor/UnrealEd/UnrealEd.Build.cs
@@ -140,7 +140,8 @@ public class UnrealEd : ModuleRules
                 "VectorVM",
 				"TreeMap",
                 "MaterialUtilities",
-                "LocalizationService",
+				"Localization",
+				"LocalizationService",
 				"AddContentDialog",
 				"GameProjectGeneration",
                 "HierarchicalLODUtilities",
diff --git a/Engine/Source/Programs/CrashReporter/CrashReportClient/Private/CrashDescription.h b/Engine/Source/Programs/CrashReporter/CrashReportClient/Private/CrashDescription.h
--- a/Engine/Source/Programs/CrashReporter/CrashReportClient/Private/CrashDescription.h
+++ b/Engine/Source/Programs/CrashReporter/CrashReportClient/Private/CrashDescription.h
@@ -352,7 +352,7 @@ struct FPrimaryCrashProperties
 	FCrashProperty CrashReporterMessage;
 
 	/**
-	 *	Windows only. Non-zero integrity values are to be discounted as "genuine" crashes.
+	 *	Platform-specific UE4 Core value (integer).
 	 */
 	FCrashProperty PlatformCallbackResult;
 
diff --git a/Engine/Source/Programs/CrashReporter/CrashReportClient/Private/SCrashReportClient.cpp b/Engine/Source/Programs/CrashReporter/CrashReportClient/Private/SCrashReportClient.cpp
--- a/Engine/Source/Programs/CrashReporter/CrashReportClient/Private/SCrashReportClient.cpp
+++ b/Engine/Source/Programs/CrashReporter/CrashReportClient/Private/SCrashReportClient.cpp
@@ -179,7 +179,7 @@ void SCrashReportClient::Construct(const FArguments& InArgs, TSharedRef<FCrashRe
 							.Font( FSlateFontInfo( FPaths::EngineContentDir() / TEXT( "Slate/Fonts/Roboto-Regular.ttf" ), 8 ) )
 							.AutoWrapText( false )
 							.IsReadOnly( true )
-							.ReadOnlyForegroundColor( FSlateColor( FLinearColor::Black ) )
+							.ReadOnlyForegroundColor( FSlateColor( FLinearColor::White * 0.8f) )
 							.BackgroundColor( FSlateColor( FLinearColor::Black ) )
 							.ForegroundColor( FSlateColor( FLinearColor::White * 0.8f ) )
 							.Text( Client, &FCrashReportClient::GetDiagnosticText )
diff --git a/Engine/Source/Programs/UnrealBuildTool/Windows/WindowsDeploy.cs b/Engine/Source/Programs/UnrealBuildTool/Windows/WindowsDeploy.cs
--- a/Engine/Source/Programs/UnrealBuildTool/Windows/WindowsDeploy.cs
+++ b/Engine/Source/Programs/UnrealBuildTool/Windows/WindowsDeploy.cs
@@ -18,10 +18,16 @@ namespace UnrealBuildTool
         public override bool PrepForUATPackageOrDeploy(FileReference ProjectFile, string ProjectName, string ProjectDirectory, string ExecutablePath, string EngineDirectory, bool bForDistribution, string CookFlavor, bool bIsDataDeploy)
         {
             string ApplicationIconPath = Path.Combine(ProjectDirectory, "Build/Windows/Application.ico");
-            // Does a Project icon exist, if not point to the default UE4 icon instead
+            // Does a Project icon exist?
             if (!File.Exists(ApplicationIconPath))
             {
-                ApplicationIconPath = Path.Combine(EngineDirectory, "Source/Runtime/Launch/Resources/Windows/UE4.ico");
+                // Also check for legacy location
+                ApplicationIconPath = Path.Combine(ProjectDirectory, "Source", ProjectName, "Resources", "Windows", ProjectName + ".ico");
+                if (!File.Exists(ApplicationIconPath))
+                {
+                    // point to the default UE4 icon instead
+                    ApplicationIconPath = Path.Combine(EngineDirectory, "Source/Runtime/Launch/Resources/Windows/UE4.ico");
+                }
             }
             // sets the icon on the original exe this will be used in the task bar when the bootstrap exe runs
             if (File.Exists(ApplicationIconPath))
diff --git a/Engine/Source/Runtime/AIModule/Private/BehaviorTree/BehaviorTreeTypes.cpp b/Engine/Source/Runtime/AIModule/Private/BehaviorTree/BehaviorTreeTypes.cpp
--- a/Engine/Source/Runtime/AIModule/Private/BehaviorTree/BehaviorTreeTypes.cpp
+++ b/Engine/Source/Runtime/AIModule/Private/BehaviorTree/BehaviorTreeTypes.cpp
@@ -329,33 +329,32 @@ void FBlackboardKeySelector::InitSelection(const UBlackboardData& BlackboardAsse
 
 void FBlackboardKeySelector::AddObjectFilter(UObject* Owner, FName PropertyName, TSubclassOf<UObject> AllowedClass)
 {
-	static int32 FilterCounter = 0;
-	const FString FilterName = FString::Printf(TEXT("%s_Object_%d"), *PropertyName.ToString(), ++FilterCounter);
-	UBlackboardKeyType_Object* FilterOb = NewObject<UBlackboardKeyType_Object>(Owner, *FilterName);
+	const FName FilterName = MakeUniqueObjectName(Owner, UBlackboardKeyType_Object::StaticClass(), *FString::Printf(TEXT("%s_Object"), *PropertyName.ToString()));
+	UBlackboardKeyType_Object* FilterOb = NewObject<UBlackboardKeyType_Object>(Owner, FilterName);
 	FilterOb->BaseClass = AllowedClass;
 	AllowedTypes.Add(FilterOb);
 }
 
 void FBlackboardKeySelector::AddClassFilter(UObject* Owner, FName PropertyName, TSubclassOf<UClass> AllowedClass)
 {
-	const FString FilterName = PropertyName.ToString() + TEXT("_Class");
-	UBlackboardKeyType_Class* FilterOb = NewObject<UBlackboardKeyType_Class>(Owner, *FilterName);
+	const FName FilterName = MakeUniqueObjectName(Owner, UBlackboardKeyType_Class::StaticClass(), *FString::Printf(TEXT("%_Class"), *PropertyName.ToString()));
+	UBlackboardKeyType_Class* FilterOb = NewObject<UBlackboardKeyType_Class>(Owner, FilterName);
 	FilterOb->BaseClass = AllowedClass;
 	AllowedTypes.Add(FilterOb);
 }
 
 void FBlackboardKeySelector::AddEnumFilter(UObject* Owner, FName PropertyName, UEnum* AllowedEnum)
 {
-	const FString FilterName = PropertyName.ToString() + TEXT("_Enum");
-	UBlackboardKeyType_Enum* FilterOb = NewObject<UBlackboardKeyType_Enum>(Owner, *FilterName);
+	const FName FilterName = MakeUniqueObjectName(Owner, UBlackboardKeyType_Enum::StaticClass(), *FString::Printf(TEXT("%_Enum"), *PropertyName.ToString()));
+	UBlackboardKeyType_Enum* FilterOb = NewObject<UBlackboardKeyType_Enum>(Owner, FilterName);
 	FilterOb->EnumType = AllowedEnum;
 	AllowedTypes.Add(FilterOb);
 }
 
 void FBlackboardKeySelector::AddNativeEnumFilter(UObject* Owner, FName PropertyName, const FString& AllowedEnumName)
 {
-	const FString FilterName = PropertyName.ToString() + TEXT("_NativeEnum");
-	UBlackboardKeyType_NativeEnum* FilterOb = NewObject<UBlackboardKeyType_NativeEnum>(Owner, *FilterName);
+	const FName FilterName = MakeUniqueObjectName(Owner, UBlackboardKeyType_NativeEnum::StaticClass(), *FString::Printf(TEXT("%_NativeEnum"), *PropertyName.ToString()));
+	UBlackboardKeyType_NativeEnum* FilterOb = NewObject<UBlackboardKeyType_NativeEnum>(Owner, FilterName);
 	FilterOb->EnumName = AllowedEnumName;
 	AllowedTypes.Add(FilterOb);
 }
diff --git a/Engine/Source/Runtime/AppFramework/Private/Widgets/Colors/SColorGradingPicker.cpp b/Engine/Source/Runtime/AppFramework/Private/Widgets/Colors/SColorGradingPicker.cpp
new file mode 100644
--- /dev/null
+++ b/Engine/Source/Runtime/AppFramework/Private/Widgets/Colors/SColorGradingPicker.cpp
@@ -0,0 +1,149 @@
+// Copyright 1998-2016 Epic Games, Inc. All Rights Reserved.
+#include "SColorGradingPicker.h"
+#include "SColorGradingWheel.h"
+#include "SNumericEntryBox.h"
+
+
+#define LOCTEXT_NAMESPACE "ColorGradingWheel"
+
+SColorGradingPicker::~SColorGradingPicker()
+{
+}
+
+void SColorGradingPicker::Construct( const FArguments& InArgs )
+{
+	ValuesMin = InArgs._ValuesMin;
+	ValuesMax = InArgs._ValuesMax;
+	MainDelta = InArgs._MainDelta;
+	MainShiftMouseMovePixelPerDelta = InArgs._MainShiftMouseMovePixelPerDelta;
+	ColorGradingModes = InArgs._ColorGradingModes;
+	OnColorCommitted = InArgs._OnColorCommitted;
+	OnQueryCurrentColor = InArgs._OnQueryCurrentColor;
+	
+	float ColorGradingWheelExponent = 2.4f;
+	if (ColorGradingModes == EColorGradingModes::Offset)
+	{
+		ColorGradingWheelExponent = 3.0f;
+	}
+
+	ChildSlot
+	[
+		SNew(SVerticalBox)
+		+ SVerticalBox::Slot()
+			.AutoHeight()
+			[
+				SNew(SBox)
+				.HAlign(HAlign_Center)
+				.VAlign(VAlign_Center)
+				.WidthOverride(125.0f)
+				.HeightOverride(125.0f)
+				.MinDesiredWidth(125.0f)
+				.MaxDesiredWidth(125.0f)
+				[
+					SNew(SColorGradingWheel)
+					.SelectedColor(TAttribute<FLinearColor>::Create(TAttribute<FLinearColor>::FGetter::CreateSP(this, &SColorGradingPicker::GetCurrentLinearColor)))
+					.DesiredWheelSize(125)
+					.ExponentDisplacement(ColorGradingWheelExponent)
+					.OnValueChanged(this, &SColorGradingPicker::HandleCurrentColorValueChanged)
+				]
+			]
+		+ SVerticalBox::Slot()
+			.AutoHeight()
+			[
+				SNew(SNumericEntryBox<float>)
+				.EditableTextBoxStyle(&FCoreStyle::Get().GetWidgetStyle<FEditableTextBoxStyle>("DarkEditableTextBox"))
+				.Value(this, &SColorGradingPicker::OnGetMainValue)
+				.OnValueChanged(this, &SColorGradingPicker::OnMainValueChanged)
+				.AllowSpin(true)
+				.MinValue(ValuesMin)
+				.MaxValue(ValuesMax)
+				.MinSliderValue(ValuesMin)
+				.MaxSliderValue(ValuesMax)
+				.Delta(MainDelta)
+				.ShiftMouseMovePixelPerDelta(MainShiftMouseMovePixelPerDelta)
+			]
+	];
+}
+
+void SColorGradingPicker::OnMainValueChanged(float InValue)
+{
+	TransformColorGradingRangeToLinearColorRange(InValue);
+
+	FVector4 CurrentValue(0.0f, 0.0f, 0.0f, 0.0f);
+	OnQueryCurrentColor.ExecuteIfBound(CurrentValue);
+	TransformColorGradingRangeToLinearColorRange(CurrentValue);
+	
+	
+	//The MainValue is the maximum of any channel value
+	float MaxCurrentValue = FMath::Max3(CurrentValue.X, CurrentValue.Y, CurrentValue.Z);
+	if (MaxCurrentValue <= 0.0f)
+	{
+		//We need the neutral value for the type of color grading, currently only offset is an addition(0.0) all other are multiplier(1.0)
+		CurrentValue = FVector4(InValue, InValue, InValue, CurrentValue.W);
+	}
+	else
+	{
+		float Ratio = InValue / MaxCurrentValue;
+		CurrentValue *= FVector4(Ratio, Ratio, Ratio, 1.0f);
+	}
+	TransformLinearColorRangeToColorGradingRange(CurrentValue);
+	OnColorCommitted.ExecuteIfBound(CurrentValue);
+}
+
+TOptional<float> SColorGradingPicker::OnGetMainValue() const
+{
+	FVector4 CurrentValue(0.0f, 0.0f, 0.0f, 0.0f);
+	OnQueryCurrentColor.ExecuteIfBound(CurrentValue);
+	//The MainValue is the maximum of any channel value
+	TOptional<float> CurrentValueOption = FMath::Max3(CurrentValue.X, CurrentValue.Y, CurrentValue.Z);
+	return CurrentValueOption;
+}
+
+void SColorGradingPicker::TransformLinearColorRangeToColorGradingRange(FVector4 &VectorValue) const
+{
+	float Ratio = (ValuesMax - ValuesMin);
+	VectorValue *= Ratio;
+	VectorValue += FVector4(ValuesMin, ValuesMin, ValuesMin, ValuesMin);
+}
+
+void SColorGradingPicker::TransformColorGradingRangeToLinearColorRange(FVector4 &VectorValue) const
+{
+	VectorValue += -FVector4(ValuesMin, ValuesMin, ValuesMin, ValuesMin);
+	float Ratio = 1.0f/(ValuesMax - ValuesMin);
+	VectorValue *= Ratio;
+}
+
+void SColorGradingPicker::TransformColorGradingRangeToLinearColorRange(float &FloatValue)
+{
+	FloatValue -= ValuesMin;
+	float Ratio = 1.0f/(ValuesMax - ValuesMin);
+	FloatValue *= Ratio;
+}
+
+FLinearColor SColorGradingPicker::GetCurrentLinearColor()
+{
+	FLinearColor CurrentColor;
+	FVector4 CurrentValue;
+	if (OnQueryCurrentColor.ExecuteIfBound(CurrentValue))
+	{
+		TransformColorGradingRangeToLinearColorRange(CurrentValue);
+		CurrentColor = FLinearColor(CurrentValue.X, CurrentValue.Y, CurrentValue.Z);
+	}
+	return CurrentColor.LinearRGBToHSV();
+}
+
+void SColorGradingPicker::HandleCurrentColorValueChanged( FLinearColor NewValue )
+{
+	//Query the current vector4 so we can pass back the w value
+	FVector4 CurrentValue(0.0f, 0.0f, 0.0f, 0.0f);
+	OnQueryCurrentColor.ExecuteIfBound(CurrentValue);
+
+	FLinearColor NewValueRGB = NewValue.HSVToLinearRGB();
+	FVector4 NewValueVector(NewValueRGB.R, NewValueRGB.G, NewValueRGB.B);
+	TransformLinearColorRangeToColorGradingRange(NewValueVector);
+	//Set the W with the original value
+	NewValueVector.W = CurrentValue.W;
+	OnColorCommitted.ExecuteIfBound(NewValueVector);
+}
+
+#undef LOCTEXT_NAMESPACE
diff --git a/Engine/Source/Runtime/AppFramework/Private/Widgets/Colors/SColorPicker.cpp b/Engine/Source/Runtime/AppFramework/Private/Widgets/Colors/SColorPicker.cpp
--- a/Engine/Source/Runtime/AppFramework/Private/Widgets/Colors/SColorPicker.cpp
+++ b/Engine/Source/Runtime/AppFramework/Private/Widgets/Colors/SColorPicker.cpp
@@ -549,6 +549,7 @@ void SColorPicker::GenerateDefaultColorPickerContent( bool bAdvancedSectionExpan
 
 										+ SHorizontalBox::Slot()
 											.AutoWidth()
+											.MaxWidth(72.0f)
 											[
 												SNew(SEditableTextBox)
 													.MinDesiredWidth(72.0f)
@@ -577,6 +578,7 @@ void SColorPicker::GenerateDefaultColorPickerContent( bool bAdvancedSectionExpan
 
 										+ SHorizontalBox::Slot()
 											.AutoWidth()
+											.MaxWidth(72.0f)
 											[
 												SNew(SEditableTextBox)
 												.MinDesiredWidth(72.0f)
@@ -981,7 +983,10 @@ TSharedRef<SWidget> SColorPicker::MakeColorSpinBox( EColorPickerChannels Channel
 	default:
 		return SNullWidget::NullWidget;
 	}
-	
+
+	// Define a maximum size for the spin box containers to prevent them from stretching out the color picker window
+	static const float MaxSpinBoxSize = 192.0f;
+
 	return SNew(SHorizontalBox)
 
 	+ SHorizontalBox::Slot()
@@ -995,6 +1000,7 @@ TSharedRef<SWidget> SColorPicker::MakeColorSpinBox( EColorPickerChannels Channel
 
 	+ SHorizontalBox::Slot()
 		.FillWidth(1.0f)
+		.MaxWidth(MaxSpinBoxSize)
 		[
 			SNew(SVerticalBox)
 				.ToolTipText(SliderTooltip)
diff --git a/Engine/Source/Runtime/AppFramework/Public/Widgets/Colors/SColorGradingPicker.h b/Engine/Source/Runtime/AppFramework/Public/Widgets/Colors/SColorGradingPicker.h
new file mode 100644
--- /dev/null
+++ b/Engine/Source/Runtime/AppFramework/Public/Widgets/Colors/SColorGradingPicker.h
@@ -0,0 +1,101 @@
+// Copyright 1998-2016 Epic Games, Inc. All Rights Reserved.
+
+#pragma once
+
+#include "CoreMinimal.h"
+#include "Misc/Attribute.h"
+#include "Input/Reply.h"
+#include "Widgets/DeclarativeSyntaxSupport.h"
+#include "Widgets/SLeafWidget.h"
+#include "Framework/SlateDelegates.h"
+
+/** Callback to get the current FVector4 value */
+DECLARE_DELEGATE_OneParam(FOnGetCurrentVector4Value, FVector4&)
+
+/**
+* Enumerates color picker modes.
+*/
+enum class EColorGradingModes
+{
+	Saturation,
+	Contrast,
+	Gamma,
+	Gain,
+	Offset
+};
+
+
+/**
+ * Class for placing a color picker. If all you need is a standalone color picker,
+ * use the functions OpenColorGradingWheel and DestroyColorGradingWheel, since they hold a static
+ * instance of the color picker.
+ */
+class APPFRAMEWORK_API SColorGradingPicker
+	: public SCompoundWidget
+{
+public:
+
+	SLATE_BEGIN_ARGS(SColorGradingPicker)
+		: _ValuesMin(0.0f)
+		, _ValuesMax(2.0f)
+		, _MainDelta(0.01f)
+		, _MainShiftMouseMovePixelPerDelta(10)
+		, _ColorGradingModes(EColorGradingModes::Saturation)
+		, _OnColorCommitted()
+		, _OnQueryCurrentColor()
+	{ }
+		
+		SLATE_ARGUMENT( float, ValuesMin )
+
+		SLATE_ARGUMENT( float, ValuesMax )
+
+		SLATE_ARGUMENT( float, MainDelta )
+
+		SLATE_ARGUMENT( int32, MainShiftMouseMovePixelPerDelta)
+
+		SLATE_ARGUMENT( EColorGradingModes, ColorGradingModes)
+
+		/** The event called when the color is committed */
+		SLATE_EVENT( FOnVector4ValueChanged, OnColorCommitted )
+
+			/** Callback to get the current FVector4 value */
+		SLATE_EVENT(FOnGetCurrentVector4Value, OnQueryCurrentColor)
+
+	SLATE_END_ARGS()
+
+	/**	Destructor. */
+	~SColorGradingPicker();
+
+public:
+
+	/**
+	 * Construct the widget
+	 *
+	 * @param InArgs Declaration from which to construct the widget.
+	 */
+	void Construct(const FArguments& InArgs );
+
+protected:
+
+	void TransformLinearColorRangeToColorGradingRange(FVector4 &VectorValue) const;
+	void TransformColorGradingRangeToLinearColorRange(FVector4 &VectorValue) const;
+	void TransformColorGradingRangeToLinearColorRange(float &FloatValue);
+
+	TOptional<float> OnGetMainValue() const;
+	void OnMainValueChanged(float InValue);
+	FLinearColor GetCurrentLinearColor();
+
+	// Callback for value changes in the color spectrum picker.
+	void HandleCurrentColorValueChanged(FLinearColor NewValue);
+
+	float ValuesMin;
+	float ValuesMax;
+	float MainDelta;
+	int32 MainShiftMouseMovePixelPerDelta;
+	EColorGradingModes ColorGradingModes;
+
+	/** Invoked when a new value is selected on the color wheel */
+	FOnVector4ValueChanged OnColorCommitted;
+
+	FOnGetCurrentVector4Value OnQueryCurrentColor;
+};
diff --git a/Engine/Source/Runtime/AutomationMessages/Classes/AutomationWorkerMessages.h b/Engine/Source/Runtime/AutomationMessages/Classes/AutomationWorkerMessages.h
--- a/Engine/Source/Runtime/AutomationMessages/Classes/AutomationWorkerMessages.h
+++ b/Engine/Source/Runtime/AutomationMessages/Classes/AutomationWorkerMessages.h
@@ -417,6 +417,9 @@ public:
 	UPROPERTY(EditAnywhere, Category="Message")
 	FString Context;
 
+	UPROPERTY(EditAnywhere, Category="Message")
+	FGuid Id;
+
 	UPROPERTY(EditAnywhere, Category="Message")
 	int32 Width;
 	UPROPERTY(EditAnywhere, Category="Message")
@@ -493,9 +496,16 @@ public:
 
 	FAutomationScreenshotMetadata(const FAutomationScreenshotData& Data)
 	{
+		static_assert( (sizeof(FAutomationScreenshotMetadata) + sizeof(FString)) == sizeof(FAutomationScreenshotData), "These objects must match in size, to ensure we're copying all the members, except for FAutomationScreenshotData.Path, which we don't copy over." );
+
+		// Human readable name and associated context the screenshot was taken in.
 		Name = Data.Name;
 		Context = Data.Context;
 
+		// Unique Id so we know if this screenshot has already been imported.
+		Id = Data.Id;
+
+		// Resolution Details
 		Width = Data.Width;
 		Height = Data.Height;
 
@@ -522,6 +532,9 @@ public:
 		EffectsQuality = Data.EffectsQuality;
 		FoliageQuality = Data.FoliageQuality;
 
+		// Enabled Features
+		// TBD
+
 		// Comparison Requests
 		bHasComparisonRules = Data.bHasComparisonRules;
 		ToleranceRed = Data.ToleranceRed;
@@ -537,6 +550,75 @@ public:
 		bIgnoreAntiAliasing = Data.bIgnoreAntiAliasing;
 		bIgnoreColors = Data.bIgnoreColors;
 	}
+
+	int32 Compare(const FAutomationScreenshotMetadata& OtherMetadata) const
+	{
+		int32 Score = 0;
+
+		if ( Width != OtherMetadata.Width || Height != OtherMetadata.Height || bIsStereo != OtherMetadata.bIsStereo )
+		{
+			return 0;
+		}
+		else
+		{
+			Score += 1000;
+		}
+
+		if (ResolutionQuality == OtherMetadata.ResolutionQuality && 
+			ViewDistanceQuality == OtherMetadata.ViewDistanceQuality &&
+			AntiAliasingQuality == OtherMetadata.AntiAliasingQuality &&
+			ShadowQuality == OtherMetadata.ShadowQuality &&
+			PostProcessQuality == OtherMetadata.PostProcessQuality &&
+			TextureQuality == OtherMetadata.TextureQuality &&
+			EffectsQuality == OtherMetadata.EffectsQuality &&
+			FoliageQuality == OtherMetadata.FoliageQuality
+			)
+		{
+			Score += 100;
+		}
+
+		if ( FeatureLevel == OtherMetadata.FeatureLevel )
+		{
+			Score += 100;
+		}
+
+		if ( UniqueDeviceId == OtherMetadata.UniqueDeviceId )
+		{
+			Score += 100;
+		}
+
+		if ( Rhi == OtherMetadata.Rhi )
+		{
+			Score += 100;
+		}
+
+		if ( Platform == OtherMetadata.Platform )
+		{
+			Score += 10;
+		}
+
+		if ( Vendor == OtherMetadata.Vendor )
+		{
+			Score += 10;
+		}
+
+		if ( AdapterName == OtherMetadata.AdapterName )
+		{
+			Score += 10;
+		}
+
+		if ( AdapterInternalDriverVersion == OtherMetadata.AdapterInternalDriverVersion )
+		{
+			Score += 10;
+		}
+
+		if ( AdapterUserDriverVersion == OtherMetadata.AdapterUserDriverVersion )
+		{
+			Score += 10;
+		}
+
+		return Score;
+	}
 };
 
 
diff --git a/Engine/Source/Runtime/AutomationWorker/Private/AutomationWorkerModule.cpp b/Engine/Source/Runtime/AutomationWorker/Private/AutomationWorkerModule.cpp
--- a/Engine/Source/Runtime/AutomationWorker/Private/AutomationWorkerModule.cpp
+++ b/Engine/Source/Runtime/AutomationWorker/Private/AutomationWorkerModule.cpp
@@ -440,7 +440,7 @@ void FAutomationWorkerModule::HandleScreenShotCapturedWithName(const TArray<FCol
 			if ( FJsonObjectConverter::UStructToJsonObjectString(Metadata, Json) )
 			{
 				FString MetadataPath = FPaths::ChangeExtension(Data.Path, TEXT("json"));
-				FFileHelper::SaveStringToFile(Json, *MetadataPath, FFileHelper::EEncodingOptions::ForceUTF8);
+				FFileHelper::SaveStringToFile(Json, *MetadataPath, FFileHelper::EEncodingOptions::ForceUTF8WithoutBOM);
 			}
 		}
 	}
diff --git a/Engine/Source/Runtime/Core/Private/GenericPlatform/GenericPlatformMisc.cpp b/Engine/Source/Runtime/Core/Private/GenericPlatform/GenericPlatformMisc.cpp
--- a/Engine/Source/Runtime/Core/Private/GenericPlatform/GenericPlatformMisc.cpp
+++ b/Engine/Source/Runtime/Core/Private/GenericPlatform/GenericPlatformMisc.cpp
@@ -233,7 +233,10 @@ FString FGenericPlatformMisc::GetUniqueDeviceId()
 
 FString FGenericPlatformMisc::GetDeviceId()
 {
-	return FString();
+	// @todo: When this function is finally removed, the functionality used will need to be moved in here.
+	PRAGMA_DISABLE_DEPRECATION_WARNINGS
+	return GetUniqueDeviceId();
+	PRAGMA_ENABLE_DEPRECATION_WARNINGS
 }
 
 void FGenericPlatformMisc::SubmitErrorReport( const TCHAR* InErrorHist, EErrorReportMode::Type InMode )
diff --git a/Engine/Source/Runtime/Core/Private/Internationalization/ICULineBreakIterator.cpp b/Engine/Source/Runtime/Core/Private/Internationalization/ICULineBreakIterator.cpp
--- a/Engine/Source/Runtime/Core/Private/Internationalization/ICULineBreakIterator.cpp
+++ b/Engine/Source/Runtime/Core/Private/Internationalization/ICULineBreakIterator.cpp
@@ -10,6 +10,38 @@
 #if UE_ENABLE_ICU
 #include "Internationalization/ICUBreakIterator.h"
 #include "Internationalization/ICUTextCharacterIterator.h"
+#include "IConsoleManager.h"
+
+enum class EHangulTextWrappingMethod : uint8
+{
+	/** Wrap per-syllable (default Geumchik wrapping) */
+	PerSyllable = 0,
+
+	/** Wrap per-word (preferred native wrapping) */
+	PerWord,
+};
+
+static TAutoConsoleVariable<int32> CVarHangulTextWrappingMethod(
+	TEXT("Localization.HangulTextWrappingMethod"),
+	static_cast<int32>(EHangulTextWrappingMethod::PerWord),
+	TEXT("0: PerSyllable, 1: PerWord (default)."),
+	ECVF_Default
+	);
+
+EHangulTextWrappingMethod GetHangulTextWrappingMethod()
+{
+	const int32 HangulTextWrappingMethodAsInt = CVarHangulTextWrappingMethod.AsVariable()->GetInt();
+	if (HangulTextWrappingMethodAsInt >= static_cast<int32>(EHangulTextWrappingMethod::PerSyllable) && HangulTextWrappingMethodAsInt <= static_cast<int32>(EHangulTextWrappingMethod::PerWord))
+	{
+		return static_cast<EHangulTextWrappingMethod>(HangulTextWrappingMethodAsInt);
+	}
+	return EHangulTextWrappingMethod::PerWord;
+}
+
+FORCEINLINE bool IsHangul(const TCHAR InChar)
+{
+	return InChar >= 0xAC00 && InChar <= 0xD7A3;
+}
 
 class FICULineBreakIterator : public IBreakIterator
 {
@@ -37,23 +69,14 @@ protected:
 	int32 MoveToNextImpl();
 
 	TSharedRef<icu::BreakIterator> GetInternalLineBreakIterator() const;
-	TSharedRef<icu::BreakIterator> GetInternalWordBreakIterator() const;
 
 private:
-	/**
-	 * We use both a line-break and word-break iterator here, as CJK languages use dictionary-based word-breaking which isn't supported correctly by the line-break iterator.
-	 * To compensate for that, we use both iterators in tandem and treat the word-break result as the minimum breaking point (to avoid breaking words in CJK).
-	 * For example, in Korean the line-break iterator tries to break per-glyph, whereas the word-break iterator knows where the word boundaries are; however in English, the line-break
-	 * iterator works correctly, but the word-break iterator may fall short as it doesn't handle symbols as would be expected.
-	 */
 	TWeakPtr<icu::BreakIterator> ICULineBreakIteratorHandle;
-	TWeakPtr<icu::BreakIterator> ICUWordBreakIteratorHandle;
 	int32 CurrentPosition;
 };
 
 FICULineBreakIterator::FICULineBreakIterator()
 	: ICULineBreakIteratorHandle(FICUBreakIteratorManager::Get().CreateLineBreakIterator())
-	, ICUWordBreakIteratorHandle(FICUBreakIteratorManager::Get().CreateWordBreakIterator())
 	, CurrentPosition(0)
 {
 }
@@ -62,34 +85,29 @@ FICULineBreakIterator::~FICULineBreakIterator()
 {
 	// This assumes that FICULineBreakIterator owns the iterators, and that nothing ever copies an FICULineBreakIterator instance
 	FICUBreakIteratorManager::Get().DestroyIterator(ICULineBreakIteratorHandle);
-	FICUBreakIteratorManager::Get().DestroyIterator(ICUWordBreakIteratorHandle);
 }
 
 void FICULineBreakIterator::SetString(const FText& InText)
 {
 	GetInternalLineBreakIterator()->adoptText(new FICUTextCharacterIterator(InText)); // ICUBreakIterator takes ownership of this instance
-	GetInternalWordBreakIterator()->adoptText(new FICUTextCharacterIterator(InText)); // ICUBreakIterator takes ownership of this instance
 	ResetToBeginning();
 }
 
 void FICULineBreakIterator::SetString(const FString& InString)
 {
 	GetInternalLineBreakIterator()->adoptText(new FICUTextCharacterIterator(InString)); // ICUBreakIterator takes ownership of this instance
-	GetInternalWordBreakIterator()->adoptText(new FICUTextCharacterIterator(InString)); // ICUBreakIterator takes ownership of this instance
 	ResetToBeginning();
 }
 
 void FICULineBreakIterator::SetString(const TCHAR* const InString, const int32 InStringLength)
 {
 	GetInternalLineBreakIterator()->adoptText(new FICUTextCharacterIterator(InString, InStringLength)); // ICUBreakIterator takes ownership of this instance
-	GetInternalWordBreakIterator()->adoptText(new FICUTextCharacterIterator(InString, InStringLength)); // ICUBreakIterator takes ownership of this instance
 	ResetToBeginning();
 }
 
 void FICULineBreakIterator::ClearString()
 {
 	GetInternalLineBreakIterator()->adoptText(new FICUTextCharacterIterator(FString())); // ICUBreakIterator takes ownership of this instance
-	GetInternalWordBreakIterator()->adoptText(new FICUTextCharacterIterator(FString())); // ICUBreakIterator takes ownership of this instance
 	ResetToBeginning();
 }
 
@@ -102,7 +120,6 @@ int32 FICULineBreakIterator::ResetToBeginning()
 {
 	TSharedRef<icu::BreakIterator> LineBrkIt = GetInternalLineBreakIterator();
 	CurrentPosition = LineBrkIt->first();
-	GetInternalWordBreakIterator()->first();
 	CurrentPosition = static_cast<FICUTextCharacterIterator&>(LineBrkIt->getText()).InternalIndexToSourceIndex(CurrentPosition);
 	return CurrentPosition;
 }
@@ -111,7 +128,6 @@ int32 FICULineBreakIterator::ResetToEnd()
 {
 	TSharedRef<icu::BreakIterator> LineBrkIt = GetInternalLineBreakIterator();
 	CurrentPosition = LineBrkIt->last();
-	GetInternalWordBreakIterator()->last();
 	CurrentPosition = static_cast<FICUTextCharacterIterator&>(LineBrkIt->getText()).InternalIndexToSourceIndex(CurrentPosition);
 	return CurrentPosition;
 }
@@ -141,13 +157,28 @@ int32 FICULineBreakIterator::MoveToCandidateAfter(const int32 InIndex)
 int32 FICULineBreakIterator::MoveToPreviousImpl()
 {
 	TSharedRef<icu::BreakIterator> LineBrkIt = GetInternalLineBreakIterator();
-	CurrentPosition = static_cast<FICUTextCharacterIterator&>(LineBrkIt->getText()).SourceIndexToInternalIndex(CurrentPosition);
+	FICUTextCharacterIterator& CharIt = static_cast<FICUTextCharacterIterator&>(LineBrkIt->getText());
 
-	const int32 NewLinePosition = LineBrkIt->preceding(CurrentPosition);
-	const int32 NewWordPosition = GetInternalWordBreakIterator()->preceding(CurrentPosition);
+	int32 InternalPosition = CharIt.SourceIndexToInternalIndex(CurrentPosition);
 
-	CurrentPosition = FMath::Min(NewLinePosition, NewWordPosition);
-	CurrentPosition = static_cast<FICUTextCharacterIterator&>(LineBrkIt->getText()).InternalIndexToSourceIndex(CurrentPosition);
+	// For Hangul using per-word wrapping, we walk back to the first Hangul character in the word and use that as the starting point for the 
+	// line-break iterator, as this will correctly handle the remaining Geumchik wrapping rules, without also causing per-syllable wrapping
+	if (GetHangulTextWrappingMethod() == EHangulTextWrappingMethod::PerWord)
+	{
+		CharIt.setIndex32(InternalPosition);
+
+		if (IsHangul(CharIt.current32()))
+		{
+			// Walk to the start of the Hangul characters
+			while (CharIt.hasPrevious() && IsHangul(static_cast<TCHAR>(CharIt.previous32())))
+			{
+				InternalPosition = CharIt.getIndex();
+			}
+		}
+	}
+
+	InternalPosition = LineBrkIt->preceding(InternalPosition);
+	CurrentPosition = CharIt.InternalIndexToSourceIndex(InternalPosition);
 
 	return CurrentPosition;
 }
@@ -155,13 +186,28 @@ int32 FICULineBreakIterator::MoveToPreviousImpl()
 int32 FICULineBreakIterator::MoveToNextImpl()
 {
 	TSharedRef<icu::BreakIterator> LineBrkIt = GetInternalLineBreakIterator();
-	CurrentPosition = static_cast<FICUTextCharacterIterator&>(LineBrkIt->getText()).SourceIndexToInternalIndex(CurrentPosition);
+	FICUTextCharacterIterator& CharIt = static_cast<FICUTextCharacterIterator&>(LineBrkIt->getText());
 
-	const int32 NewLinePosition = LineBrkIt->following(CurrentPosition);
-	const int32 NewWordPosition = GetInternalWordBreakIterator()->following(CurrentPosition);
+	int32 InternalPosition = CharIt.SourceIndexToInternalIndex(CurrentPosition);
 
-	CurrentPosition = FMath::Max(NewLinePosition, NewWordPosition);
-	CurrentPosition = static_cast<FICUTextCharacterIterator&>(LineBrkIt->getText()).InternalIndexToSourceIndex(CurrentPosition);
+	// For Hangul using per-word wrapping, we walk forward to the last Hangul character in the word and use that as the starting point for the 
+	// line-break iterator, as this will correctly handle the remaining Geumchik wrapping rules, without also causing per-syllable wrapping
+	if (GetHangulTextWrappingMethod() == EHangulTextWrappingMethod::PerWord)
+	{
+		CharIt.setIndex32(InternalPosition);
+
+		if (IsHangul(CharIt.current32()))
+		{
+			// Walk to the end of the Hangul characters
+			while (CharIt.hasNext() && IsHangul(static_cast<TCHAR>(CharIt.next32())))
+			{
+				InternalPosition = CharIt.getIndex();
+			}
+		}
+	}
+
+	InternalPosition = LineBrkIt->following(InternalPosition);
+	CurrentPosition = CharIt.InternalIndexToSourceIndex(InternalPosition);
 
 	return CurrentPosition;
 }
@@ -171,11 +217,6 @@ TSharedRef<icu::BreakIterator> FICULineBreakIterator::GetInternalLineBreakIterat
 	return ICULineBreakIteratorHandle.Pin().ToSharedRef();
 }
 
-TSharedRef<icu::BreakIterator> FICULineBreakIterator::GetInternalWordBreakIterator() const
-{
-	return ICUWordBreakIteratorHandle.Pin().ToSharedRef();
-}
-
 TSharedRef<IBreakIterator> FBreakIterator::CreateLineBreakIterator()
 {
 	return MakeShareable(new FICULineBreakIterator());
diff --git a/Engine/Source/Runtime/Core/Private/Internationalization/TextFormatter.cpp b/Engine/Source/Runtime/Core/Private/Internationalization/TextFormatter.cpp
--- a/Engine/Source/Runtime/Core/Private/Internationalization/TextFormatter.cpp
+++ b/Engine/Source/Runtime/Core/Private/Internationalization/TextFormatter.cpp
@@ -50,7 +50,7 @@ FORCEINLINE bool IsLiteralBreakChar(const TCHAR InChar)
 /** Token representing a literal string inside the text */
 struct FStringLiteral
 {
-	FStringLiteral(const FStringToken& InString)
+	explicit FStringLiteral(const FStringToken& InString)
 		: StringStartPos(InString.GetTokenStartPos())
 		, StringLen(InString.GetTokenEndPos() - InString.GetTokenStartPos())
 	{
@@ -66,7 +66,7 @@ struct FStringLiteral
 /** Token representing a format argument */
 struct FArgumentTokenSpecifier
 {
-	FArgumentTokenSpecifier(const FStringToken& InArgument)
+	explicit FArgumentTokenSpecifier(const FStringToken& InArgument)
 		: ArgumentNameStartPos(InArgument.GetTokenStartPos())
 		, ArgumentNameLen(InArgument.GetTokenEndPos() - InArgument.GetTokenStartPos())
 		, ArgumentIndex(INDEX_NONE)
@@ -105,11 +105,19 @@ struct FArgumentTokenSpecifier
 /** Token representing a format argument modifier */
 struct FArgumentModifierTokenSpecifier
 {
-	FArgumentModifierTokenSpecifier(TSharedRef<ITextFormatArgumentModifier> InTextFormatArgumentModifier)
-		: TextFormatArgumentModifier(MoveTemp(InTextFormatArgumentModifier))
+	FArgumentModifierTokenSpecifier(const FStringToken& InModifierPatternWithPipe, TSharedRef<ITextFormatArgumentModifier> InTextFormatArgumentModifier)
+		: ModifierPatternStartPos(InModifierPatternWithPipe.GetTokenStartPos() + 1) // We don't want to store the pipe
+		, ModifierPatternLen(InModifierPatternWithPipe.GetTokenEndPos() - InModifierPatternWithPipe.GetTokenStartPos() - 1)
+		, TextFormatArgumentModifier(MoveTemp(InTextFormatArgumentModifier))
 	{
 	}
 
+	/** The start of the pattern this modifier was generated from */
+	const TCHAR* ModifierPatternStartPos;
+
+	/** The length of the pattern this modifier was generated from */
+	int32 ModifierPatternLen;
+
 	/** The compiled argument modifier that should be evaluated */
 	TSharedRef<ITextFormatArgumentModifier> TextFormatArgumentModifier;
 };
@@ -117,7 +125,7 @@ struct FArgumentModifierTokenSpecifier
 /** Token representing an escaped character */
 struct FEscapedCharacter
 {
-	FEscapedCharacter(TCHAR InChar)
+	explicit FEscapedCharacter(TCHAR InChar)
 		: Character(InChar)
 	{
 	}
@@ -259,7 +267,7 @@ TOptional<FExpressionError> ParseArgumentModifier(FExpressionTokenConsumer& Cons
 	}
 
 	// Add the token to the consumer - this moves the read position in the stream to the end of the token
-	Consumer.Add(EntireToken, FArgumentModifierTokenSpecifier(CompiledTextArgumentModifier.ToSharedRef()));
+	Consumer.Add(EntireToken, FArgumentModifierTokenSpecifier(EntireToken, CompiledTextArgumentModifier.ToSharedRef()));
 	return TOptional<FExpressionError>();
 }
 
@@ -669,6 +677,7 @@ void FTextFormatData::ConditionalCompile_NoLock()
 
 	if (bRequiresCompile)
 	{
+		bRequiresCompile = false;
 		if (!CompiledTextSnapshot.IdenticalTo(SourceText))
 		{
 			if (!CompiledTextSnapshot.IsDisplayStringEqualTo(SourceText))
@@ -687,6 +696,11 @@ void FTextFormatData::ConditionalCompile_NoLock()
 
 FString FTextFormatData::Format_NoLock(const FPrivateTextFormatArguments& InFormatArgs)
 {
+	if (SourceType == ESourceType::Text && InFormatArgs.bRebuildText)
+	{
+		SourceText.Rebuild();
+	}
+
 	ConditionalCompile_NoLock();
 
 	if (LexedExpression.Num() == 0)
@@ -737,6 +751,13 @@ FString FTextFormatData::Format_NoLock(const FPrivateTextFormatArguments& InForm
 				ResultString.AppendChar(TextFormatTokens::ArgEndChar);
 			}
 		}
+		else if (const auto* ArgumentModifierToken = Token.Node.Cast<TextFormatTokens::FArgumentModifierTokenSpecifier>())
+		{
+			// If we find an argument modifier token on its own then it means an argument value failed to evaluate (likely due to InFormatArgs.GetArgumentValue returning null)
+			// In this case we just write the literal value of the argument modifier back into the final string...
+			ResultString.AppendChar(TextFormatTokens::ArgModChar);
+			ResultString.AppendChars(ArgumentModifierToken->ModifierPatternStartPos, ArgumentModifierToken->ModifierPatternLen);
+		}
 	}
 
 	return ResultString;
@@ -939,7 +960,23 @@ FText FTextFormatter::Format(FTextFormat&& InFmt, TArray<FFormatArgumentData>&&
 
 FString FTextFormatter::Format(const FTextFormat& InFmt, const FPrivateTextFormatArguments& InFormatArgs)
 {
-	return InFmt.TextFormatData->Format(InFormatArgs);
+	FTextFormat FmtPattern = InFmt;
+
+	// If we're rebuilding as source then we need to handle that before we call Format
+	// We don't need to worry about any rebuilding that needs to happen as non-source, as Format takes care of that internally
+	if (InFormatArgs.bRebuildAsSource)
+	{
+		FText FmtText = InFmt.GetSourceText();
+
+		if (InFormatArgs.bRebuildText)
+		{
+			FmtText.Rebuild();
+		}
+
+		FmtPattern = FTextFormat(FmtText.BuildSourceString());
+	}
+
+	return FmtPattern.TextFormatData->Format(InFormatArgs);
 }
 
 void FTextFormatter::ArgumentValueToFormattedString(const FFormatArgumentValue& InValue, const FPrivateTextFormatArguments& InFormatArgs, FString& OutResult)
diff --git a/Engine/Source/Runtime/Core/Private/Mac/MacPlatformProcess.cpp b/Engine/Source/Runtime/Core/Private/Mac/MacPlatformProcess.cpp
--- a/Engine/Source/Runtime/Core/Private/Mac/MacPlatformProcess.cpp
+++ b/Engine/Source/Runtime/Core/Private/Mac/MacPlatformProcess.cpp
@@ -618,6 +618,16 @@ FProcHandle FMacPlatformProcess::CreateProc( const TCHAR* URL, const TCHAR* Parm
 	return FProcHandle(ProcessHandle);
 }
 
+FProcHandle FMacPlatformProcess::OpenProcess(uint32 ProcessID)
+{
+	for (NSRunningApplication *app in[[NSWorkspace sharedWorkspace] runningApplications])
+	{
+		NSLog(@"%@",[app localizedName]);
+	}
+
+	return FProcHandle();
+}
+
 bool FMacPlatformProcess::IsProcRunning( FProcHandle& ProcessHandle )
 {
 	SCOPED_AUTORELEASE_POOL;
diff --git a/Engine/Source/Runtime/Core/Private/Misc/FbxErrors.cpp b/Engine/Source/Runtime/Core/Private/Misc/FbxErrors.cpp
--- a/Engine/Source/Runtime/Core/Private/Misc/FbxErrors.cpp
+++ b/Engine/Source/Runtime/Core/Private/Misc/FbxErrors.cpp
@@ -5,6 +5,8 @@
 /** Generic */
 FName FFbxErrors::Generic_ImportingNewObjectFailed(TEXT("Generic_ImportingNewObjectFailed"));
 
+FName FFbxErrors::Generic_ReimportingObjectFailed(TEXT("Generic_ReimportingObjectFailed"));
+
 FName FFbxErrors::Generic_LoadingSceneFailed(TEXT("Generic_LoadingSceneFailed"));
 
 FName FFbxErrors::Generic_InvalidCharacterInName(TEXT("Generic_InvalidCharacterInName"));
@@ -65,6 +67,8 @@ FName FFbxErrors::SkeletalMesh_DuplicateBones(TEXT("SkeletalMesh_DuplicateBones"
 
 FName FFbxErrors::SkeletalMesh_NoInfluences(TEXT("SkeletalMesh_NoInfluences"));
 
+FName FFbxErrors::SkeletalMesh_TooManyInfluences(TEXT("SkeletalMesh_TooManyInfluences"));
+
 FName FFbxErrors::SkeletalMesh_RestoreSortingMismatchedStrips(TEXT("SkeletalMesh_RestoreSortingMismatchedStrips"));
 
 FName FFbxErrors::SkeletalMesh_RestoreSortingNoSectionMatch(TEXT("SkeletalMesh_RestoreSortingNoSectionMatch"));
diff --git a/Engine/Source/Runtime/Core/Private/Windows/TextStoreACP.cpp b/Engine/Source/Runtime/Core/Private/Windows/TextStoreACP.cpp
--- a/Engine/Source/Runtime/Core/Private/Windows/TextStoreACP.cpp
+++ b/Engine/Source/Runtime/Core/Private/Windows/TextStoreACP.cpp
@@ -567,22 +567,20 @@ STDAPI FTextStoreACP::QueryInsert(LONG acpInsertStart, LONG acpInsertEnd, ULONG
 {
 	UE_LOG(LogTextStoreACP, Verbose, TEXT("QueryInsert"));
 
-	const LONG StringLength = TextContext->GetTextLength();
-
-	// Query range is invalid.
-	if(acpInsertStart < 0 || acpInsertStart > StringLength || acpInsertEnd < 0 || acpInsertEnd > StringLength)
-	{
-		return E_INVALIDARG;
-	}
-
 	// Can't successfully query if there's nowhere to write a result.
 	if(!pacpResultStart || !pacpResultEnd)
 	{
 		return E_INVALIDARG;
 	}
 
-	*pacpResultStart = acpInsertStart;
-	*pacpResultEnd = acpInsertEnd;
+	uint32 BeginIndex;
+	uint32 Length;
+	ITextInputMethodContext::ECaretPosition CaretPosition;
+	TextContext->GetSelectionRange(BeginIndex, Length, CaretPosition);
+
+	// Workaround for Microsoft IMEs that expect QueryInsert to return the current selection range (since they omit a call to GetSelection).
+	*pacpResultStart = static_cast<LONG>(BeginIndex);
+	*pacpResultEnd = static_cast<LONG>(BeginIndex + Length);
 
 	return S_OK;
 }
diff --git a/Engine/Source/Runtime/Core/Public/Internationalization/Text.h b/Engine/Source/Runtime/Core/Public/Internationalization/Text.h
--- a/Engine/Source/Runtime/Core/Public/Internationalization/Text.h
+++ b/Engine/Source/Runtime/Core/Public/Internationalization/Text.h
@@ -539,6 +539,7 @@ private:
 public:
 	friend class FTextCache;
 	friend class FTextFormatter;
+	friend class FTextFormatData;
 	friend class FTextSnapshot;
 	friend class FTextInspector;
 	friend class FArchive;
@@ -895,11 +896,9 @@ private:
 class CORE_API FTextBuilder
 {
 public:
-
 	FTextBuilder()
 		: IndentCount(0)
 	{
-
 	}
 
 	void Indent()
@@ -978,6 +977,11 @@ public:
 		Report.Empty();
 	}
 
+	bool IsEmpty()
+	{
+		return Report.IsEmpty();
+	}
+
 	FText ToText() const
 	{
 		return FText::FromString(Report);
diff --git a/Engine/Source/Runtime/Core/Public/Mac/MacPlatformProcess.h b/Engine/Source/Runtime/Core/Public/Mac/MacPlatformProcess.h
--- a/Engine/Source/Runtime/Core/Public/Mac/MacPlatformProcess.h
+++ b/Engine/Source/Runtime/Core/Public/Mac/MacPlatformProcess.h
@@ -116,6 +116,7 @@ struct CORE_API FMacPlatformProcess : public FGenericPlatformProcess
 	static void LaunchURL(const TCHAR* URL, const TCHAR* Parms, FString* Error);
 	static FString GetGameBundleId();
 	static FProcHandle CreateProc( const TCHAR* URL, const TCHAR* Parms, bool bLaunchDetached, bool bLaunchHidden, bool bLaunchReallyHidden, uint32* OutProcessID, int32 PriorityModifier, const TCHAR* OptionalWorkingDirectory, void* PipeWriteChild, void * PipeReadChild = nullptr);
+	static FProcHandle OpenProcess(uint32 ProcessID);
 	static bool IsProcRunning( FProcHandle & ProcessHandle );
 	static void WaitForProc( FProcHandle & ProcessHandle );
 	static void CloseProc( FProcHandle & ProcessHandle );
diff --git a/Engine/Source/Runtime/Core/Public/Math/Vector4.h b/Engine/Source/Runtime/Core/Public/Math/Vector4.h
--- a/Engine/Source/Runtime/Core/Public/Math/Vector4.h
+++ b/Engine/Source/Runtime/Core/Public/Math/Vector4.h
@@ -351,6 +351,20 @@ public:
 	 */
 	float SizeSquared3() const;
 
+	/**
+	 * Get the length (magnitude) of this vector, taking the W component into account
+	 *
+	 * @return The length of this vector
+	 */
+	float Size() const;
+
+	/**
+	 * Get the squared length of this vector, taking the W component into account
+	 *
+	 * @return The squared length of this vector
+	 */
+	float SizeSquared() const;
+
 	/** Utility to check if there are any non-finite values (NaN or Inf) in this vector. */
 	bool ContainsNaN() const;
 
@@ -636,6 +650,16 @@ FORCEINLINE float FVector4::SizeSquared3() const
 	return X*X + Y*Y + Z*Z;
 }
 
+FORCEINLINE float FVector4::Size() const
+{
+	return FMath::Sqrt(X*X + Y*Y + Z*Z + W*W);
+}
+
+FORCEINLINE float FVector4::SizeSquared() const
+{
+	return X*X + Y*Y + Z*Z + W*W;
+}
+
 
 FORCEINLINE bool FVector4::IsUnit3(float LengthSquaredTolerance) const
 {
diff --git a/Engine/Source/Runtime/Core/Public/Misc/AutomationTest.h b/Engine/Source/Runtime/Core/Public/Misc/AutomationTest.h
--- a/Engine/Source/Runtime/Core/Public/Misc/AutomationTest.h
+++ b/Engine/Source/Runtime/Core/Public/Misc/AutomationTest.h
@@ -21,6 +21,7 @@
 #include "Misc/FeedbackContext.h"
 #include "Future.h"
 #include "Async.h"
+#include "Misc/Guid.h"
 
 /** Flags for specifying automation test requirements/behavior */
 namespace EAutomationTestFlags
@@ -465,6 +466,8 @@ struct FAutomationScreenshotData
 	FString Name;
 	FString Context;
 
+	FGuid Id;
+
 	int32 Width;
 	int32 Height;
 
diff --git a/Engine/Source/Runtime/Core/Public/Misc/FbxErrors.h b/Engine/Source/Runtime/Core/Public/Misc/FbxErrors.h
--- a/Engine/Source/Runtime/Core/Public/Misc/FbxErrors.h
+++ b/Engine/Source/Runtime/Core/Public/Misc/FbxErrors.h
@@ -17,6 +17,8 @@ struct CORE_API FFbxErrors
 	/** Generic */
 	static FName Generic_ImportingNewObjectFailed;
 
+	static FName Generic_ReimportingObjectFailed;
+
 	static FName Generic_LoadingSceneFailed;
 
 	static FName Generic_InvalidCharacterInName;
@@ -77,6 +79,8 @@ struct CORE_API FFbxErrors
 
 	static FName SkeletalMesh_NoInfluences;
 
+	static FName SkeletalMesh_TooManyInfluences;
+
 	static FName SkeletalMesh_RestoreSortingMismatchedStrips;
 
 	static FName SkeletalMesh_RestoreSortingNoSectionMatch;
diff --git a/Engine/Source/Runtime/Core/Public/Templates/SharedPointer.h b/Engine/Source/Runtime/Core/Public/Templates/SharedPointer.h
--- a/Engine/Source/Runtime/Core/Public/Templates/SharedPointer.h
+++ b/Engine/Source/Runtime/Core/Public/Templates/SharedPointer.h
@@ -125,6 +125,18 @@ FORCEINLINE TSharedRef< CastToType, Mode > StaticCastSharedRef( TSharedRef< Cast
 }
 
 
+namespace UE4SharedPointer_Private
+{
+	// Needed to work around an Android compiler bug - we need to construct a TSharedRef
+	// from MakeShared without making MakeShared a friend in order to access the private constructor.
+	template <typename ObjectType, ESPMode Mode>
+	FORCEINLINE TSharedRef<ObjectType, Mode> MakeSharedRef(ObjectType* InObject, SharedPointerInternals::FReferenceControllerBase* InSharedReferenceCount)
+	{
+		return TSharedRef<ObjectType, Mode>(InObject, InSharedReferenceCount);
+	}
+}
+
+
 /**
  * TSharedRef is a non-nullable, non-intrusive reference-counted authoritative object reference.
  *
@@ -473,8 +485,17 @@ private:
 		controller object is shared by all shared and weak pointers that refer to the object */
 	SharedPointerInternals::FSharedReferencer< Mode > SharedReferenceCount;
 
-	template <typename InObjectType, ESPMode InMode, typename... InArgTypes>
-	friend TSharedRef<InObjectType, InMode> MakeShared(InArgTypes&&... Args);
+	// VC emits an erroneous warning here - there is no inline specifier!
+	#ifdef _MSC_VER
+		#pragma warning(push)
+		#pragma warning(disable : 4396) // warning: the inline specifier cannot be used when a friend declaration refers to a specialization of a function template
+	#endif
+
+	friend TSharedRef UE4SharedPointer_Private::MakeSharedRef<ObjectType, Mode>(ObjectType* InObject, SharedPointerInternals::FReferenceControllerBase* InSharedReferenceCount);
+
+	#ifdef _MSC_VER
+		#pragma warning(pop)
+	#endif
 
 	FORCEINLINE explicit TSharedRef(ObjectType* InObject, SharedPointerInternals::FReferenceControllerBase* InSharedReferenceCount)
 		: Object(InObject)
@@ -505,10 +526,6 @@ struct FMakeReferenceTo<void>
 };
 
 
-template <typename InObjectType, ESPMode InMode, typename... InArgTypes>
-TSharedRef<InObjectType, InMode> MakeShared(InArgTypes&&... Args);
-
-
 /**
  * TSharedPtr is a non-intrusive reference-counted authoritative object pointer.  This shared pointer
  * will be conditionally thread-safe when the optional Mode template argument is set to ThreadSafe.
@@ -1658,7 +1675,7 @@ template <typename InObjectType, ESPMode InMode = ESPMode::Fast, typename... InA
 FORCEINLINE TSharedRef<InObjectType, InMode> MakeShared(InArgTypes&&... Args)
 {
 	SharedPointerInternals::TIntrusiveReferenceController<InObjectType>* Controller = SharedPointerInternals::NewIntrusiveReferenceController<InObjectType>(Forward<InArgTypes>(Args)...);
-	return TSharedRef<InObjectType, InMode>(Controller->GetObjectPtr(), (SharedPointerInternals::FReferenceControllerBase*)Controller);
+	return UE4SharedPointer_Private::MakeSharedRef<InObjectType, InMode>(Controller->GetObjectPtr(), (SharedPointerInternals::FReferenceControllerBase*)Controller);
 }
 
 
diff --git a/Engine/Source/Runtime/Core/Public/UObject/EditorObjectVersion.h b/Engine/Source/Runtime/Core/Public/UObject/EditorObjectVersion.h
--- a/Engine/Source/Runtime/Core/Public/UObject/EditorObjectVersion.h
+++ b/Engine/Source/Runtime/Core/Public/UObject/EditorObjectVersion.h
@@ -29,6 +29,16 @@ struct CORE_API FEditorObjectVersion
 		RefactorMeshEditorMaterials,
 		// Added UFontFace assets
 		AddedFontFaceAssets,
+		// Add UPROPERTY for TMap of Mesh section, so the serialize will be done normally (and export to text will work correctly)
+		UPropertryForMeshSection,
+		// Update the schema of all widget blueprints to use the WidgetGraphSchema
+		WidgetGraphSchema,
+		// Added a specialized content slot to the background blur widget
+		AddedBackgroundBlurContentSlot,
+		// Updated UserDefinedEnums to have stable keyed display names
+		StableUserDefinedEnumDisplayNames,
+		// Added "Inline" option to UFontFace assets
+		AddedInlineFontFaceAssets,
 		// -----<new versions can be added above this line>-------------------------------------------------
 		VersionPlusOne,
 		LatestVersion = VersionPlusOne - 1
diff --git a/Engine/Source/Runtime/CoreUObject/Private/Internationalization/TextPackageNamespaceUtil.cpp b/Engine/Source/Runtime/CoreUObject/Private/Internationalization/TextPackageNamespaceUtil.cpp
--- a/Engine/Source/Runtime/CoreUObject/Private/Internationalization/TextPackageNamespaceUtil.cpp
+++ b/Engine/Source/Runtime/CoreUObject/Private/Internationalization/TextPackageNamespaceUtil.cpp
@@ -41,6 +41,7 @@ FString FindOrAddPackageNamespace(UPackage* InPackage, const bool bCanAdd)
 				if (PackageLocalizationNamespaceValue.IsEmpty())
 				{
 					PackageLocalizationNamespaceValue = FGuid::NewGuid().ToString();
+					InPackage->MarkPackageDirty();
 				}
 				return PackageLocalizationNamespaceValue;
 			}
@@ -67,7 +68,10 @@ void ClearPackageNamespace(UPackage* InPackage)
 		if (!FPackageName::IsScriptPackage(PackageName))
 		{
 			UMetaData* PackageMetaData = InPackage->GetMetaData();
-			PackageMetaData->RootMetaDataMap.Remove(PackageLocalizationNamespaceKey);
+			if (PackageMetaData->RootMetaDataMap.Remove(PackageLocalizationNamespaceKey) > 0)
+			{
+				InPackage->MarkPackageDirty();
+			}
 		}
 	}
 }
@@ -84,6 +88,7 @@ void ForcePackageNamespace(UPackage* InPackage, const FString& InNamespace)
 
 			FString& PackageLocalizationNamespaceValue = PackageMetaData->RootMetaDataMap.FindOrAdd(PackageLocalizationNamespaceKey);
 			PackageLocalizationNamespaceValue = InNamespace;
+			InPackage->MarkPackageDirty();
 		}
 	}
 }
diff --git a/Engine/Source/Runtime/CoreUObject/Private/Serialization/PropertyLocalizationDataGathering.cpp b/Engine/Source/Runtime/CoreUObject/Private/Serialization/PropertyLocalizationDataGathering.cpp
--- a/Engine/Source/Runtime/CoreUObject/Private/Serialization/PropertyLocalizationDataGathering.cpp
+++ b/Engine/Source/Runtime/CoreUObject/Private/Serialization/PropertyLocalizationDataGathering.cpp
@@ -39,7 +39,7 @@ FPropertyLocalizationDataGatherer::FPropertyLocalizationDataGatherer(TArray<FGat
 	}
 }
 
-bool FPropertyLocalizationDataGatherer::ShouldProcessObject(const UObject* Object, const EPropertyLocalizationGathererTextFlags GatherTextFlags)
+bool FPropertyLocalizationDataGatherer::ShouldProcessObject(const UObject* Object, const EPropertyLocalizationGathererTextFlags GatherTextFlags) const
 {
 	if (Object->HasAnyFlags(RF_Transient))
 	{
@@ -48,11 +48,15 @@ bool FPropertyLocalizationDataGatherer::ShouldProcessObject(const UObject* Objec
 	}
 
 	// Skip objects that we've already processed to avoid repeated work and cyclic chains
-	bool bAlreadyProcessed = false;
-	ProcessedObjects.Add(FObjectAndGatherFlags(Object, GatherTextFlags), &bAlreadyProcessed);
+	const bool bAlreadyProcessed = ProcessedObjects.Contains(FObjectAndGatherFlags(Object, GatherTextFlags));
 	return !bAlreadyProcessed;
 }
 
+void FPropertyLocalizationDataGatherer::MarkObjectProcessed(const UObject* Object, const EPropertyLocalizationGathererTextFlags GatherTextFlags)
+{
+	ProcessedObjects.Add(FObjectAndGatherFlags(Object, GatherTextFlags));
+}
+
 void FPropertyLocalizationDataGatherer::GatherLocalizationDataFromObjectWithCallbacks(const UObject* Object, const EPropertyLocalizationGathererTextFlags GatherTextFlags)
 {
 	// See if we have a custom handler for this type
@@ -72,11 +76,13 @@ void FPropertyLocalizationDataGatherer::GatherLocalizationDataFromObjectWithCall
 
 		if (ShouldProcessObject(Object, GatherTextFlags))
 		{
+			MarkObjectProcessed(Object, GatherTextFlags);
 			(*CustomCallback)(Object, *this, GatherTextFlags);
 		}
 	}
-	else
+	else if (ShouldProcessObject(Object, GatherTextFlags))
 	{
+		MarkObjectProcessed(Object, GatherTextFlags);
 		GatherLocalizationDataFromObject(Object, GatherTextFlags);
 	}
 }
@@ -85,11 +91,6 @@ void FPropertyLocalizationDataGatherer::GatherLocalizationDataFromObject(const U
 {
 	checkf(IsObjectValidForGather(Object), TEXT("Cannot gather for objects outside of the current package! Package: '%s'. Object: '%s'."), *Package->GetFullName(), *Object->GetFullName());
 
-	if (!ShouldProcessObject(Object, GatherTextFlags))
-	{
-		return;
-	}
-
 	const FString Path = Object->GetPathName();
 
 	// Gather text from our fields.
@@ -110,6 +111,7 @@ void FPropertyLocalizationDataGatherer::GatherLocalizationDataFromObject(const U
 	}
 
 	// Gather from anything that has us as their outer, as not all objects are reachable via a property pointer.
+	if (!(GatherTextFlags & EPropertyLocalizationGathererTextFlags::SkipSubObjects))
 	{
 		TArray<UObject*> ChildObjects;
 		GetObjectsWithOuter(Object, ChildObjects, false, RF_Transient, EInternalObjectFlags::PendingKill);
@@ -123,6 +125,8 @@ void FPropertyLocalizationDataGatherer::GatherLocalizationDataFromObject(const U
 
 void FPropertyLocalizationDataGatherer::GatherLocalizationDataFromObjectFields(const FString& PathToParent, const UObject* Object, const EPropertyLocalizationGathererTextFlags GatherTextFlags)
 {
+	const UObject* ArchetypeObject = Object->GetArchetype();
+
 	for (TFieldIterator<const UField> FieldIt(Object->GetClass(), EFieldIteratorFlags::IncludeSuper, EFieldIteratorFlags::ExcludeDeprecated, EFieldIteratorFlags::IncludeInterfaces); FieldIt; ++FieldIt)
 	{
 		// Gather text from the property data
@@ -130,23 +134,38 @@ void FPropertyLocalizationDataGatherer::GatherLocalizationDataFromObjectFields(c
 			const UProperty* PropertyField = Cast<UProperty>(*FieldIt);
 			if (PropertyField)
 			{
-				GatherLocalizationDataFromChildTextProperties(PathToParent, PropertyField, PropertyField->ContainerPtrToValuePtr<void>(Object), GatherTextFlags | (PropertyField->HasAnyPropertyFlags(CPF_EditorOnly) ? EPropertyLocalizationGathererTextFlags::ForceEditorOnly : EPropertyLocalizationGathererTextFlags::None));
+				const void* ValueAddress = PropertyField->ContainerPtrToValuePtr<void>(Object);
+				const void* DefaultValueAddress = nullptr;
+
+				if (ArchetypeObject)
+				{
+					const UProperty* ArchetypePropertyField = FindField<UProperty>(ArchetypeObject->GetClass(), *PropertyField->GetName());
+					if (ArchetypePropertyField && ArchetypePropertyField->IsA(PropertyField->GetClass()))
+					{
+						DefaultValueAddress = ArchetypePropertyField->ContainerPtrToValuePtr<void>(ArchetypeObject);
+					}
+				}
+
+				GatherLocalizationDataFromChildTextProperties(PathToParent, PropertyField, ValueAddress, DefaultValueAddress, GatherTextFlags | (PropertyField->HasAnyPropertyFlags(CPF_EditorOnly) ? EPropertyLocalizationGathererTextFlags::ForceEditorOnly : EPropertyLocalizationGathererTextFlags::None));
 			}
 		}
 
 		// Gather text from the script bytecode
 		{
 			const UStruct* StructField = Cast<UStruct>(*FieldIt);
-			if (StructField && IsObjectValidForGather(StructField))
+			if (StructField && IsObjectValidForGather(StructField) && ShouldProcessObject(StructField, GatherTextFlags))
 			{
+				MarkObjectProcessed(StructField, GatherTextFlags);
 				GatherLocalizationDataFromObject(StructField, GatherTextFlags);
 			}
 		}
 	}
 }
 
-void FPropertyLocalizationDataGatherer::GatherLocalizationDataFromStructFields(const FString& PathToParent, const UStruct* Struct, const void* StructData, const EPropertyLocalizationGathererTextFlags GatherTextFlags)
+void FPropertyLocalizationDataGatherer::GatherLocalizationDataFromStructFields(const FString& PathToParent, const UStruct* Struct, const void* StructData, const void* DefaultStructData, const EPropertyLocalizationGathererTextFlags GatherTextFlags)
 {
+	const UStruct* ArchetypeStruct = Cast<UStruct>(Struct->GetArchetype());
+
 	for (TFieldIterator<const UField> FieldIt(Struct, EFieldIteratorFlags::IncludeSuper, EFieldIteratorFlags::ExcludeDeprecated, EFieldIteratorFlags::IncludeInterfaces); FieldIt; ++FieldIt)
 	{
 		// Gather text from the property data
@@ -154,22 +173,35 @@ void FPropertyLocalizationDataGatherer::GatherLocalizationDataFromStructFields(c
 			const UProperty* PropertyField = Cast<UProperty>(*FieldIt);
 			if (PropertyField)
 			{
-				GatherLocalizationDataFromChildTextProperties(PathToParent, PropertyField, PropertyField->ContainerPtrToValuePtr<void>(StructData), GatherTextFlags | (PropertyField->HasAnyPropertyFlags(CPF_EditorOnly) ? EPropertyLocalizationGathererTextFlags::ForceEditorOnly : EPropertyLocalizationGathererTextFlags::None));
+				const void* ValueAddress = PropertyField->ContainerPtrToValuePtr<void>(StructData);
+				const void* DefaultValueAddress = nullptr;
+
+				if (ArchetypeStruct && DefaultStructData)
+				{
+					const UProperty* ArchetypePropertyField = FindField<UProperty>(ArchetypeStruct, *PropertyField->GetName());
+					if (ArchetypePropertyField && ArchetypePropertyField->IsA(PropertyField->GetClass()))
+					{
+						DefaultValueAddress = ArchetypePropertyField->ContainerPtrToValuePtr<void>(DefaultStructData);
+					}
+				}
+
+				GatherLocalizationDataFromChildTextProperties(PathToParent, PropertyField, ValueAddress, DefaultValueAddress, GatherTextFlags | (PropertyField->HasAnyPropertyFlags(CPF_EditorOnly) ? EPropertyLocalizationGathererTextFlags::ForceEditorOnly : EPropertyLocalizationGathererTextFlags::None));
 			}
 		}
 
 		// Gather text from the script bytecode
 		{
 			const UStruct* StructField = Cast<UStruct>(*FieldIt);
-			if (StructField && IsObjectValidForGather(StructField))
+			if (StructField && IsObjectValidForGather(StructField) && ShouldProcessObject(StructField, GatherTextFlags))
 			{
+				MarkObjectProcessed(StructField, GatherTextFlags);
 				GatherLocalizationDataFromObject(StructField, GatherTextFlags);
 			}
 		}
 	}
 }
 
-void FPropertyLocalizationDataGatherer::GatherLocalizationDataFromChildTextProperties(const FString& PathToParent, const UProperty* const Property, const void* const ValueAddress, const EPropertyLocalizationGathererTextFlags GatherTextFlags)
+void FPropertyLocalizationDataGatherer::GatherLocalizationDataFromChildTextProperties(const FString& PathToParent, const UProperty* const Property, const void* const ValueAddress, const void* const DefaultValueAddress, const EPropertyLocalizationGathererTextFlags GatherTextFlags)
 {
 	if (Property->HasAnyPropertyFlags(CPF_Transient))
 	{
@@ -184,103 +216,96 @@ void FPropertyLocalizationDataGatherer::GatherLocalizationDataFromChildTextPrope
 	const UStructProperty* const StructProperty = Cast<const UStructProperty>(Property);
 	const UObjectPropertyBase* const ObjectProperty = Cast<const UObjectPropertyBase>(Property);
 
-	// Property is a text property.
-	if (TextProperty)
-	{
-		GatherLocalizationDataFromTextProperty(PathToParent, TextProperty, ValueAddress, GatherTextFlags);
-	}
-	else
+	const EPropertyLocalizationGathererTextFlags ChildPropertyGatherTextFlags = GatherTextFlags | (Property->HasAnyPropertyFlags(CPF_EditorOnly) ? EPropertyLocalizationGathererTextFlags::ForceEditorOnly : EPropertyLocalizationGathererTextFlags::None);
+
+	// Handles both native, fixed-size arrays and plain old non-array properties.
+	const bool IsFixedSizeArray = Property->ArrayDim > 1;
+	for(int32 i = 0; i < Property->ArrayDim; ++i)
 	{
-		const EPropertyLocalizationGathererTextFlags ChildPropertyGatherTextFlags = GatherTextFlags | (Property->HasAnyPropertyFlags(CPF_EditorOnly) ? EPropertyLocalizationGathererTextFlags::ForceEditorOnly : EPropertyLocalizationGathererTextFlags::None);
+		const FString PathToElement = FString(PathToParent.IsEmpty() ? TEXT("") : PathToParent + TEXT(".")) + (IsFixedSizeArray ? Property->GetName() + FString::Printf(TEXT("[%d]"), i) : Property->GetName());
+		const void* const ElementValueAddress = reinterpret_cast<const uint8*>(ValueAddress) + Property->ElementSize * i;
+		const void* const DefaultElementValueAddress = DefaultValueAddress ? (reinterpret_cast<const uint8*>(DefaultValueAddress) + Property->ElementSize * i) : nullptr;
 
-		// Handles both native, fixed-size arrays and plain old non-array properties.
-		const bool IsFixedSizeArray = Property->ArrayDim > 1;
-		for(int32 i = 0; i < Property->ArrayDim; ++i)
+		const bool bIsDefaultValue = DefaultElementValueAddress && Property->Identical(ElementValueAddress, DefaultElementValueAddress, PPF_None);
+		if (bIsDefaultValue)
 		{
-			const FString PathToElement = FString(PathToParent.IsEmpty() ? TEXT("") : PathToParent + TEXT(".")) + (IsFixedSizeArray ? Property->GetName() + FString::Printf(TEXT("[%d]"), i) : Property->GetName());
-			const void* const ElementValueAddress = reinterpret_cast<const uint8*>(ValueAddress) + Property->ElementSize * i;
+			// Skip any properties that have the default value, as those will be gathered from the default instance
+			continue;
+		}
+
+		// Property is a text property.
+		if (TextProperty)
+		{
+			const FText* const TextElementValueAddress = static_cast<const FText*>(ElementValueAddress);
 
-			// Property is a DYNAMIC array property.
-			if (ArrayProperty)
+			UPackage* const PropertyPackage = TextProperty->GetOutermost();
+			if (FTextInspector::GetFlags(*TextElementValueAddress) & ETextFlag::ConvertedProperty)
 			{
-				// Iterate over all elements of the array.
-				FScriptArrayHelper ScriptArrayHelper(ArrayProperty, ElementValueAddress);
-				const int32 ElementCount = ScriptArrayHelper.Num();
-				for(int32 j = 0; j < ElementCount; ++j)
-				{
-					GatherLocalizationDataFromChildTextProperties(PathToElement + FString::Printf(TEXT("(%d)"), j), ArrayProperty->Inner, ScriptArrayHelper.GetRawPtr(j), ChildPropertyGatherTextFlags);
-				}
+				PropertyPackage->MarkPackageDirty();
 			}
-			// Property is a map property.
-			else if (MapProperty)
-			{
-				// Iterate over all elements of the map.
-				FScriptMapHelper ScriptMapHelper(MapProperty, ElementValueAddress);
-				const int32 ElementCount = ScriptMapHelper.Num();
-				for(int32 j = 0; j < ElementCount; ++j)
-				{
-					if (!ScriptMapHelper.IsValidIndex(j))
-					{
-						continue;
-					}
 
-					const uint8* MapPairPtr = ScriptMapHelper.GetPairPtr(j);
-					GatherLocalizationDataFromChildTextProperties(PathToElement + FString::Printf(TEXT("(%d - Key)"), j), MapProperty->KeyProp, MapPairPtr + MapProperty->MapLayout.KeyOffset, ChildPropertyGatherTextFlags);
-					GatherLocalizationDataFromChildTextProperties(PathToElement + FString::Printf(TEXT("(%d - Value)"), j), MapProperty->ValueProp, MapPairPtr + MapProperty->MapLayout.ValueOffset, ChildPropertyGatherTextFlags);
-				}
+			GatherTextInstance(*TextElementValueAddress, PathToElement, !!(GatherTextFlags & EPropertyLocalizationGathererTextFlags::ForceEditorOnlyProperties) || TextProperty->HasAnyPropertyFlags(CPF_EditorOnly));
+		}
+		// Property is a DYNAMIC array property.
+		else if (ArrayProperty)
+		{
+			// Iterate over all elements of the array.
+			FScriptArrayHelper ScriptArrayHelper(ArrayProperty, ElementValueAddress);
+			const int32 ElementCount = ScriptArrayHelper.Num();
+			for(int32 j = 0; j < ElementCount; ++j)
+			{
+				GatherLocalizationDataFromChildTextProperties(PathToElement + FString::Printf(TEXT("(%d)"), j), ArrayProperty->Inner, ScriptArrayHelper.GetRawPtr(j), nullptr, ChildPropertyGatherTextFlags);
 			}
-			// Property is a set property.
-			else if (SetProperty)
+		}
+		// Property is a map property.
+		else if (MapProperty)
+		{
+			// Iterate over all elements of the map.
+			FScriptMapHelper ScriptMapHelper(MapProperty, ElementValueAddress);
+			const int32 ElementCount = ScriptMapHelper.Num();
+			for(int32 j = 0; j < ElementCount; ++j)
 			{
-				// Iterate over all elements of the Set.
-				FScriptSetHelper ScriptSetHelper(SetProperty, ElementValueAddress);
-				const int32 ElementCount = ScriptSetHelper.Num();
-				for(int32 j = 0; j < ElementCount; ++j)
+				if (!ScriptMapHelper.IsValidIndex(j))
 				{
-					if (!ScriptSetHelper.IsValidIndex(j))
-					{
-						continue;
-					}
-
-					const uint8* ElementPtr = ScriptSetHelper.GetElementPtr(j);
-					GatherLocalizationDataFromChildTextProperties(PathToElement + FString::Printf(TEXT("(%d)"), j), SetProperty->ElementProp, ElementPtr + SetProperty->SetLayout.ElementOffset, ChildPropertyGatherTextFlags);
+					continue;
 				}
+
+				const uint8* MapPairPtr = ScriptMapHelper.GetPairPtr(j);
+				GatherLocalizationDataFromChildTextProperties(PathToElement + FString::Printf(TEXT("(%d - Key)"), j), MapProperty->KeyProp, MapPairPtr + MapProperty->MapLayout.KeyOffset, nullptr, ChildPropertyGatherTextFlags);
+				GatherLocalizationDataFromChildTextProperties(PathToElement + FString::Printf(TEXT("(%d - Value)"), j), MapProperty->ValueProp, MapPairPtr + MapProperty->MapLayout.ValueOffset, nullptr, ChildPropertyGatherTextFlags);
 			}
-			// Property is a struct property.
-			else if (StructProperty)
-			{
-				GatherLocalizationDataFromStructFields(PathToElement, StructProperty->Struct, ElementValueAddress, ChildPropertyGatherTextFlags);
-			}
-			// Property is an object property.
-			else if (ObjectProperty)
+		}
+		// Property is a set property.
+		else if (SetProperty)
+		{
+			// Iterate over all elements of the Set.
+			FScriptSetHelper ScriptSetHelper(SetProperty, ElementValueAddress);
+			const int32 ElementCount = ScriptSetHelper.Num();
+			for(int32 j = 0; j < ElementCount; ++j)
 			{
-				const UObject* InnerObject = ObjectProperty->GetObjectPropertyValue(ElementValueAddress);
-				if (InnerObject && IsObjectValidForGather(InnerObject))
+				if (!ScriptSetHelper.IsValidIndex(j))
 				{
-					GatherLocalizationDataFromObjectWithCallbacks(InnerObject, ChildPropertyGatherTextFlags);
+					continue;
 				}
+
+				const uint8* ElementPtr = ScriptSetHelper.GetElementPtr(j);
+				GatherLocalizationDataFromChildTextProperties(PathToElement + FString::Printf(TEXT("(%d)"), j), SetProperty->ElementProp, ElementPtr + SetProperty->SetLayout.ElementOffset, nullptr, ChildPropertyGatherTextFlags);
 			}
 		}
-	}
-}
-
-void FPropertyLocalizationDataGatherer::GatherLocalizationDataFromTextProperty(const FString& PathToParent, const UTextProperty* const TextProperty, const void* const ValueAddress, const EPropertyLocalizationGathererTextFlags GatherTextFlags)
-{
-	const bool IsFixedSizeArray = TextProperty->ArrayDim > 1;
-	for(int32 i = 0; i < TextProperty->ArrayDim; ++i)
-	{
-		const FString PathToElement = FString(PathToParent.IsEmpty() ? TEXT("") : PathToParent + TEXT(".")) + (IsFixedSizeArray ? TextProperty->GetName() + FString::Printf(TEXT("[%d]"), i) : TextProperty->GetName());
-		const FText* const ElementValueAddress = reinterpret_cast<const FText*>(reinterpret_cast<const uint8*>(ValueAddress) + TextProperty->ElementSize * i);
-
+		// Property is a struct property.
+		else if (StructProperty)
 		{
-			UPackage* const PropertyPackage = TextProperty->GetOutermost();
-			if( FTextInspector::GetFlags(*ElementValueAddress) & ETextFlag::ConvertedProperty )
+			GatherLocalizationDataFromStructFields(PathToElement, StructProperty->Struct, ElementValueAddress, DefaultElementValueAddress, ChildPropertyGatherTextFlags);
+		}
+		// Property is an object property.
+		else if (ObjectProperty && !(GatherTextFlags & EPropertyLocalizationGathererTextFlags::SkipSubObjects))
+		{
+			const UObject* InnerObject = ObjectProperty->GetObjectPropertyValue(ElementValueAddress);
+			if (InnerObject && IsObjectValidForGather(InnerObject))
 			{
-				PropertyPackage->MarkPackageDirty();
+				GatherLocalizationDataFromObjectWithCallbacks(InnerObject, ChildPropertyGatherTextFlags);
 			}
 		}
-
-		GatherTextInstance(*ElementValueAddress, PathToElement, !!(GatherTextFlags & EPropertyLocalizationGathererTextFlags::ForceEditorOnlyProperties) || TextProperty->HasAnyPropertyFlags(CPF_EditorOnly));
 	}
 }
 
diff --git a/Engine/Source/Runtime/CoreUObject/Private/UObject/Enum.cpp b/Engine/Source/Runtime/CoreUObject/Private/UObject/Enum.cpp
--- a/Engine/Source/Runtime/CoreUObject/Private/UObject/Enum.cpp
+++ b/Engine/Source/Runtime/CoreUObject/Private/UObject/Enum.cpp
@@ -287,7 +287,7 @@ void UEnum::AddNamesToMasterList()
 	for (TPair<FName, int64> Kvp : Names)
 	{
 		UEnum* Enum = AllEnumNames.FindRef(Kvp.Key);
-		if (Enum == nullptr)
+		if (Enum == nullptr || Enum->HasAnyFlags(RF_NewerVersionExists))
 		{
 			AllEnumNames.Add(Kvp.Key, this);
 		}
diff --git a/Engine/Source/Runtime/CoreUObject/Private/UObject/LazyObjectPtr.cpp b/Engine/Source/Runtime/CoreUObject/Private/UObject/LazyObjectPtr.cpp
--- a/Engine/Source/Runtime/CoreUObject/Private/UObject/LazyObjectPtr.cpp
+++ b/Engine/Source/Runtime/CoreUObject/Private/UObject/LazyObjectPtr.cpp
@@ -137,21 +137,29 @@ void FLazyObjectPtr::PossiblySerializeObjectGuid(UObject *Object, FArchive& Ar)
 
 					if (bDuplicate || bReassigning)
 					{
-						if (!bReassigning)
+						if (!bReassigning && OtherObject && OtherObject->HasAnyFlags(RF_NewerVersionExists))
 						{
-							// Always warn for non-map packages, skip map packages in PIE or game
-							const bool bInGame = FApp::IsGame() || Package->HasAnyPackageFlags(PKG_PlayInEditor);
-
-							UE_CLOG(!Package->ContainsMap() || !bInGame, LogUObjectGlobals, Warning, 
-								TEXT("Guid referenced by %s is already used by %s, which should never happen in the editor but could happen at runtime with duplicate level loading or PIE"), 
-								*Object->GetFullName(), *OtherObject->GetFullName());
+							GuidAnnotation.RemoveAnnotation(OtherObject);
+							GuidAnnotation.AddAnnotation(Object, Guid);
 						}
 						else
 						{
-							UE_LOG(LogUObjectGlobals, Warning, TEXT("Assigning new Guid to %s"), *Object->GetFullName());
+							if (!bReassigning)
+							{
+								// Always warn for non-map packages, skip map packages in PIE or game
+								const bool bInGame = FApp::IsGame() || Package->HasAnyPackageFlags(PKG_PlayInEditor);
+
+								UE_CLOG(!Package->ContainsMap() || !bInGame, LogUObjectGlobals, Warning,
+									TEXT("Guid referenced by %s is already used by %s, which should never happen in the editor but could happen at runtime with duplicate level loading or PIE"),
+									*Object->GetFullName(), *OtherObject->GetFullName());
+							}
+							else
+							{
+								UE_LOG(LogUObjectGlobals, Warning, TEXT("Assigning new Guid to %s"), *Object->GetFullName());
+							}
+							// This guid is in use, which should never happen in the editor but could happen at runtime with duplicate level loading or PIE. If so give it an invalid GUID and don't add to the annotation map.
+							Guid = FGuid();
 						}
-						// This guid is in use, which should never happen in the editor but could happen at runtime with duplicate level loading or PIE. If so give it an invalid GUID and don't add to the annotation map.
-						Guid = FGuid();
 					}
 					else
 					{
diff --git a/Engine/Source/Runtime/CoreUObject/Private/UObject/Obj.cpp b/Engine/Source/Runtime/CoreUObject/Private/UObject/Obj.cpp
--- a/Engine/Source/Runtime/CoreUObject/Private/UObject/Obj.cpp
+++ b/Engine/Source/Runtime/CoreUObject/Private/UObject/Obj.cpp
@@ -1172,6 +1172,42 @@ void UObject::SerializeScriptProperties( FArchive& Ar ) const
 }
 
 
+void UObject::BuildSubobjectMapping(UObject* OtherObject, TMap<UObject*, UObject*>& ObjectMapping) const
+{
+	UPackage* ThisPackage = GetOutermost();
+	UPackage* OtherPackage = OtherObject->GetOutermost();
+
+	ForEachObjectWithOuter(this, [&](UObject* InSubObject)
+	{
+		if (ObjectMapping.Contains(InSubObject))
+		{
+			return;
+		}
+
+		FString NewSubObjectName = InSubObject->GetName();
+
+		UClass* OtherSubObjectClass = InSubObject->GetClass();
+		if (OtherSubObjectClass->ClassGeneratedBy && OtherSubObjectClass->ClassGeneratedBy->GetOutermost() == ThisPackage)
+		{
+			// This is a generated class type, so we actually need to use the new generated class type from the new package otherwise our type check will fail
+			FString NewClassName = OtherSubObjectClass->GetPathName(ThisPackage);
+			NewClassName = FString::Printf(TEXT("%s.%s"), *OtherPackage->GetName(), *NewClassName);
+
+			OtherSubObjectClass = LoadObject<UClass>(OtherPackage, *NewClassName);
+		}
+
+		//UObject* OtherSubObject = StaticLoadObject(OtherSubObjectClass, OtherObject, *NewSubObjectName, nullptr, LOAD_Quiet | LOAD_NoRedirects, nullptr, true);
+		UObject* OtherSubObject = StaticFindObjectFast(OtherSubObjectClass, OtherObject, *NewSubObjectName);
+		ObjectMapping.Emplace(InSubObject, OtherSubObject);
+
+		if (OtherSubObject)
+		{
+			InSubObject->BuildSubobjectMapping(OtherSubObject, ObjectMapping);
+		}
+	}, false, RF_NoFlags, EInternalObjectFlags::PendingKill);
+}
+
+
 void UObject::CollectDefaultSubobjects( TArray<UObject*>& OutSubobjectArray, bool bIncludeNestedSubobjects/*=false*/ )
 {
 	OutSubobjectArray.Empty();
@@ -1461,7 +1497,7 @@ void UObject::GetAssetRegistryTagMetadata(TMap<FName, FAssetRegistryTagMetadata>
 }
 #endif
 
-bool UObject::IsAsset () const
+bool UObject::IsAsset() const
 {
 	// Assets are not transient or CDOs. They must be public.
 	const bool bHasValidObjectFlags = !HasAnyFlags(RF_Transient | RF_ClassDefaultObject) && HasAnyFlags(RF_Public);
@@ -1471,8 +1507,8 @@ bool UObject::IsAsset () const
 		// Don't count objects embedded in other objects (e.g. font textures, sequences, material expressions)
 		if ( UPackage* LocalOuterPackage = dynamic_cast<UPackage*>(GetOuter()) )
 		{
-			// Also exclude any objects found in the transient package.
-			return LocalOuterPackage != GetTransientPackage();
+			// Also exclude any objects found in the transient package, or in a package that is transient.
+			return LocalOuterPackage != GetTransientPackage() && !LocalOuterPackage->HasAnyFlags(RF_Transient);
 		}
 	}
 
@@ -3812,7 +3848,7 @@ void StaticUObjectInit()
 	UObjectBaseInit();
 
 	// Allocate special packages.
-	GObjTransientPkg = NewObject<UPackage>(nullptr, TEXT("/Engine/Transient"));
+	GObjTransientPkg = NewObject<UPackage>(nullptr, TEXT("/Engine/Transient"), RF_Transient);
 	GObjTransientPkg->AddToRoot();
 
 	if( FParse::Param( FCommandLine::Get(), TEXT("VERIFYGC") ) )
diff --git a/Engine/Source/Runtime/CoreUObject/Private/UObject/PropertyMap.cpp b/Engine/Source/Runtime/CoreUObject/Private/UObject/PropertyMap.cpp
--- a/Engine/Source/Runtime/CoreUObject/Private/UObject/PropertyMap.cpp
+++ b/Engine/Source/Runtime/CoreUObject/Private/UObject/PropertyMap.cpp
@@ -1003,40 +1003,6 @@ bool UMapProperty::ConvertFromType(const FPropertyTag& Tag, FArchive& Ar, uint8*
 	return false;
 }
 
-/**
- * Checks to see if this property already has the supplied value as a key
- * @param	InMap			The address of the map
- * @param	InBaseAddress	The base address of the map
- * @param	InValue			The value to find in the map
- * @return True if InValue is a key in the map, false otherwise
- */
-bool UMapProperty::HasKey(void* InMap, void* InBaseAddress, const FString& InValue) const
-{
-	FScriptMapHelper MapHelper(this, InMap);
-
-	for (int32 Index = 0, ItemsLeft = MapHelper.Num(); ItemsLeft > 0; ++Index)
-	{
-		if (MapHelper.IsValidIndex(Index))
-		{
-			--ItemsLeft;
-
-			uint8* PairPtr = MapHelper.GetPairPtr(Index);
-			uint8* KeyPtr = KeyProp->ContainerPtrToValuePtr<uint8>(PairPtr);
-
-			FString KeyValue;
-			if ( KeyPtr != InBaseAddress && KeyProp->ExportText_Direct(KeyValue, KeyPtr, KeyPtr, nullptr, 0) )
-			{
-				if ( (Cast<UObjectProperty>(KeyProp) != nullptr && KeyValue.Contains(InValue)) || InValue == KeyValue )
-				{
-					return true;
-				}
-			}
-		}
-	}
-
-	return false;
-}
-
 IMPLEMENT_CORE_INTRINSIC_CLASS(UMapProperty, UProperty,
 	{
 		Class->EmitObjectReference(STRUCT_OFFSET(UMapProperty, KeyProp),   TEXT("KeyProp"));
diff --git a/Engine/Source/Runtime/CoreUObject/Private/UObject/PropertyNumeric.cpp b/Engine/Source/Runtime/CoreUObject/Private/UObject/PropertyNumeric.cpp
--- a/Engine/Source/Runtime/CoreUObject/Private/UObject/PropertyNumeric.cpp
+++ b/Engine/Source/Runtime/CoreUObject/Private/UObject/PropertyNumeric.cpp
@@ -29,6 +29,14 @@ uint8 UNumericProperty::ReadEnumAsUint8(FArchive& Ar, UStruct* DefaultsStruct, c
 
 	int64 Result = Enum->GetValueByName(EnumName);
 	if (!Enum->IsValidEnumValue(Result))
+	{
+		const int32 EnumIndex = UEnum::FindEnumRedirects(Enum, EnumName);
+		if (EnumIndex != INDEX_NONE)
+		{
+			Result = Enum->GetValueByIndex(EnumIndex);
+		}
+	}
+	if (!Enum->IsValidEnumValue(Result))
 	{
 		UE_LOG(
 			LogClass,
diff --git a/Engine/Source/Runtime/CoreUObject/Private/UObject/PropertySet.cpp b/Engine/Source/Runtime/CoreUObject/Private/UObject/PropertySet.cpp
--- a/Engine/Source/Runtime/CoreUObject/Private/UObject/PropertySet.cpp
+++ b/Engine/Source/Runtime/CoreUObject/Private/UObject/PropertySet.cpp
@@ -869,38 +869,6 @@ bool USetProperty::ConvertFromType(const FPropertyTag& Tag, FArchive& Ar, uint8*
 	return false;
 }
 
-/**
- * Checks to see if this property already has the supplied value as an element
- * @param	InSet			The address of the set
- * @param	InBaseAddress	The base address of the set
- * @param	InValue			The value to find in the set
- * @return True if InValue is an element in the set, false otherwise
- */
-bool USetProperty::HasElement(void* InSet, void* InBaseAddress, const FString& InValue) const
-{
-	FScriptSetHelper SetHelper(this, InSet);
-
-	for ( int32 Index = 0, ItemsLeft = SetHelper.Num(); ItemsLeft > 0; ++Index )
-	{
-		if (SetHelper.IsValidIndex(Index))
-		{
-			--ItemsLeft;
-
-			uint8* Element = SetHelper.GetElementPtr(Index);
-
-			FString ElementValue;
-			if (Element != InBaseAddress && ElementProp->ExportText_Direct(ElementValue, Element, Element, nullptr, 0))
-			{
-				if ( (Cast<UObjectProperty>(ElementProp) != nullptr && ElementValue.Contains(InValue)) || ElementValue == InValue)
-				{
-					return true;
-				}
-			}
-		}
-	}
-
-	return false;
-}
 
 IMPLEMENT_CORE_INTRINSIC_CLASS(USetProperty, UProperty,
 	{
diff --git a/Engine/Source/Runtime/CoreUObject/Private/UObject/TextProperty.cpp b/Engine/Source/Runtime/CoreUObject/Private/UObject/TextProperty.cpp
--- a/Engine/Source/Runtime/CoreUObject/Private/UObject/TextProperty.cpp
+++ b/Engine/Source/Runtime/CoreUObject/Private/UObject/TextProperty.cpp
@@ -84,8 +84,8 @@ bool UTextProperty::Identical( const void* A, const void* B, uint32 PortFlags )
 
 void UTextProperty::SerializeItem( FArchive& Ar, void* Value, void const* Defaults ) const
 {
-	const TCppType PropertyValue = GetPropertyValue(Value);
-	Ar << *GetPropertyValuePtr(Value);
+	TCppType* TextPtr = GetPropertyValuePtr(Value);
+	Ar << *TextPtr;
 }
 
 void UTextProperty::ExportTextItem( FString& ValueStr, const void* PropertyValue, const void* DefaultValue, UObject* Parent, int32 PortFlags, UObject* ExportRootScope ) const
diff --git a/Engine/Source/Runtime/CoreUObject/Private/UObject/UObjectGlobals.cpp b/Engine/Source/Runtime/CoreUObject/Private/UObject/UObjectGlobals.cpp
--- a/Engine/Source/Runtime/CoreUObject/Private/UObject/UObjectGlobals.cpp
+++ b/Engine/Source/Runtime/CoreUObject/Private/UObject/UObjectGlobals.cpp
@@ -99,6 +99,8 @@ FCoreUObjectDelegates::FReinstanceHotReloadedClassesDelegate FCoreUObjectDelegat
 FCoreUObjectDelegates::FIsPackageOKToSaveDelegate FCoreUObjectDelegates::IsPackageOKToSaveDelegate;
 FCoreUObjectDelegates::FAutoPackageBackupDelegate FCoreUObjectDelegates::AutoPackageBackupDelegate;
 
+FCoreUObjectDelegates::FOnPackageReloaded FCoreUObjectDelegates::OnPackageReloaded;
+
 FCoreUObjectDelegates::FOnPreObjectPropertyChanged FCoreUObjectDelegates::OnPreObjectPropertyChanged;
 FCoreUObjectDelegates::FOnObjectPropertyChanged FCoreUObjectDelegates::OnObjectPropertyChanged;
 
@@ -718,12 +720,17 @@ bool ResolveName(UObject*& InPackage, FString& InOutName, bool Create, bool Thro
 			Create         = false;
 		}
 
-		// In case this is a short script package name, convert to long name before passing to CreatePackage/FindObject.
-		FName* ScriptPackageName = FPackageName::FindScriptPackageName(*PartialName);
-		if (ScriptPackageName)
+		FName* ScriptPackageName = nullptr;
+		if (!bSubobjectPath)
 		{
-			PartialName = ScriptPackageName->ToString();
+			// In case this is a short script package name, convert to long name before passing to CreatePackage/FindObject.
+			ScriptPackageName = FPackageName::FindScriptPackageName(*PartialName);
+			if (ScriptPackageName)
+			{
+				PartialName = ScriptPackageName->ToString();
+			}
 		}
+
 		// Only long package names are allowed so don't even attempt to create one because whatever the name represents
 		// it's not a valid package name anyway.
 		
diff --git a/Engine/Source/Runtime/CoreUObject/Private/UObject/UObjectHash.cpp b/Engine/Source/Runtime/CoreUObject/Private/UObject/UObjectHash.cpp
--- a/Engine/Source/Runtime/CoreUObject/Private/UObject/UObjectHash.cpp
+++ b/Engine/Source/Runtime/CoreUObject/Private/UObject/UObjectHash.cpp
@@ -469,7 +469,6 @@ UObject* StaticFindObjectFastInternalThreadSafe(FUObjectHashTables& ThreadHash,
 					Or, if the object is a top-level package then accept it immediately.*/
 					&& (bAnyPackage || !Object->GetOuter())
 
-
 					/** If a class was specified, check that the object is of the correct class */
 					&& (ObjectClass == nullptr || (bExactClass ? Object->GetClass() == ObjectClass : Object->IsA(ObjectClass)))
 
@@ -507,7 +506,7 @@ UObject* StaticFindObjectFastInternal(UClass* ObjectClass, UObject* ObjectPackag
 	check(ObjectPackage != ANY_PACKAGE); // this could never have returned anything but nullptr
 	// If they specified an outer use that during the hashing
 	auto& ThreadHash = FUObjectHashTables::Get();
-	UObject* Result = StaticFindObjectFastInternalThreadSafe(ThreadHash, ObjectClass, ObjectPackage, ObjectName, bExactClass, bAnyPackage, ExcludeFlags, ExclusiveInternalFlags);
+	UObject* Result = StaticFindObjectFastInternalThreadSafe(ThreadHash, ObjectClass, ObjectPackage, ObjectName, bExactClass, bAnyPackage, ExcludeFlags | RF_NewerVersionExists, ExclusiveInternalFlags);
 	return Result;
 }
 
