Copying //UE4/Dev-Editor to //UE4/Dev-Main (Source: //UE4/Dev-Editor @ 3228984)

diff --git a/Engine/Source/Runtime/CoreUObject/Public/Serialization/PropertyLocalizationDataGathering.h b/Engine/Source/Runtime/CoreUObject/Public/Serialization/PropertyLocalizationDataGathering.h
--- a/Engine/Source/Runtime/CoreUObject/Public/Serialization/PropertyLocalizationDataGathering.h
+++ b/Engine/Source/Runtime/CoreUObject/Public/Serialization/PropertyLocalizationDataGathering.h
@@ -33,6 +33,11 @@ enum class EPropertyLocalizationGathererTextFlags : uint8
 	 * Force all gathered text to be treated as editor-only data.
 	 */
 	ForceEditorOnly = ForceEditorOnlyProperties | ForceEditorOnlyScriptData,
+
+	/**
+	 * Don't process any sub-objects (either inner objects or object pointers).
+	 */
+	SkipSubObjects = 1<<3,
 };
 ENUM_CLASS_FLAGS(EPropertyLocalizationGathererTextFlags);
 
@@ -66,13 +71,15 @@ public:
 	void GatherLocalizationDataFromObjectWithCallbacks(const UObject* Object, const EPropertyLocalizationGathererTextFlags GatherTextFlags);
 	void GatherLocalizationDataFromObject(const UObject* Object, const EPropertyLocalizationGathererTextFlags GatherTextFlags);
 	void GatherLocalizationDataFromObjectFields(const FString& PathToParent, const UObject* Object, const EPropertyLocalizationGathererTextFlags GatherTextFlags);
-	void GatherLocalizationDataFromStructFields(const FString& PathToParent, const UStruct* Struct, const void* StructData, const EPropertyLocalizationGathererTextFlags GatherTextFlags);
-	void GatherLocalizationDataFromChildTextProperties(const FString& PathToParent, const UProperty* const Property, const void* const ValueAddress, const EPropertyLocalizationGathererTextFlags GatherTextFlags);
-	void GatherLocalizationDataFromTextProperty(const FString& PathToParent, const UTextProperty* const TextProperty, const void* const ValueAddress, const EPropertyLocalizationGathererTextFlags GatherTextFlags);
+	void GatherLocalizationDataFromStructFields(const FString& PathToParent, const UStruct* Struct, const void* StructData, const void* DefaultStructData, const EPropertyLocalizationGathererTextFlags GatherTextFlags);
+	void GatherLocalizationDataFromChildTextProperties(const FString& PathToParent, const UProperty* const Property, const void* const ValueAddress, const void* const DefaultValueAddress, const EPropertyLocalizationGathererTextFlags GatherTextFlags);
 
 	void GatherTextInstance(const FText& Text, const FString& Description, const bool bIsEditorOnly);
 	void GatherScriptBytecode(const FString& PathToScript, const TArray<uint8>& ScriptData, const bool bIsEditorOnly);
 
+	bool ShouldProcessObject(const UObject* Object, const EPropertyLocalizationGathererTextFlags GatherTextFlags) const;
+	void MarkObjectProcessed(const UObject* Object, const EPropertyLocalizationGathererTextFlags GatherTextFlags);
+
 	static FLocalizationDataGatheringCallbackMap& GetTypeSpecificLocalizationDataGatheringCallbacks();
 
 	FORCEINLINE TArray<FGatherableTextData>& GetGatherableTextDataArray() const
@@ -118,8 +125,6 @@ private:
 		uint32 KeyHash;
 	};
 
-	bool ShouldProcessObject(const UObject* Object, const EPropertyLocalizationGathererTextFlags GatherTextFlags);
-
 	TArray<FGatherableTextData>& GatherableTextDataArray;
 	const UPackage* Package;
 	EPropertyLocalizationGathererResultFlags& ResultFlags;
diff --git a/Engine/Source/Runtime/CoreUObject/Public/UObject/GCObjectScopeGuard.h b/Engine/Source/Runtime/CoreUObject/Public/UObject/GCObjectScopeGuard.h
new file mode 100644
--- /dev/null
+++ b/Engine/Source/Runtime/CoreUObject/Public/UObject/GCObjectScopeGuard.h
@@ -0,0 +1,35 @@
+// Copyright 1998-2016 Epic Games, Inc. All Rights Reserved.
+
+#pragma once
+
+#include "GCObject.h"
+
+/**
+ * Specific implementation of FGCObject that prevents a single UObject-based pointer from being GC'd while this guard is in scope.
+ * @note This is the lean version of TStrongObjectPtr which uses an inline FGCObject so *can't* safely be used with containers that treat types as trivially relocatable.
+ */
+class FGCObjectScopeGuard : public FGCObject
+{
+public:
+	explicit FGCObjectScopeGuard(const UObject* InObject)
+		: Object(InObject)
+	{
+	}
+
+	virtual ~FGCObjectScopeGuard()
+	{
+	}
+
+	/** Non-copyable */
+	FGCObjectScopeGuard(const FGCObjectScopeGuard&) = delete;
+	FGCObjectScopeGuard& operator=(const FGCObjectScopeGuard&) = delete;
+
+	//~ FGCObject interface
+	virtual void AddReferencedObjects(FReferenceCollector& Collector) override
+	{
+		Collector.AddReferencedObject(Object);
+	}
+
+private:
+	const UObject* Object;
+};
diff --git a/Engine/Source/Runtime/CoreUObject/Public/UObject/Object.h b/Engine/Source/Runtime/CoreUObject/Public/UObject/Object.h
--- a/Engine/Source/Runtime/CoreUObject/Public/UObject/Object.h
+++ b/Engine/Source/Runtime/CoreUObject/Public/UObject/Object.h
@@ -802,6 +802,14 @@ public:
 	UFunction* FindFunction( FName InName ) const;
 	UFunction* FindFunctionChecked( FName InName ) const;
 
+	/**
+	 * Given OtherObject (which will be the same type as 'this'), recursively find any matching sub-objects from 'this' that also exist within OtherObject, and add the mappings to ObjectMapping.
+	 *
+	 * @param	OtherObject		The to find matching sub-objects within.
+	 * @param	ObjectMapping	The complete mapping between this object hierarchy and the other object hierarchy.
+	 */
+	virtual void BuildSubobjectMapping(UObject* OtherObject, TMap<UObject*, UObject*>& ObjectMapping) const;
+
 	/**
 	 * Uses the TArchiveObjectReferenceCollector to build a list of all components referenced by this object which have this object as the outer
 	 *
diff --git a/Engine/Source/Runtime/CoreUObject/Public/UObject/ObjectMacros.h b/Engine/Source/Runtime/CoreUObject/Public/UObject/ObjectMacros.h
--- a/Engine/Source/Runtime/CoreUObject/Public/UObject/ObjectMacros.h
+++ b/Engine/Source/Runtime/CoreUObject/Public/UObject/ObjectMacros.h
@@ -105,7 +105,7 @@ enum EPackageFlags
 	PKG_ServerSideOnly				= 0x00000004,   // Only needed on the server side.
 	PKG_CompiledIn					= 0x00000010,   // This package is from "compiled in" classes.
 	PKG_ForDiffing					= 0x00000020,	// This package was loaded just for the purposes of diffing
-	PKG_EditorOnly					= 0x00000040, // This is editor-only package (for example: editor module script package)
+	PKG_EditorOnly					= 0x00000040,	// This is editor-only package (for example: editor module script package)
 	PKG_Developer					= 0x00000080,	// Developer module
 //	PKG_Unused						= 0x00000100,
 //	PKG_Unused						= 0x00000200,
@@ -117,7 +117,7 @@ enum EPackageFlags
 	PKG_Need						= 0x00008000,	// Client needs to download this package.
 	PKG_Compiling					= 0x00010000,	// package is currently being compiled
 	PKG_ContainsMap					= 0x00020000,	// Set if the package contains a ULevel/ UWorld object
-	PKG_RequiresLocalizationGather		= 0x00040000,	// Set if the package contains any data to be gathered by localization
+	PKG_RequiresLocalizationGather	= 0x00040000,	// Set if the package contains any data to be gathered by localization
 	PKG_DisallowLazyLoading			= 0x00080000,	// Set if the archive serializing this package cannot use lazy loading
 	PKG_PlayInEditor				= 0x00100000,	// Set if the package was created for the purpose of PIE
 	PKG_ContainsScript				= 0x00200000,	// Package is allowed to contain UClass objects
@@ -473,7 +473,7 @@ enum EObjectFlags
 	RF_KeepForCooker			=0x00000800,	///< Keep this object during garbage collection because it's still being used by the cooker
 	RF_NeedPostLoad				=0x00001000,	///< Object needs to be postloaded.
 	RF_NeedPostLoadSubobjects	=0x00002000,	///< During load, indicates that the object still needs to instance subobjects and fixup serialized component references
-	//RF_Unused				=0x00004000,	///
+	RF_NewerVersionExists		=0x00004000,	///< Object has been consigned to oblivion due to its owner package being reloaded, and a newer version currently exists
 	RF_BeginDestroyed			=0x00008000,	///< BeginDestroy has been called on the object.
 	RF_FinishDestroyed			=0x00010000,	///< FinishDestroy has been called on the object.
 
diff --git a/Engine/Source/Runtime/CoreUObject/Public/UObject/StrongObjectPtr.h b/Engine/Source/Runtime/CoreUObject/Public/UObject/StrongObjectPtr.h
new file mode 100644
--- /dev/null
+++ b/Engine/Source/Runtime/CoreUObject/Public/UObject/StrongObjectPtr.h
@@ -0,0 +1,135 @@
+// Copyright 1998-2016 Epic Games, Inc. All Rights Reserved.
+
+#pragma once
+
+#include "GCObject.h"
+#include "UniquePtr.h"
+
+/**
+ * Specific implementation of FGCObject that prevents a single UObject-based pointer from being GC'd while this guard is in scope.
+ * @note This is the "full-fat" version of FGCObjectScopeGuard which uses a heap-allocated FGCObject so *can* safely be used with containers that treat types as trivially relocatable.
+ */
+template <typename ObjectType>
+class TStrongObjectPtr
+{
+public:
+	TStrongObjectPtr()
+		: ReferenceCollector(MakeUnique<FInternalReferenceCollector>(nullptr))
+	{
+		static_assert(TPointerIsConvertibleFromTo<ObjectType, const volatile UObject>::Value, "TStrongObjectPtr can only be constructed with UObject types");
+	}
+
+	explicit TStrongObjectPtr(ObjectType* InObject)
+		: ReferenceCollector(MakeUnique<FInternalReferenceCollector>(InObject))
+	{
+		static_assert(TPointerIsConvertibleFromTo<ObjectType, const volatile UObject>::Value, "TStrongObjectPtr can only be constructed with UObject types");
+	}
+
+	TStrongObjectPtr(const TStrongObjectPtr& InOther)
+		: ReferenceCollector(MakeUnique<FInternalReferenceCollector>(InOther.Get()))
+	{
+	}
+
+	template <typename OtherObjectType>
+	TStrongObjectPtr(const TStrongObjectPtr<OtherObjectType>& InOther)
+		: ReferenceCollector(MakeUnique<FInternalReferenceCollector>(InOther.Get()))
+	{
+	}
+
+	FORCEINLINE TStrongObjectPtr& operator=(const TStrongObjectPtr& InOther)
+	{
+		ReferenceCollector->Set(InOther.Get());
+		return *this;
+	}
+
+	template <typename OtherObjectType>
+	FORCEINLINE TStrongObjectPtr& operator=(const TStrongObjectPtr<OtherObjectType>& InOther)
+	{
+		ReferenceCollector->Set(InOther.Get());
+		return *this;
+	}
+
+	FORCEINLINE ObjectType& operator*() const
+	{
+		check(IsValid());
+		return *Get();
+	}
+
+	FORCEINLINE ObjectType* operator->() const
+	{
+		check(IsValid());
+		return Get();
+	}
+
+	FORCEINLINE bool IsValid() const
+	{
+		return Get() != nullptr;
+	}
+
+	FORCEINLINE explicit operator bool() const
+	{
+		return Get() != nullptr;
+	}
+
+	FORCEINLINE ObjectType* Get() const
+	{
+		return ReferenceCollector->Get();
+	}
+
+	FORCEINLINE void Reset(ObjectType* InNewObject = nullptr)
+	{
+		ReferenceCollector->Set(InNewObject);
+	}
+
+	FORCEINLINE friend uint32 GetTypeHash(const TStrongObjectPtr& InStrongObjectPtr)
+	{
+		return GetTypeHash(InStrongObjectPtr.Get());
+	}
+
+private:
+	class FInternalReferenceCollector : public FGCObject
+	{
+	public:
+		FInternalReferenceCollector(ObjectType* InObject)
+			: Object(InObject)
+		{
+		}
+
+		virtual ~FInternalReferenceCollector()
+		{
+		}
+
+		FORCEINLINE ObjectType* Get() const
+		{
+			return Object;
+		}
+
+		FORCEINLINE void Set(ObjectType* InObject)
+		{
+			Object = InObject;
+		}
+
+		//~ FGCObject interface
+		virtual void AddReferencedObjects(FReferenceCollector& Collector) override
+		{
+			Collector.AddReferencedObject(Object);
+		}
+
+	private:
+		ObjectType* Object;
+	};
+
+	TUniquePtr<FInternalReferenceCollector> ReferenceCollector;
+};
+
+template <typename LHSObjectType, typename RHSObjectType>
+FORCEINLINE bool operator==(const TStrongObjectPtr<LHSObjectType>& InLHS, const TStrongObjectPtr<RHSObjectType>& InRHS)
+{
+	return InLHS.Get() == InRHS.Get();
+}
+
+template <typename LHSObjectType, typename RHSObjectType>
+FORCEINLINE bool operator!=(const TStrongObjectPtr<LHSObjectType>& InLHS, const TStrongObjectPtr<RHSObjectType>& InRHS)
+{
+	return InLHS.Get() != InRHS.Get();
+}
diff --git a/Engine/Source/Runtime/CoreUObject/Public/UObject/UObjectGlobals.h b/Engine/Source/Runtime/CoreUObject/Public/UObject/UObjectGlobals.h
--- a/Engine/Source/Runtime/CoreUObject/Public/UObject/UObjectGlobals.h
+++ b/Engine/Source/Runtime/CoreUObject/Public/UObject/UObjectGlobals.h
@@ -1824,6 +1824,10 @@ protected:
 DECLARE_DELEGATE_RetVal_TwoParams( bool, FCheckForAutoAddDelegate, UPackage*, const FString& );
 DECLARE_DELEGATE_OneParam( FAddPackageToDefaultChangelistDelegate, const TCHAR* );
 
+/** Defined in PackageReload.h */
+enum class EPackageReloadPhase : uint8;
+class FPackageReloadedEvent;
+
 /**
  * Global CoreUObject delegates
  */
@@ -1844,6 +1848,10 @@ struct COREUOBJECT_API FCoreUObjectDelegates
 	/** Delegate type for making auto backup of package */
 	DECLARE_DELEGATE_RetVal_OneParam(bool, FAutoPackageBackupDelegate, const UPackage&);
 
+	/** Called by ReloadPackage during package reloading. It will be called several times for different phases of fix-up to allow custom code to handle updating objects as needed */
+	DECLARE_MULTICAST_DELEGATE_TwoParams(FOnPackageReloaded, EPackageReloadPhase, FPackageReloadedEvent*);
+	static FOnPackageReloaded OnPackageReloaded;
+
 #if WITH_EDITOR
 	// Callback for all object modifications
 	DECLARE_MULTICAST_DELEGATE_OneParam(FOnObjectModified, UObject*);
diff --git a/Engine/Source/Runtime/CoreUObject/Public/UObject/UnrealType.h b/Engine/Source/Runtime/CoreUObject/Public/UObject/UnrealType.h
--- a/Engine/Source/Runtime/CoreUObject/Public/UObject/UnrealType.h
+++ b/Engine/Source/Runtime/CoreUObject/Public/UObject/UnrealType.h
@@ -2564,8 +2564,6 @@ public:
 	virtual bool SameType(const UProperty* Other) const override;
 	virtual bool ConvertFromType(const FPropertyTag& Tag, FArchive& Ar, uint8* Data, UStruct* DefaultsStruct, bool& bOutAdvanceProperty) override;
 	// End of UProperty interface
-
-	bool HasKey(void* InMap, void* InBaseAddress, const FString& InValue) const;
 };
 
 // need to break this out a different type so that the DECLARE_CASTED_CLASS_INTRINSIC macro can digest the comma
@@ -2616,8 +2614,6 @@ public:
 	virtual bool SameType(const UProperty* Other) const override;
 	virtual bool ConvertFromType(const FPropertyTag& Tag, FArchive& Ar, uint8* Data, UStruct* DefaultsStruct, bool& bOutAdvanceProperty) override;
 	// End of UProperty interface
-
-	bool HasElement(void* InSet, void* InBaseAddress, const FString& InValue) const;
 };
 
 /**
@@ -3333,6 +3329,39 @@ public:
 		}
 	}
 
+	/**
+	 * Checks if a key in the map matches the specified key
+	 *
+	 * @param	InBaseAddress	The base address of the map
+	 * @param	InKeyValue		The key to find within the map
+	 *
+	 * @return	True if the key is found, false otherwise
+	 */
+	bool HasKey(const void* InBaseAddress, const FString& InKeyValue) const
+	{
+		for (int32 Index = 0, ItemsLeft = Num(); ItemsLeft > 0; ++Index)
+		{
+			if (IsValidIndex(Index))
+			{
+				--ItemsLeft;
+
+				const uint8* PairPtr = GetPairPtr(Index);
+				const uint8* KeyPtr = KeyProp->ContainerPtrToValuePtr<const uint8>(PairPtr);
+
+				FString KeyValue;
+				if (KeyPtr != InBaseAddress && KeyProp->ExportText_Direct(KeyValue, KeyPtr, KeyPtr, nullptr, 0))
+				{
+					if ((Cast<UObjectProperty>(KeyProp) != nullptr && KeyValue.Contains(InKeyValue)) || InKeyValue == KeyValue)
+					{
+						return true;
+					}
+				}
+			}
+		}
+
+		return false;
+	}
+
 private:
 	/**
 	 * Internal function to call into the property system to construct / initialize elements.
@@ -3772,6 +3801,38 @@ public:
 		}
 	}
 
+	/**
+	 * Checks if an element has already been added to the set
+	 *
+	 * @param	InBaseAddress	The base address of the set
+	 * @param	InElementValue	The element value to check for
+	 *
+	 * @return	True if the element is found in the set, false otherwise
+	 */
+	bool HasElement(void* InBaseAddress, const FString& InElementValue) const
+	{
+		for (int32 Index = 0, ItemsLeft = Num(); ItemsLeft > 0; ++Index)
+		{
+			if (IsValidIndex(Index))
+			{
+				--ItemsLeft;
+
+				const uint8* Element = GetElementPtr(Index);
+
+				FString ElementValue;
+				if (Element != InBaseAddress && ElementProp->ExportText_Direct(ElementValue, Element, Element, nullptr, 0))
+				{
+					if ((Cast<UObjectProperty>(ElementProp) != nullptr && ElementValue.Contains(InElementValue)) || ElementValue == InElementValue)
+					{
+						return true;
+					}
+				}
+			}
+		}
+
+		return false;
+	}
+
 private:
 	/**
 	* Internal function to call into the property system to construct / initialize elements.
@@ -4178,6 +4239,8 @@ namespace EPropertyChangeType
 	const Type Duplicate = 1 << 3;
 	//Interactive, e.g. dragging a slider. Will be followed by a ValueSet when finished.
 	const Type Interactive = 1 << 4;
+	//Redirected.  Used when property references are updated due to content hot-reloading, or an asset being replaced during asset deletion (aka, asset consolidation).
+	const Type Redirected = 1 << 5;
 };
 
 /**
diff --git a/Engine/Source/Runtime/Engine/Classes/Components/StaticMeshComponent.h b/Engine/Source/Runtime/Engine/Classes/Components/StaticMeshComponent.h
--- a/Engine/Source/Runtime/Engine/Classes/Components/StaticMeshComponent.h
+++ b/Engine/Source/Runtime/Engine/Classes/Components/StaticMeshComponent.h
@@ -240,7 +240,7 @@ class ENGINE_API UStaticMeshComponent : public UMeshComponent
 	UPROPERTY(EditAnywhere, BlueprintReadOnly, Category=Lighting, meta=(InlineEditConditionToggle))
 	uint32 bOverrideLightMapRes:1;
 
-	/** Light map resolution to use on this component, used if bOverrideLightMapRes is true */
+	/** Light map resolution to use on this component, used if bOverrideLightMapRes is true and there is a valid StaticMesh. */
 	UPROPERTY(EditAnywhere, BlueprintReadOnly, Category=Lighting, meta=(ClampMax = 4096, editcondition="bOverrideLightMapRes") )
 	int32 OverriddenLightMapRes;
 
@@ -298,6 +298,10 @@ class ENGINE_API UStaticMeshComponent : public UMeshComponent
 	/** Material Bounds used for texture streaming. */
 	UPROPERTY(NonTransactional)
 	TArray<uint32> MaterialStreamingRelativeBoxes;
+
+	/** The component has some custom painting on LODs or not. */
+	UPROPERTY()
+	bool bCustomOverrideVertexColorPerLOD;
 #endif
 
 	/** The Lightmass settings for this object. */
diff --git a/Engine/Source/Runtime/Engine/Classes/Curves/CurveOwnerInterface.h b/Engine/Source/Runtime/Engine/Classes/Curves/CurveOwnerInterface.h
--- a/Engine/Source/Runtime/Engine/Classes/Curves/CurveOwnerInterface.h
+++ b/Engine/Source/Runtime/Engine/Classes/Curves/CurveOwnerInterface.h
@@ -24,6 +24,9 @@ public:
 	/** Called to modify the owner of the curve */
 	virtual void ModifyOwner() = 0;
 
+	/** Returns the owner(s) of the curve */
+	virtual TArray<const UObject*> GetOwners() const = 0;
+
 	/** Called to make curve owner transactional */
 	virtual void MakeTransactional() = 0;
 
diff --git a/Engine/Source/Runtime/Engine/Classes/EdGraph/EdGraph.h b/Engine/Source/Runtime/Engine/Classes/EdGraph/EdGraph.h
--- a/Engine/Source/Runtime/Engine/Classes/EdGraph/EdGraph.h
+++ b/Engine/Source/Runtime/Engine/Classes/EdGraph/EdGraph.h
@@ -68,6 +68,8 @@ class ENGINE_API UEdGraph : public UObject
 {
 	GENERATED_UCLASS_BODY()
 
+public:
+
 	/** The schema that this graph obeys */
 	UPROPERTY()
 	TSubclassOf<class UEdGraphSchema>  Schema;
@@ -118,8 +120,9 @@ public:
 	/** Remove a listener for OnGraphChanged events */
 	void RemoveOnGraphChangedHandler( FDelegateHandle Handle );
 
-#if WITH_EDITORONLY_DATA
 	//~ Begin UObject interface
+	virtual void BuildSubobjectMapping(UObject* OtherObject, TMap<UObject*, UObject*>& ObjectMapping) const override;
+#if WITH_EDITORONLY_DATA
 	virtual void Serialize( FArchive& Ar ) override;
 	virtual void PostInitProperties() override;
 	virtual void PostLoad() override;
diff --git a/Engine/Source/Runtime/Engine/Classes/EdGraph/EdGraphNode.h b/Engine/Source/Runtime/Engine/Classes/EdGraph/EdGraphNode.h
--- a/Engine/Source/Runtime/Engine/Classes/EdGraph/EdGraphNode.h
+++ b/Engine/Source/Runtime/Engine/Classes/EdGraph/EdGraphNode.h
@@ -295,6 +295,12 @@ public:
 	/** Gets the index for a pin */
 	int32 GetPinIndex(UEdGraphPin* Pin) const;
 
+	/** Gets the pin at a given index 
+	* @param Index The zero-based index of the pin to access.
+	* @return The pin found at this location or nullptr if invalid index.
+	*/
+	UEdGraphPin* GetPinAt(int32 Index) const;
+
 	/** Break all links on this node */
 	void BreakAllNodeLinks();
 
@@ -514,6 +520,14 @@ public:
 	/** Return whether to draw this node as a comment node */
 	virtual bool ShouldDrawNodeAsComment() const { return false; }
 
+	/** 
+	* Returns whether to draw this node as a control point only (knot/reroute node). Note that this means that the node should only have on input and output pin.
+	* @param OutInputPinIndex The index in the pins array associated with the control point input pin.
+	* @param OutOutputPinIndex The index in the pins array associated with the control point output pin.
+	* @return Whether or not to draw this node as a control point.
+	*/
+	virtual bool ShouldDrawNodeAsControlPointOnly(int32& OutInputPinIndex, int32& OutOutputPinIndex) const;
+
 	/**
 	 * Add's node data to the search metadata, override to collect more data that may be desirable to search for
 	 *
diff --git a/Engine/Source/Runtime/Engine/Classes/EdGraph/EdGraphSchema.h b/Engine/Source/Runtime/Engine/Classes/EdGraph/EdGraphSchema.h
--- a/Engine/Source/Runtime/Engine/Classes/EdGraph/EdGraphSchema.h
+++ b/Engine/Source/Runtime/Engine/Classes/EdGraph/EdGraphSchema.h
@@ -852,6 +852,14 @@ class ENGINE_API UEdGraphSchema : public UObject
 	/* When a node is removed, this method determines whether we should remove it immediately or use the old (slower) code path that results in all node being recreated: */
 	virtual bool ShouldAlwaysPurgeOnModification() const { return true; }
 
+	/** 
+	 * Perform any logic necessary to safely remove this node from the graph.  
+	 * @param Graph		Type of pin to drop onto the node
+	 * @param Node		Direction of the source pin
+	 * @return			Whether or not the node was successfully deleted.
+    */
+	virtual bool SafeDeleteNodeFromGraph(UEdGraph* Graph, UEdGraphNode* Node) const { return false; }
+
 	/*
 	 * Some schemas have nodes that support the user dynamically adding pins when dropping a connection on the node
 	 *
diff --git a/Engine/Source/Runtime/Engine/Classes/Engine/Blueprint.h b/Engine/Source/Runtime/Engine/Classes/Engine/Blueprint.h
--- a/Engine/Source/Runtime/Engine/Classes/Engine/Blueprint.h
+++ b/Engine/Source/Runtime/Engine/Classes/Engine/Blueprint.h
@@ -664,6 +664,8 @@ public:
 	/** Collect blueprints that depend on this blueprint. */
 	virtual void GatherDependencies(TSet<TWeakObjectPtr<UBlueprint>>& InDependencies) const;
 
+	virtual void ReplaceDeprecatedNodes();
+
 #endif	//#if WITH_EDITOR
 
 	//~ Begin UObject Interface
diff --git a/Engine/Source/Runtime/Engine/Classes/Engine/DeveloperSettings.h b/Engine/Source/Runtime/Engine/Classes/Engine/DeveloperSettings.h
--- a/Engine/Source/Runtime/Engine/Classes/Engine/DeveloperSettings.h
+++ b/Engine/Source/Runtime/Engine/Classes/Engine/DeveloperSettings.h
@@ -32,6 +32,9 @@ public:
 	virtual FText GetSectionText() const;
 	/** Gets the description for the section, uses the classes ToolTip by default. */
 	virtual FText GetSectionDescription() const;
+
+	/** Whether or not this class supports auto registration or if the settings have a custom setup */
+	virtual bool SupportsAutoRegistration() const { return true; }
 #endif
 
 	/** Gets a custom widget for the settings.  This is only for very custom situations. */
diff --git a/Engine/Source/Runtime/Engine/Classes/Engine/FontFace.h b/Engine/Source/Runtime/Engine/Classes/Engine/FontFace.h
--- a/Engine/Source/Runtime/Engine/Classes/Engine/FontFace.h
+++ b/Engine/Source/Runtime/Engine/Classes/Engine/FontFace.h
@@ -13,7 +13,7 @@ struct FPropertyChangedEvent;
 
 /**
  * A font face asset contains the raw payload data for a source TTF/OTF file as used by FreeType.
- * During cook this asset type generates a ".ufont" file containing the raw payload data.
+ * During cook this asset type generates a ".ufont" file containing the raw payload data (unless loaded "Inline").
  */
 UCLASS(hidecategories=Object, autoexpandcategories=FontFace, MinimalAPI, BlueprintType)
 class UFontFace : public UObject, public IFontFaceInterface
@@ -25,6 +25,7 @@ public:
 	UFontFace();
 
 	//~ Begin UObject Interface
+	virtual void Serialize(FArchive& Ar) override;
 	virtual void GetResourceSizeEx(FResourceSizeEx& CumulativeResourceSize) override;
 #if WITH_EDITOR
 	virtual void PostEditChangeProperty(FPropertyChangedEvent& PropertyChangedEvent) override;
@@ -41,9 +42,7 @@ public:
 	virtual const FString& GetFontFilename() const override;
 	virtual EFontHinting GetHinting() const override;
 	virtual EFontLoadingPolicy GetLoadingPolicy() const override;
-#if WITH_EDITORONLY_DATA
-	virtual const TArray<uint8>& GetFontFaceData() const override;
-#endif // WITH_EDITORONLY_DATA
+	virtual FFontFaceDataConstRef GetFontFaceData() const override;
 	virtual FString GetCookedFilename() const override;
 	//~ End IFontFaceInterface interface
 
@@ -59,9 +58,12 @@ public:
 	UPROPERTY(EditAnywhere, BlueprintReadWrite, Category=FontFace)
 	EFontLoadingPolicy LoadingPolicy;
 
+	/** The data associated with the font face. This should always be filled in providing the source filename is valid. */
+	FFontFaceDataRef FontFaceData;
+
 #if WITH_EDITORONLY_DATA
 	/** The data associated with the font face. This should always be filled in providing the source filename is valid. */
 	UPROPERTY()
-	TArray<uint8> FontFaceData;
+	TArray<uint8> FontFaceData_DEPRECATED;
 #endif // WITH_EDITORONLY_DATA
 };
diff --git a/Engine/Source/Runtime/Engine/Classes/Engine/PostProcessVolume.h b/Engine/Source/Runtime/Engine/Classes/Engine/PostProcessVolume.h
--- a/Engine/Source/Runtime/Engine/Classes/Engine/PostProcessVolume.h
+++ b/Engine/Source/Runtime/Engine/Classes/Engine/PostProcessVolume.h
@@ -24,30 +24,30 @@ class APostProcessVolume : public AVolume, public IInterface_PostProcessVolume
 	GENERATED_UCLASS_BODY()
 
 	/** Post process settings to use for this volume. */
-	UPROPERTY(interp, Category=PostProcessVolume, meta=(ShowOnlyInnerProperties))
+	UPROPERTY(interp, Category=PostProcessVolumeSettings)
 	struct FPostProcessSettings Settings;
 
 	/**
 	 * Priority of this volume. In the case of overlapping volumes the one with the highest priority
 	 * overrides the lower priority ones. The order is undefined if two or more overlapping volumes have the same priority.
 	 */
-	UPROPERTY(EditAnywhere, BlueprintReadWrite, Category=PostProcessVolume)
+	UPROPERTY(EditAnywhere, BlueprintReadWrite, Category=PostProcessVolumeSettings)
 	float Priority;
 
 	/** World space radius around the volume that is used for blending (only if not unbound).			*/
-	UPROPERTY(interp, Category=PostProcessVolume, meta=(ClampMin = "0.0", UIMin = "0.0", UIMax = "6000.0"))
+	UPROPERTY(interp, Category=PostProcessVolumeSettings, meta=(ClampMin = "0.0", UIMin = "0.0", UIMax = "6000.0"))
 	float BlendRadius;
 
 	/** 0:no effect, 1:full effect */
-	UPROPERTY(interp, Category=PostProcessVolume, BlueprintReadWrite, meta=(UIMin = "0.0", UIMax = "1.0"))
+	UPROPERTY(interp, Category=PostProcessVolumeSettings, BlueprintReadWrite, meta=(UIMin = "0.0", UIMax = "1.0"))
 	float BlendWeight;
 
 	/** Whether this volume is enabled or not.															*/
-	UPROPERTY(EditAnywhere, BlueprintReadWrite, Category=PostProcessVolume)
+	UPROPERTY(EditAnywhere, BlueprintReadWrite, Category=PostProcessVolumeSettings)
 	uint32 bEnabled:1;
 
 	/** Whether this volume bounds are used or it affects the whole world.								*/
-	UPROPERTY(EditAnywhere, BlueprintReadWrite, Category=PostProcessVolume)
+	UPROPERTY(EditAnywhere, BlueprintReadWrite, Category=PostProcessVolumeSettings)
 	uint32 bUnbound:1;
 
 	//~ Begin IInterface_PostProcessVolume Interface
