Fixing UT linker errors after Dev-Core copy up

diff --git a/Engine/Source/Runtime/CoreUObject/Public/UObject/LinkerLoad.h b/Engine/Source/Runtime/CoreUObject/Public/UObject/LinkerLoad.h
--- a/Engine/Source/Runtime/CoreUObject/Public/UObject/LinkerLoad.h
+++ b/Engine/Source/Runtime/CoreUObject/Public/UObject/LinkerLoad.h
@@ -485,7 +485,7 @@ public:
 	 *					If Object is a UClass and the class default object has already been created, calls
 	 *					Preload for the class default object as well.
 	 */
-	void Preload( UObject* Object ) override;
+	COREUOBJECT_API void Preload( UObject* Object ) override;
 
 	/**
 	 * Before loading a persistent object from disk, this function can be used to discover
