Edit integrate of CL #3231129 on Core

diff --git a/Engine/Source/Runtime/CoreUObject/Private/UObject/UObjectGlobals.cpp b/Engine/Source/Runtime/CoreUObject/Private/UObject/UObjectGlobals.cpp
--- a/Engine/Source/Runtime/CoreUObject/Private/UObject/UObjectGlobals.cpp
+++ b/Engine/Source/Runtime/CoreUObject/Private/UObject/UObjectGlobals.cpp
@@ -3698,13 +3698,10 @@ UObject* FObjectInitializer::CreateDefaultSubobject(UObject* Outer, FName Subobj
 			bool bOwnerArchetypeIsNotNative;
 			UClass* OuterArchetypeClass;
 
-			if (bIsTransient)
-			{
-				SubobjectFlags |= RF_Transient;
-			}
+			// It is not safe to mark this component as properly transient, that results in it being nulled incorrectly
 
-				OuterArchetypeClass = Outer->GetArchetype()->GetClass();
-				bOwnerArchetypeIsNotNative = !OuterArchetypeClass->HasAnyClassFlags(CLASS_Native | CLASS_Intrinsic);
+			OuterArchetypeClass = Outer->GetArchetype()->GetClass();
+			bOwnerArchetypeIsNotNative = !OuterArchetypeClass->HasAnyClassFlags(CLASS_Native | CLASS_Intrinsic);
 
 			const bool bOwnerTemplateIsNotCDO = ObjectArchetype != nullptr && ObjectArchetype != Outer->GetClass()->GetDefaultObject(false) && !Outer->HasAnyFlags(RF_ClassDefaultObject);
 #if !UE_BUILD_SHIPPING
diff --git a/Engine/Source/Runtime/CoreUObject/Public/UObject/Object.h b/Engine/Source/Runtime/CoreUObject/Public/UObject/Object.h
--- a/Engine/Source/Runtime/CoreUObject/Public/UObject/Object.h
+++ b/Engine/Source/Runtime/CoreUObject/Public/UObject/Object.h
@@ -65,7 +65,7 @@ class COREUOBJECT_API UObject : public UObjectBaseUtility
 	* Create a component or subobject only to be used with the editor.
 	* @param	TReturnType					class of return type, all overrides must be of this type
 	* @param	SubobjectName				name of the new component
-	* @param	bTransient					true if the component is being assigned to a transient property
+	* @param	bTransient					true if the component is being assigned to a transient property. This does not make the component itself transient, but does stop it from inheriting parent defaults
 	*/
 	template<class TReturnType>
 	TReturnType* CreateEditorOnlyDefaultSubobject(FName SubobjectName, bool bTransient = false)
@@ -78,7 +78,7 @@ class COREUOBJECT_API UObject : public UObjectBaseUtility
 	* Create a component or subobject
 	* @param	TReturnType					class of return type, all overrides must be of this type
 	* @param	SubobjectName				name of the new component
-	* @param	bTransient					true if the component is being assigned to a transient property
+	* @param	bTransient					true if the component is being assigned to a transient property. This does not make the component itself transient, but does stop it from inheriting parent defaults
 	*/
 	template<class TReturnType>
 	TReturnType* CreateDefaultSubobject(FName SubobjectName, bool bTransient = false)
@@ -92,7 +92,7 @@ class COREUOBJECT_API UObject : public UObjectBaseUtility
 	* @param TReturnType class of return type, all overrides must be of this type
 	* @param TClassToConstructByDefault class to construct by default
 	* @param SubobjectName name of the new component
-	* @param bTransient		true if the component is being assigned to a transient property
+	* @param bTransient		true if the component is being assigned to a transient property. This does not make the component itself transient, but does stop it from inheriting parent defaults
 	*/
 	template<class TReturnType, class TClassToConstructByDefault>
 	TReturnType* CreateDefaultSubobject(FName SubobjectName, bool bTransient = false)
@@ -105,7 +105,7 @@ class COREUOBJECT_API UObject : public UObjectBaseUtility
 	* when a derived class specified DoNotCreateDefaultSubobject with the subobject's name.
 	* @param	TReturnType					class of return type, all overrides must be of this type
 	* @param	SubobjectName				name of the new component
-	* @param	bTransient					true if the component is being assigned to a transient property
+	* @param	bTransient					true if the component is being assigned to a transient property. This does not make the component itself transient, but does stop it from inheriting parent defaults
 	*/
 	template<class TReturnType>
 	TReturnType* CreateOptionalDefaultSubobject(FName SubobjectName, bool bTransient = false)
@@ -119,7 +119,7 @@ class COREUOBJECT_API UObject : public UObjectBaseUtility
 	* when a derived class specified DoNotCreateDefaultSubobject with the subobject's name.
 	* @param	TReturnType					class of return type, all overrides must be of this type
 	* @param	SubobjectName				name of the new component
-	* @param	bTransient					true if the component is being assigned to a transient property
+	* @param	bTransient					true if the component is being assigned to a transient property. This does not make the component itself transient, but does stop it from inheriting parent defaults
 	*/
 	template<class TReturnType>
 	TReturnType* CreateAbstractDefaultSubobject(FName SubobjectName, bool bTransient = false)
