Fixing compile errors when compiling with Malloc Profiler enabled

diff --git a/Engine/Source/Runtime/Core/Private/ProfilingDebugging/MallocProfiler.cpp b/Engine/Source/Runtime/Core/Private/ProfilingDebugging/MallocProfiler.cpp
--- a/Engine/Source/Runtime/Core/Private/ProfilingDebugging/MallocProfiler.cpp
+++ b/Engine/Source/Runtime/Core/Private/ProfilingDebugging/MallocProfiler.cpp
@@ -14,6 +14,8 @@
 #include "Misc/Paths.h"
 #include "HAL/TlsAutoCleanup.h"
 #include "Misc/App.h"
+#include "Misc/CoreDelegates.h"
+#include "ProfilingDebugging/ProfilingHelpers.h"
 
 #if USE_MALLOC_PROFILER
 
diff --git a/Engine/Source/Runtime/CoreUObject/Private/UObject/UObjectBaseUtility.cpp b/Engine/Source/Runtime/CoreUObject/Private/UObject/UObjectBaseUtility.cpp
--- a/Engine/Source/Runtime/CoreUObject/Private/UObject/UObjectBaseUtility.cpp
+++ b/Engine/Source/Runtime/CoreUObject/Private/UObject/UObjectBaseUtility.cpp
@@ -10,6 +10,7 @@
 #include "Templates/Casts.h"
 #include "UObject/Interface.h"
 #include "Modules/ModuleManager.h"
+#include "ProfilingDebugging/MallocProfiler.h"
 
 /***********************/
 /******** Names ********/
