Fixing fallout from EDL 'easy switch' refactor: s.PreloadPackageDependencies will not be used when EDL is enabled

diff --git a/Engine/Source/Runtime/CoreUObject/Private/Serialization/AsyncLoading.cpp b/Engine/Source/Runtime/CoreUObject/Private/Serialization/AsyncLoading.cpp
--- a/Engine/Source/Runtime/CoreUObject/Private/Serialization/AsyncLoading.cpp
+++ b/Engine/Source/Runtime/CoreUObject/Private/Serialization/AsyncLoading.cpp
@@ -584,7 +584,7 @@ int32 FAsyncLoadingThread::CreateAsyncPackagesFromQueue(bool bUseTimeLimit, bool
 
 			if (GPreloadPackageDependencies && IsPlatformFileCompatibleWithDependencyPreloading() && 
 					// we don't want preloading of dependencies with event driven loading, we would rather find the dependencies naturally.
-				 (!IsEventDrivenLoaderEnabledInCookedBuilds() || !FPlatformProperties::RequiresCookedData()))
+				 !GEventDrivenLoaderEnabled)
 			{
 				AssetRegistry = IAssetRegistryInterface::GetPtr();
 			}
@@ -5874,7 +5874,7 @@ void FAsyncPackage::UpdateLoadPercentage()
 int32 LoadPackageAsync(const FString& InName, const FGuid* InGuid /*= nullptr*/, const TCHAR* InPackageToLoadFrom /*= nullptr*/, FLoadPackageAsyncDelegate InCompletionDelegate /*= FLoadPackageAsyncDelegate()*/, EPackageFlags InPackageFlags /*= PKG_None*/, int32 InPIEInstanceID /*= INDEX_NONE*/, int32 InPackagePriority /*= 0*/)
 {
 #if !WITH_EDITOR
-	if (GPreloadPackageDependencies)
+	if (GPreloadPackageDependencies && !GEventDrivenLoaderEnabled)
 	{
 		// If dependency preloading is enabled, we need to force the asset registry module to be loaded on the game thread
 		// as it will potentiall be used on the async loading thread, which isn't allowed to load modules.
diff --git a/Engine/Source/Runtime/CoreUObject/Private/UObject/UObjectGlobals.cpp b/Engine/Source/Runtime/CoreUObject/Private/UObject/UObjectGlobals.cpp
--- a/Engine/Source/Runtime/CoreUObject/Private/UObject/UObjectGlobals.cpp
+++ b/Engine/Source/Runtime/CoreUObject/Private/UObject/UObjectGlobals.cpp
@@ -1424,7 +1424,7 @@ UPackage* LoadPackageInternal(UPackage* InOuter, const TCHAR* InLongPackageName,
 	bool bAllowDependencyPreloading = !InOuter && ((LoadFlags & (LOAD_DisableDependencyPreloading | LOAD_ForFileDiff)) == 0);
 	static auto CVarPreloadDependencies = IConsoleManager::Get().FindConsoleVariable(TEXT("s.PreloadPackageDependencies"));
 	
-	if (bAllowDependencyPreloading && CVarPreloadDependencies && CVarPreloadDependencies->GetInt() != 0)
+	if (!GEventDrivenLoaderEnabled && bAllowDependencyPreloading && CVarPreloadDependencies && CVarPreloadDependencies->GetInt() != 0)
 	{		
 		if (IsPlatformFileCompatibleWithDependencyPreloading())
 		{
