Preflight now shows success with no warnings.

diff --git a/Engine/Documentation/Extras/API/Build/Source/DocumentationEditor.Target.cs b/Engine/Documentation/Extras/API/Build/Source/DocumentationEditor.Target.cs
--- a/Engine/Documentation/Extras/API/Build/Source/DocumentationEditor.Target.cs
+++ b/Engine/Documentation/Extras/API/Build/Source/DocumentationEditor.Target.cs
@@ -6,24 +6,13 @@ using System.Collections.Generic;
 public class DocumentationEditorTarget : TargetRules
 {
 	public DocumentationEditorTarget(TargetInfo Target)
+	: base(Target)
 	{
 		Type = TargetType.Editor;
-	}
-
-	//
-	// TargetRules interface.
-	//
-
-	public override void SetupBinaries(
-		TargetInfo Target,
-		ref List<UEBuildBinaryConfiguration> OutBuildBinaryConfigurations,
-		ref List<string> OutExtraModuleNames
-		)
-	{
-		OutExtraModuleNames.Add("UE4Game");
-		OutExtraModuleNames.Add("OnlineSubsystemAmazon");
-		OutExtraModuleNames.Add("OnlineSubsystemFacebook");
-		OutExtraModuleNames.Add("OnlineSubsystemNull");
-		OutExtraModuleNames.Add("OnlineSubsystemSteam");
+		ExtraModuleNames.Add("UE4Game");
+		ExtraModuleNames.Add("OnlineSubsystemAmazon");
+		ExtraModuleNames.Add("OnlineSubsystemFacebook");
+		ExtraModuleNames.Add("OnlineSubsystemNull");
+		ExtraModuleNames.Add("OnlineSubsystemSteam");
 	}
 }
diff --git a/Engine/Source/Programs/UnrealDocTool/DoxygenInputFilter/DoxygenInputFilter.pl b/Engine/Source/Programs/UnrealDocTool/DoxygenInputFilter/DoxygenInputFilter.pl
--- a/Engine/Source/Programs/UnrealDocTool/DoxygenInputFilter/DoxygenInputFilter.pl
+++ b/Engine/Source/Programs/UnrealDocTool/DoxygenInputFilter/DoxygenInputFilter.pl
@@ -52,12 +52,16 @@ while (<FILE1>)
 #print "lcn2 = " . $lastclassname . "\n";
 			}
 		}
-		else
+		elsif (index($line, "GENERATED_UCLASS_BODY()") >= 0)
 		{
-			if (index($line, "GENERATED_UCLASS_BODY()") >= 0)
+			# $lastclassname might not be initialized. If so, the previous logic block failed to catch a class name, or "GENERATED_UCLASS_BODY()" is floating around loose in a file somewhere.
+			$line = "GENERATED_BODY() public: " . $lastclassname . "(const FObjectInitializer& ObjectInitializer);\n";
+		}
+		elsif (index($line, "enum ") >= 0)
+		{
+			if ($line =~ /\s+enum \w+_API E\w+/)
 			{
-				# $lastclassname might not be initialized. If so, the previous logic block failed to catch a class name, or "GENERATED_UCLASS_BODY()" is floating around loose in a file somewhere.
-				$line = "GENERATED_BODY() public: " . $lastclassname . "(const FObjectInitializer& ObjectInitializer);\n";
+				$line = "enum " . substr($line, (index($line, "_API") + 5));
 			}
 		}
 		$commentposition = index($line, "\/\/");
diff --git a/Engine/Source/Runtime/CoreUObject/Public/UObject/NoExportTypes.h b/Engine/Source/Runtime/CoreUObject/Public/UObject/NoExportTypes.h
--- a/Engine/Source/Runtime/CoreUObject/Public/UObject/NoExportTypes.h
+++ b/Engine/Source/Runtime/CoreUObject/Public/UObject/NoExportTypes.h
@@ -220,65 +220,40 @@ UENUM(BlueprintType)
 enum class EUnit : uint8
 {
 	/** Scalar distance/length unit. */
-
-	Micrometers, Millimeters, Centimeters, Meters, Kilometers,
-	Inches, Feet, Yards, Miles,
-	Lightyears,
-
-
+	Micrometers, Millimeters, Centimeters, Meters, Kilometers, Inches, Feet, Yards, Miles, Lightyears,
+	
 	/** Angular units */
-
 	Degrees, Radians,
-
-
+	
 	/** Speed units */
-
 	MetersPerSecond, KilometersPerHour, MilesPerHour,
-
-
+	
 	/** Temperature units */
-
 	Celsius, Farenheit, Kelvin,
-
-
+	
 	/** Mass units */
-
-	Micrograms, Milligrams, Grams, Kilograms, MetricTons,
-	Ounces, Pounds, Stones,
-
-
+	Micrograms, Milligrams, Grams, Kilograms, MetricTons, Ounces, Pounds, Stones,
+	
 	/** Force units */
-
 	Newtons, PoundsForce, KilogramsForce,
-
-
+	
 	/** Frequency units */
-
 	Hertz, Kilohertz, Megahertz, Gigahertz, RevolutionsPerMinute,
-
-
-	/** Data Size units */
 	
+	/** Data Size units */	
 	Bytes, Kilobytes, Megabytes, Gigabytes, Terabytes,
-
-
-	/** Luminous flux units */
 	
+	/** Luminous flux units */	
 	Lumens,
-
-
-	/** Time units */
 	
+	/** Time units */	
 	Milliseconds, Seconds, Minutes, Hours, Days, Months, Years,
 
-	/** Arbitrary multiplier */
-	
+	/** Arbitrary multiplier */	
 	Multiplier,
 
-
-	/** Symbolic entry, not specifiable on meta data. */
-	
-	Unspecified
+	/** Symbolic entry, not specifiable on meta data. */	
+	Unspecified,
 };
 
 /** A globally unique identifier. */
diff --git a/Engine/Source/Runtime/Engine/Classes/GameFramework/PlayerStart.h b/Engine/Source/Runtime/Engine/Classes/GameFramework/PlayerStart.h
--- a/Engine/Source/Runtime/Engine/Classes/GameFramework/PlayerStart.h
+++ b/Engine/Source/Runtime/Engine/Classes/GameFramework/PlayerStart.h
@@ -15,9 +15,12 @@
 UCLASS(Blueprintable, ClassGroup=Common, hidecategories=Collision)
 class ENGINE_API APlayerStart : public ANavigationObjectBase
 {
-	/** To take more control over PlayerStart selection, you can override the virtual AGameModeBase::FindPlayerStart and AGameModeBase::ChoosePlayerStart function */
+	GENERATED_BODY()
+public:
+
+	APlayerStart(const FObjectInitializer& ObjectInitializer);
 
-	GENERATED_UCLASS_BODY()
+	/*~ To take more control over PlayerStart selection, you can override the virtual AGameModeBase::FindPlayerStart and AGameModeBase::ChoosePlayerStart functions. */
 
 	/** Used when searching for which playerstart to use. */
 	UPROPERTY(EditAnywhere, BlueprintReadWrite, Category=Object)
@@ -37,6 +40,3 @@ public:
 	class UArrowComponent* GetArrowComponent() const;
 #endif
 };
-
-
-
