Copying //UE4/Dev-Core to //UE4/Dev-Main (Source: //UE4/Dev-Core @ 3283640)

diff --git a/Engine/Source/Runtime/Core/Public/GenericPlatform/GenericPlatformFile.h b/Engine/Source/Runtime/Core/Public/GenericPlatform/GenericPlatformFile.h
--- a/Engine/Source/Runtime/Core/Public/GenericPlatform/GenericPlatformFile.h
+++ b/Engine/Source/Runtime/Core/Public/GenericPlatform/GenericPlatformFile.h
@@ -336,6 +336,9 @@ public:
 
 	virtual void GetTimeStampPair(const TCHAR* PathA, const TCHAR* PathB, FDateTime& OutTimeStampA, FDateTime& OutTimeStampB);
 
+	/** Return the modification time of a file in the local time of the calling code (GetTimeStamp returns UTC). Returns FDateTime::MinValue() on failure **/
+	virtual FDateTime	GetTimeStampLocal(const TCHAR* Filename);
+
 	/** 
 	 * Call the Visit function of the visitor once for each file or directory in a directory tree. This function explores subdirectories.
 	 * @param Directory		The directory to iterate the contents of, recursively.
diff --git a/Engine/Source/Runtime/Core/Public/HAL/IPlatformFileOpenLogWrapper.h b/Engine/Source/Runtime/Core/Public/HAL/IPlatformFileOpenLogWrapper.h
--- a/Engine/Source/Runtime/Core/Public/HAL/IPlatformFileOpenLogWrapper.h
+++ b/Engine/Source/Runtime/Core/Public/HAL/IPlatformFileOpenLogWrapper.h
@@ -11,11 +11,37 @@
 #include "Misc/DateTime.h"
 #include "GenericPlatform/GenericPlatformFile.h"
 #include "Misc/Paths.h"
+#include "Async/AsyncFileHandle.h"
 
 class IAsyncReadFileHandle;
+class FPlatformFileOpenLog;
 
 #if !UE_BUILD_SHIPPING
 
+class FLoggingAsyncReadFileHandle final : public IAsyncReadFileHandle
+{
+	FPlatformFileOpenLog* Owner;
+	FString Filename;
+	IAsyncReadFileHandle* ActualRequest;
+public:
+	FLoggingAsyncReadFileHandle(FPlatformFileOpenLog* InOwner, const TCHAR* InFilename, IAsyncReadFileHandle* InActualRequest)
+		: Owner(InOwner)
+		, Filename(InFilename)
+		, ActualRequest(InActualRequest)
+	{
+	}
+	~FLoggingAsyncReadFileHandle()
+	{
+		delete ActualRequest;
+	}
+	virtual IAsyncReadRequest* SizeRequest(FAsyncFileCallBack* CompleteCallback = nullptr) override
+	{
+		return ActualRequest->SizeRequest(CompleteCallback);
+	}
+	virtual IAsyncReadRequest* ReadRequest(int64 Offset, int64 BytesToRead, EAsyncIOPriority Priority = AIOP_Normal, FAsyncFileCallBack* CompleteCallback = nullptr, uint8* UserSuppliedMemory = nullptr) override;
+};
+
+
 class CORE_API FPlatformFileOpenLog : public IPlatformFile
 {
 protected:
@@ -153,17 +179,7 @@ public:
 		IFileHandle* Result = LowerLevel->OpenRead(Filename, bAllowWrite);
 		if (Result)
 		{
-			CriticalSection.Lock();
-			if (FilenameAccessMap.Find(Filename) == nullptr)
-			{
-				FilenameAccessMap.Emplace(Filename, ++OpenOrder);
-				FString Text = FString::Printf(TEXT("\"%s\" %llu\n"), Filename, OpenOrder);
-				for (auto File = LogOutput.CreateIterator(); File; ++File)
-				{
-					(*File)->Write((uint8*)StringCast<ANSICHAR>(*Text).Get(), Text.Len());
-				}
-			}
-			CriticalSection.Unlock();
+			AddToOpenLog(Filename);
 		}
 		return Result;
 	}
@@ -233,23 +249,30 @@ public:
 	}
 	virtual IAsyncReadFileHandle* OpenAsyncRead(const TCHAR* Filename) override
 	{
-		IAsyncReadFileHandle* Result = LowerLevel->OpenAsyncRead(Filename);
-		if (Result)
+		// we must not record the "open" here...what matters is when we start reading the file!
+		return new FLoggingAsyncReadFileHandle(this, Filename, LowerLevel->OpenAsyncRead(Filename));
+	}
+
+	void AddToOpenLog(const TCHAR* Filename)
+	{
+		CriticalSection.Lock();
+		if (FilenameAccessMap.Find(Filename) == nullptr)
 		{
-			CriticalSection.Lock();
-			if (FilenameAccessMap.Find(Filename) == nullptr)
+			FilenameAccessMap.Emplace(Filename, ++OpenOrder);
+			FString Text = FString::Printf(TEXT("\"%s\" %llu\n"), Filename, OpenOrder);
+			for (auto File = LogOutput.CreateIterator(); File; ++File)
 			{
-				FilenameAccessMap.Emplace(Filename, ++OpenOrder);
-				FString Text = FString::Printf(TEXT("\"%s\" %llu\n"), Filename, OpenOrder);
-				for (auto File = LogOutput.CreateIterator(); File; ++File)
-				{
-					(*File)->Write((uint8*)StringCast<ANSICHAR>(*Text).Get(), Text.Len());
-				}
+				(*File)->Write((uint8*)StringCast<ANSICHAR>(*Text).Get(), Text.Len());
 			}
-			CriticalSection.Unlock();
 		}
-		return Result;
+		CriticalSection.Unlock();
 	}
 };
 
+IAsyncReadRequest* FLoggingAsyncReadFileHandle::ReadRequest(int64 Offset, int64 BytesToRead, EAsyncIOPriority Priority, FAsyncFileCallBack* CompleteCallback, uint8* UserSuppliedMemory)
+{
+	Owner->AddToOpenLog(*Filename);
+	return ActualRequest->ReadRequest(Offset, BytesToRead, Priority, CompleteCallback, UserSuppliedMemory);
+}
+
 #endif // !UE_BUILD_SHIPPING
diff --git a/Engine/Source/Runtime/Core/Public/HAL/Platform.h b/Engine/Source/Runtime/Core/Public/HAL/Platform.h
--- a/Engine/Source/Runtime/Core/Public/HAL/Platform.h
+++ b/Engine/Source/Runtime/Core/Public/HAL/Platform.h
@@ -264,6 +264,14 @@
 	#define PLATFORM_SUPPORTS_TBB 0
 #endif
 
+#ifndef PLATFORM_MAX_CACHED_SYNC_FILE_HANDLES_PER_GENERIC_ASYNC_FILE_HANDLE
+	#define PLATFORM_MAX_CACHED_SYNC_FILE_HANDLES_PER_GENERIC_ASYNC_FILE_HANDLE 4
+#endif
+
+#ifndef PLATFORM_FORCE_SINGLE_SYNC_FILE_HANDLE_PER_GENERIC_ASYNC_FILE_HANDLE
+	#define PLATFORM_FORCE_SINGLE_SYNC_FILE_HANDLE_PER_GENERIC_ASYNC_FILE_HANDLE 0
+#endif
+
 #ifndef PLATFORM_SUPPORTS_JEMALLOC
 	#define PLATFORM_SUPPORTS_JEMALLOC 0
 #endif
diff --git a/Engine/Source/Runtime/Core/Public/Misc/ConfigCacheIni.h b/Engine/Source/Runtime/Core/Public/Misc/ConfigCacheIni.h
--- a/Engine/Source/Runtime/Core/Public/Misc/ConfigCacheIni.h
+++ b/Engine/Source/Runtime/Core/Public/Misc/ConfigCacheIni.h
@@ -763,7 +763,7 @@ public:
 	 * @param Platform The platform to load the .ini for (if NULL, uses current)
 	 * @param bForceReload If true, the destination .in will be regenerated from the source, otherwise this will only process if the dest isn't in GConfig
 	 * @param bRequireDefaultIni If true, the Default*.ini file is required to exist when generating the final ini file.
-	 * @param bAllowGeneratedIniWhenCooked If true, the engine will attempt to load the generated/user INI file when loading.
+	 * @param bAllowGeneratedIniWhenCooked If true, the engine will attempt to load the generated/user INI file when loading cooked games
 	 * @param GeneratedConfigDir The location where generated config files are made.
 	 * @return true if the final ini was created successfully.
 	 */
@@ -779,11 +779,12 @@ public:
 	 * @param bIsBaseIniName true if IniName is a Base name, which can be overridden on commandline, etc.
 	 * @param Platform The platform to use for Base ini names, NULL means to use the current platform
 	 * @param bForceReload force reload the ini file from disk this is required if you make changes to the ini file not using the config system as the hierarchy cache will not be updated in this case
+	 * @return true if the ini file was loaded successfully
 	 */
-	static void LoadLocalIniFile(FConfigFile& ConfigFile, const TCHAR* IniName, bool bIsBaseIniName, const TCHAR* Platform=NULL, const bool bForceReload=false);
+	static bool LoadLocalIniFile(FConfigFile& ConfigFile, const TCHAR* IniName, bool bIsBaseIniName, const TCHAR* Platform=NULL, bool bForceReload=false);
 
 	/**
-	 * Load an ini file directly into an FConfigFile from the specified config folders, and nothing is written to GConfig or disk. 
+	 * Load an ini file directly into an FConfigFile from the specified config folders, optionally writing to disk. 
 	 * The passed in .ini name can be a "base" (Engine, Game) which will be modified by platform and/or commandline override,
 	 * or it can be a full ini filenname (ie WrangleContent) loaded from the Source config directory
 	 *
@@ -791,11 +792,15 @@ public:
 	 * @param IniName Either a Base ini name (Engine) or a full ini name (WrangleContent). NO PATH OR EXTENSION SHOULD BE USED!
 	 * @param EngineConfigDir Engine config directory.
 	 * @param SourceConfigDir Game config directory.
-	 * @param bGenerateDestIni true if IniName is a Base name, which can be overridden on commandline, etc.
+	 * @param bIsBaseIniName true if IniName is a Base name, which can be overridden on commandline, etc.
 	 * @param Platform The platform to use for Base ini names
 	 * @param bForceReload force reload the ini file from disk this is required if you make changes to the ini file not using the config system as the hierarchy cache will not be updated in this case
+	 * @param bWriteDestIni write out a destination ini file to the Saved folder, only valid if bIsBaseIniName is true
+	 * @param bAllowGeneratedIniWhenCooked If true, the engine will attempt to load the generated/user INI file when loading cooked games
+	 * @param GeneratedConfigDir The location where generated config files are made.
+	 * @return true if the ini file was loaded successfully
 	 */
-	static void LoadExternalIniFile(FConfigFile& ConfigFile, const TCHAR* IniName, const TCHAR* EngineConfigDir, const TCHAR* SourceConfigDir, bool bGenerateDestIni, const TCHAR* Platform=NULL, const bool bForceReload=false);
+	static bool LoadExternalIniFile(FConfigFile& ConfigFile, const TCHAR* IniName, const TCHAR* EngineConfigDir, const TCHAR* SourceConfigDir, bool bIsBaseIniName, const TCHAR* Platform=NULL, bool bForceReload=false, bool bWriteDestIni=false, bool bAllowGeneratedIniWhenCooked = true, const TCHAR* GeneratedConfigDir = *FPaths::GeneratedConfigDir());
 
 	/**
 	 * Needs to be called after GConfig is set and LoadCoalescedFile was called.
diff --git a/Engine/Source/Runtime/Core/Public/Misc/Parse.h b/Engine/Source/Runtime/Core/Public/Misc/Parse.h
--- a/Engine/Source/Runtime/Core/Public/Misc/Parse.h
+++ b/Engine/Source/Runtime/Core/Public/Misc/Parse.h
@@ -23,7 +23,7 @@ struct CORE_API FParse
 	/** Parses a globally unique identifier. */
 	static bool Value( const TCHAR* Stream, const TCHAR* Match, struct FGuid& Guid );
 	/** Parses a string from a text string. */
-	static bool Value( const TCHAR* Stream, const TCHAR* Match, TCHAR* Value, int32 MaxLen, bool bShouldStopOnComma=true );
+	static bool Value( const TCHAR* Stream, const TCHAR* Match, TCHAR* Value, int32 MaxLen, bool bShouldStopOnSeparator=true );
 	/** Parses a byte. */
 	static bool Value( const TCHAR* Stream, const TCHAR* Match, uint8& Value );
 	/** Parses a signed byte. */
@@ -37,7 +37,7 @@ struct CORE_API FParse
 	/** Parses a signed double word. */
 	static bool Value( const TCHAR* Stream, const TCHAR* Match, int32& Value );
 	/** Parses a string. */
-	static bool Value( const TCHAR* Stream, const TCHAR* Match, FString& Value, bool bShouldStopOnComma=true );
+	static bool Value( const TCHAR* Stream, const TCHAR* Match, FString& Value, bool bShouldStopOnSeparator =true );
 	/** Parses an FText. */
 	static bool Value( const TCHAR* Stream, const TCHAR* Match, FText& Value, const TCHAR* Namespace = NULL );
 	/** Parses a quadword. */
diff --git a/Engine/Source/Runtime/Core/Public/Misc/Paths.h b/Engine/Source/Runtime/Core/Public/Misc/Paths.h
--- a/Engine/Source/Runtime/Core/Public/Misc/Paths.h
+++ b/Engine/Source/Runtime/Core/Public/Misc/Paths.h
@@ -478,58 +478,31 @@ public:
 	/** Gets the relative path to get from BaseDir to RootDirectory  */
 	static const FString& GetRelativePathToRoot();
 
-	FORCEINLINE static FString Combine(const TCHAR* PathA, const TCHAR* PathB) 
-	{ 
-		const TCHAR* Pathes[] = { PathA, PathB };
+	template <typename... PathTypes>
+	FORCEINLINE static FString Combine(PathTypes&&... InPaths)
+	{
+		const TCHAR* Paths[] = { GetTCharPtr(Forward<PathTypes>(InPaths))... };
 		FString Out;
 		
-		CombineInternal(Out, Pathes, 2);
+		CombineInternal(Out, Paths, ARRAY_COUNT(Paths));
 		return Out;
 	}
 
-	FORCEINLINE static FString Combine(const TCHAR* PathA, const TCHAR* PathB, const TCHAR* PathC) 
-	{ 
-		const TCHAR* Pathes[] = { PathA, PathB, PathC };
-		FString Out;
-		
-		CombineInternal(Out, Pathes, 3);
-		return Out;
-	}
-	
-	FORCEINLINE static FString Combine(const TCHAR* PathA, const TCHAR* PathB, const TCHAR* PathC, const TCHAR* PathD) 
-	{ 
-		const TCHAR* Pathes[] = { PathA, PathB, PathC, PathD };
-		FString Out;
-		
-		CombineInternal(Out, Pathes, 4);
-		return Out;
-	}
+protected:
 
-	FORCEINLINE static FString Combine(const TCHAR* PathA, const TCHAR* PathB, const TCHAR* PathC, const TCHAR* PathD, const TCHAR* PathE) 
-	{ 
-		const TCHAR* Pathes[] = { PathA, PathB, PathC, PathD, PathE };
-		FString Out;
-		
-		CombineInternal(Out, Pathes, 5);
-		return Out;
-	}
+	static void CombineInternal(FString& OutPath, const TCHAR** Paths, int32 NumPaths);
 
-	FORCEINLINE static FString Combine(const FString& PathA, const FString& PathB)
-	{ 
-		return Combine(*PathA, *PathB);
+private:
+	FORCEINLINE static const TCHAR* GetTCharPtr(const TCHAR* Ptr)
+	{
+		return Ptr;
 	}
 
-	FORCEINLINE static FString Combine(const FString& PathA, const FString& PathB, const FString& PathC)
+	FORCEINLINE static const TCHAR* GetTCharPtr(const FString& Str)
 	{
-		return Combine(*PathA, *PathB, *PathC);
+		return *Str;
 	}
 
-protected:
-
-	static void CombineInternal(FString& OutPath, const TCHAR** Pathes, int32 NumPathes);
-
-private:
-
 	/** Holds the path to the currently loaded game project file. */
 	static FString GameProjectFilePath;
 
diff --git a/Engine/Source/Runtime/Core/Public/ProfilingDebugging/LoadTimeTracker.h b/Engine/Source/Runtime/Core/Public/ProfilingDebugging/LoadTimeTracker.h
--- a/Engine/Source/Runtime/Core/Public/ProfilingDebugging/LoadTimeTracker.h
+++ b/Engine/Source/Runtime/Core/Public/ProfilingDebugging/LoadTimeTracker.h
@@ -86,7 +86,6 @@ public:
 	double LinkerLoad_SerializeImportMap;
 	double LinkerLoad_SerializeExportMap;
 	double LinkerLoad_FixupImportMap;
-	double LinkerLoad_RemapImports;
 	double LinkerLoad_FixupExportMap;
 	double LinkerLoad_SerializeDependsMap;
 	double LinkerLoad_SerializePreloadDependencies;
diff --git a/Engine/Source/Runtime/Core/Public/Serialization/Archive.h b/Engine/Source/Runtime/Core/Public/Serialization/Archive.h
--- a/Engine/Source/Runtime/Core/Public/Serialization/Archive.h
+++ b/Engine/Source/Runtime/Core/Public/Serialization/Archive.h
@@ -17,11 +17,16 @@ class ITargetPlatform;
 struct FUntypedBulkData;
 template<class TEnum> class TEnumAsByte;
 
-// this is the master switch
-//@todoio if this is off, then we should leave the package file format completely unchanged....!!!!! this is really important to fix before we merge to main
-//#define COOK_FOR_EVENT_DRIVEN_LOAD (1)
 
-#define EVENT_DRIVEN_ASYNC_LOAD_ACTIVE_AT_RUNTIME (!GIsInitialLoad) // set to (!GIsInitialLoad) to avoid using the EDL at boot time
+// Temporary while we shake out the EDL at boot
+#define USE_EVENT_DRIVEN_ASYNC_LOAD_AT_BOOT_TIME (1)
+
+#if USE_EVENT_DRIVEN_ASYNC_LOAD_AT_BOOT_TIME
+	#define EVENT_DRIVEN_ASYNC_LOAD_ACTIVE_AT_RUNTIME (1)
+#else
+	#define EVENT_DRIVEN_ASYNC_LOAD_ACTIVE_AT_RUNTIME (!GIsInitialLoad) 
+#endif
+
 #define DEVIRTUALIZE_FLinkerLoad_Serialize (!WITH_EDITORONLY_DATA)
 
 
diff --git a/Engine/Source/Runtime/Core/Public/Stats/Stats.h b/Engine/Source/Runtime/Core/Public/Stats/Stats.h
--- a/Engine/Source/Runtime/Core/Public/Stats/Stats.h
+++ b/Engine/Source/Runtime/Core/Public/Stats/Stats.h
@@ -313,13 +313,6 @@ enum EStatType
 	STATTYPE_Error
 };
 
-template<class T> inline const TCHAR* GetStatFormatString(void) { return TEXT(""); }
-template<> inline const TCHAR* GetStatFormatString<uint32>(void) { return TEXT("%u"); }
-template<> inline const TCHAR* GetStatFormatString<uint64>(void) { return TEXT("%llu"); }
-template<> inline const TCHAR* GetStatFormatString<int64>(void) { return TEXT("%lld"); }
-template<> inline const TCHAR* GetStatFormatString<float>(void) { return TEXT("%.1f"); }
-template<> inline const TCHAR* GetStatFormatString<double>(void) { return TEXT("%.1f"); }
-
 /*----------------------------------------------------------------------------
 	Stats helpers
 ----------------------------------------------------------------------------*/
diff --git a/Engine/Source/Runtime/Core/Public/Stats/Stats2.h b/Engine/Source/Runtime/Core/Public/Stats/Stats2.h
--- a/Engine/Source/Runtime/Core/Public/Stats/Stats2.h
+++ b/Engine/Source/Runtime/Core/Public/Stats/Stats2.h
@@ -2051,8 +2051,8 @@ DECLARE_STATS_GROUP_VERBOSE(TEXT("GnmVerbose"), STATGROUP_PS4RHIVERBOSE, STATCAT
 DECLARE_STATS_GROUP(TEXT("Init Views"),STATGROUP_InitViews, STATCAT_Advanced);
 DECLARE_STATS_GROUP(TEXT("Landscape"),STATGROUP_Landscape, STATCAT_Advanced);
 DECLARE_STATS_GROUP(TEXT("Light Rendering"),STATGROUP_LightRendering, STATCAT_Advanced);
-DECLARE_STATS_GROUP(TEXT("Load Time"), STATGROUP_LoadTime, STATCAT_Advanced);
-DECLARE_STATS_GROUP_VERBOSE(TEXT("Load Time (Verbose)"), STATGROUP_LoadTimeVerbose, STATCAT_Advanced);
+DECLARE_STATS_GROUP(TEXT("LoadTime"), STATGROUP_LoadTime, STATCAT_Advanced);
+DECLARE_STATS_GROUP_VERBOSE(TEXT("LoadTimeVerbose"), STATGROUP_LoadTimeVerbose, STATCAT_Advanced);
 DECLARE_STATS_GROUP_VERBOSE(TEXT("MathVerbose"), STATGROUP_MathVerbose, STATCAT_Advanced);
 DECLARE_STATS_GROUP(TEXT("Memory Allocator"),STATGROUP_MemoryAllocator, STATCAT_Advanced);
 DECLARE_STATS_GROUP(TEXT("Memory Platform"),STATGROUP_MemoryPlatform, STATCAT_Advanced);
diff --git a/Engine/Source/Runtime/Core/Public/Templates/UnrealTypeTraits.h b/Engine/Source/Runtime/Core/Public/Templates/UnrealTypeTraits.h
--- a/Engine/Source/Runtime/Core/Public/Templates/UnrealTypeTraits.h
+++ b/Engine/Source/Runtime/Core/Public/Templates/UnrealTypeTraits.h
@@ -85,7 +85,7 @@ struct TFormatSpecifier
 		// If the static_assert below were a constant 0 or something not dependent on T, compilers are free to detect this and fail to compile the template.
 		// As specified in the C++ standard s14.6p8. A compiler is free to give a diagnostic here or not. MSVC ignores it, and clang/gcc instantiates the 
 		// template and triggers the static_assert.
-		static_assert(sizeof(T) < 0, "Format specifer not supported for this type."); // request for a format specifier for a type we do not know about
+		static_assert(sizeof(T) < 0, "Format specifier not supported for this type."); // request for a format specifier for a type we do not know about
 		return TEXT("Unknown");
 	}
 };
@@ -93,9 +93,10 @@ struct TFormatSpecifier
 template<> \
 struct TFormatSpecifier<type> \
 {  \
-	FORCEINLINE static TCHAR const* GetFormatSpecifier() \
+	FORCEINLINE static const TCHAR (&GetFormatSpecifier())[5] \
 	{ \
-		return TEXT(format); \
+		static const TCHAR Spec[5] = TEXT(format); \
+		return Spec; \
 	} \
 };
 
diff --git a/Engine/Source/Runtime/Core/Public/UObject/ObjectVersion.h b/Engine/Source/Runtime/Core/Public/UObject/ObjectVersion.h
--- a/Engine/Source/Runtime/Core/Public/UObject/ObjectVersion.h
+++ b/Engine/Source/Runtime/Core/Public/UObject/ObjectVersion.h
@@ -617,6 +617,9 @@ enum EUnrealEngineObjectUE4Version
 	VER_UE4_PROPERTY_TAG_SET_MAP_SUPPORT,
 	// Added SearchableNames to the package summary and asset registry
 	VER_UE4_ADDED_SEARCHABLE_NAMES,
+	// Increased size of SerialSize and SerialOffset in export map entries to 64 bit, allow support for bigger files
+	VER_UE4_64BIT_EXPORTMAP_SERIALSIZES,
+
 
 	// -----<new versions can be added before this line>-------------------------------------------------
 	// - this needs to be the last line (see note below)
diff --git a/Engine/Source/Runtime/Core/Public/UObject/UnrealNames.inl b/Engine/Source/Runtime/Core/Public/UObject/UnrealNames.inl
--- a/Engine/Source/Runtime/Core/Public/UObject/UnrealNames.inl
+++ b/Engine/Source/Runtime/Core/Public/UObject/UnrealNames.inl
@@ -73,6 +73,8 @@ REGISTER_NAME(102,Actor)
 REGISTER_NAME(103,ObjectRedirector)
 REGISTER_NAME(104,ObjectArchetype)
 REGISTER_NAME(105,Class)
+REGISTER_NAME(106,ScriptStruct)
+REGISTER_NAME(107,Function)
 
 // Misc.
 REGISTER_NAME(200,State)
diff --git a/Engine/Source/Runtime/Core/Public/Windows/WindowsPlatformCompilerPreSetup.h b/Engine/Source/Runtime/Core/Public/Windows/WindowsPlatformCompilerPreSetup.h
--- a/Engine/Source/Runtime/Core/Public/Windows/WindowsPlatformCompilerPreSetup.h
+++ b/Engine/Source/Runtime/Core/Public/Windows/WindowsPlatformCompilerPreSetup.h
@@ -137,6 +137,7 @@
 			__pragma(warning(push)) \
 			__pragma(warning(disable: 4510))  /* '<class>': default constructor could not be generated. */ \
 			__pragma(warning(disable: 4610))  /* object '<class>' can never be instantiated - user-defined constructor required. */ \
+			__pragma(warning(disable: 4946))  /* reinterpret_cast used between related classes: '<class1>' and '<class1>' */ \
 			__pragma(warning(disable: 4996))  /* '<obj>' was declared deprecated. */ \
 			__pragma(warning(disable: 6011))  /* Dereferencing NULL pointer '<ptr>'. */ \
 			__pragma(warning(disable: 6101))  /* Returning uninitialized memory '<expr>'.  A successful path through the function does not set the named _Out_ parameter. */ \
diff --git a/Engine/Source/Runtime/CoreUObject/Private/Blueprint/BlueprintSupport.cpp b/Engine/Source/Runtime/CoreUObject/Private/Blueprint/BlueprintSupport.cpp
--- a/Engine/Source/Runtime/CoreUObject/Private/Blueprint/BlueprintSupport.cpp
+++ b/Engine/Source/Runtime/CoreUObject/Private/Blueprint/BlueprintSupport.cpp
@@ -291,6 +291,7 @@ struct FPreloadMembersHelper
 			UObject* CurrentObject = *it;
 			if (!CurrentObject->HasAnyFlags(RF_LoadCompleted))
 			{
+				check(!GEventDrivenLoaderEnabled);
 				CurrentObject->SetFlags(RF_NeedLoad);
 				if (auto Linker = CurrentObject->GetLinker())
 				{
@@ -305,6 +306,7 @@ struct FPreloadMembersHelper
 	{
 		if (InObject && !InObject->HasAnyFlags(RF_LoadCompleted))
 		{
+			check(!GEventDrivenLoaderEnabled);
 			InObject->SetFlags(RF_NeedLoad);
 			if (FLinkerLoad* Linker = InObject->GetLinker())
 			{
@@ -1330,6 +1332,7 @@ void FLinkerLoad::FinalizeBlueprint(UClass* LoadClass)
 			// to re-run the serialization)
 			if ( (SuperCDO != nullptr) && !SuperCDO->HasAnyFlags(RF_NeedLoad|RF_LoadCompleted) )
 			{
+				check(!GEventDrivenLoaderEnabled);
 				SuperCDO->SetFlags(RF_NeedLoad);
 			}
 			SuperLinker->FinalizeBlueprint(SuperClass);
@@ -1864,6 +1867,10 @@ void FLinkerLoad::CreateDynamicTypeLoader()
 	{
 		const FString DynamicTypePath = GetExportPathName(DynamicTypeExportIndex);
 		const FName DynamicTypeClassName = GetDynamicTypeClassName(*DynamicTypePath);
+		if (DynamicTypeClassName == NAME_None)
+		{
+			UE_LOG(LogTemp, Error, TEXT("Exports %d, DynamicTypePath %s, Export Name %s, Package Root %s"), ExportMap.Num(), *DynamicTypePath, *DynamicTypeExport->ObjectName.ToString(), *LinkerRoot->GetPathName());
+		}
 		ensure(DynamicTypeClassName != NAME_None);
 		const bool bIsDynamicClass = DynamicTypeClassName == DynamicClassName;
 		const bool bIsDynamicStruct = DynamicTypeClassName == UScriptStruct::StaticClass()->GetFName();
@@ -1881,6 +1888,11 @@ void FLinkerLoad::CreateDynamicTypeLoader()
 				CDOExport->ClassIndex = DynamicTypeExport->ThisIndex;
 			}
 
+			// Note, the layout of the fake export table is assumed elsewhere
+				//check(ImportLinker->ExportMap.Num() == 2); // we assume there are two elements in the fake export table and the second one is the CDO
+				//LocalExportIndex = FPackageIndex::FromExport(1);
+
+
 			FObjectExport* const FakeExports[] = { DynamicTypeExport , CDOExport }; // must be sync'ed with FBlueprintDependencyData::DependencyTypes
 			int32 RunningIndex = 0;
 			for(int32 LocExportIndex = 0; LocExportIndex < (sizeof(FakeExports)/sizeof(FakeExports[0])); LocExportIndex++)
@@ -2258,6 +2270,7 @@ void FDeferredObjInitializerTracker::ResolveDeferredSubObjects(UObject* CDO)
 		{
 			if (FLinkerLoad* SubObjLinker = SubObjArchetype->GetLinker())
 			{
+				check(!GEventDrivenLoaderEnabled);
 				SubObjArchetype->SetFlags(RF_NeedLoad);
 				SubObjLinker->Preload(SubObjArchetype);
 			}
diff --git a/Engine/Source/Runtime/CoreUObject/Private/Misc/ExclusiveLoadPackageTimeTracker.cpp b/Engine/Source/Runtime/CoreUObject/Private/Misc/ExclusiveLoadPackageTimeTracker.cpp
--- a/Engine/Source/Runtime/CoreUObject/Private/Misc/ExclusiveLoadPackageTimeTracker.cpp
+++ b/Engine/Source/Runtime/CoreUObject/Private/Misc/ExclusiveLoadPackageTimeTracker.cpp
@@ -190,7 +190,7 @@ void FExclusiveLoadPackageTimeTracker::InternalDumpReport() const
 
 		if (LowThresholdCount > 0)
 		{
-			UE_LOG(LogLoad, Log, TEXT("    ... skipped %d assets slower than %.1fms totaling %.1fms"), LowThresholdCount, LowTimeThreshold * 1000, TotalLowTime * 1000);
+			UE_LOG(LogLoad, Log, TEXT("    ... skipped %d assets that loaded in less than %.1fms totaling %.1fms"), LowThresholdCount, LowTimeThreshold * 1000, TotalLowTime * 1000);
 		}
 	}
 
diff --git a/Engine/Source/Runtime/CoreUObject/Private/Misc/WorldCompositionUtility.cpp b/Engine/Source/Runtime/CoreUObject/Private/Misc/WorldCompositionUtility.cpp
--- a/Engine/Source/Runtime/CoreUObject/Private/Misc/WorldCompositionUtility.cpp
+++ b/Engine/Source/Runtime/CoreUObject/Private/Misc/WorldCompositionUtility.cpp
@@ -114,6 +114,7 @@ bool FWorldTileInfo::Read(const FString& InPackageFileName, FWorldTileInfo& OutI
 		//make sure the filereader gets the correct version number (it defaults to latest version)
 		FileReader->SetUE4Ver(FileSummary.GetFileVersionUE4());
 		FileReader->SetEngineVer(FileSummary.SavedByEngineVersion);
+		FileReader->SetLicenseeUE4Ver(FileSummary.GetFileVersionLicenseeUE4());
 
 		// Load the structure
 		*FileReader << OutInfo;
diff --git a/Engine/Source/Runtime/CoreUObject/Private/Serialization/AsyncLoading.cpp b/Engine/Source/Runtime/CoreUObject/Private/Serialization/AsyncLoading.cpp
--- a/Engine/Source/Runtime/CoreUObject/Private/Serialization/AsyncLoading.cpp
+++ b/Engine/Source/Runtime/CoreUObject/Private/Serialization/AsyncLoading.cpp
@@ -37,6 +37,7 @@
 #include "UObject/UObjectHash.h"
 #include "UniquePtr.h"
 #include "Serialization/BufferReader.h"
+#include "TaskGraphInterfaces.h"
 
 #define FIND_MEMORY_STOMPS (1 && (PLATFORM_WINDOWS || PLATFORM_LINUX) && !WITH_EDITORONLY_DATA)
 
@@ -410,6 +411,11 @@ FAsyncPackage* FAsyncLoadingThread::FindExistingPackageAndAddCompletionCallback(
 void FAsyncLoadingThread::UpdateExistingPackagePriorities(FAsyncPackage* InPackage, TAsyncLoadPriority InNewPriority, IAssetRegistryInterface* InAssetRegistry)
 {
 	check(!IsInGameThread() || !IsMultithreaded());
+	if (GEventDrivenLoaderEnabled)
+	{
+		InPackage->SetPriority(InNewPriority);
+		return;
+	}
 	if (InNewPriority > InPackage->GetPriority())
 	{
 		AsyncPackages.Remove(InPackage);
@@ -582,7 +588,7 @@ int32 FAsyncLoadingThread::CreateAsyncPackagesFromQueue(bool bUseTimeLimit, bool
 		{
 			IAssetRegistryInterface* AssetRegistry = nullptr;
 
-			if (GPreloadPackageDependencies && IsPlatformFileCompatibleWithDependencyPreloading() && 
+			if (GPreloadPackageDependencies && !GEventDrivenLoaderEnabled && IsPlatformFileCompatibleWithDependencyPreloading() &&
 					// we don't want preloading of dependencies with event driven loading, we would rather find the dependencies naturally.
 				 !GEventDrivenLoaderEnabled)
 			{
@@ -607,6 +613,305 @@ int32 FAsyncLoadingThread::CreateAsyncPackagesFromQueue(bool bUseTimeLimit, bool
 	return NumCreated;
 }
 
+struct FEDLBootObjectState
+{
+	ENotifyRegistrationType NotifyRegistrationType;
+	ENotifyRegistrationPhase LastNotifyRegistrationPhase;
+	UObject *(*Register)();
+	bool bDynamic;
+};
+
+struct FEDLBootNotificationManager
+{
+	TMap<FName, FEDLBootObjectState> PathToState;
+	TMultiMap<FName, FEventLoadNodePtr> PathToWaitingPackageNodes;
+	TArray<FName> PathsToFire;
+	TArray<UClass*> CDORecursiveStack;
+	TArray<UClass*> CDORecursives;
+	FCriticalSection EDLBootNotificationManagerLock;
+
+	// return true if you are waiting for this compiled in object
+	FORCENOINLINE bool AddWaitingPackage(FAsyncPackage* Pkg, FName PackageName, FName ObjectName, FPackageIndex Import)
+	{
+		if (PackageName == GLongCoreUObjectPackageName)
+		{
+			return false; // We assume nothing in coreuobject ever loads assets in a constructor
+		}
+		FScopeLock Lock(&EDLBootNotificationManagerLock);
+		check(GIsInitialLoad);
+		check(Import.IsImport()); // compiled in exports make no sense
+		FString ObjectNameString = ObjectName.ToString();
+		FName LongFName(*(PackageName.ToString() / ObjectNameString));
+		check(LongFName != NAME_None);
+		FName WaitName = LongFName;
+		FEDLBootObjectState* ExistingState = PathToState.Find(LongFName);
+		if (!ExistingState)
+		{
+			//if (ObjectName.ToString().EndsWith(HEADER_GENERATED_DELEGATE_SIGNATURE_SUFFIX))
+			// there are also some arg structs and other things that are just part of the package with no registration
+			{
+				ExistingState = PathToState.Find(PackageName);
+				WaitName = PackageName;
+			}
+			if (!ExistingState)
+			{
+				UE_LOG(LogStreaming, Fatal, TEXT("Compiled in export %s not found; it was never registered."), *LongFName.ToString());
+				return false;
+			}
+		}
+		if (ExistingState->LastNotifyRegistrationPhase == ENotifyRegistrationPhase::NRP_Finished)
+		{
+			return false;
+		}
+		FEventLoadNodePtr MyNode;
+		MyNode.WaitingPackage = FCheckedWeakAsyncPackagePtr(Pkg);
+		MyNode.ImportOrExportIndex = Import;
+		MyNode.Phase = EEventLoadNode::ImportOrExport_Create;
+
+		PathToWaitingPackageNodes.Add(WaitName, MyNode);
+
+		return true;
+	}
+
+	void NotifyRegistrationEvent(const TCHAR* PackageName, const TCHAR* Name, ENotifyRegistrationType NotifyRegistrationType, ENotifyRegistrationPhase NotifyRegistrationPhase, UObject *(*InRegister)(), bool InbDynamic)
+	{
+		if (!GIsInitialLoad)
+		{
+			return;
+		}
+		static FName LongCoreUObjectPackageName(TEXT("/Script/CoreUObject")); // can't use the global here because it may not be initialized yet
+		FName PackageFName(PackageName);
+		if (PackageFName == LongCoreUObjectPackageName)
+		{
+			return; // We assume nothing in coreuobject ever loads assets in a constructor
+		}
+
+		FScopeLock Lock(&EDLBootNotificationManagerLock);
+
+		FName LongFName(*(FString(PackageName) / Name));
+
+		//FPlatformMisc::LowLevelOutputDebugStringf(TEXT("NotifyRegistrationEvent %s %d %d\r\n"), *LongFName.ToString(), int32(NotifyRegistrationType), int32(NotifyRegistrationPhase));
+
+		// some things, like delegate signatures, are not registered; rather they are part of the package singleton, so we track the package state as being the max of any member of that package
+		FEDLBootObjectState* ExistingPackageState = PathToState.Find(PackageFName);
+		FEDLBootObjectState* ExistingState = PathToState.Find(LongFName);
+
+		if (!ExistingState)
+		{
+			if (NotifyRegistrationPhase != ENotifyRegistrationPhase::NRP_Added)
+			{
+				UE_LOG(LogStreaming, Fatal, TEXT("Attempt to process %s before it has been added."), *LongFName.ToString());
+			}
+			FEDLBootObjectState NewState;
+			NewState.LastNotifyRegistrationPhase = NotifyRegistrationPhase;
+			NewState.NotifyRegistrationType = NotifyRegistrationType;
+			NewState.Register = InRegister;
+			NewState.bDynamic = InbDynamic;
+			PathToState.Add(LongFName, NewState);
+
+			if (!ExistingPackageState)
+			{
+				NewState.NotifyRegistrationType = ENotifyRegistrationType::NRT_Package;
+				PathToState.Add(PackageFName, NewState);
+			}
+		}
+		else
+		{
+			if (int32(ExistingState->LastNotifyRegistrationPhase) + 1 != int32(NotifyRegistrationPhase))
+			{
+				UE_CLOG(GEventDrivenLoaderEnabled, LogStreaming, Fatal, TEXT("Invalid state transition %d %d with %s when it has already been processed."), int32(ExistingState->LastNotifyRegistrationPhase), int32(NotifyRegistrationPhase), *LongFName.ToString());
+			}
+			if (ExistingState->NotifyRegistrationType != (NotifyRegistrationType))
+			{
+				UE_CLOG(GEventDrivenLoaderEnabled, LogStreaming, Fatal, TEXT("Multiple types %d %d with %s when it has already been processed."), int32(ExistingState->NotifyRegistrationType), int32(NotifyRegistrationType), *LongFName.ToString());
+			}
+			ExistingState->LastNotifyRegistrationPhase = NotifyRegistrationPhase;
+			if (NotifyRegistrationPhase == ENotifyRegistrationPhase::NRP_Finished)
+			{
+				ExistingState->Register = nullptr; // we don't need to do this in ConstructWaitingBootObjects
+				PathsToFire.Add(LongFName);
+			}
+			check(ExistingPackageState); // if we have an existing state for the thing, we should also have a 
+			if (ExistingPackageState && int32(NotifyRegistrationPhase) > int32(ExistingPackageState->LastNotifyRegistrationPhase))
+			{
+				ExistingPackageState->LastNotifyRegistrationPhase = NotifyRegistrationPhase;
+				if (NotifyRegistrationPhase == ENotifyRegistrationPhase::NRP_Finished)
+				{
+					//FPlatformMisc::LowLevelOutputDebugStringf(TEXT("Fired package %s %d %d\r\n"), *PackageFName.ToString(), int32(NotifyRegistrationType), int32(NotifyRegistrationPhase));
+					PathsToFire.Add(PackageFName);
+				}
+			}
+		}
+	}
+
+	void NotifyRegistrationComplete()
+	{
+#if USE_EVENT_DRIVEN_ASYNC_LOAD_AT_BOOT_TIME
+		FireCompletedCompiledInImports(true);
+		FlushAsyncLoading();
+#endif
+#if !HACK_HEADER_GENERATOR
+		check(!GIsInitialLoad && IsInGameThread());
+		FScopeLock Lock(&EDLBootNotificationManagerLock);
+		for (auto& Pair : PathToState)
+		{
+			if (Pair.Value.LastNotifyRegistrationPhase != ENotifyRegistrationPhase::NRP_Finished && !Pair.Value.bDynamic)
+			{
+#if USE_EVENT_DRIVEN_ASYNC_LOAD_AT_BOOT_TIME
+				UE_CLOG(GEventDrivenLoaderEnabled, LogStreaming, Fatal, TEXT("%s (%d) was not complete (%d) after registration was complete."), *Pair.Key.ToString(), int32(Pair.Value.NotifyRegistrationType), int32(Pair.Value.LastNotifyRegistrationPhase));
+#else
+				UE_LOG(LogStreaming, Warning, TEXT("%s was not complete (%d) after registration was complete."), *Pair.Key.ToString(), int32(Pair.Value.LastNotifyRegistrationPhase));
+#endif
+			}
+		}
+		if (PathToWaitingPackageNodes.Num())
+		{
+			UE_LOG(LogStreaming, Fatal, TEXT("Initial load is complete, but we still have %d waiting packages."), PathToWaitingPackageNodes.Num());
+		}
+		if (GEventDrivenLoaderEnabled && PathsToFire.Num())
+		{
+			for (FName Path : PathsToFire)
+			{
+				UE_LOG(LogStreaming, Error, TEXT("%s was not fired."), *Path.ToString());
+			}
+			UE_LOG(LogStreaming, Fatal, TEXT("Initial load is complete, but we still have %d imports to fire (listed above)."), PathsToFire.Num());
+		}
+#endif
+		PathToState.Empty();
+		PathsToFire.Empty();
+	}
+
+	bool ConstructWaitingBootObjects()
+	{
+		check(GIsInitialLoad && IsInGameThread());
+		UObject *(*BootObjectRegister)() = nullptr;
+		UObject *(*BootPackageObjectRegister)() = nullptr;
+		FName WaitingPackage;
+		bool bIsCDO = false;
+		{
+			FScopeLock Lock(&EDLBootNotificationManagerLock);
+			for (auto& Pair : PathToWaitingPackageNodes)
+			{
+				FEDLBootObjectState* ExistingState = PathToState.Find(Pair.Key);
+				if (ExistingState)
+				{
+					if (ExistingState->Register)
+					{
+						BootObjectRegister = ExistingState->Register;
+						ExistingState->Register = nullptr; // we don't need to do this more than once
+						bIsCDO = ExistingState->NotifyRegistrationType == ENotifyRegistrationType::NRT_ClassCDO;
+						break;
+					}
+				}
+			}
+		}
+		if (BootObjectRegister)
+		{
+			UObject* BootObject = BootObjectRegister();
+			check(BootObject);
+			UObjectForceRegistration(BootObject);
+			if (bIsCDO)
+			{
+				UClass* Class = CastChecked<UClass>(BootObject);
+				bool bAnyParentOnStack = false;
+				UClass* Super = Class;
+				while (Super)
+				{
+					if (CDORecursiveStack.Contains(Super))
+					{
+						bAnyParentOnStack = true;
+						break;
+					}
+					Super = Super->GetSuperClass();
+				}
+
+				if (!bAnyParentOnStack)
+				{
+					CDORecursiveStack.Push(Class);
+					Class->GetDefaultObject();
+					verify(CDORecursiveStack.Pop() == Class);
+				}
+				else
+				{
+					CDORecursives.Add(Class);
+				}
+			}
+			return true;
+		}
+		if (CDORecursives.Num())
+		{
+			UClass* OkToRun = nullptr;
+			for (UClass* Class : CDORecursives)
+			{
+				bool bAnyParentOnStack = false;
+				UClass* Super = Class;
+				while (Super)
+				{
+					if (CDORecursiveStack.Contains(Super))
+					{
+						bAnyParentOnStack = true;
+						break;
+					}
+					Super = Super->GetSuperClass();
+				}
+				if (!bAnyParentOnStack)
+				{
+					OkToRun = Class;
+					break;
+				}
+			}
+			if (OkToRun)
+			{
+				CDORecursives.Remove(OkToRun);
+				CDORecursiveStack.Push(OkToRun);
+				OkToRun->GetDefaultObject();
+				verify(CDORecursiveStack.Pop() == OkToRun);
+			}
+			else
+			{
+				FPlatformProcess::Sleep(.001f);
+			}
+			return true; // even if we didn't do anything we need to return true to avoid checking for cycles
+		}
+		return false;
+	}
+
+	bool IsWaitingForSomething()
+	{
+		FScopeLock Lock(&EDLBootNotificationManagerLock);
+		return PathToWaitingPackageNodes.Num() > 0;
+	}
+
+	bool IsObjComplete(UObject* Obj)
+	{
+		static FName LongCoreUObjectPackageName(TEXT("/Script/CoreUObject")); // can't use the global here because it may not be initialized yet
+		FName PackageName = Obj->GetOutermost()->GetFName();
+		if (PackageName == LongCoreUObjectPackageName)
+		{
+			return true; // We assume nothing in coreuobject ever loads assets in a constructor, therefore it can be considered complete
+		}
+		FScopeLock Lock(&EDLBootNotificationManagerLock);
+		FName LongFName(*(PackageName.ToString() / Obj->GetName()));
+
+		FEDLBootObjectState* ExistingState = PathToState.Find(LongFName);
+
+		if (!ExistingState || ExistingState->LastNotifyRegistrationPhase == ENotifyRegistrationPhase::NRP_Finished)
+		{
+			return true;
+		}
+		return false;
+	}
+
+	bool FireCompletedCompiledInImports(bool bFinalRun = false);
+};
+
+static FEDLBootNotificationManager& GetGEDLBootNotificationManager()
+{
+	static FEDLBootNotificationManager Singleton;
+	return Singleton;
+}
+
+
 static FThreadSafeCounter AsyncPackageSerialNumber;
 
 FName FUnsafeWeakAsyncPackagePtr::HumanReadableStringForDebugging() const
@@ -970,6 +1275,7 @@ bool FEventLoadGraph::CheckForCyclesInner(const TMultiMap<FEventLoadNodePtr, FEv
 
 void FEventLoadGraph::CheckForCycles()
 {
+	int32 NumWaitingBoot = 0;
 #if !UE_BUILD_SHIPPING && !UE_BUILD_TEST
 	TMultiMap<FEventLoadNodePtr, FEventLoadNodePtr> Arcs;
 	TArray<FEventLoadNodePtr> AddedNodes;
@@ -984,8 +1290,16 @@ void FEventLoadGraph::CheckForCycles()
 
 		if (!Node.NumPrerequistes)
 		{
+			if (GIsInitialLoad && Node.bFired)
+			{
+				// this is something that is compiled in, but has not been finished yet
+				NumWaitingBoot++;
+			}
+			else
+			{
 			UE_LOG(LogStreaming, Fatal, TEXT("Node %s has zero prerequisites, but has not been queued."), *Ptr.HumanReadableStringForDebugging());
 		}
+		}
 		for (FEventLoadNodePtr Other : Node.NodesWaitingForMe)
 		{
 			Arcs.Add(Other, Ptr);
@@ -1009,12 +1323,16 @@ void FEventLoadGraph::CheckForCycles()
 			UE_LOG(LogStreaming, Fatal, TEXT("Async loading event graph contained a cycle, see above."));
 		}
 	}
-	if (AddedNodes.Num())
+	if (AddedNodes.Num() - NumWaitingBoot != 0)
 	{
-		UE_LOG(LogStreaming, Fatal, TEXT("No outstanding IO, no nodes in the queue, yet we still have 'AddedNodes' in the graph."));
+		for (FEventLoadNodePtr &Ptr : AddedNodes)
+		{
+			FPlatformMisc::LowLevelOutputDebugStringf(TEXT("AddedNode: %s\r\n"), *Ptr.HumanReadableStringForDebugging());
+		}
+		UE_LOG(LogStreaming, Fatal, TEXT("No outstanding IO, no nodes in the queue, yet we still have %d 'AddedNodes' in the graph (with %d boot nodes)."), AddedNodes.Num(), NumWaitingBoot);
 	}
 #endif
-	if (PackagesWithNodes.Num())
+	if (PackagesWithNodes.Num() && !NumWaitingBoot)
 	{
 		UE_LOG(LogStreaming, Fatal, TEXT("No outstanding IO, no nodes in the queue, yet we still have 'PackagesWithNodes' in the graph."));
 	}
@@ -1401,8 +1719,17 @@ void FAsyncPackage::Event_FinishLinker()
 				NumImplicit = Linker->ImportMap.Num() + Linker->ExportMap.Num();
 #endif
 
-				FEventLoadNodePtr MyDoneNode = AddNode(EEventLoadNode::Package_ExportsSerialized, FPackageIndex(), false, NumImplicit);
+				AddNode(EEventLoadNode::Package_ExportsSerialized, FPackageIndex(), false, NumImplicit);
 
+				AddNode(EEventLoadNode::Package_SetupImports, FPackageIndex(), true);
+				EventNodeArray.Init(Linker->ImportMap.Num(), Linker->ExportMap.Num());
+			}
+			if (PackagesWaitingToLinkImports.Num())
+			{
+				FCheckedWeakAsyncPackagePtr WeakThis(this);
+				FEventLoadNodePtr MyDoneNode;
+				MyDoneNode.WaitingPackage = WeakThis;
+				MyDoneNode.Phase = EEventLoadNode::Package_ExportsSerialized;
 				// There are things waiting to link imports. I need to not finish until those links are made
 				// Package_ExportsSerialized is actually earlier than we need. We just need to make sure the linker is not destroyed before the other packages link
 				for (FCheckedWeakAsyncPackagePtr& Waiter : PackagesWaitingToLinkImports)
@@ -1413,9 +1740,6 @@ void FAsyncPackage::Event_FinishLinker()
 					AddArc(Prereq, MyDoneNode);
 				}
 				PackagesWaitingToLinkImports.Empty();
-
-				AddNode(EEventLoadNode::Package_SetupImports, FPackageIndex(), true);
-				EventNodeArray.Init(Linker->ImportMap.Num(), Linker->ExportMap.Num());
 			}
 			FEventLoadNodePtr MyDependentExportsSerializedNode;
 			MyDependentExportsSerializedNode.WaitingPackage = FCheckedWeakAsyncPackagePtr(this);
@@ -1547,7 +1871,6 @@ FObjectImport* FAsyncPackage::FindExistingImport(int32 LocalImportIndex)
 				if (Import->XObject)
 				{
 //native blueprint 
-#if DO_CHECK
 					FName NameImportClass(Import->ClassName);
 					FName NameActualImportClass(Import->XObject->GetClass()->GetFName());
 					if (NameActualImportClass != NameImportClass)
@@ -1562,7 +1885,6 @@ FObjectImport* FAsyncPackage::FindExistingImport(int32 LocalImportIndex)
 							|| (NameImportClass == NAME_Function && NameActualImportClass == NAME_DelegateFunction);
 						UE_CLOG(!bSafeException, LogStreaming, Fatal, TEXT("FAsyncPackage::FindExistingImport class mismatch %s != %s"), *NameActualImportClass.ToString(), *NameImportClass.ToString());
 					}
-#endif //DO_CHECK
 					AddObjectReference(Import->XObject);
 				}
 			}
@@ -1571,6 +1893,51 @@ FObjectImport* FAsyncPackage::FindExistingImport(int32 LocalImportIndex)
 	return Import;
 }
 
+static bool IsFullyLoadedObj(UObject* Obj)
+{
+	if (!Obj)
+	{
+		return false;
+	}
+	if (Obj->HasAllFlags(RF_WasLoaded | RF_LoadCompleted) 
+		|| Obj->IsA(UPackage::StaticClass())) // packages are never really loaded, so if it exists, it is loaded
+	{
+		return true;
+	}
+	if (Obj->HasAnyFlags(RF_WasLoaded | RF_NeedLoad | RF_WillBeLoaded))
+	{
+		return false;
+	}
+	if (GIsInitialLoad && Obj->GetOutermost()->HasAnyPackageFlags(PKG_CompiledIn))
+	{
+		return GetGEDLBootNotificationManager().IsObjComplete(Obj);
+	}
+	//native blueprint 
+	UDynamicClass* UD = Cast<UDynamicClass>(Obj);
+	if (!UD)
+	{
+		return true;
+	}
+
+	if (GEventDrivenLoaderEnabled && EVENT_DRIVEN_ASYNC_LOAD_ACTIVE_AT_RUNTIME)
+	{
+		if (0 != (UD->ClassFlags & CLASS_Constructed))
+		{
+			return true;
+		}
+	}
+	else
+	{
+		if (UD->GetDefaultObject(false))
+		{
+			UE_CLOG(!UD->HasAnyClassFlags(CLASS_TokenStreamAssembled), LogStreaming, Fatal, TEXT("Class %s is fully loaded, but does not have its token stream assembled."), *UD->GetFullName());
+			return true;
+		}
+	}
+	return false;
+}
+
+
 EAsyncPackageState::Type FAsyncPackage::LoadImports_Event()
 {
 	SCOPE_CYCLE_COUNTER(STAT_FAsyncPackage_LoadImports);
@@ -1596,7 +1963,14 @@ EAsyncPackageState::Type FAsyncPackage::LoadImports_Event()
 
 		if (Import->XObject)
 		{
-			continue; // we already have this thing
+			if (!IsFullyLoadedObj(Import->XObject))
+			{
+				UE_LOG(LogStreaming, Verbose, TEXT("FAsyncPackage::LoadImports for %s: import %s was found but was not fully loaded yet."), *Desc.NameToLoad.ToString(), *OriginalImport->ObjectName.ToString());
+			}
+			else
+			{
+				continue; // we already have this thing
+			}
 		}
 
 		bool bForcePackageLoad = false;
@@ -1630,7 +2004,7 @@ EAsyncPackageState::Type FAsyncPackage::LoadImports_Event()
 		}
 
 		// Don't try to import a package that is in an import table that we know is an invalid entry
-		if (FLinkerLoad::KnownMissingPackages.Contains(Import->ObjectName))
+		if (FLinkerLoad::IsKnownMissingPackage(Import->ObjectName))
 		{
 			continue;
 		}
@@ -1655,8 +2029,11 @@ EAsyncPackageState::Type FAsyncPackage::LoadImports_Event()
 			if (bCompiledInNotDynamic)
 			{
 				// This can happen with editor only classes, not sure if this should be a warning or a silent continue
+				if (!GIsInitialLoad)
+				{
 				UE_LOG(LogStreaming, Warning, TEXT("FAsyncPackage::LoadImports for %s: Skipping import %s, depends on missing native class"), *Desc.NameToLoad.ToString(), *OriginalImport->ObjectName.ToString());
 			}
+			}
 			else if (!ExistingPackage || bForcePackageLoad)
 			{
 				// The package doesn't exist and this import is not in the dependency list so add it now.
@@ -1689,7 +2066,7 @@ EAsyncPackageState::Type FAsyncPackage::LoadImports_Event()
 				PendingPackage->PackagesWaitingToLinkImports.Add(WeakThis);
 				bDidSomething = true;
 			}
-			else if (int32(PendingPackage->AsyncPackageLoadingState) <= int32(EAsyncPackageLoadingState::WaitingForPostLoad))
+			else if (int32(PendingPackage->AsyncPackageLoadingState) < int32(EAsyncPackageLoadingState::WaitingForPostLoad))
 			{
 				FEventLoadNodePtr MyPrerequisiteNode;
 				MyPrerequisiteNode.WaitingPackage = WeakThis;
@@ -1749,44 +2126,6 @@ void FAsyncPackage::Event_SetupImports()
 	AsyncLoadingThread.QueueEvent_SetupExports(this);
 }
 
-static bool IsFullyLoadedObj(UObject* Obj)
-{
-	if (!Obj)
-	{
-		return false;
-	}
-	if (Obj->HasAllFlags(RF_WasLoaded | RF_LoadCompleted))
-	{
-		return true;
-	}
-	if (Obj->HasAnyFlags(RF_WasLoaded | RF_NeedLoad))
-	{
-		return false;
-	}
-//native blueprint 
-	UDynamicClass* UD = Cast<UDynamicClass>(Obj);
-	if (!UD)
-	{
-		return true;
-	}
-
-	if (GEventDrivenLoaderEnabled && EVENT_DRIVEN_ASYNC_LOAD_ACTIVE_AT_RUNTIME)
-	{
-		if (0 != (UD->ClassFlags & CLASS_Constructed))
-		{
-			return true;
-		}
-	}
-	else
-	{
-		if (UD->GetDefaultObject(false))
-		{
-			UE_CLOG(!UD->HasAnyClassFlags(CLASS_TokenStreamAssembled), LogStreaming, Fatal, TEXT("Class %s is fully loaded, but does not have its token stream assembled."), *UD->GetFullName());
-			return true;
-		}
-	}
-	return false;
-}
 
 EAsyncPackageState::Type FAsyncPackage::SetupImports_Event()
 {
@@ -1807,8 +2146,10 @@ EAsyncPackageState::Type FAsyncPackage::SetupImports_Event()
 	}
 
 	// Create imports.
-	while (ImportIndex < Linker->ImportMap.Num()) //todo timeslice && !IsTimeLimitExceeded())
+	bool bAnyImportArcsAdded = false;
+	while (ImportIndex < Linker->ImportMap.Num())
 	{
+		bool bFireIfNoArcsAdded = true;
 		int32 LocalImportIndex = ImportIndex++;
 		FScopedCreateImportCounter ScopedCounter(Linker, LocalImportIndex);
 		FObjectImport& Import = Linker->ImportMap[LocalImportIndex];
@@ -1841,9 +2182,11 @@ EAsyncPackageState::Type FAsyncPackage::SetupImports_Event()
 				}
 			}
 		}
-		else if (!Import.XObject || !IsFullyLoadedObj(Import.XObject)) 
+		else if (!Import.XObject || !IsFullyLoadedObj(Import.XObject) 
+			|| GIsInitialLoad) // during the initial load, we might find the object, but it hasn't yet actually been finished
 		{
 			FPackageIndex OuterMostIndex = FPackageIndex::FromImport(LocalImportIndex);
+			FPackageIndex OuterMostNonPackageIndex = OuterMostIndex;
 			while (true)
 			{
 				check(!OuterMostIndex.IsNull() && OuterMostIndex.IsImport());
@@ -1853,6 +2196,7 @@ EAsyncPackageState::Type FAsyncPackage::SetupImports_Event()
 				{
 					break;
 				}
+				OuterMostNonPackageIndex = OuterMostIndex;
 				OuterMostIndex = OuterMostImport.OuterIndex;
 			}
 			FObjectImport& OuterMostImport = Linker->Imp(OuterMostIndex);
@@ -1874,7 +2218,17 @@ EAsyncPackageState::Type FAsyncPackage::SetupImports_Event()
 			{
 				FLinkerLoad* ImportLinker = ImportPackage->LinkerLoad;
 				bool bDynamicImport = ImportLinker && ImportLinker->bDynamicClassLinker;
-				if (!ImportLinker || !ImportLinker->AsyncRoot)
+
+#if USE_EVENT_DRIVEN_ASYNC_LOAD_AT_BOOT_TIME
+				if (GIsInitialLoad && !ImportLinker && ImportPackage->HasAnyPackageFlags(PKG_CompiledIn) && !bDynamicImport)
+				{
+					check(ImportPackage->GetName().StartsWith(TEXT("/Script/"))); // sanity check though doesn't really matter; we just need to correctly identify compiled in things.
+					// OuterMostNonPackageIndex is used here because if it is a CDO or subobject, etc, we wait for the outermost thing that is not a package
+					bFireIfNoArcsAdded = !GetGEDLBootNotificationManager().AddWaitingPackage(this, OuterMostImport.ObjectName, Linker->Imp(OuterMostNonPackageIndex).ObjectName, FPackageIndex::FromImport(LocalImportIndex));
+				}
+#endif
+				if (bFireIfNoArcsAdded && // if bFireIfNoArcsAdded is false, then we know we are waiting for compiled in thing, so don't bother to look for the import now
+					(!ImportLinker || !ImportLinker->AsyncRoot))
 				{
 					FindExistingImport(LocalImportIndex);
 					bool bFinishedLoading = IsFullyLoadedObj(Import.XObject);
@@ -1929,7 +2283,8 @@ EAsyncPackageState::Type FAsyncPackage::SetupImports_Event()
 					// this is a hack because the fake export table is missing lots
 					if (bDynamicSomethingMissingFromTheFakeExportTable)
 					{
-						LocalExportIndex = FPackageIndex::FromExport(0);
+						check(ImportLinker->ExportMap.Num() == 2); // we assume there are two elements in the fake export table and the second one is the CDO
+						LocalExportIndex = FPackageIndex::FromExport(1);
 					}
 
 					Import.bImportFailed = LocalExportIndex.IsNull();
@@ -1938,108 +2293,128 @@ EAsyncPackageState::Type FAsyncPackage::SetupImports_Event()
 					if (!Import.bImportFailed)
 					{
 						FObjectExport& Export = ImportLinker->Exp(LocalExportIndex);
-						if (bDynamicSomethingMissingFromTheFakeExportTable)
+						Import.bImportFailed = Export.bExportLoadFailed;
+						if (!Import.bImportFailed)
 						{
-//native blueprint 
-
-							// we can't set Import.SourceIndex because they would be incorrect
-
-							// We hope this things is available when the class is constructed
-							if (!IsFullyLoadedObj(Export.Object))
+							if (bDynamicSomethingMissingFromTheFakeExportTable)
 							{
-								FEventLoadNodePtr MyDependentNode;
-								MyDependentNode.WaitingPackage = WeakThis;
-								MyDependentNode.ImportOrExportIndex = FPackageIndex::FromImport(LocalImportIndex);
-								MyDependentNode.Phase = EEventLoadNode::ImportOrExport_Create;
-
-								{
-									check(int32(ImportLinker->AsyncRoot->AsyncPackageLoadingState) >= int32(EAsyncPackageLoadingState::StartImportPackages));
-									FEventLoadNodePtr PrereqisiteNode;
-									PrereqisiteNode.WaitingPackage = FCheckedWeakAsyncPackagePtr(ImportLinker->AsyncRoot);
-									PrereqisiteNode.ImportOrExportIndex = LocalExportIndex;
-									PrereqisiteNode.Phase = EEventLoadNode::ImportOrExport_Serialize;
+									//native blueprint 
 
-									// can't consider an import serialized until the corresponding export is serialized
-									AddArc(PrereqisiteNode, MyDependentNode);
-								}
+								// we can't set Import.SourceIndex because they would be incorrect
 
+								// We hope this things is available when the class is constructed
+								if (!IsFullyLoadedObj(Export.Object))
 								{
-									FEventLoadNodePtr DependentNode;
-									DependentNode.WaitingPackage = FCheckedWeakAsyncPackagePtr(ImportLinker->AsyncRoot);
-									DependentNode.Phase = EEventLoadNode::Package_ExportsSerialized; // this could be much later...really all we care about is that the linker isn't destroyed.
-
-									// The other package should not leave the event driven loader until we have linked this import
-									AddArc(MyDependentNode, DependentNode);
+									bAnyImportArcsAdded = true;
+									FEventLoadNodePtr MyDependentNode;
+									MyDependentNode.WaitingPackage = WeakThis;
+									MyDependentNode.ImportOrExportIndex = FPackageIndex::FromImport(LocalImportIndex);
+									MyDependentNode.Phase = EEventLoadNode::ImportOrExport_Create;
+
+									{
+										check(int32(ImportLinker->AsyncRoot->AsyncPackageLoadingState) >= int32(EAsyncPackageLoadingState::StartImportPackages));
+										FEventLoadNodePtr PrereqisiteNode;
+										PrereqisiteNode.WaitingPackage = FCheckedWeakAsyncPackagePtr(ImportLinker->AsyncRoot);
+										PrereqisiteNode.ImportOrExportIndex = LocalExportIndex;
+										PrereqisiteNode.Phase = EEventLoadNode::ImportOrExport_Serialize;
+
+										// can't consider an import serialized until the corresponding export is serialized
+										AddArc(PrereqisiteNode, MyDependentNode);
+									}
+
+									{
+										FEventLoadNodePtr DependentNode;
+										DependentNode.WaitingPackage = FCheckedWeakAsyncPackagePtr(ImportLinker->AsyncRoot);
+										DependentNode.Phase = EEventLoadNode::Package_ExportsSerialized; // this could be much later...really all we care about is that the linker isn't destroyed.
+
+										// The other package should not leave the event driven loader until we have linked this import
+										AddArc(MyDependentNode, DependentNode);
+									}
 								}
 							}
-						}
-						else
-						{
-						Import.SourceIndex = LocalExportIndex.ToExport();
-						Import.SourceLinker = ImportLinker;
-							if (!Export.Object)
+							else
 							{
-								FEventLoadNodePtr MyDependentNode;
-								MyDependentNode.WaitingPackage = WeakThis;
-								MyDependentNode.ImportOrExportIndex = FPackageIndex::FromImport(LocalImportIndex);
-								MyDependentNode.Phase = EEventLoadNode::ImportOrExport_Create;
-
+								Import.SourceIndex = LocalExportIndex.ToExport();
+								Import.SourceLinker = ImportLinker;
+								if (!Export.Object)
 								{
-									FEventLoadNodePtr PrereqisiteNode;
-									PrereqisiteNode.WaitingPackage = FCheckedWeakAsyncPackagePtr(ImportLinker->AsyncRoot);
-									PrereqisiteNode.ImportOrExportIndex = LocalExportIndex;
-									PrereqisiteNode.Phase = EEventLoadNode::ImportOrExport_Create;
-
-									// can't create an import until the corresponding export is created
-									AddArc(PrereqisiteNode, MyDependentNode);
+									bAnyImportArcsAdded = true;
+									FEventLoadNodePtr MyDependentNode;
+									MyDependentNode.WaitingPackage = WeakThis;
+									MyDependentNode.ImportOrExportIndex = FPackageIndex::FromImport(LocalImportIndex);
+									MyDependentNode.Phase = EEventLoadNode::ImportOrExport_Create;
+
+									{
+										FEventLoadNodePtr PrereqisiteNode;
+										PrereqisiteNode.WaitingPackage = FCheckedWeakAsyncPackagePtr(ImportLinker->AsyncRoot);
+										PrereqisiteNode.ImportOrExportIndex = LocalExportIndex;
+										PrereqisiteNode.Phase = EEventLoadNode::ImportOrExport_Create;
+
+										// can't create an import until the corresponding export is created
+										AddArc(PrereqisiteNode, MyDependentNode);
+									}
+
+
+									{
+										FEventLoadNodePtr DependentNode;
+										DependentNode.WaitingPackage = FCheckedWeakAsyncPackagePtr(ImportLinker->AsyncRoot);
+										DependentNode.Phase = EEventLoadNode::Package_ExportsSerialized; // this could be much later...really all we care about is that the linker isn't destroyed.
+
+										// The other package should not leave the event driven loader until we have linked this import
+										AddArc(MyDependentNode, DependentNode);
+									}
 								}
-
-
+								else
 								{
-									FEventLoadNodePtr DependentNode;
-									DependentNode.WaitingPackage = FCheckedWeakAsyncPackagePtr(ImportLinker->AsyncRoot);
-									DependentNode.Phase = EEventLoadNode::Package_ExportsSerialized; // this could be much later...really all we care about is that the linker isn't destroyed.
-
-									// The other package should not leave the event driven loader until we have linked this import
-									AddArc(MyDependentNode, DependentNode);
+									check(!Import.XObject || Import.XObject == Export.Object);
+									Import.XObject = Export.Object;
+									if (Import.XObject)
+									{
+										AddObjectReference(Import.XObject);
+									}
 								}
-							}
-							else
-							{
-								check(!Import.XObject || Import.XObject == Export.Object);
-								Import.XObject = Export.Object;
-								if (Import.XObject)
+								if (!IsFullyLoadedObj(Export.Object))
 								{
-									AddObjectReference(Import.XObject);
-								}
-							}
-							if (!IsFullyLoadedObj(Export.Object))
-							{
-								FEventLoadNodePtr MyDependentNode;
-								MyDependentNode.WaitingPackage = WeakThis;
-								MyDependentNode.ImportOrExportIndex = FPackageIndex::FromImport(LocalImportIndex);
-								MyDependentNode.Phase = EEventLoadNode::ImportOrExport_Serialize;
-
-								FEventLoadNodePtr PrereqisiteNode;
-								PrereqisiteNode.WaitingPackage = FCheckedWeakAsyncPackagePtr(ImportLinker->AsyncRoot);
-								PrereqisiteNode.ImportOrExportIndex = LocalExportIndex;
-								PrereqisiteNode.Phase = EEventLoadNode::ImportOrExport_Serialize;
+										bAnyImportArcsAdded = true;
+									FEventLoadNodePtr MyDependentNode;
+									MyDependentNode.WaitingPackage = WeakThis;
+									MyDependentNode.ImportOrExportIndex = FPackageIndex::FromImport(LocalImportIndex);
+									MyDependentNode.Phase = EEventLoadNode::ImportOrExport_Serialize;
 
-								// can't consider an import serialized until the corresponding export is serialized
-								AddArc(PrereqisiteNode, MyDependentNode);
+									FEventLoadNodePtr PrereqisiteNode;
+									PrereqisiteNode.WaitingPackage = FCheckedWeakAsyncPackagePtr(ImportLinker->AsyncRoot);
+									PrereqisiteNode.ImportOrExportIndex = LocalExportIndex;
+									PrereqisiteNode.Phase = EEventLoadNode::ImportOrExport_Serialize;
 
+									// can't consider an import serialized until the corresponding export is serialized
+									AddArc(PrereqisiteNode, MyDependentNode);
+								}
 							}
 						}
 					}
 				}
 			}
 		}
-		DoneAddingPrerequistesFireIfNone(EEventLoadNode::ImportOrExport_Create, FPackageIndex::FromImport(LocalImportIndex));
+		if (bFireIfNoArcsAdded)
+		{
+			DoneAddingPrerequistesFireIfNone(EEventLoadNode::ImportOrExport_Create, FPackageIndex::FromImport(LocalImportIndex));
+		}
+		else
+		{
+			NodeWillBeFiredExternally(EEventLoadNode::ImportOrExport_Create, FPackageIndex::FromImport(LocalImportIndex));
+		}
+	}
+
+	if (bAnyImportArcsAdded && Linker->GetFArchiveAsync2Loader())
+	{
+		// we are waiting for imports, so drop our precache requests
+		Linker->GetFArchiveAsync2Loader()->FlushCache();
 	}
 
 	return ImportIndex == Linker->ImportMap.Num() ? EAsyncPackageState::Complete : EAsyncPackageState::TimeOut;
 }
 
+
 void FAsyncLoadingThread::QueueEvent_SetupExports(FAsyncPackage* Package, int32 EventSystemPriority)
 {
 	check(Package->AsyncPackageLoadingState == EAsyncPackageLoadingState::SetupExports);
@@ -2279,6 +2654,10 @@ EAsyncPackageState::Type FAsyncPackage::SetupExports_Event()
 
 				}
 			}
+			else
+			{
+				Export.bExportLoadFailed = true;
+			}
 		}
 		DoneAddingPrerequistesFireIfNone(EEventLoadNode::ImportOrExport_Create, FPackageIndex::FromExport(LocalExportIndex));
 	}
@@ -2554,8 +2933,9 @@ void FAsyncPackage::EventDrivenCreateExport(int32 LocalExportIndex)
 	SCOPED_LOADTIMER(Package_CreateExports);
 	FObjectExport& Export = Linker->ExportMap[LocalExportIndex];
 	// Check whether we already loaded the object and if not whether the context flags allow loading it.
-	check(!Export.Object || Export.Object->HasAnyFlags(RF_ClassDefaultObject)); // we should not have this yet, unless it is a CDO
-	if (!Export.Object)
+	//check(!Export.Object || Export.Object->HasAnyFlags(RF_ClassDefaultObject)); // we should not have this yet, unless it is a CDO
+	check(!Export.Object); // we should not have this yet
+	if (!Export.Object && !Export.bExportLoadFailed)
 	{
 		if (!Export.Object && !Linker->FilterExport(Export)) // for some acceptable position, it was not "not for" 
 		{
@@ -2686,8 +3066,10 @@ void FAsyncPackage::EventDrivenCreateExport(int32 LocalExportIndex)
 					// Do this for all subobjects created in the native constructor.
 					if (!Export.Object->HasAnyFlags(RF_LoadCompleted))
 					{
-						UE_LOG(LogStreaming, Verbose, TEXT("Note: %s was found in memory but created by some other process and needs loading."), *Export.Object->GetFullName());
+							UE_LOG(LogStreaming, VeryVerbose, TEXT("Note2: %s was constructed during load and is an export and so needs loading."), *Export.Object->GetFullName());
+							UE_CLOG(!Export.Object->HasAllFlags(RF_WillBeLoaded), LogStreaming, Fatal, TEXT("%s was found in memory and is an export but does not have all load flags."), *Export.Object->GetFullName());
 						Export.Object->SetFlags(RF_NeedLoad | RF_NeedPostLoad | RF_NeedPostLoadSubobjects | RF_WasLoaded);
+							Export.Object->ClearFlags(RF_WillBeLoaded);
 					}
 				}
 			}
@@ -2724,7 +3106,7 @@ void FAsyncPackage::EventDrivenCreateExport(int32 LocalExportIndex)
 
 				// If we are about to create a CDO, we need to ensure that all parent sub-objects are loaded
 				// to get default value initialization to work.
-#if DO_CHECK
+	#if DO_CHECK
 				if ((ObjectLoadFlags & RF_ClassDefaultObject) != 0)
 				{
 					UClass* SuperClass = LoadClass->GetSuperClass();
@@ -2763,7 +3145,7 @@ void FAsyncPackage::EventDrivenCreateExport(int32 LocalExportIndex)
 						check(Template->IsA(LoadClass));
 					}
 				}
-#endif
+	#endif
 				if (LoadClass->HasAnyFlags(RF_NeedLoad))
 				{
 					UE_LOG(LogStreaming, Fatal, TEXT("LoadClass %s had RF_NeedLoad while creating %s"), *LoadClass->GetFullName(), *Export.ObjectName.ToString());
@@ -2814,10 +3196,51 @@ void FAsyncPackage::EventDrivenCreateExport(int32 LocalExportIndex)
 	{
 		AddObjectReference(Export.Object);
 	}
+	else
+	{
+		Export.bExportLoadFailed = true;
+	}
 	LastObjectWorkWasPerformedOn = Export.Object;
-	check(Linker->FilterExport(Export) || (Export.Object && !Export.bExportLoadFailed));
+	check(Export.Object || Export.bExportLoadFailed);
+}
+
+static FPackageIndex FindExportFromObject(FLinkerLoad* Linker, UObject *Object)
+{
+	check(Linker && Linker->AsyncRoot && Linker->AsyncRoot->ObjectNameToImportOrExport.Num());
+	FPackageIndex Result;
+	UObject* Outer = Object->GetOuter();
+	if (Outer)
+	{
+		FPackageIndex OuterIndex = FindExportFromObject(Linker, Outer);
+		for (auto It = Linker->AsyncRoot->ObjectNameToImportOrExport.CreateKeyIterator(Object->GetFName()); It; ++It)
+		{
+			FPackageIndex PotentialExport = It.Value();
+			if (PotentialExport.IsExport())
+			{
+				FObjectExport& Export = Linker->Exp(PotentialExport);
+				if (OuterIndex == Export.OuterIndex)
+				{
+					Result = PotentialExport;
+					break;
+				}
+			}
+		}
+	}
+	return Result;
 }
 
+void FAsyncPackage::MarkNewObjectForLoadIfItIsAnExport(UObject *Object)
+{
+	if (!Object->HasAnyFlags(RF_WillBeLoaded | RF_LoadCompleted | RF_NeedLoad))
+	{
+		FPackageIndex MaybeExportIndex = FindExportFromObject(Linker, Object);
+		if (MaybeExportIndex.IsExport())
+		{
+			UE_LOG(LogStreaming, VeryVerbose, TEXT("Note: %s was constructed during load and is an export and so needs loading."), *Object->GetFullName());
+			Object->SetFlags(RF_WillBeLoaded);
+		}
+	}
+}
 void FAsyncPackage::EventDrivenSerializeExport(int32 LocalExportIndex)
 {
 	SCOPED_LOADTIMER(Package_PreLoadObjects);
@@ -2834,6 +3257,7 @@ void FAsyncPackage::EventDrivenSerializeExport(int32 LocalExportIndex)
 			UObject* LocObj = ConstructDynamicType(*Linker->GetExportPathName(LocalExportIndex), EConstructDynamicType::CallZConstructor);
 			check(UD == LocObj);
 		}
+		Object->ClearFlags(RF_NeedLoad | RF_WillBeLoaded);
 	}
 	else if (Object && Object->HasAnyFlags(RF_NeedLoad))
 	{
@@ -2904,8 +3328,6 @@ void FAsyncPackage::EventDrivenSerializeExport(int32 LocalExportIndex)
 		check(Linker->TemplateForGetArchetypeFromLoader == Template);
 		Linker->TemplateForGetArchetypeFromLoader = nullptr;
 
-
-
 		Object->SetFlags(RF_LoadCompleted);
 		ThreadContext.SerializedObject = PrevSerializedObject;
 		Linker->bForceSimpleIndexToObject = false;
@@ -3162,7 +3584,7 @@ EAsyncPackageState::Type FAsyncPackage::ProcessImportsAndExports_Event()
 		}
 		if (bDidSomething)
 		{
-			continue; // check time limit
+			continue; // check time limitE
 		}
 		if (ImportsThatAreNowSerialized.Num())
 		{
@@ -3222,6 +3644,11 @@ EAsyncPackageState::Type FAsyncPackage::ProcessImportsAndExports_Event()
 				GCurrentExportIndex = LocalExportIndex;
 				EventDrivenSerializeExport(LocalExportIndex);
 				GCurrentExportIndex = -1;
+				{
+					FObjectExport& Export = Linker->ExportMap[LocalExportIndex];
+					UObject* Object = Export.Object;
+					check(!Object || !Object->HasAnyFlags(RF_NeedLoad));
+				}
 			}
 			else
 			{
@@ -3420,6 +3847,7 @@ void FAsyncPackage::Event_StartPostload()
 			checkSlow(!(Object && !ReferencedObjects.Contains(Object)));
 			if (Object && (Object->HasAnyFlags(RF_NeedPostLoad) || Linker->bDynamicClassLinker || Object->HasAnyInternalFlags(EInternalObjectFlags::AsyncLoading)))
 			{
+				check(Object->IsValidLowLevelFast());
 				ObjLoaded.Add(Object);
 			}
 		}
@@ -3511,6 +3939,31 @@ void FAsyncPackage::RemoveAllNodes()
 	}
 }
 
+bool FEDLBootNotificationManager::FireCompletedCompiledInImports(bool bFinalRun)
+{
+#if USE_EVENT_DRIVEN_ASYNC_LOAD_AT_BOOT_TIME
+	FScopeLock Lock(&EDLBootNotificationManagerLock);
+	check(bFinalRun || GIsInitialLoad);
+	bool bResult = !!PathsToFire.Num();
+	for (FName LongName : PathsToFire)
+	{
+		for (auto It = PathToWaitingPackageNodes.CreateKeyIterator(LongName); It; ++It)
+		{
+			FEventLoadNodePtr& NodeToFire = It.Value();
+			FAsyncPackage* Pkg = &NodeToFire.WaitingPackage.GetPackage();
+			check(Pkg);
+			Pkg->FireNode(NodeToFire);
+		}
+		PathToWaitingPackageNodes.Remove(LongName);
+	}
+	PathsToFire.Empty();	
+	return bResult;
+#else
+	return false;
+#endif
+}
+
+
 void FAsyncPackage::FireNode(FEventLoadNodePtr& NodeToFire)
 {
 	check(int32(AsyncPackageLoadingState) < int32(EAsyncPackageLoadingState::PostLoad_Etc));
@@ -3782,7 +4235,6 @@ EAsyncPackageState::Type FAsyncLoadingThread::ProcessAsyncLoading(int32& OutPack
 				{
 
 					// We're done, at least on this thread, so we can remove the package now.
-					AddToLoadedPackages(Package);
 					{
 #if THREADSAFE_UOBJECTS
 						FScopeLock LockAsyncPackages(&AsyncPackagesCritical);
@@ -3792,6 +4244,7 @@ EAsyncPackageState::Type FAsyncLoadingThread::ProcessAsyncLoading(int32& OutPack
 						AsyncPackages.RemoveAt(PackageIndex);
 						AsyncPackagesReadyForTick.RemoveAt(0, 1, false); //@todoio this should maybe be a heap or something to avoid the removal cost
 					}
+				AddToLoadedPackages(Package);
 				}
 				if (IsTimeLimitExceeded(TickStartTime, bUseTimeLimit, TimeLimit, TEXT("TickAsyncPackage")))
 				{
@@ -3937,6 +4390,7 @@ EAsyncPackageState::Type FAsyncLoadingThread::ProcessLoadedPackages(bool bUseTim
 	volatile int32 CurrentAsyncLoadingCounter = AsyncLoadingTickCounter;
 
 	double TickStartTime = FPlatformTime::Seconds();
+
 	{
 #if THREADSAFE_UOBJECTS
 		FScopeLock LoadedPackagesLock(&LoadedPackagesCritical);
@@ -3947,6 +4401,20 @@ EAsyncPackageState::Type FAsyncLoadingThread::ProcessLoadedPackages(bool bUseTim
 		LoadedPackages.Reset();
 		LoadedPackagesNameLookup.Reset();
 	}
+#if USE_EVENT_DRIVEN_ASYNC_LOAD_AT_BOOT_TIME
+	if (IsMultithreaded() && GEventDrivenLoaderEnabled &&
+		ENamedThreads::RenderThread == ENamedThreads::GameThread) // render thread tasks are actually being sent to the game thread.
+	{
+		// The async loading thread might have queued some render thread tasks (we don't have a render thread yet, so these are actually sent to the game thread)
+		// We need to process them now before we do any postloads.
+		FTaskGraphInterface::Get().ProcessThreadUntilIdle(ENamedThreads::GameThread);
+		if (IsTimeLimitExceeded(TickStartTime, bUseTimeLimit, TimeLimit, TEXT("ProcessRenderThreadCommandsOnGameThread")))
+		{
+			return EAsyncPackageState::TimeOut;
+		}
+	}
+#endif
+
 		
 	bDidSomething = LoadedPackagesToProcess.Num() > 0;
 	for (int32 PackageIndex = 0; PackageIndex < LoadedPackagesToProcess.Num() && !IsAsyncLoadingSuspended(); ++PackageIndex)
@@ -4094,17 +4562,46 @@ EAsyncPackageState::Type FAsyncLoadingThread::TickAsyncLoading(bool bUseTimeLimi
 					IsTimeLimitExceeded(TickStartTime, bUseTimeLimit, TimeLimit, TEXT("FDeferredMessageLog::Flush()"));
 				}
 			}
-			if (!bDidSomething && !bIsMultithreaded && GEventDrivenLoaderEnabled)
+			if (!bDidSomething && GEventDrivenLoaderEnabled)
+			{
+				if (bIsMultithreaded)
+				{
+#if USE_EVENT_DRIVEN_ASYNC_LOAD_AT_BOOT_TIME
+					if (GIsInitialLoad)
+					{
+						bDidSomething = GetGEDLBootNotificationManager().ConstructWaitingBootObjects(); // with the ASL, we always create new boot objects when we have nothing else to do
+						IsTimeLimitExceeded(TickStartTime, bUseTimeLimit, TimeLimit, TEXT("ConstructWaitingBootObjects"));
+					}
+#endif
+				}
+				else
+				{
+#if USE_EVENT_DRIVEN_ASYNC_LOAD_AT_BOOT_TIME
+					if (GIsInitialLoad)
+					{
+						bDidSomething = GetGEDLBootNotificationManager().FireCompletedCompiledInImports(); // no ASL, first try to fire any completed boot objects, and if there are none, then create some boot objects
+						IsTimeLimitExceeded(TickStartTime, bUseTimeLimit, TimeLimit, TEXT("FireCompletedCompiledInImports"));
+						if (!bDidSomething)
+						{
+							bDidSomething = GetGEDLBootNotificationManager().ConstructWaitingBootObjects();
+							IsTimeLimitExceeded(TickStartTime, bUseTimeLimit, TimeLimit, TEXT("ConstructWaitingBootObjects"));
+						}
+					}
+#endif
+					if (!bDidSomething)
 			{
 				CheckForCycles();
+					}
 				IsTimeLimitExceeded(TickStartTime, bUseTimeLimit, TimeLimit, TEXT("CheckForCycles (non-shipping)"));
 			}
 		}
 	}
+	}
 
 	return Result;
 }
 
+
 FAsyncLoadingThread::FAsyncLoadingThread()
 {
 	// Current these two vars are always on or off together but can be made separate
@@ -4240,6 +4737,10 @@ EAsyncPackageState::Type  FAsyncLoadingThread::TickAsyncThread(bool bUseTimeLimi
 		double TickStartTime = FPlatformTime::Seconds();
 		if (AsyncThreadReady.GetValue())
 		{
+			if (GIsInitialLoad && GEventDrivenLoaderEnabled)
+			{
+				GetGEDLBootNotificationManager().FireCompletedCompiledInImports();
+			}
 			CreateAsyncPackagesFromQueue(bUseTimeLimit, bUseFullTimeLimit, TimeLimit, FlushTree);
 			float TimeUsed = (float)(FPlatformTime::Seconds() - TickStartTime);
 			const float RemainingTimeLimit = FMath::Max(0.0f, TimeLimit - TimeUsed);
@@ -4255,10 +4756,19 @@ EAsyncPackageState::Type  FAsyncLoadingThread::TickAsyncThread(bool bUseTimeLimi
 		}
 		if (ProcessedRequests == 0 && IsMultithreaded())
 		{
+			uint32 WaitTime = 30;
 			if (IsEventDrivenLoaderEnabled())
 			{
+
+				if (!GetGEDLBootNotificationManager().IsWaitingForSomething())
+				{
 				CheckForCycles();
 				IsTimeLimitExceeded(TickStartTime, bUseTimeLimit, TimeLimit, TEXT("CheckForCycles (non-shipping)"));
+			}
+				else
+				{
+					WaitTime = 1; // we are waiting for the game thread to deal with the boot constructors, so lets spin tighter
+				}
 			}
 			const bool bIgnoreThreadIdleStats = true;
 			SCOPED_LOADTIMER(Package_Temp3);
@@ -4388,6 +4898,16 @@ void NotifyConstructedDuringAsyncLoading(UObject* Object, bool bSubObject)
 	FUObjectThreadContext& ThreadContext = FUObjectThreadContext::Get();
 	check(ThreadContext.AsyncPackage); // Otherwise something is wrong and we're creating objects outside of async loading code
 	ThreadContext.AsyncPackage->AddObjectReference(Object);
+	if (GEventDrivenLoaderEnabled)
+	{
+		// if this is in the package and is an export, then let mark it as needing load now
+		if (Object->GetOutermost() == ThreadContext.AsyncPackage->GetLinkerRoot() && 
+			int32(ThreadContext.AsyncPackage->AsyncPackageLoadingState) <= int32(EAsyncPackageLoadingState::ProcessNewImportsAndExports) &&
+			int32(ThreadContext.AsyncPackage->AsyncPackageLoadingState) > int32(EAsyncPackageLoadingState::WaitingForSummary))
+		{
+			ThreadContext.AsyncPackage->MarkNewObjectForLoadIfItIsAnExport(Object);
+		}
+	}
 }
 
 /*-----------------------------------------------------------------------------
@@ -4487,6 +5007,7 @@ FAsyncPackage::~FAsyncPackage()
 	EmptyReferencedObjects();
 }
 
+
 void FAsyncPackage::AddReferencedObjects(FReferenceCollector& Collector)
 {
 	Collector.AllowEliminatingReferences(false);
@@ -5829,7 +6350,7 @@ void FAsyncPackage::CallCompletionCallbacks(bool bInternal, EAsyncLoadingResult:
 
 void FAsyncPackage::Cancel()
 {
-	UE_CLOG(GNewAsyncIO, LogStreaming, Fatal, TEXT("FAsyncPackage::Cancel is not supported with the new loader"));
+	UE_CLOG(GNewAsyncIO || GEventDrivenLoaderEnabled, LogStreaming, Fatal, TEXT("FAsyncPackage::Cancel is not supported with the new loader"));
 
 	// Call any completion callbacks specified.
 	const EAsyncLoadingResult::Type Result = EAsyncLoadingResult::Canceled;
@@ -5837,8 +6358,8 @@ void FAsyncPackage::Cancel()
 	{
 		if (!CompletionCallbacks[CallbackIndex].bCalled)
 		{
-		CompletionCallbacks[CallbackIndex].Callback.ExecuteIfBound(Desc.Name, nullptr, Result);
-	}
+			CompletionCallbacks[CallbackIndex].Callback.ExecuteIfBound(Desc.Name, nullptr, Result);
+		}
 	}
 	{
 		// Clear load flags from any referenced objects
@@ -5900,13 +6421,20 @@ void FAsyncPackage::UpdateLoadPercentage()
 
 int32 LoadPackageAsync(const FString& InName, const FGuid* InGuid /*= nullptr*/, const TCHAR* InPackageToLoadFrom /*= nullptr*/, FLoadPackageAsyncDelegate InCompletionDelegate /*= FLoadPackageAsyncDelegate()*/, EPackageFlags InPackageFlags /*= PKG_None*/, int32 InPIEInstanceID /*= INDEX_NONE*/, int32 InPackagePriority /*= 0*/)
 {
+	static bool bOnce = false;
+	if (!bOnce && GEventDrivenLoaderEnabled)
+	{
+		bOnce = true;
+		FGCObject::StaticInit(); // otherwise this thing is created during async loading, but not associated with a package
+	}
+
 #if !WITH_EDITOR
 	if (GPreloadPackageDependencies && !GEventDrivenLoaderEnabled)
 	{
 		// If dependency preloading is enabled, we need to force the asset registry module to be loaded on the game thread
 		// as it will potentiall be used on the async loading thread, which isn't allowed to load modules.
 		// We could do this at init time, but doing it here allows us to not load the module at all if preloading is
-		// disabled.
+		// disabled.l
 		IAssetRegistryInterface::GetPtr();
 	}
 #endif
@@ -5974,7 +6502,15 @@ void CancelAsyncLoading()
 	// Cancelling async loading while loading is suspend will result in infinite stall
 	UE_CLOG(FAsyncLoadingThread::Get().IsAsyncLoadingSuspended(), LogStreaming, Fatal, TEXT("Cannot Cancel Async Loading while async loading is suspended."));
 
-	FAsyncLoadingThread::Get().CancelAsyncLoading();
+	if (GEventDrivenLoaderEnabled)
+	{
+		UE_LOG(LogStreaming, Warning, TEXT("Cannot Cancel Async Loading using the EDL loader. Async loading will be flushed instead."));
+		FlushAsyncLoading();
+	}
+	else
+	{
+		FAsyncLoadingThread::Get().CancelAsyncLoading();
+	}
 	CollectGarbage(GARBAGE_COLLECTION_KEEPFLAGS, false);
 }
 
@@ -6128,6 +6664,19 @@ bool IsEventDrivenLoaderEnabled()
 	return GEventDrivenLoaderEnabled;
 }
 
+
+void NotifyRegistrationEvent(const TCHAR* PackageName, const TCHAR* Name, ENotifyRegistrationType NotifyRegistrationType, ENotifyRegistrationPhase NotifyRegistrationPhase, UObject *(*InRegister)(), bool InbDynamic)
+{
+	GetGEDLBootNotificationManager().NotifyRegistrationEvent(PackageName, Name, NotifyRegistrationType, NotifyRegistrationPhase, InRegister, InbDynamic);
+}
+
+
+void NotifyRegistrationComplete()
+{
+	GetGEDLBootNotificationManager().NotifyRegistrationComplete();
+}
+
+
 /*----------------------------------------------------------------------------
 	FArchiveAsync.
 ----------------------------------------------------------------------------*/
@@ -6982,8 +7531,16 @@ void FArchiveAsync2::CompleteRead()
 			PrecacheEndPos = ReadRequestOffset + ReadRequestSize;
 			INC_MEMORY_STAT_BY(STAT_FArchiveAsync2Mem, PrecacheEndPos - PrecacheStartPos);
 			DEC_MEMORY_STAT_BY(STAT_AsyncFileMemory, ReadRequestSize);
+
+			// keeps the last cache block of the header around until we process the first export
+			if (LoadPhase != ELoadPhase::ProcessingExports)
+			{
+				CompleteCancel();
+				CanceledReadRequestPtr = Handle->ReadRequest(PrecacheEndPos - HeaderSizeWhenReadingExportsFromSplitFile - 1, 1, AIOP_Precache);
+			}
 		}
 	}
+
 	delete ReadRequestPtr;
 	ReadRequestPtr = nullptr;
 	LogItem(TEXT("CompleteRead"), ReadRequestOffset, ReadRequestSize);
@@ -6997,7 +7554,7 @@ void FArchiveAsync2::CompleteCancel()
 	{
 		double StartTime = FPlatformTime::Seconds();
 		CanceledReadRequestPtr->WaitCompletion();
-		check(!CanceledReadRequestPtr->GetReadResults()); // this should have been canceled
+		//check(!CanceledReadRequestPtr->GetReadResults()); // this should have been canceled
 		delete CanceledReadRequestPtr;
 		CanceledReadRequestPtr = nullptr;
 		LogItem(TEXT("Complete Cancel"), 0, 0, StartTime);
@@ -7176,21 +7733,34 @@ void FArchiveAsync2::FirstExportStarting()
 	ExportReadTime = FPlatformTime::Seconds();
 	LogItem(TEXT("Exports"));
 	LoadPhase = ELoadPhase::ProcessingExports;
+}
 
-	if (GEventDrivenLoaderEnabled)
+IAsyncReadRequest* FArchiveAsync2::MakeEventDrivenPrecacheRequest(int64 Offset, int64 BytesToRead, FAsyncFileCallBack* CompleteCallback)
+{
+	check(GEventDrivenLoaderEnabled);
+	if (LoadPhase == ELoadPhase::WaitingForFirstExport)
 	{
+		// we need to avoid tearing down the old file and requests until we have the one in flight
+		double StartTime = FPlatformTime::Seconds();
+		HeaderSizeWhenReadingExportsFromSplitFile = HeaderSize;
+
+		FString NewFileName = FPaths::GetBaseFilename(FileName, false) + TEXT(".uexp");
+		IAsyncReadFileHandle* NewHandle = FPlatformFileManager::Get().GetPlatformFile().OpenAsyncRead(*NewFileName);
+		check(NewHandle); // this generally cannot fail because it is async
+
+
+		check(Offset - HeaderSizeWhenReadingExportsFromSplitFile >= 0);
+		IAsyncReadRequest* Precache = NewHandle->ReadRequest(Offset - HeaderSizeWhenReadingExportsFromSplitFile, BytesToRead, AIOP_Precache, CompleteCallback);
 		FlushCache();
 		if (Handle)
 		{
 			delete Handle;
 			Handle = nullptr;
 		}
+		Handle = NewHandle;
+		FileName = NewFileName;
 
-		HeaderSizeWhenReadingExportsFromSplitFile = HeaderSize;
-		FileName = FPaths::GetBaseFilename(FileName, false) + TEXT(".uexp");
-
-		Handle = FPlatformFileManager::Get().GetPlatformFile().OpenAsyncRead(*FileName);
-		check(Handle); // this generally cannot fail because it is async
+		FirstExportStarting();
 
 		check(!SizeRequestPtr);
 		SizeRequestPtr = Handle->SizeRequest();
@@ -7198,15 +7768,8 @@ void FArchiveAsync2::FirstExportStarting()
 		{
 			TotalSize(); // complete the request
 		}
-	}
-}
-
-IAsyncReadRequest* FArchiveAsync2::MakeEventDrivenPrecacheRequest(int64 Offset, int64 BytesToRead, FAsyncFileCallBack* CompleteCallback)
-{
-	check(GEventDrivenLoaderEnabled);
-	if (LoadPhase == ELoadPhase::WaitingForFirstExport)
-	{
-		FirstExportStarting();
+		LogItem(TEXT("First Precache"), Offset - HeaderSizeWhenReadingExportsFromSplitFile, BytesToRead, StartTime);
+		return Precache;
 	}
 	double StartTime = FPlatformTime::Seconds();
 	check(Offset - HeaderSizeWhenReadingExportsFromSplitFile >= 0);
diff --git a/Engine/Source/Runtime/CoreUObject/Private/UObject/Class.cpp b/Engine/Source/Runtime/CoreUObject/Private/UObject/Class.cpp
--- a/Engine/Source/Runtime/CoreUObject/Private/UObject/Class.cpp
+++ b/Engine/Source/Runtime/CoreUObject/Private/UObject/Class.cpp
@@ -33,6 +33,7 @@
 #include "UObject/LinkerPlaceholderFunction.h"
 #include "UObject/StructScriptLoader.h"
 #include "UObject/PropertyHelper.h"
+#include "UObject/CoreRedirects.h"
 #include "Serialization/ArchiveScriptReferenceCollector.h"
 
 // This flag enables some expensive class tree validation that is meant to catch mutations of 
@@ -68,6 +69,8 @@ COREUOBJECT_API void InitializePrivateStaticClass(
 	const TCHAR* Name
 	)
 {
+	NotifyRegistrationEvent(PackageName, Name, ENotifyRegistrationType::NRT_Class, ENotifyRegistrationPhase::NRP_Started);
+
 	/* No recursive ::StaticClass calls allowed. Setup extras. */
 	if (TClass_Super_StaticClass != TClass_PrivateStaticClass)
 	{
@@ -91,6 +94,7 @@ COREUOBJECT_API void InitializePrivateStaticClass(
 		// Register immediately (don't let the function name mistake you!)
 		TClass_PrivateStaticClass->DeferredRegister(UDynamicClass::StaticClass(), PackageName, Name);
 	}
+	NotifyRegistrationEvent(PackageName, Name, ENotifyRegistrationType::NRT_Class, ENotifyRegistrationPhase::NRP_Finished);
 }
 
 void FNativeFunctionRegistrar::RegisterFunction(class UClass* Class, const ANSICHAR* InName, Native InPointer)
@@ -853,47 +857,10 @@ void UStruct::SerializeBinEx( FArchive& Ar, void* Data, void const* DefaultData,
 	}
 }
 
-TMap<FName,TMap<FName,FName> > UStruct::TaggedPropertyRedirects;
-void UStruct::InitTaggedPropertyRedirectsMap()
-{
-	if( GConfig )
-	{
-		// Go over all configs so that plugins/etc have a chance to register TaggedPropertyRedirects
-		for (const auto& ConfigPair : *GConfig)
-		{
-			const FString& ConfigFilename = ConfigPair.Key;
-
-			if (FConfigSection* PackageRedirects = GConfig->GetSectionPrivate( TEXT("/Script/Engine.Engine"), false, true, ConfigFilename ))
-			{
-				for( FConfigSection::TIterator It(*PackageRedirects); It; ++It )
-				{
-					if( It.Key() == TEXT("TaggedPropertyRedirects") )
-					{
-						FName ClassName = NAME_None;
-						FName OldPropertyName = NAME_None;
-						FName NewPropertyName = NAME_None;
-
-						FParse::Value( *It.Value().GetValue(), TEXT("ClassName="), ClassName );
-						FParse::Value( *It.Value().GetValue(), TEXT("OldPropertyName="), OldPropertyName );
-						FParse::Value( *It.Value().GetValue(), TEXT("NewPropertyName="), NewPropertyName );
-
-						check(ClassName != NAME_None && OldPropertyName != NAME_None && NewPropertyName != NAME_None );
-						TaggedPropertyRedirects.FindOrAdd(ClassName).Add(OldPropertyName, NewPropertyName);
-					}
-				}
-			}
-		}
-	}
-	else
-	{
-		UE_LOG(LogClass, Warning, TEXT(" **** TAGGED PROPERTY REDIRECTS UNABLE TO INITIALIZE! **** "));
-	}
-}
-
 void UStruct::SerializeTaggedProperties(FArchive& Ar, uint8* Data, UStruct* DefaultsStruct, uint8* Defaults, const UObject* BreakRecursionIfFullyLoad) const
 {
 	//SCOPED_LOADTIMER(SerializeTaggedPropertiesTime);
-	
+
 	// Determine if this struct supports optional property guid's (UBlueprintGeneratedClasses Only)
 	const bool bArePropertyGuidsAvailable = (Ar.UE4Ver() >= VER_UE4_PROPERTY_GUID_IN_PROPERTY_TAG) && !FPlatformProperties::RequiresCookedData() && ArePropertyGuidsAvailable();
 
@@ -953,42 +920,24 @@ void UStruct::SerializeTaggedProperties(FArchive& Ar, uint8* Data, UStruct* Defa
 				// No need to check redirects on platforms where everything is cooked. Always check for save games
 				if (!FPlatformProperties::RequiresCookedData() || Ar.IsSaveGame())
 				{
-					// Look in the redirect table to see if we're searching for a different name
-					static bool bAlreadyInitialized_TaggedPropertyRedirectsMap = false;
-					if( !bAlreadyInitialized_TaggedPropertyRedirectsMap )
-					{
-						InitTaggedPropertyRedirectsMap();
-						bAlreadyInitialized_TaggedPropertyRedirectsMap = true;
-					}
-					
 					FName EachName = GetFName();
+					FName PackageName = GetOutermost()->GetFName();
 					// Search the current class first, then work up the class hierarchy to see if theres a match for our fixup.
 					UStruct* Owner = GetOwnerStruct();
 					if( Owner )
 					{
-						UStruct* SuperClass = Owner->GetSuperStruct();
-						while( EachName != NAME_None)
+						UStruct* CheckStruct = Owner;
+						while(CheckStruct)
 						{
-							const TMap<FName, FName>* ClassTaggedPropertyRedirects = TaggedPropertyRedirects.Find( EachName );
-							if (ClassTaggedPropertyRedirects)
-							{
-								const FName* NewPropertyName = ClassTaggedPropertyRedirects->Find(Tag.Name);
-								if (NewPropertyName)
-								{
-									Tag.Name = *NewPropertyName;
-									break;
-								}
-							}
-							// If theres another class name to check get it, otherwise flag the end.
-							if( SuperClass != nullptr )
-							{
-								EachName = SuperClass->GetFName();
-								SuperClass = SuperClass->GetSuperStruct();
-							}
-							else
+							FName NewTagName = UProperty::FindRedirectedPropertyName(CheckStruct, Tag.Name);
+
+							if (NewTagName != NAME_None)
 							{
-								EachName = NAME_None;
+								Tag.Name = NewTagName;
+								break;
 							}
+
+							CheckStruct = CheckStruct->GetSuperStruct();
 						}
 					}
 				}
@@ -2713,6 +2662,17 @@ UObject* UClass::CreateDefaultObject()
 			// NULL (so we don't invalidate one that has already been setup)
 			if (ClassDefaultObject == NULL)
 			{
+				FString PackageName;
+				FString CDOName;
+				bool bDoNotify = false;
+				if (GIsInitialLoad && GetOutermost()->HasAnyPackageFlags(PKG_CompiledIn))
+				{
+					PackageName = GetOutermost()->GetFName().ToString();
+					CDOName = GetDefaultObjectName().ToString();
+					NotifyRegistrationEvent(*PackageName, *CDOName, ENotifyRegistrationType::NRT_ClassCDO, ENotifyRegistrationPhase::NRP_Started);
+					bDoNotify = true;
+				}
+
 				// RF_ArchetypeObject flag is often redundant to RF_ClassDefaultObject, but we need to tag
 				// the CDO as RF_ArchetypeObject in order to propagate that flag to any default sub objects.
 				ClassDefaultObject = StaticAllocateObject(this, GetOuter(), NAME_None, EObjectFlags(RF_Public|RF_ClassDefaultObject|RF_ArchetypeObject));
@@ -2720,6 +2680,11 @@ UObject* UClass::CreateDefaultObject()
 				// Blueprint CDOs have their properties always initialized.
 				const bool bShouldInitializeProperties = !HasAnyClassFlags(CLASS_Native | CLASS_Intrinsic);
 				(*ClassConstructor)(FObjectInitializer(ClassDefaultObject, ParentDefaultObject, false, bShouldInitializeProperties));
+				if (bDoNotify)
+				{
+					NotifyRegistrationEvent(*PackageName, *CDOName, ENotifyRegistrationType::NRT_ClassCDO, ENotifyRegistrationPhase::NRP_Finished);
+				}
+				ClassDefaultObject->PostCDOContruct();
 			}
 		}
 	}
