Copying //UE4/Dev-Core to //UE4/Dev-Main (Source: //UE4/Dev-Core @ 3283640)

diff --git a/Engine/Source/Runtime/CoreUObject/Private/UObject/CoreRedirects.cpp b/Engine/Source/Runtime/CoreUObject/Private/UObject/CoreRedirects.cpp
new file mode 100644
--- /dev/null
+++ b/Engine/Source/Runtime/CoreUObject/Private/UObject/CoreRedirects.cpp
@@ -0,0 +1,1713 @@
+// Copyright 1998-2017 Epic Games, Inc. All Rights Reserved.
+
+#include "UObject/CoreRedirects.h"
+#include "UObject/Package.h"
+#include "UObject/PropertyHelper.h"
+#include "UObject/Linker.h"
+#include "Templates/Casts.h"
+#include "Misc/ConfigCacheIni.h"
+#include "Serialization/DeferredMessageLog.h"
+
+FCoreRedirectObjectName::FCoreRedirectObjectName(const FString& InString)
+{
+	if (!ExpandNames(InString, ObjectName, OuterName, PackageName))
+	{
+		Reset();
+	}
+}
+
+FCoreRedirectObjectName::FCoreRedirectObjectName(const class UObject* Object)
+{
+	// This is more efficient than going to path string and back to FNames
+	if (Object)
+	{	
+		UObject* Outer = Object->GetOuter();
+
+		if (!Outer)
+		{
+			PackageName = Object->GetFName();
+			// This is a package
+		}
+		else
+		{
+			FString OuterString;
+			ObjectName = Object->GetFName();
+
+			// Follow outer chain,
+			while (Outer)
+			{
+				UObject* NextOuter = Outer->GetOuter();
+				if (!NextOuter)
+				{
+					OuterName = FName(*OuterString);
+					PackageName = Outer->GetFName();
+					break;
+				}
+				if (!OuterString.IsEmpty())
+				{
+					OuterString.AppendChar('.');
+				}
+				OuterString.Append(Outer->GetName());
+
+				Outer = NextOuter;
+			}
+		}
+	}
+}
+
+FString FCoreRedirectObjectName::ToString() const
+{
+	return CombineNames(ObjectName, OuterName, PackageName);
+}
+
+void FCoreRedirectObjectName::Reset()
+{
+	ObjectName = OuterName = PackageName = NAME_None;
+}
+
+bool FCoreRedirectObjectName::Matches(const FCoreRedirectObjectName& Other, bool bCheckSubstring) const
+{
+	if (bCheckSubstring)
+	{
+		// Much slower
+		if (ObjectName != NAME_None && !Other.ObjectName.ToString().Contains(ObjectName.ToString()))
+		{
+			return false;
+		}
+		if (OuterName != NAME_None && !Other.OuterName.ToString().Contains(OuterName.ToString()))
+		{
+			return false;
+		}
+		if (PackageName != NAME_None && !Other.PackageName.ToString().Contains(PackageName.ToString()))
+		{
+			return false;
+		}
+	}
+	else
+	{
+		// Check all names that are non empty
+		if (ObjectName != Other.ObjectName && ObjectName != NAME_None)
+		{
+			return false;
+		}
+		if (OuterName != Other.OuterName && OuterName != NAME_None)
+		{
+			return false;
+		}
+		if (PackageName != Other.PackageName && PackageName != NAME_None)
+		{
+			return false;
+		}
+	}
+
+	return true;
+}
+
+int32 FCoreRedirectObjectName::MatchScore(const FCoreRedirectObjectName& Other) const
+{
+	int32 MatchScore = 1;
+
+	if (ObjectName != NAME_None)
+	{
+		if (ObjectName == Other.ObjectName)
+		{
+			// Object name most important
+			MatchScore += 8;
+		}
+		else
+		{
+			return 0;
+		}
+	}
+
+	if (OuterName != NAME_None)
+	{
+		if (OuterName == Other.OuterName)
+		{
+			MatchScore += 4;
+		}
+		else
+		{
+			return 0;
+		}
+	}
+
+	if (PackageName != NAME_None)
+	{
+		if (PackageName == Other.PackageName)
+		{
+			MatchScore += 2;
+		}
+		else
+		{
+			return 0;
+		}
+	}
+
+	return MatchScore;
+}
+
+bool FCoreRedirectObjectName::HasValidCharacters() const
+{
+	static FString InvalidRedirectCharacters = TEXT("\"' ,:|&!~\n\r\t@#(){}[]=;^%$`");
+
+	return ObjectName.IsValidXName(InvalidRedirectCharacters) && OuterName.IsValidXName(InvalidRedirectCharacters) && PackageName.IsValidXName(InvalidRedirectCharacters);
+}
+
+bool FCoreRedirectObjectName::ExpandNames(const FString& InString, FName& OutName, FName& OutOuter, FName &OutPackage)
+{
+	FString FullString = InString;
+	FullString.Trim();
+	FullString.TrimTrailing();
+
+	// Parse (/path.)?(outerchain.)?(name) where path and outerchain are optional
+	int32 SlashIndex = INDEX_NONE;
+	int32 FirstPeriodIndex = INDEX_NONE;
+	int32 LastPeriodIndex = INDEX_NONE;
+
+	FullString.FindChar('/', SlashIndex);
+
+	if (!FullString.FindChar('.', FirstPeriodIndex))
+	{
+		// If start with /, fill in package name, otherwise name
+		if (SlashIndex != INDEX_NONE)
+		{
+			OutPackage = FName(*FullString);
+		}
+		else
+		{
+			OutName = FName(*FullString);
+		}
+		return true;
+	}
+
+	FullString.FindLastChar('.', LastPeriodIndex);
+
+	if (SlashIndex == INDEX_NONE)
+	{
+		// No /, so start from beginning. There must be an outer if we got this far
+		OutOuter = FName(*FullString.Mid(0, LastPeriodIndex));
+	}
+	else
+	{
+		OutPackage = FName(*FullString.Left(FirstPeriodIndex));
+		if (FirstPeriodIndex != LastPeriodIndex)
+		{
+			// Extract Outer between periods
+			OutOuter = FName(*FullString.Mid(FirstPeriodIndex + 1, LastPeriodIndex - FirstPeriodIndex - 1));
+		}
+	}
+
+	OutName = FName(*FullString.Mid(LastPeriodIndex + 1));
+
+	return true;
+}
+
+FString FCoreRedirectObjectName::CombineNames(FName NewName, FName NewOuter, FName NewPackage)
+{
+	TArray<FString> CombineStrings;
+
+	if (NewOuter != NAME_None)
+	{
+		if (NewPackage != NAME_None)
+		{
+			return FString::Printf(TEXT("%s.%s.%s"), *NewPackage.ToString(), *NewOuter.ToString(), *NewName.ToString());
+		}
+		else
+		{
+			return FString::Printf(TEXT("%s.%s"), *NewOuter.ToString(), *NewName.ToString());
+		}
+	}
+	else if (NewPackage != NAME_None)
+	{
+		if (NewName != NAME_None)
+		{
+			return FString::Printf(TEXT("%s.%s"), *NewPackage.ToString(), *NewName.ToString());
+		}
+		return NewPackage.ToString();
+	}
+	return NewName.ToString();
+}
+
+void FCoreRedirect::NormalizeNewName()
+{
+	// Fill in NewName as needed
+	if (NewName.ObjectName == NAME_None)
+	{
+		NewName.ObjectName = OldName.ObjectName;
+	}
+	if (NewName.OuterName == NAME_None)
+	{
+		NewName.OuterName = OldName.OuterName;
+	}
+	if (NewName.PackageName == NAME_None)
+	{
+		NewName.PackageName = OldName.PackageName;
+	}
+}
+
+const TCHAR* FCoreRedirect::ParseValueChanges(const TCHAR* Buffer)
+{
+	// If we export an empty array we export an empty string, so ensure that if we're passed an empty string
+	// we interpret it as an empty array.
+	if (*Buffer++ != TCHAR('('))
+	{
+		return nullptr;
+	}
+
+	SkipWhitespace(Buffer);
+	if (*Buffer == TCHAR(')'))
+	{
+		return Buffer + 1;
+	}
+
+	for (;;)
+	{
+		SkipWhitespace(Buffer);
+		if (*Buffer++ != TCHAR('('))
+		{
+			return nullptr;
+		}
+
+		// Parse the key and value
+		FString KeyString, ValueString;
+		Buffer = UPropertyHelpers::ReadToken(Buffer, KeyString, true);
+		if (!Buffer)
+		{
+			return nullptr;
+		}
+
+		SkipWhitespace(Buffer);
+		if (*Buffer++ != TCHAR(','))
+		{
+			return nullptr;
+		}
+
+		// Parse the value
+		SkipWhitespace(Buffer);
+		Buffer = UPropertyHelpers::ReadToken(Buffer, ValueString, true);
+		if (!Buffer)
+		{
+			return nullptr;
+		}
+
+		SkipWhitespace(Buffer);
+		if (*Buffer++ != TCHAR(')'))
+		{
+			return nullptr;
+		}
+
+		ValueChanges.Add(KeyString, ValueString);
+
+		switch (*Buffer++)
+		{
+		case TCHAR(')'):
+			return Buffer;
+
+		case TCHAR(','):
+			break;
+
+		default:
+			return nullptr;
+		}
+	}
+}
+
+bool FCoreRedirect::Matches(ECoreRedirectFlags InFlags, const FCoreRedirectObjectName& InName) const
+{
+	// Check flags for exact match
+	if (InFlags != RedirectFlags)
+	{
+		return false;
+	}
+
+	return OldName.Matches(InName, IsSubstringMatch());
+}
+
+bool FCoreRedirect::HasValueChanges() const
+{
+	return ValueChanges.Num() > 0;
+}
+
+bool FCoreRedirect::IsSubstringMatch() const
+{
+	return !!(RedirectFlags & ECoreRedirectFlags::Option_MatchSubstring);
+}
+
+FCoreRedirectObjectName FCoreRedirect::RedirectName(const FCoreRedirectObjectName& OldObjectName) const
+{
+	FCoreRedirectObjectName ModifyName(OldObjectName);
+
+	// Convert names that are different and non empty
+	if (OldName.ObjectName != NewName.ObjectName)
+	{
+		if (IsSubstringMatch())
+		{
+			ModifyName.ObjectName = FName(*OldObjectName.ObjectName.ToString().Replace(*OldName.ObjectName.ToString(), *NewName.ObjectName.ToString()));
+		}
+		else
+		{
+			ModifyName.ObjectName = NewName.ObjectName;
+		}
+	}
+	if (OldName.OuterName != NewName.OuterName)
+	{
+		if (IsSubstringMatch())
+		{
+			ModifyName.OuterName = FName(*OldObjectName.OuterName.ToString().Replace(*OldName.OuterName.ToString(), *NewName.OuterName.ToString()));
+		}
+		else
+		{
+			ModifyName.OuterName = NewName.OuterName;
+		}
+	}
+	if (OldName.PackageName != NewName.PackageName)
+	{
+		if (IsSubstringMatch())
+		{
+			ModifyName.PackageName = FName(*OldObjectName.PackageName.ToString().Replace(*OldName.PackageName.ToString(), *NewName.PackageName.ToString()));
+		}
+		else
+		{
+			ModifyName.PackageName = NewName.PackageName;
+		}
+	}
+
+	return ModifyName;
+}
+
+bool FCoreRedirects::bInitialized;
+TMap<FName, ECoreRedirectFlags> FCoreRedirects::ConfigKeyMap;
+TMap<ECoreRedirectFlags, FCoreRedirects::FRedirectNameMap> FCoreRedirects::RedirectTypeMap;
+
+bool FCoreRedirects::RedirectNameAndValues(ECoreRedirectFlags Type, const FCoreRedirectObjectName& OldObjectName, FCoreRedirectObjectName& NewObjectName, const FCoreRedirect** FoundValueRedirect)
+{
+	NewObjectName = OldObjectName;
+	TArray<const FCoreRedirect*> FoundRedirects;
+
+	if (GetMatchingRedirects(Type, OldObjectName, FoundRedirects))
+	{
+		// Sort them based on match
+		FoundRedirects.Sort([&OldObjectName](const FCoreRedirect& A, const FCoreRedirect& B) { return A.OldName.MatchScore(OldObjectName) > B.OldName.MatchScore(OldObjectName); });
+
+		// Apply in order
+		for (int32 i = 0; i < FoundRedirects.Num(); i++)
+		{
+			const FCoreRedirect* Redirect = FoundRedirects[i];
+
+			if (!Redirect)
+			{
+				continue;
+			}
+
+			// Only apply if name match is still valid, if it already renamed part of it it may not apply any more. Don't want to check flags as those were checked in the gather step
+			if (Redirect->OldName.Matches(NewObjectName, Redirect->IsSubstringMatch()))
+			{
+				if (FoundValueRedirect && (Redirect->HasValueChanges() || Redirect->OverrideClassName.IsValid()))
+				{
+					if (*FoundValueRedirect)
+					{
+						UE_LOG(LogLinker, Error, TEXT("RedirectNameAndValues(%s) found multiple conflicting value redirects, %s and %s!"), *OldObjectName.ToString(), *(*FoundValueRedirect)->OldName.ToString(), *Redirect->OldName.ToString());
+					}
+					else
+					{
+						// Set value redirects for processing outside
+						*FoundValueRedirect = Redirect;
+					}
+				}
+
+				NewObjectName = Redirect->RedirectName(NewObjectName);
+			}
+		}
+	}
+
+	return NewObjectName != OldObjectName;
+}
+
+FCoreRedirectObjectName FCoreRedirects::GetRedirectedName(ECoreRedirectFlags Type, const FCoreRedirectObjectName& OldObjectName)
+{
+	FCoreRedirectObjectName NewObjectName;
+
+	RedirectNameAndValues(Type, OldObjectName, NewObjectName, nullptr);
+
+	return NewObjectName;
+}
+
+const TMap<FString, FString>* FCoreRedirects::GetValueRedirects(ECoreRedirectFlags Type, const FCoreRedirectObjectName& OldObjectName)
+{
+	FCoreRedirectObjectName NewObjectName;
+	const FCoreRedirect* FoundRedirect = nullptr;
+
+	RedirectNameAndValues(Type, OldObjectName, NewObjectName, &FoundRedirect);
+
+	if (FoundRedirect && FoundRedirect->ValueChanges.Num() > 0)
+	{
+		return &FoundRedirect->ValueChanges;
+	}
+
+	return nullptr;
+}
+
+bool FCoreRedirects::GetMatchingRedirects(ECoreRedirectFlags Type, const FCoreRedirectObjectName& OldObjectName, TArray<const FCoreRedirect*>& FoundRedirects)
+{
+	// Look for all redirects that match the given names and flags
+	bool bFound = false;
+	FRedirectNameMap* RedirectNameMap = RedirectTypeMap.Find(Type);
+	if (RedirectNameMap)
+	{
+		TArray<FCoreRedirect>* RedirectsForName = RedirectNameMap->RedirectMap.Find(OldObjectName.GetSearchKey(Type));
+
+		if (RedirectsForName)
+		{
+			for (FCoreRedirect& CheckRedirect : *RedirectsForName)
+			{
+				if (CheckRedirect.Matches(Type, OldObjectName))
+				{
+					bFound = true;
+					FoundRedirects.Add(&CheckRedirect);
+				}
+			}
+		}
+	}
+
+	// Add Package redirects now as well
+	if (!(Type & ECoreRedirectFlags::Type_Package))
+	{
+		bFound |= GetMatchingRedirects(ECoreRedirectFlags::Type_Package, OldObjectName, FoundRedirects);
+	}
+
+	// Add substring matches as well, these can be slow
+	if (!(Type & ECoreRedirectFlags::Option_MatchSubstring))
+	{
+		bFound |= GetMatchingRedirects(Type | ECoreRedirectFlags::Option_MatchSubstring, OldObjectName, FoundRedirects);
+	}
+
+	return bFound;
+}
+
+bool FCoreRedirects::FindPreviousNames(ECoreRedirectFlags Type, const FCoreRedirectObjectName& NewObjectName, TArray<FCoreRedirectObjectName>& PreviousNames)
+{
+	bool bFound = false;
+
+	FRedirectNameMap* RedirectNameMap = RedirectTypeMap.Find(Type);
+	if (RedirectNameMap)
+	{
+		for (TPair<FName, TArray<FCoreRedirect> >& Pair : RedirectNameMap->RedirectMap)
+		{
+			for (FCoreRedirect& Redirect : Pair.Value)
+			{
+				if (Redirect.NewName.Matches(NewObjectName, Redirect.IsSubstringMatch()))
+				{
+					// Construct a reverse redirect
+					FCoreRedirect ReverseRedirect = FCoreRedirect(Redirect);
+					ReverseRedirect.OldName = Redirect.NewName;
+					ReverseRedirect.NewName = Redirect.OldName;
+
+					FCoreRedirectObjectName OldName = ReverseRedirect.RedirectName(NewObjectName);
+
+					if (OldName != NewObjectName)
+					{
+						bFound = true;
+						PreviousNames.AddUnique(OldName);
+					}
+				}
+			}
+		}
+	}
+
+	// Add Package redirects now as well
+	if (!(Type & ECoreRedirectFlags::Type_Package))
+	{
+		bFound |= FindPreviousNames(ECoreRedirectFlags::Type_Package, NewObjectName, PreviousNames);
+	}
+
+	// Add substring matches as well, these can be slow
+	if (!(Type & ECoreRedirectFlags::Option_MatchSubstring))
+	{
+		bFound |= FindPreviousNames(Type | ECoreRedirectFlags::Option_MatchSubstring, NewObjectName, PreviousNames);
+	}
+
+	return bFound;
+}
+
+bool FCoreRedirects::IsKnownMissing(ECoreRedirectFlags Type, const FCoreRedirectObjectName& ObjectName)
+{
+	TArray<const FCoreRedirect*> FoundRedirects;
+	return FCoreRedirects::GetMatchingRedirects(Type | ECoreRedirectFlags::Option_Removed, ObjectName, FoundRedirects);
+}
+
+bool FCoreRedirects::AddKnownMissing(ECoreRedirectFlags Type, const FCoreRedirectObjectName& ObjectName)
+{
+	TArray<FCoreRedirect> NewRedirects;
+	NewRedirects.Emplace(Type | ECoreRedirectFlags::Option_Removed, ObjectName, FCoreRedirectObjectName());
+	return AddRedirectList(NewRedirects, TEXT("AddKnownMissing"));
+}
+
+void FCoreRedirects::RunTests()
+{
+	TMap<ECoreRedirectFlags, FRedirectNameMap > BackupMap = RedirectTypeMap;
+	RedirectTypeMap.Empty();
+
+	TArray<FCoreRedirect> NewRedirects;
+
+	new (NewRedirects) FCoreRedirect(ECoreRedirectFlags::Type_Property, TEXT("Property"), TEXT("Property2"));
+	new (NewRedirects) FCoreRedirect(ECoreRedirectFlags::Type_Property, TEXT("Class.Property"), TEXT("Property3"));
+	new (NewRedirects) FCoreRedirect(ECoreRedirectFlags::Type_Property, TEXT("/game/PackageSpecific.Class.Property"), TEXT("Property4"));
+	new (NewRedirects) FCoreRedirect(ECoreRedirectFlags::Type_Property, TEXT("/game/Package.Class.OtherProperty"), TEXT("/game/Package.Class.OtherProperty2"));
+	new (NewRedirects) FCoreRedirect(ECoreRedirectFlags::Type_Class, TEXT("Class"), TEXT("Class2"));
+	new (NewRedirects) FCoreRedirect(ECoreRedirectFlags::Type_Class, TEXT("/game/Package.Class"), TEXT("Class3"));
+	new (NewRedirects) FCoreRedirect(ECoreRedirectFlags::Type_Class | ECoreRedirectFlags::Option_InstanceOnly, TEXT("/game/Package.Class"), TEXT("ClassInstance"));
+	new (NewRedirects) FCoreRedirect(ECoreRedirectFlags::Type_Package, TEXT("/game/Package"), TEXT("/game/Package2"));
+	new (NewRedirects) FCoreRedirect(ECoreRedirectFlags::Type_Package | ECoreRedirectFlags::Option_MatchSubstring, TEXT("/oldgame"), TEXT("/newgame"));
+	new (NewRedirects) FCoreRedirect(ECoreRedirectFlags::Type_Package | ECoreRedirectFlags::Option_Removed, TEXT("/game/RemovedPackage"), TEXT("/game/RemovedPackage"));
+
+	AddRedirectList(NewRedirects, TEXT("RunTests"));
+
+	struct FRedirectTest
+	{
+		FString OldName;
+		FString NewName;
+		ECoreRedirectFlags Type;
+
+		FRedirectTest(const FString& InOldName, const FString& InNewName, ECoreRedirectFlags InType)
+			: OldName(InOldName), NewName(InNewName), Type(InType)
+		{}
+	};
+
+	TArray<FRedirectTest> Tests;
+
+	UE_LOG(LogLinker, Log, TEXT("Running FCoreRedirect Tests"));
+
+	// Package-specific property rename and package rename apply
+	Tests.Emplace(TEXT("/game/PackageSpecific.Class.Property"), TEXT("/game/PackageSpecific.Class.Property4"), ECoreRedirectFlags::Type_Property);
+	// Wrong type, no replacement
+	Tests.Emplace(TEXT("/game/PackageSpecific.Class.Property"), TEXT("/game/PackageSpecific.Class.Property"), ECoreRedirectFlags::Type_Function);
+	// Class-specific property rename and package rename apply
+	Tests.Emplace(TEXT("/game/Package.Class.Property"), TEXT("/game/Package2.Class.Property3"), ECoreRedirectFlags::Type_Property);
+	// Package-Specific class rename applies
+	Tests.Emplace(TEXT("/game/Package.Class"), TEXT("/game/Package2.Class3"), ECoreRedirectFlags::Type_Class);
+	// Generic class rename applies
+	Tests.Emplace(TEXT("/game/PackageOther.Class"), TEXT("/game/PackageOther.Class2"), ECoreRedirectFlags::Type_Class);
+	// Check instance option
+	Tests.Emplace(TEXT("/game/Package.Class"), TEXT("/game/Package2.ClassInstance"), ECoreRedirectFlags::Type_Class | ECoreRedirectFlags::Option_InstanceOnly);
+	// Substring test
+	Tests.Emplace(TEXT("/oldgame/Package.DefaultClass"), TEXT("/newgame/Package.DefaultClass"), ECoreRedirectFlags::Type_Class);
+
+	for (FRedirectTest& Test : Tests)
+	{
+		FCoreRedirectObjectName OldName = FCoreRedirectObjectName(Test.OldName);
+		FCoreRedirectObjectName NewName = GetRedirectedName(Test.Type, OldName);
+
+		if (NewName.ToString() != Test.NewName)
+		{
+			UE_LOG(LogLinker, Error, TEXT("FCoreRedirect Test Failed: %s to %s, should be %s!"), *OldName.ToString(), *NewName.ToString(), *Test.NewName);
+		}
+	}
+
+	// Check reverse lookup
+	TArray<FCoreRedirectObjectName> OldNames;
+
+	FindPreviousNames(ECoreRedirectFlags::Type_Class, FCoreRedirectObjectName(TEXT("/game/PackageOther.Class2")), OldNames);
+
+	if (OldNames.Num() != 1 || OldNames[0].ToString() != TEXT("/game/PackageOther.Class"))
+	{
+		UE_LOG(LogLinker, Error, TEXT("FCoreRedirect Test Failed: ReverseLookup!"));
+	}
+
+	// Check removed
+	if (!IsKnownMissing(ECoreRedirectFlags::Type_Package, FCoreRedirectObjectName(TEXT("/game/RemovedPackage"))))
+	{
+		UE_LOG(LogLinker, Error, TEXT("FCoreRedirect Test Failed: /game/RemovedPackage should be removed!"));
+	}
+
+	if (IsKnownMissing(ECoreRedirectFlags::Type_Package, FCoreRedirectObjectName(TEXT("/game/NotRemovedPackage"))))
+	{
+		UE_LOG(LogLinker, Error, TEXT("FCoreRedirect Test Failed: /game/NotRemovedPackage should be removed!"));
+	}
+
+	AddKnownMissing(ECoreRedirectFlags::Type_Package, FCoreRedirectObjectName(TEXT("/game/NotRemovedPackage")));
+
+	if (!IsKnownMissing(ECoreRedirectFlags::Type_Package, FCoreRedirectObjectName(TEXT("/game/NotRemovedPackage"))))
+	{
+		UE_LOG(LogLinker, Error, TEXT("FCoreRedirect Test Failed: /game/NotRemovedPackage should be removed now!"));
+	}
+
+	// Restore old state
+	RedirectTypeMap = BackupMap;
+}
+
+bool FCoreRedirects::ReadRedirectsFromIni(const FString& IniName)
+{
+	// Mark that this has been done at least once
+	if (!bInitialized)
+	{
+		// Setup map
+		ConfigKeyMap.Add(TEXT("ObjectRedirects"), ECoreRedirectFlags::Type_Object);
+		ConfigKeyMap.Add(TEXT("ClassRedirects"), ECoreRedirectFlags::Type_Class);
+		ConfigKeyMap.Add(TEXT("StructRedirects"), ECoreRedirectFlags::Type_Struct);
+		ConfigKeyMap.Add(TEXT("EnumRedirects"), ECoreRedirectFlags::Type_Enum);
+		ConfigKeyMap.Add(TEXT("FunctionRedirects"), ECoreRedirectFlags::Type_Function);
+		ConfigKeyMap.Add(TEXT("PropertyRedirects"), ECoreRedirectFlags::Type_Property);
+		ConfigKeyMap.Add(TEXT("PackageRedirects"), ECoreRedirectFlags::Type_Package);
+		ConfigKeyMap.Add(TEXT("DelegateRedirects"), ECoreRedirectFlags::Type_Delegate);
+
+		RegisterNativeRedirects();
+
+		// Enable to run startup tests
+		//RunTests();
+
+		bInitialized = true;
+	}
+
+	if (GConfig)
+	{
+		FConfigSection* RedirectSection = GConfig->GetSectionPrivate(TEXT("CoreRedirects"), false, true, IniName);
+		if (RedirectSection)
+		{
+			TArray<FCoreRedirect> NewRedirects;
+
+			for (FConfigSection::TIterator It(*RedirectSection); It; ++It)
+			{
+				FString OldName, NewName, OverrideClassName;
+
+				bool bInstanceOnly = false;
+				bool bRemoved = false;
+				bool bMatchSubstring = false;
+
+				const FString& ValueString = It.Value().GetValue();
+
+				FParse::Bool(*ValueString, TEXT("InstanceOnly="), bInstanceOnly);
+				FParse::Bool(*ValueString, TEXT("Removed="), bRemoved);
+				FParse::Bool(*ValueString, TEXT("MatchSubstring="), bMatchSubstring);
+
+				FParse::Value(*ValueString, TEXT("OldName="), OldName);
+				FParse::Value(*ValueString, TEXT("NewName="), NewName);
+
+				FParse::Value(*ValueString, TEXT("OverrideClassName="), OverrideClassName);
+
+				ECoreRedirectFlags* FlagPtr = ConfigKeyMap.Find(It.Key());
+
+				// If valid type
+				if (FlagPtr)
+				{
+					ECoreRedirectFlags NewFlags = *FlagPtr;
+
+					if (bInstanceOnly)
+					{
+						NewFlags |= ECoreRedirectFlags::Option_InstanceOnly;
+					}
+
+					if (bRemoved)
+					{
+						NewFlags |= ECoreRedirectFlags::Option_Removed;
+					}
+
+					if (bMatchSubstring)
+					{
+						NewFlags |= ECoreRedirectFlags::Option_MatchSubstring;
+					}
+
+					FCoreRedirect* Redirect = new (NewRedirects) FCoreRedirect(NewFlags, FCoreRedirectObjectName(OldName), FCoreRedirectObjectName(NewName));
+
+					if (!OverrideClassName.IsEmpty())
+					{
+						Redirect->OverrideClassName = FCoreRedirectObjectName(OverrideClassName);
+					}
+
+					int32 ValueChangesIndex = ValueString.Find(TEXT("ValueChanges="));
+				
+					if (ValueChangesIndex != INDEX_NONE)
+					{
+						// Look for first (
+						ValueChangesIndex = ValueString.Find(TEXT("("), ESearchCase::CaseSensitive, ESearchDir::FromStart, ValueChangesIndex);
+
+						FString ValueChangesString = ValueString.Mid(ValueChangesIndex);
+						const TCHAR* Buffer = Redirect->ParseValueChanges(GetData(ValueChangesString));
+
+						if (!Buffer)
+						{
+							UE_LOG(LogLinker, Error, TEXT("ReadRedirectsFromIni failed to parse ValueChanges for Redirect %s!"), *ValueString);
+
+							// Remove added redirect
+							Redirect = nullptr;
+							NewRedirects.RemoveAt(NewRedirects.Num() - 1);
+						}
+					}
+				}
+				else
+				{
+					UE_LOG(LogLinker, Error, TEXT("ReadRedirectsFromIni failed to parse type for Redirect %s!"), *ValueString);
+				}
+			}
+
+			return AddRedirectList(NewRedirects, IniName);
+		}
+	}
+	else
+	{
+		UE_LOG(LogLinker, Warning, TEXT(" **** CORE REDIRECTS UNABLE TO INITIALIZE! **** "));
+	}
+	return false;
+}
+
+bool FCoreRedirects::AddRedirectList(const TArray<FCoreRedirect>& Redirects, const FString& SourceString)
+{
+	bool bAddedAny = false;
+	for (const FCoreRedirect& NewRedirect : Redirects)
+	{
+		if (!NewRedirect.OldName.IsValid() || !NewRedirect.NewName.IsValid())
+		{
+			UE_LOG(LogLinker, Error, TEXT("AddRedirectList(%s) failed to add redirector from %s to %s with empty name!"), *SourceString, *NewRedirect.OldName.ToString(), *NewRedirect.NewName.ToString());
+			continue;
+		}
+
+		if (!NewRedirect.OldName.HasValidCharacters() || !NewRedirect.NewName.HasValidCharacters())
+		{
+			UE_LOG(LogLinker, Error, TEXT("AddRedirectList(%s) failed to add redirector from %s to %s with invalid characters!"), *SourceString, *NewRedirect.OldName.ToString(), *NewRedirect.NewName.ToString());
+			continue;
+		}
+
+
+		if (NewRedirect.NewName.PackageName != NewRedirect.OldName.PackageName && NewRedirect.OldName.OuterName != NAME_None)
+		{
+			UE_LOG(LogLinker, Error, TEXT("AddRedirectList(%s) failed to add redirector, it's not valid to modify package from %s to %s while specifying outer!"), *SourceString, *NewRedirect.OldName.ToString(), *NewRedirect.NewName.ToString());
+			continue;
+		}
+
+		if (NewRedirect.IsSubstringMatch())
+		{
+			UE_LOG(LogLinker, Log, TEXT("AddRedirectList(%s) has substring redirect %s, these are very slow and should be resolved as soon as possible!"), *SourceString, *NewRedirect.OldName.ToString());
+		}
+		
+		if (AddSingleRedirect(NewRedirect, SourceString))
+		{
+			bAddedAny = true;
+
+			// If value redirect, add a value redirect from NewName->NewName as well, this will merge with existing ones as needed
+			if (NewRedirect.OldName != NewRedirect.NewName && NewRedirect.HasValueChanges())
+			{
+				FCoreRedirect ValueRedirect = NewRedirect;
+				ValueRedirect.OldName = ValueRedirect.NewName;
+
+				AddSingleRedirect(ValueRedirect, SourceString);
+			}
+		}
+	}
+
+	return bAddedAny;
+}
+
+bool FCoreRedirects::AddSingleRedirect(const FCoreRedirect& NewRedirect, const FString& SourceString)
+{
+	FRedirectNameMap& ExistingNameMap = RedirectTypeMap.FindOrAdd(NewRedirect.RedirectFlags);
+	TArray<FCoreRedirect>& ExistingRedirects = ExistingNameMap.RedirectMap.FindOrAdd(NewRedirect.GetSearchKey());
+
+	bool bFoundDuplicate = false;
+
+	// Check for duplicate
+	for (FCoreRedirect& ExistingRedirect : ExistingRedirects)
+	{
+		if (ExistingRedirect.IdenticalMatchRules(NewRedirect))
+		{
+			bFoundDuplicate = true;
+			if (ExistingRedirect.NewName == NewRedirect.NewName && NewRedirect.HasValueChanges())
+			{
+				// Merge value redirects as names are the same
+				ExistingRedirect.ValueChanges.Append(NewRedirect.ValueChanges);
+			}
+			else if (ExistingRedirect.NewName != NewRedirect.NewName)
+			{
+				UE_LOG(LogLinker, Error, TEXT("AddRedirectList(%s) found conflicting redirectors for %s! Old: %s, New: %s"), *SourceString, *ExistingRedirect.OldName.ToString(), *ExistingRedirect.NewName.ToString(), *NewRedirect.NewName.ToString());
+			}
+			else
+			{
+				// ENABLE THIS ONCE INIS ARE CONVERTED
+				//UE_LOG(LogLinker, Error, TEXT("AddRedirectList(%s) found duplicate redirectors for %s to %s!), *SourceString, *ExistingRedirect.OldName.ToString(), *ExistingRedirect.NewName.ToString());
+			}
+			break;
+		}
+	}
+
+	if (bFoundDuplicate)
+	{
+		return false;
+	}
+
+	ExistingRedirects.Add(NewRedirect);
+	return true;
+}
+
+ECoreRedirectFlags FCoreRedirects::GetFlagsForTypeName(FName PackageName, FName TypeName)
+{
+	if (PackageName == GLongCoreUObjectPackageName)
+	{
+		if (TypeName == NAME_Class)
+		{
+			return ECoreRedirectFlags::Type_Class;
+		}
+		else if (TypeName == NAME_ScriptStruct)
+		{
+			return ECoreRedirectFlags::Type_Struct;
+		}
+		else if (TypeName == NAME_Enum)
+		{
+			return ECoreRedirectFlags::Type_Enum;
+		}
+		else if (TypeName == NAME_Package)
+		{
+			return ECoreRedirectFlags::Type_Package;
+		}
+		else if (TypeName == NAME_Function)
+		{
+			return ECoreRedirectFlags::Type_Function;
+		}
+
+		// If ending with property, it's a property
+		if (TypeName.ToString().EndsWith(TEXT("Property")))
+		{
+			return ECoreRedirectFlags::Type_Property;
+		}
+	}
+
+	// If ending with GeneratedClass this has to be a class subclass, some of these are in engine or plugins
+	if (TypeName.ToString().EndsWith(TEXT("GeneratedClass")))
+	{
+		return ECoreRedirectFlags::Type_Class;
+	}
+
+	return ECoreRedirectFlags::Type_Object;
+}
+
+ECoreRedirectFlags FCoreRedirects::GetFlagsForTypeClass(UClass *TypeClass)
+{
+	// Use Name version for consistency, if we can't figure it out from just the name it isn't safe
+	return GetFlagsForTypeName(TypeClass->GetOutermost()->GetFName(), TypeClass->GetFName());
+}
+
+#if WITH_EDITORONLY_DATA
+// These redirects are very old, only load them in editor builds
+
+PRAGMA_DISABLE_OPTIMIZATION
+
+// The compiler doesn't like having a massive string table in a single function so split it up
+#define CLASS_REDIRECT(OldName, NewName) new (Redirects) FCoreRedirect(ECoreRedirectFlags::Type_Class, TEXT(OldName), TEXT(NewName))
+#define CLASS_REDIRECT_INSTANCES(OldName, NewName) new (Redirects) FCoreRedirect(ECoreRedirectFlags::Type_Class | ECoreRedirectFlags::Option_InstanceOnly, TEXT(OldName), TEXT(NewName))
+#define STRUCT_REDIRECT(OldName, NewName) new (Redirects) FCoreRedirect(ECoreRedirectFlags::Type_Struct, TEXT(OldName), TEXT(NewName))
+#define ENUM_REDIRECT(OldName, NewName) new (Redirects) FCoreRedirect(ECoreRedirectFlags::Type_Enum, TEXT(OldName), TEXT(NewName))
+#define PROPERTY_REDIRECT(OldName, NewName) new (Redirects) FCoreRedirect(ECoreRedirectFlags::Type_Property, TEXT(OldName), TEXT(NewName))
+#define FUNCTION_REDIRECT(OldName, NewName) new (Redirects) FCoreRedirect(ECoreRedirectFlags::Type_Function, TEXT(OldName), TEXT(NewName))
+#define PACKAGE_REDIRECT(OldName, NewName) new (Redirects) FCoreRedirect(ECoreRedirectFlags::Type_Package, TEXT(OldName), TEXT(NewName))
+
+static void RegisterNativeRedirects40(TArray<FCoreRedirect>& Redirects)
+{
+	CLASS_REDIRECT("AIDebugComponent", "GameplayDebuggingComponent");
+	CLASS_REDIRECT("AnimTreeInstance", "AnimInstance");
+	CLASS_REDIRECT("AnimationCompressionAlgorithm", "AnimCompress");
+	CLASS_REDIRECT("AnimationCompressionAlgorithm_Automatic", "AnimCompress_Automatic");
+	CLASS_REDIRECT("AnimationCompressionAlgorithm_BitwiseCompressOnly", "AnimCompress_BitwiseCompressOnly");
+	CLASS_REDIRECT("AnimationCompressionAlgorithm_LeastDestructive", "AnimCompress_LeastDestructive");
+	CLASS_REDIRECT("AnimationCompressionAlgorithm_PerTrackCompression", "AnimCompress_PerTrackCompression");
+	CLASS_REDIRECT("AnimationCompressionAlgorithm_RemoveEverySecondKey", "AnimCompress_RemoveEverySecondKey");
+	CLASS_REDIRECT("AnimationCompressionAlgorithm_RemoveLinearKeys", "AnimCompress_RemoveLinearKeys");
+	CLASS_REDIRECT("AnimationCompressionAlgorithm_RemoveTrivialKeys", "AnimCompress_RemoveTrivialKeys");
+	CLASS_REDIRECT("BlueprintActorBase", "Actor");
+	CLASS_REDIRECT("DefaultPawnMovement", "FloatingPawnMovement");
+	CLASS_REDIRECT("DirectionalLightMovable", "DirectionalLight");
+	CLASS_REDIRECT("DirectionalLightStatic", "DirectionalLight");
+	CLASS_REDIRECT("DirectionalLightStationary", "DirectionalLight");
+	CLASS_REDIRECT("DynamicBlockingVolume", "BlockingVolume");
+	CLASS_REDIRECT("DynamicPhysicsVolume", "PhysicsVolume");
+	CLASS_REDIRECT("DynamicTriggerVolume", "TriggerVolume");
+	CLASS_REDIRECT("GameInfo", "/Script/Engine.GameMode");
+	CLASS_REDIRECT("GameReplicationInfo", "/Script/Engine.GameState");
+	CLASS_REDIRECT("InterpActor", "StaticMeshActor");
+	CLASS_REDIRECT("K2Node_CallSuperFunction", "/Script/BlueprintGraph.K2Node_CallParentFunction");
+	CLASS_REDIRECT("MaterialSpriteComponent", "MaterialBillboardComponent");
+	CLASS_REDIRECT("MovementComp_Character", "CharacterMovementComponent");
+	CLASS_REDIRECT("MovementComp_Projectile", "ProjectileMovementComponent");
+	CLASS_REDIRECT("MovementComp_Rotating", "RotatingMovementComponent");
+	CLASS_REDIRECT("NavAreaDefault", "/Script/Engine.NavArea_Default");
+	CLASS_REDIRECT("NavAreaDefinition", "/Script/Engine.NavArea");
+	CLASS_REDIRECT("NavAreaNull", "/Script/Engine.NavArea_Null");
+	CLASS_REDIRECT("PhysicsActor", "StaticMeshActor");
+	CLASS_REDIRECT("PhysicsBSJointActor", "PhysicsConstraintActor");
+	CLASS_REDIRECT("PhysicsHingeActor", "PhysicsConstraintActor");
+	CLASS_REDIRECT("PhysicsPrismaticActor", "PhysicsConstraintActor");
+	CLASS_REDIRECT("PlayerCamera", "PlayerCameraManager");
+	CLASS_REDIRECT("PlayerReplicationInfo", "/Script/Engine.PlayerState");
+	CLASS_REDIRECT("PointLightMovable", "PointLight");
+	CLASS_REDIRECT("PointLightStatic", "PointLight");
+	CLASS_REDIRECT("PointLightStationary", "PointLight");
+	CLASS_REDIRECT("RB_BSJointSetup", "PhysicsConstraintTemplate");
+	CLASS_REDIRECT("RB_BodySetup", "BodySetup");
+	CLASS_REDIRECT("RB_ConstraintActor", "PhysicsConstraintActor");
+	CLASS_REDIRECT("RB_ConstraintComponent", "PhysicsConstraintComponent");
+	CLASS_REDIRECT("RB_ConstraintSetup", "PhysicsConstraintTemplate");
+	CLASS_REDIRECT("RB_Handle", "PhysicsHandleComponent");
+	CLASS_REDIRECT("RB_HingeSetup", "PhysicsConstraintTemplate");
+	CLASS_REDIRECT("RB_PrismaticSetup", "PhysicsConstraintTemplate");
+	CLASS_REDIRECT("RB_RadialForceComponent", "RadialForceComponent");
+	CLASS_REDIRECT("RB_SkelJointSetup", "PhysicsConstraintTemplate");
+	CLASS_REDIRECT("RB_Thruster", "PhysicsThruster");
+	CLASS_REDIRECT("RB_ThrusterComponent", "PhysicsThrusterComponent");
+	CLASS_REDIRECT("SensingComponent", "PawnSensingComponent");
+	CLASS_REDIRECT("SingleAnimSkeletalActor", "SkeletalMeshActor");
+	CLASS_REDIRECT("SingleAnimSkeletalComponent", "SkeletalMeshComponent");
+	CLASS_REDIRECT("SkeletalMeshReplicatedComponent", "SkeletalMeshComponent");
+	CLASS_REDIRECT("SkeletalPhysicsActor", "SkeletalMeshActor");
+	CLASS_REDIRECT("SoundMode", "SoundMix");
+	CLASS_REDIRECT("SpotLightMovable", "SpotLight");
+	CLASS_REDIRECT("SpotLightStatic", "SpotLight");
+	CLASS_REDIRECT("SpotLightStationary", "SpotLight");
+	CLASS_REDIRECT("SpriteComponent", "BillboardComponent");
+	CLASS_REDIRECT("StaticMeshReplicatedComponent", "StaticMeshComponent");
+	CLASS_REDIRECT("VimBlueprint", "AnimBlueprint");
+	CLASS_REDIRECT("VimGeneratedClass", "AnimBlueprintGeneratedClass");
+	CLASS_REDIRECT("VimInstance", "AnimInstance");
+	CLASS_REDIRECT("WorldInfo", "WorldSettings");
+	CLASS_REDIRECT_INSTANCES("NavAreaMeta", "/Script/Engine.NavArea_Default");
+
+	STRUCT_REDIRECT("VimDebugData", "AnimBlueprintDebugData");
+
+	FUNCTION_REDIRECT("Actor.GetController", "Pawn.GetController");
+	FUNCTION_REDIRECT("Actor.GetTouchingActors", "Actor.GetOverlappingActors");
+	PROPERTY_REDIRECT("Actor.GetOverlappingActors.OutTouchingActors", "OverlappingActors");
+	FUNCTION_REDIRECT("Actor.GetTouchingComponents", "Actor.GetOverlappingComponents");
+	PROPERTY_REDIRECT("Actor.GetOverlappingComponents.TouchingComponents", "OverlappingComponents");
+	FUNCTION_REDIRECT("Actor.HasTag", "Actor.ActorHasTag");
+	FUNCTION_REDIRECT("Actor.ReceiveActorTouch", "Actor.ReceiveActorBeginOverlap");
+	PROPERTY_REDIRECT("Actor.ReceiveActorBeginOverlap.Other", "OtherActor");
+	FUNCTION_REDIRECT("Actor.ReceiveActorUntouch", "Actor.ReceiveActorEndOverlap");
+	PROPERTY_REDIRECT("Actor.ReceiveActorEndOverlap.Other", "OtherActor");
+	PROPERTY_REDIRECT("Actor.ReceiveHit.NormalForce", "NormalImpulse");
+	FUNCTION_REDIRECT("Actor.SetActorHidden", "Actor.SetActorHiddenInGame");
+	PROPERTY_REDIRECT("Actor.LifeSpan", "Actor.InitialLifeSpan");
+	PROPERTY_REDIRECT("Actor.OnActorTouch", "OnActorBeginOverlap");
+	PROPERTY_REDIRECT("Actor.OnActorUnTouch", "OnActorEndOverlap");
+
+	FUNCTION_REDIRECT("AnimInstance.GetSequencePlayerLength", "GetAnimAssetPlayerLength");
+	FUNCTION_REDIRECT("AnimInstance.GetSequencePlayerTimeFraction", "GetAnimAssetPlayerTimeFraction");
+	FUNCTION_REDIRECT("AnimInstance.GetSequencePlayerTimeFromEnd", "GetAnimAssetPlayerTimeFromEnd");
+	FUNCTION_REDIRECT("AnimInstance.GetSequencePlayerTimeFromEndFraction", "GetAnimAssetPlayerTimeFromEndFraction");
+	FUNCTION_REDIRECT("AnimInstance.KismetInitializeAnimation", "AnimInstance.BlueprintInitializeAnimation");
+	FUNCTION_REDIRECT("AnimInstance.KismetUpdateAnimation", "AnimInstance.BlueprintUpdateAnimation");
+	PROPERTY_REDIRECT("AnimInstance.GetAnimAssetPlayerLength.Sequence", "AnimAsset");
+	PROPERTY_REDIRECT("AnimInstance.GetAnimAssetPlayerTimeFraction.Sequence", "AnimAsset");
+	PROPERTY_REDIRECT("AnimInstance.GetAnimAssetPlayerTimeFromEnd.Sequence", "AnimAsset");
+	PROPERTY_REDIRECT("AnimInstance.GetAnimAssetPlayerTimeFromEndFraction.Sequence", "AnimAsset");
+	PROPERTY_REDIRECT("AnimInstance.VimVertexAnims", "AnimInstance.VertexAnims");
+
+	FUNCTION_REDIRECT("GameplayStatics.ClearSoundMode", "GameplayStatics.ClearSoundMixModifiers");
+	FUNCTION_REDIRECT("GameplayStatics.GetGameInfo", "GetGameMode");
+	FUNCTION_REDIRECT("GameplayStatics.GetGameReplicationInfo", "GetGameState");
+	FUNCTION_REDIRECT("GameplayStatics.GetPlayerCamera", "GameplayStatics.GetPlayerCameraManager");
+	FUNCTION_REDIRECT("GameplayStatics.K2_SetSoundMode", "GameplayStatics.SetBaseSoundMix");
+	FUNCTION_REDIRECT("GameplayStatics.PopSoundMixModifier.InSoundMode", "InSoundMixModifier");
+	FUNCTION_REDIRECT("GameplayStatics.PopSoundMode", "GameplayStatics.PopSoundMixModifier");
+	FUNCTION_REDIRECT("GameplayStatics.PushSoundMixModifier.InSoundMode", "InSoundMixModifier");
+	FUNCTION_REDIRECT("GameplayStatics.PushSoundMode", "GameplayStatics.PushSoundMixModifier");
+	FUNCTION_REDIRECT("GameplayStatics.SetBaseSoundMix.InSoundMode", "InSoundMix");
+	FUNCTION_REDIRECT("GameplayStatics.SetTimeDilation", "GameplayStatics.SetGlobalTimeDilation");
+
+	FUNCTION_REDIRECT("KismetMaterialLibrary.CreateMaterialInstanceDynamic", "KismetMaterialLibrary.CreateDynamicMaterialInstance");
+	FUNCTION_REDIRECT("KismetMaterialParameterCollectionLibrary.GetScalarParameterValue", "KismetMaterialLibrary.GetScalarParameterValue");
+	FUNCTION_REDIRECT("KismetMaterialParameterCollectionLibrary.GetVectorParameterValue", "KismetMaterialLibrary.GetVectorParameterValue");
+	FUNCTION_REDIRECT("KismetMaterialParameterCollectionLibrary.SetScalarParameterValue", "KismetMaterialLibrary.SetScalarParameterValue");
+	FUNCTION_REDIRECT("KismetMaterialParameterCollectionLibrary.SetVectorParameterValue", "KismetMaterialLibrary.SetVectorParameterValue");
+
+	FUNCTION_REDIRECT("KismetMathLibrary.BreakTransform.Translation", "Location");
+	FUNCTION_REDIRECT("KismetMathLibrary.Conv_VectorToTransform.InTranslation", "InLocation");
+	FUNCTION_REDIRECT("KismetMathLibrary.FRand", "RandomFloat");
+	FUNCTION_REDIRECT("KismetMathLibrary.FRandFromStream", "RandomFloatFromStream");
+	FUNCTION_REDIRECT("KismetMathLibrary.FRandRange", "RandomFloatInRange");
+	FUNCTION_REDIRECT("KismetMathLibrary.FRandRangeFromStream", "RandomFloatInRangeFromStream");
+	FUNCTION_REDIRECT("KismetMathLibrary.InverseTransformPosition", "KismetMathLibrary.InverseTransformLocation");
+	PROPERTY_REDIRECT("KismetMathLibrary.InverseTransformLocation.Position", "Location");
+	PROPERTY_REDIRECT("KismetMathLibrary.MakeTransform.Translation", "Location");
+	FUNCTION_REDIRECT("KismetMathLibrary.Rand", "RandomInteger");
+	FUNCTION_REDIRECT("KismetMathLibrary.RandBool", "RandomBool");
+	FUNCTION_REDIRECT("KismetMathLibrary.RandBoolFromStream", "RandomBoolFromStream");
+	FUNCTION_REDIRECT("KismetMathLibrary.RandFromStream", "RandomIntegerFromStream");
+	FUNCTION_REDIRECT("KismetMathLibrary.RandRange", "RandomIntegerInRange");
+	FUNCTION_REDIRECT("KismetMathLibrary.RandRangeFromStream", "RandomIntegerInRangeFromStream");
+	FUNCTION_REDIRECT("KismetMathLibrary.RotRand", "RandomRotator");
+	FUNCTION_REDIRECT("KismetMathLibrary.RotRandFromStream", "RandomRotatorFromStream");
+	FUNCTION_REDIRECT("KismetMathLibrary.TransformPosition", "KismetMathLibrary.TransformLocation");
+	PROPERTY_REDIRECT("KismetMathLibrary.TransformLocation.Position", "Location");
+	FUNCTION_REDIRECT("KismetMathLibrary.VRand", "RandomUnitVector");
+	FUNCTION_REDIRECT("KismetMathLibrary.VRandFromStream", "RandomUnitVectorFromStream");
+
+	PROPERTY_REDIRECT("KismetSystemLibrary.CapsuleTraceMultiForObjects.ObjectsToTrace", "ObjectTypes");
+	PROPERTY_REDIRECT("KismetSystemLibrary.CapsuleTraceSingleForObjects.ObjectsToTrace", "ObjectTypes");
+	PROPERTY_REDIRECT("KismetSystemLibrary.LineTraceMultiForObjects.ObjectsToTrace", "ObjectTypes");
+	PROPERTY_REDIRECT("KismetSystemLibrary.LineTraceSingleForObjects.ObjectsToTrace", "ObjectTypes");
+	PROPERTY_REDIRECT("KismetSystemLibrary.PrintKismetWarning", "PrintWarning");
+	PROPERTY_REDIRECT("KismetSystemLibrary.SphereTraceMultiForObjects.ObjectsToTrace", "ObjectTypes");
+	PROPERTY_REDIRECT("KismetSystemLibrary.SphereTraceSingleForObjects.ObjectsToTrace", "ObjectTypes");
+
+	FUNCTION_REDIRECT("AIController.ClearFocus", "AIController.K2_ClearFocus");
+	FUNCTION_REDIRECT("AIController.SetFocalPoint", "AIController.K2_SetFocalPoint");
+	FUNCTION_REDIRECT("AIController.SetFocus", "AIController.K2_SetFocus");
+	FUNCTION_REDIRECT("ArrowComponent.SetArrowColor_New", "ArrowComponent.SetArrowColor");
+	FUNCTION_REDIRECT("Character.Launch", "Character.LaunchCharacter");
+	FUNCTION_REDIRECT("Controller.K2_GetActorRotation", "Controller.GetControlRotation");
+	FUNCTION_REDIRECT("DecalActor.CreateMIDForDecal", "DecalActor.CreateDynamicMaterialInstance");
+	FUNCTION_REDIRECT("DecalComponent.CreateMIDForDecal", "DecalComponent.CreateDynamicMaterialInstance");
+	PROPERTY_REDIRECT("HUD.AddHitBox.InPos", "Position");
+	PROPERTY_REDIRECT("HUD.AddHitBox.InPriority", "Priority");
+	PROPERTY_REDIRECT("HUD.AddHitBox.InSize", "Size");
+	PROPERTY_REDIRECT("HUD.AddHitBox.bInConsumesInput", "bConsumesInput");
+	FUNCTION_REDIRECT("LevelScriptActor.BeginGame", "Actor.ReceiveBeginPlay");
+	FUNCTION_REDIRECT("LevelScriptActor.LoadStreamLevel", "GameplayStatics.LoadStreamLevel");
+	FUNCTION_REDIRECT("LevelScriptActor.OpenLevel", "GameplayStatics.OpenLevel");
+	FUNCTION_REDIRECT("LevelScriptActor.UnloadStreamLevel", "GameplayStatics.UnloadStreamLevel");
+	FUNCTION_REDIRECT("MovementComponent.ConstrainPositionToPlane", "MovementComponent.ConstrainLocationToPlane");
+	PROPERTY_REDIRECT("MovementComponent.ConstrainLocationToPlane.Position", "Location");
+	FUNCTION_REDIRECT("PlayerCameraManager.KismetUpdateCamera", "BlueprintUpdateCamera");
+	FUNCTION_REDIRECT("PlayerController.AddLookUpInput", "PlayerController.AddPitchInput");
+	FUNCTION_REDIRECT("PlayerController.AddTurnInput", "PlayerController.AddYawInput");
+	PROPERTY_REDIRECT("PlayerController.DeprojectMousePositionToWorld.Direction", "WorldDirection");
+	PROPERTY_REDIRECT("PlayerController.DeprojectMousePositionToWorld.WorldPosition", "WorldLocation");
+	FUNCTION_REDIRECT("PrimitiveComponent.AddForceAtPosition", "PrimitiveComponent.AddForceAtLocation");
+	PROPERTY_REDIRECT("PrimitiveComponent.AddForceAtLocation.Position", "Location");
+	FUNCTION_REDIRECT("PrimitiveComponent.AddImpulseAtPosition", "PrimitiveComponent.AddImpulseAtLocation");
+	PROPERTY_REDIRECT("PrimitiveComponent.AddImpulseAtLocation.Position", "Location");
+	FUNCTION_REDIRECT("PrimitiveComponent.CreateAndSetMaterialInstanceDynamic", "PrimitiveComponent.CreateDynamicMaterialInstance");
+	FUNCTION_REDIRECT("PrimitiveComponent.CreateAndSetMaterialInstanceDynamicFromMaterial", "PrimitiveComponent.CreateDynamicMaterialInstance");
+	PROPERTY_REDIRECT("PrimitiveComponent.CreateDynamicMaterialInstance.Parent", "SourceMaterial");
+	FUNCTION_REDIRECT("PrimitiveComponent.GetRBAngularVelocity", "GetPhysicsAngularVelocity");
+	FUNCTION_REDIRECT("PrimitiveComponent.GetRBLinearVelocity", "GetPhysicsLinearVelocity");
+	FUNCTION_REDIRECT("PrimitiveComponent.GetTouchingActors", "PrimitiveComponent.GetOverlappingActors");
+	PROPERTY_REDIRECT("PrimitiveComponent.GetOverlappingActors.TouchingActors", "OverlappingActors");
+	FUNCTION_REDIRECT("PrimitiveComponent.GetTouchingComponents", "PrimitiveComponent.GetOverlappingComponents");
+	PROPERTY_REDIRECT("PrimitiveComponent.GetOverlappingComponents.TouchingComponents", "OverlappingComponents");
+	FUNCTION_REDIRECT("PrimitiveComponent.KismetTraceComponent", "PrimitiveComponent.K2_LineTraceComponent");
+	FUNCTION_REDIRECT("PrimitiveComponent.SetAllRBLinearVelocity", "SetAllPhysicsLinearVelocity");
+	FUNCTION_REDIRECT("PrimitiveComponent.SetMovementChannel", "PrimitiveComponent.SetCollisionObjectType");
+	FUNCTION_REDIRECT("PrimitiveComponent.SetRBAngularVelocity", "SetPhysicsAngularVelocity");
+	FUNCTION_REDIRECT("PrimitiveComponent.SetRBLinearVelocity", "SetPhysicsLinearVelocity");
+	FUNCTION_REDIRECT("ProjectileMovementComponent.StopMovement", "ProjectileMovementComponent.StopSimulating");
+	FUNCTION_REDIRECT("SceneComponent.GetComponentToWorld", "K2_GetComponentToWorld");
+	FUNCTION_REDIRECT("SingleAnimSkeletalComponent.GetPlayRate", "SkeletalMeshComponent.GetPlayRate");
+	FUNCTION_REDIRECT("SingleAnimSkeletalComponent.GetPosition", "SkeletalMeshComponent.GetPosition");
+	FUNCTION_REDIRECT("SingleAnimSkeletalComponent.IsPlaying", "SkeletalMeshComponent.IsPlaying");
+	FUNCTION_REDIRECT("SingleAnimSkeletalComponent.Play", "SkeletalMeshComponent.Play");
+	FUNCTION_REDIRECT("SingleAnimSkeletalComponent.SetAnim", "SkeletalMeshComponent.SetAnimation");
+	FUNCTION_REDIRECT("SingleAnimSkeletalComponent.SetPlayRate", "SkeletalMeshComponent.SetPlayRate");
+	FUNCTION_REDIRECT("SingleAnimSkeletalComponent.SetPosition", "SkeletalMeshComponent.SetPosition");
+	FUNCTION_REDIRECT("SingleAnimSkeletalComponent.Stop", "SkeletalMeshComponent.Stop");
+	FUNCTION_REDIRECT("SkinnedMeshComponent.MatchRefBone", "SkinnedMeshComponent.GetBoneIndex");
+
+	PROPERTY_REDIRECT("AnimNotifyEvent.Time", "AnimNotifyEvent.DisplayTime");
+	PROPERTY_REDIRECT("AnimSequence.BasePose", "AnimSequence.RetargetSource");
+	PROPERTY_REDIRECT("AudioComponent.PitchMultiplierMax", "AudioComponent.PitchModulationMax");
+	PROPERTY_REDIRECT("AudioComponent.PitchMultiplierMin", "AudioComponent.PitchModulationMin");
+	PROPERTY_REDIRECT("AudioComponent.VolumeMultiplierMax", "AudioComponent.VolumeModulationMax");
+	PROPERTY_REDIRECT("AudioComponent.VolumeMultiplierMin", "AudioComponent.VolumeModulationMin");
+	PROPERTY_REDIRECT("BodyInstance.MovementChannel", "BodyInstance.ObjectType");
+	PROPERTY_REDIRECT("BranchingPoint.Time", "BranchingPoint.DisplayTime");
+	PROPERTY_REDIRECT("CapsuleComponent.CapsuleHeight", "CapsuleComponent.CapsuleHalfHeight");
+	PROPERTY_REDIRECT("CharacterMovementComponent.AccelRate", "CharacterMovementComponent.MaxAcceleration");
+	PROPERTY_REDIRECT("CharacterMovementComponent.AirSpeed", "CharacterMovementComponent.MaxFlySpeed");
+	PROPERTY_REDIRECT("CharacterMovementComponent.BrakingDeceleration", "CharacterMovementComponent.BrakingDecelerationWalking");
+	PROPERTY_REDIRECT("CharacterMovementComponent.CrouchHeight", "CharacterMovementComponent.CrouchedHalfHeight");
+	PROPERTY_REDIRECT("CharacterMovementComponent.CrouchedPct", "CharacterMovementComponent.CrouchedSpeedMultiplier");
+	PROPERTY_REDIRECT("CharacterMovementComponent.CrouchedSpeedPercent", "CharacterMovementComponent.CrouchedSpeedMultiplier");
+	PROPERTY_REDIRECT("CharacterMovementComponent.GroundSpeed", "CharacterMovementComponent.MaxWalkSpeed");
+	PROPERTY_REDIRECT("CharacterMovementComponent.JumpZ", "CharacterMovementComponent.JumpZVelocity");
+	PROPERTY_REDIRECT("CharacterMovementComponent.WaterSpeed", "CharacterMovementComponent.MaxSwimSpeed");
+	PROPERTY_REDIRECT("CharacterMovementComponent.bOrientToMovement", "CharacterMovementComponent.bOrientRotationToMovement");
+	PROPERTY_REDIRECT("CollisionResponseContainer.Dynamic", "CollisionResponseContainer.WorldDynamic");
+	PROPERTY_REDIRECT("CollisionResponseContainer.RigidBody", "CollisionResponseContainer.PhysicsBody");
+	PROPERTY_REDIRECT("CollisionResponseContainer.Static", "CollisionResponseContainer.WorldStatic");
+	PROPERTY_REDIRECT("Controller.PlayerReplicationInfo", "Controller.PlayerState");
+	PROPERTY_REDIRECT("DefaultPawn.DefaultPawnMovement", "DefaultPawn.MovementComponent");
+	PROPERTY_REDIRECT("DirectionalLightComponent.MovableWholeSceneDynamicShadowRadius", "DirectionalLightComponent.DynamicShadowDistanceMovableLight");
+	PROPERTY_REDIRECT("DirectionalLightComponent.StationaryWholeSceneDynamicShadowRadius", "DirectionalLightComponent.DynamicShadowDistanceStationaryLight");
+	PROPERTY_REDIRECT("FloatingPawnMovement.AccelRate", "FloatingPawnMovement.Acceleration");
+	PROPERTY_REDIRECT("FloatingPawnMovement.DecelRate", "FloatingPawnMovement.Deceleration");
+	PROPERTY_REDIRECT("GameMode.GameReplicationInfoClass", "GameMode.GameStateClass");
+	PROPERTY_REDIRECT("GameMode.PlayerReplicationInfoClass", "GameMode.PlayerStateClass");
+	PROPERTY_REDIRECT("GameState.GameClass", "GameState.GameModeClass");
+	PROPERTY_REDIRECT("K2Node_TransitionRuleGetter.AssociatedSequencePlayerNode", "K2Node_TransitionRuleGetter.AssociatedAnimAssetPlayerNode");
+	PROPERTY_REDIRECT("LightComponent.InverseSquaredFalloff", "PointLightComponent.bUseInverseSquaredFalloff");
+	PROPERTY_REDIRECT("LightComponentBase.Brightness", "LightComponentBase.Intensity");
+	PROPERTY_REDIRECT("Material.RefractionBias", "Material.RefractionDepthBias");
+	PROPERTY_REDIRECT("MaterialEditorInstanceConstant.RefractionBias", "MaterialEditorInstanceConstant.RefractionDepthBias");
+	PROPERTY_REDIRECT("NavLinkProxy.NavLinks", "NavLinkProxy.PointLinks");
+	PROPERTY_REDIRECT("NavLinkProxy.NavSegmentLinks", "NavLinkProxy.SegmentLinks");
+	PROPERTY_REDIRECT("Pawn.ControllerClass", "Pawn.AIControllerClass");
+	PROPERTY_REDIRECT("Pawn.PlayerReplicationInfo", "Pawn.PlayerState");
+	PROPERTY_REDIRECT("PawnSensingComponent.SightCounterInterval", "PawnSensingComponent.SensingInterval");
+	PROPERTY_REDIRECT("PawnSensingComponent.bWantsSeePlayerNotify", "PawnSensingComponent.bSeePawns");
+	PROPERTY_REDIRECT("PlayerController.LookRightScale", "PlayerController.InputYawScale");
+	PROPERTY_REDIRECT("PlayerController.LookUpScale", "PlayerController.InputPitchScale");
+	PROPERTY_REDIRECT("PlayerController.PlayerCamera", "PlayerController.PlayerCameraManager");
+	PROPERTY_REDIRECT("PlayerController.PlayerCameraClass", "PlayerController.PlayerCameraManagerClass");
+	PROPERTY_REDIRECT("PointLightComponent.Radius", "PointLightComponent.AttenuationRadius");
+	PROPERTY_REDIRECT("PostProcessSettings.ExposureOffset", "PostProcessSettings.AutoExposureBias");
+	PROPERTY_REDIRECT("PostProcessSettings.EyeAdaptationHighPercent", "PostProcessSettings.AutoExposureHighPercent");
+	PROPERTY_REDIRECT("PostProcessSettings.EyeAdaptationLowPercent", "PostProcessSettings.AutoExposureLowPercent");
+	PROPERTY_REDIRECT("PostProcessSettings.EyeAdaptationMaxBrightness", "PostProcessSettings.AutoExposureMaxBrightness");
+	PROPERTY_REDIRECT("PostProcessSettings.EyeAdaptationMinBrightness", "PostProcessSettings.AutoExposureMinBrightness");
+	PROPERTY_REDIRECT("PostProcessSettings.EyeAdaptionSpeedDown", "PostProcessSettings.AutoExposureSpeedDown");
+	PROPERTY_REDIRECT("PostProcessSettings.EyeAdaptionSpeedUp", "PostProcessSettings.AutoExposureSpeedUp");
+	PROPERTY_REDIRECT("PostProcessSettings.bOverride_ExposureOffset", "PostProcessSettings.bOverride_AutoExposureBias");
+	PROPERTY_REDIRECT("PostProcessSettings.bOverride_EyeAdaptationHighPercent", "PostProcessSettings.bOverride_AutoExposureHighPercent");
+	PROPERTY_REDIRECT("PostProcessSettings.bOverride_EyeAdaptationLowPercent", "PostProcessSettings.bOverride_AutoExposureLowPercent");
+	PROPERTY_REDIRECT("PostProcessSettings.bOverride_EyeAdaptationMaxBrightness", "PostProcessSettings.bOverride_AutoExposureMaxBrightness");
+	PROPERTY_REDIRECT("PostProcessSettings.bOverride_EyeAdaptationMinBrightness", "PostProcessSettings.bOverride_AutoExposureMinBrightness");
+	PROPERTY_REDIRECT("PostProcessSettings.bOverride_EyeAdaptionSpeedDown", "PostProcessSettings.bOverride_AutoExposureSpeedDown");
+	PROPERTY_REDIRECT("PostProcessSettings.bOverride_EyeAdaptionSpeedUp", "PostProcessSettings.bOverride_AutoExposureSpeedUp");
+	PROPERTY_REDIRECT("ProjectileMovementComponent.Speed", "ProjectileMovementComponent.InitialSpeed");
+	PROPERTY_REDIRECT("SceneComponent.ModifyFrequency", "SceneComponent.Mobility");
+	PROPERTY_REDIRECT("SceneComponent.RelativeTranslation", "SceneComponent.RelativeLocation");
+	PROPERTY_REDIRECT("SceneComponent.bAbsoluteTranslation", "SceneComponent.bAbsoluteLocation");
+	PROPERTY_REDIRECT("SkeletalMeshComponent.AnimationBlueprint", "SkeletalMeshComponent.AnimBlueprintGeneratedClass");
+	PROPERTY_REDIRECT("SkinnedMeshComponent.SkinnedMeshUpdateFlag", "SkinnedMeshComponent.MeshComponentUpdateFlag");
+	PROPERTY_REDIRECT("SlateBrush.TextureName", "SlateBrush.ResourceName");
+	PROPERTY_REDIRECT("SlateBrush.TextureObject", "SlateBrush.ResourceObject");
+	PROPERTY_REDIRECT("WorldSettings.DefaultGameType", "WorldSettings.DefaultGameMode");
+
+	FCoreRedirect* PointLightComponent = CLASS_REDIRECT("PointLightComponent", "PointLightComponent");
+	PointLightComponent->ValueChanges.Add(TEXT("PointLightComponent0"), TEXT("LightComponent0"));
+
+	FCoreRedirect* DirectionalLightComponent = CLASS_REDIRECT("DirectionalLightComponent", "DirectionalLightComponent");
+	DirectionalLightComponent->ValueChanges.Add(TEXT("DirectionalLightComponent0"), TEXT("LightComponent0"));
+
+	FCoreRedirect* SpotLightComponent = CLASS_REDIRECT("SpotLightComponent", "SpotLightComponent");
+	SpotLightComponent->ValueChanges.Add(TEXT("SpotLightComponent0"), TEXT("LightComponent0"));
+
+	FCoreRedirect* ETransitionGetterType = ENUM_REDIRECT("ETransitionGetterType", "ETransitionGetter");
+	ETransitionGetterType->ValueChanges.Add(TEXT("TGT_ArbitraryState_GetBlendWeight"), TEXT("ETransitionGetter::ArbitraryState_GetBlendWeight"));
+	ETransitionGetterType->ValueChanges.Add(TEXT("TGT_CurrentState_ElapsedTime"), TEXT("ETransitionGetter::CurrentState_ElapsedTime"));
+	ETransitionGetterType->ValueChanges.Add(TEXT("TGT_CurrentState_GetBlendWeight"), TEXT("ETransitionGetter::CurrentState_GetBlendWeight"));
+	ETransitionGetterType->ValueChanges.Add(TEXT("TGT_CurrentTransitionDuration"), TEXT("ETransitionGetter::CurrentTransitionDuration"));
+	ETransitionGetterType->ValueChanges.Add(TEXT("TGT_SequencePlayer_GetCurrentTime"), TEXT("ETransitionGetter::AnimationAsset_GetCurrentTime"));
+	ETransitionGetterType->ValueChanges.Add(TEXT("TGT_SequencePlayer_GetCurrentTimeFraction"), TEXT("ETransitionGetter::AnimationAsset_GetCurrentTimeFraction"));
+	ETransitionGetterType->ValueChanges.Add(TEXT("TGT_SequencePlayer_GetLength"), TEXT("ETransitionGetter::AnimationAsset_GetLength"));
+	ETransitionGetterType->ValueChanges.Add(TEXT("TGT_SequencePlayer_GetTimeFromEnd"), TEXT("ETransitionGetter::AnimationAsset_GetTimeFromEnd"));
+	ETransitionGetterType->ValueChanges.Add(TEXT("TGT_SequencePlayer_GetTimeFromEndFraction"), TEXT("ETransitionGetter::AnimationAsset_GetTimeFromEndFraction"));
+
+	FCoreRedirect* EModifyFrequency = ENUM_REDIRECT("EModifyFrequency", "EComponentMobility");
+	EModifyFrequency->ValueChanges.Add(TEXT("MF_Dynamic"), TEXT("EComponentMobility::Movable"));
+	EModifyFrequency->ValueChanges.Add(TEXT("MF_OccasionallyModified"), TEXT("EComponentMobility::Stationary"));
+	EModifyFrequency->ValueChanges.Add(TEXT("MF_Static"), TEXT("EComponentMobility::Static"));
+
+	FCoreRedirect* EAttachLocationType = ENUM_REDIRECT("EAttachLocationType", "EAttachLocation");
+	EAttachLocationType->ValueChanges.Add(TEXT("EAttachLocationType_AbsoluteWorld"), TEXT("EAttachLocation::KeepWorldPosition"));
+	EAttachLocationType->ValueChanges.Add(TEXT("EAttachLocationType_RelativeOffset"), TEXT("EAttachLocation::KeepRelativeOffset"));
+	EAttachLocationType->ValueChanges.Add(TEXT("EAttachLocationType_SnapTo"), TEXT("EAttachLocation::SnapToTarget"));
+
+	FCoreRedirect* EAxis = ENUM_REDIRECT("EAxis", "EAxis");
+	EAxis->ValueChanges.Add(TEXT("AXIS_BLANK"), TEXT("EAxis::None"));
+	EAxis->ValueChanges.Add(TEXT("AXIS_NONE"), TEXT("EAxis::None"));
+	EAxis->ValueChanges.Add(TEXT("AXIS_X"), TEXT("EAxis::X"));
+	EAxis->ValueChanges.Add(TEXT("AXIS_Y"), TEXT("EAxis::Y"));
+	EAxis->ValueChanges.Add(TEXT("AXIS_Z"), TEXT("EAxis::Z"));
+
+	FCoreRedirect* EKeys = ENUM_REDIRECT("EKeys", "EKeys");
+	EKeys->ValueChanges.Add(TEXT("EKeys::XboxTypeS_A"), TEXT("EKeys::Gamepad_FaceButton_Bottom"));
+	EKeys->ValueChanges.Add(TEXT("EKeys::XboxTypeS_B"), TEXT("EKeys::Gamepad_FaceButton_Right"));
+	EKeys->ValueChanges.Add(TEXT("EKeys::XboxTypeS_X"), TEXT("EKeys::Gamepad_FaceButton_Left"));
+	EKeys->ValueChanges.Add(TEXT("EKeys::XboxTypeS_Y"), TEXT("EKeys::Gamepad_FaceButton_Top"));
+	EKeys->ValueChanges.Add(TEXT("EKeys::XboxTypeS_Back"), TEXT("EKeys::Gamepad_Special_Left"));
+	EKeys->ValueChanges.Add(TEXT("EKeys::XboxTypeS_Start"), TEXT("EKeys::Gamepad_Special_Right"));
+	EKeys->ValueChanges.Add(TEXT("EKeys::XboxTypeS_DPad_Down"), TEXT("EKeys::Gamepad_DPad_Down"));
+	EKeys->ValueChanges.Add(TEXT("EKeys::XboxTypeS_DPad_Left"), TEXT("EKeys::Gamepad_DPad_Left"));
+	EKeys->ValueChanges.Add(TEXT("EKeys::XboxTypeS_DPad_Right"), TEXT("EKeys::Gamepad_DPad_Right"));
+	EKeys->ValueChanges.Add(TEXT("EKeys::XboxTypeS_DPad_Up"), TEXT("EKeys::Gamepad_DPad_Up"));
+	EKeys->ValueChanges.Add(TEXT("EKeys::XboxTypeS_LeftShoulder"), TEXT("EKeys::Gamepad_LeftShoulder"));
+	EKeys->ValueChanges.Add(TEXT("EKeys::XboxTypeS_LeftThumbstick"), TEXT("EKeys::Gamepad_LeftThumbstick"));
+	EKeys->ValueChanges.Add(TEXT("EKeys::XboxTypeS_LeftTrigger"), TEXT("EKeys::Gamepad_LeftTrigger"));
+	EKeys->ValueChanges.Add(TEXT("EKeys::XboxTypeS_LeftTriggerAxis"), TEXT("EKeys::Gamepad_LeftTriggerAxis"));
+	EKeys->ValueChanges.Add(TEXT("EKeys::XboxTypeS_LeftX"), TEXT("EKeys::Gamepad_LeftX"));
+	EKeys->ValueChanges.Add(TEXT("EKeys::XboxTypeS_LeftY"), TEXT("EKeys::Gamepad_LeftY"));
+	EKeys->ValueChanges.Add(TEXT("EKeys::XboxTypeS_RightShoulder"), TEXT("EKeys::Gamepad_RightShoulder"));
+	EKeys->ValueChanges.Add(TEXT("EKeys::XboxTypeS_RightThumbstick"), TEXT("EKeys::Gamepad_RightThumbstick"));
+	EKeys->ValueChanges.Add(TEXT("EKeys::XboxTypeS_RightTrigger"), TEXT("EKeys::Gamepad_RightTrigger"));
+	EKeys->ValueChanges.Add(TEXT("EKeys::XboxTypeS_RightTriggerAxis"), TEXT("EKeys::Gamepad_RightTriggerAxis"));
+	EKeys->ValueChanges.Add(TEXT("EKeys::XboxTypeS_RightX"), TEXT("EKeys::Gamepad_RightX"));
+	EKeys->ValueChanges.Add(TEXT("EKeys::XboxTypeS_RightY"), TEXT("EKeys::Gamepad_RightY"));
+
+	FCoreRedirect* EMaxConcurrentResolutionRule = ENUM_REDIRECT("EMaxConcurrentResolutionRule", "EMaxConcurrentResolutionRule");
+	EMaxConcurrentResolutionRule->ValueChanges.Add(TEXT("EMaxConcurrentResolutionRule::StopFarthest"), TEXT("EMaxConcurrentResolutionRule::StopFarthestThenPreventNew"));
+
+	FCoreRedirect* EMeshComponentUpdateFlag = ENUM_REDIRECT("EMeshComponentUpdateFlag", "EMeshComponentUpdateFlag");
+	EMeshComponentUpdateFlag->ValueChanges.Add(TEXT("SMU_AlwaysTickPose"), TEXT("EMeshComponentUpdateFlag::AlwaysTickPose"));
+	EMeshComponentUpdateFlag->ValueChanges.Add(TEXT("SMU_AlwaysTickPoseAndRefreshBones"), TEXT("EMeshComponentUpdateFlag::AlwaysTickPoseAndRefreshBones"));
+	EMeshComponentUpdateFlag->ValueChanges.Add(TEXT("SMU_OnlyTickPoseWhenRendered"), TEXT("EMeshComponentUpdateFlag::OnlyTickPoseWhenRendered"));
+
+	FCoreRedirect* EParticleEventType = ENUM_REDIRECT("EParticleEventType", "EParticleEventType");
+	EParticleEventType->ValueChanges.Add(TEXT("EPET_Kismet"), TEXT("EPET_Blueprint"));
+
+	FCoreRedirect* ETranslucencyLightingMode = ENUM_REDIRECT("ETranslucencyLightingMode", "ETranslucencyLightingMode");
+	ETranslucencyLightingMode->ValueChanges.Add(TEXT("TLM_PerPixel"), TEXT("TLM_VolumetricDirectional"));
+	ETranslucencyLightingMode->ValueChanges.Add(TEXT("TLM_PerPixelNonDirectional"), TEXT("TLM_VolumetricNonDirectional"));
+}
+
+static void RegisterNativeRedirects46(TArray<FCoreRedirect>& Redirects)
+{
+	// 4.1-4.4 
+
+	CLASS_REDIRECT("K2Node_CastToInterface", "/Script/BlueprintGraph.K2Node_DynamicCast");
+	CLASS_REDIRECT("K2Node_MathExpression", "/Script/BlueprintGraph.K2Node_MathExpression");
+	CLASS_REDIRECT("EmitterSpawnable", "Emitter");
+	CLASS_REDIRECT("SlateWidgetStyleAsset", "/Script/SlateCore.SlateWidgetStyleAsset");
+	CLASS_REDIRECT("SlateWidgetStyleContainerBase", "/Script/SlateCore.SlateWidgetStyleContainerBase");
+	CLASS_REDIRECT("SmartNavLinkComponent", "/Script/Engine.NavLinkCustomComponent");
+	CLASS_REDIRECT("WidgetBlueprint", "/Script/UMGEditor.WidgetBlueprint");
+
+	PROPERTY_REDIRECT("AnimNotify.Received_Notify.AnimSeq", "Animation");
+	PROPERTY_REDIRECT("AnimNotifyState.Received_NotifyBegin.AnimSeq", "Animation");
+	PROPERTY_REDIRECT("AnimNotifyState.Received_NotifyEnd.AnimSeq", "Animation");
+	PROPERTY_REDIRECT("AnimNotifyState.Received_NotifyTick.AnimSeq", "Animation");
+	FUNCTION_REDIRECT("Character.IsJumping", "Character.IsJumpProvidingForce");
+	PROPERTY_REDIRECT("CharacterMovementComponent.AddImpulse.InMomentum", "Impulse");
+	PROPERTY_REDIRECT("CharacterMovementComponent.AddImpulse.bMassIndependent", "bVelocityChange");
+	FUNCTION_REDIRECT("CharacterMovementComponent.AddMomentum", "CharacterMovementComponent.AddImpulse");
+	FUNCTION_REDIRECT("Controller.GetControlledPawn", "Controller.K2_GetPawn");
+	FUNCTION_REDIRECT("DefaultPawn.LookUp", "Pawn.AddControllerPitchInput");
+	FUNCTION_REDIRECT("DefaultPawn.Turn", "Pawn.AddControllerYawInput");
+	FUNCTION_REDIRECT("KismetSystemLibrary.EXPERIMENTAL_ShowGameCenterLeaderboard", "KismetSystemLibrary.ShowPlatformSpecificLeaderboardScreen");
+	FUNCTION_REDIRECT("MovementComponent.GetMaxSpeedModifier", "MovementComponent.K2_GetMaxSpeedModifier");
+	FUNCTION_REDIRECT("MovementComponent.GetModifiedMaxSpeed", "MovementComponent.K2_GetModifiedMaxSpeed");
+	FUNCTION_REDIRECT("Pawn.AddLookUpInput", "Pawn.AddControllerPitchInput");
+	FUNCTION_REDIRECT("Pawn.AddPitchInput", "Pawn.AddControllerPitchInput");
+	FUNCTION_REDIRECT("Pawn.AddRollInput", "Pawn.AddControllerRollInput");
+	FUNCTION_REDIRECT("Pawn.AddTurnInput", "Pawn.AddControllerYawInput");
+	FUNCTION_REDIRECT("Pawn.AddYawInput", "Pawn.AddControllerYawInput");
+	FUNCTION_REDIRECT("PawnMovementComponent.StopActiveMovement", "NavMovementComponent.StopActiveMovement");
+	FUNCTION_REDIRECT("PointLightComponent.SetRadius", "PointLightComponent.SetAttenuationRadius");
+	FUNCTION_REDIRECT("SkeletalMeshComponent.SetAnimBlueprint", "SkeletalMeshComponent.SetAnimInstanceClass");
+	FUNCTION_REDIRECT("SkeletalMeshComponent.SetAnimClass", "SkeletalMeshComponent.SetAnimInstanceClass");
+	PROPERTY_REDIRECT("SkeletalMeshComponent.SetAnimInstanceClass.NewBlueprint", "NewClass");
+
+	PROPERTY_REDIRECT("StringClassReference.ClassName", "StringClassReference.AssetLongPathname");
+	PROPERTY_REDIRECT("Material.LightingModel", "Material.ShadingModel");
+	PROPERTY_REDIRECT("MaterialInstanceBasePropertyOverrides.LightingModel", "MaterialInstanceBasePropertyOverrides.ShadingModel");
+	PROPERTY_REDIRECT("MaterialInstanceBasePropertyOverrides.bOverride_LightingModel", "MaterialInstanceBasePropertyOverrides.bOverride_ShadingModel");
+	PROPERTY_REDIRECT("PassiveSoundMixModifier.VolumeThreshold", "PassiveSoundMixModifier.MinVolumeThreshold");
+	PROPERTY_REDIRECT("PrimitiveComponent.CanBeCharacterBase", "PrimitiveComponent.CanCharacterStepUpOn");
+	PROPERTY_REDIRECT("SkeletalMeshLODInfo.DisplayFactor", "SkeletalMeshLODInfo.ScreenSize");
+	PROPERTY_REDIRECT("SplineMeshComponent.SplineXDir", "SplineMeshComponent.SplineUpDir");
+	PROPERTY_REDIRECT("TextureFactory.LightingModel", "TextureFactory.ShadingModel");
+
+	FCoreRedirect* EKinematicBonesUpdateToPhysics = ENUM_REDIRECT("EKinematicBonesUpdateToPhysics", "EKinematicBonesUpdateToPhysics");
+	EKinematicBonesUpdateToPhysics->ValueChanges.Add(TEXT("EKinematicBonesUpdateToPhysics::SkipFixedAndSimulatingBones"), TEXT("EKinematicBonesUpdateToPhysics::SkipAllBones"));
+
+	FCoreRedirect* EMaterialLightingModel = ENUM_REDIRECT("EMaterialLightingModel", "EMaterialShadingModel");
+	EMaterialLightingModel->ValueChanges.Add(TEXT("MLM_DefaultLit"), TEXT("MSM_DefaultLit"));
+	EMaterialLightingModel->ValueChanges.Add(TEXT("MLM_PreintegratedSkin"), TEXT("MSM_PreintegratedSkin"));
+	EMaterialLightingModel->ValueChanges.Add(TEXT("MLM_Subsurface"), TEXT("MSM_Subsurface"));
+	EMaterialLightingModel->ValueChanges.Add(TEXT("MLM_Unlit"), TEXT("MSM_Unlit"));
+
+	FCoreRedirect* ESmartNavLinkDir = ENUM_REDIRECT("ESmartNavLinkDir", "ENavLinkDirection");
+	ESmartNavLinkDir->ValueChanges.Add(TEXT("ESmartNavLinkDir::BothWays"), TEXT("ENavLinkDirection::BothWays"));
+	ESmartNavLinkDir->ValueChanges.Add(TEXT("ESmartNavLinkDir::OneWay"), TEXT("ENavLinkDirection::LeftToRight"));
+
+	FCoreRedirect* EPhysicsType = ENUM_REDIRECT("EPhysicsType", "EPhysicsType");
+	EPhysicsType->ValueChanges.Add(TEXT("PhysType_Fixed"), TEXT("PhysType_Kinematic"));
+	EPhysicsType->ValueChanges.Add(TEXT("PhysType_Unfixed"), TEXT("PhysType_Simulated"));
+
+	FCoreRedirect* ESceneTextureId = ENUM_REDIRECT("ESceneTextureId", "ESceneTextureId");
+	EPhysicsType->ValueChanges.Add(TEXT("PPI_LightingModel"), TEXT("PPI_ShadingModel"));
+
+	// 4.5
+
+	CLASS_REDIRECT("AIController", "/Script/AIModule.AIController");
+	CLASS_REDIRECT("AIResourceInterface", "/Script/AIModule.AIResourceInterface");
+	CLASS_REDIRECT("AISystem", "/Script/AIModule.AISystem");
+	CLASS_REDIRECT("AITypes", "/Script/AIModule.AITypes");
+	CLASS_REDIRECT("BTAuxiliaryNode", "/Script/AIModule.BTAuxiliaryNode");
+	CLASS_REDIRECT("BTCompositeNode", "/Script/AIModule.BTCompositeNode");
+	CLASS_REDIRECT("BTComposite_Selector", "/Script/AIModule.BTComposite_Selector");
+	CLASS_REDIRECT("BTComposite_Sequence", "/Script/AIModule.BTComposite_Sequence");
+	CLASS_REDIRECT("BTComposite_SimpleParallel", "/Script/AIModule.BTComposite_SimpleParallel");
+	CLASS_REDIRECT("BTDecorator", "/Script/AIModule.BTDecorator");
+	CLASS_REDIRECT("BTDecorator_Blackboard", "/Script/AIModule.BTDecorator_Blackboard");
+	CLASS_REDIRECT("BTDecorator_BlackboardBase", "/Script/AIModule.BTDecorator_BlackboardBase");
+	CLASS_REDIRECT("BTDecorator_BlueprintBase", "/Script/AIModule.BTDecorator_BlueprintBase");
+	CLASS_REDIRECT("BTDecorator_CompareBBEntries", "/Script/AIModule.BTDecorator_CompareBBEntries");
+	CLASS_REDIRECT("BTDecorator_ConeCheck", "/Script/AIModule.BTDecorator_ConeCheck");
+	CLASS_REDIRECT("BTDecorator_Cooldown", "/Script/AIModule.BTDecorator_Cooldown");
+	CLASS_REDIRECT("BTDecorator_DoesPathExist", "/Script/AIModule.BTDecorator_DoesPathExist");
+	CLASS_REDIRECT("BTDecorator_ForceSuccess", "/Script/AIModule.BTDecorator_ForceSuccess");
+	CLASS_REDIRECT("BTDecorator_KeepInCone", "/Script/AIModule.BTDecorator_KeepInCone");
+	CLASS_REDIRECT("BTDecorator_Loop", "/Script/AIModule.BTDecorator_Loop");
+	CLASS_REDIRECT("BTDecorator_Optional", "/Script/AIModule.BTDecorator_ForceSuccess");
+	CLASS_REDIRECT("BTDecorator_ReachedMoveGoal", "/Script/AIModule.BTDecorator_ReachedMoveGoal");
+	CLASS_REDIRECT("BTDecorator_TimeLimit", "/Script/AIModule.BTDecorator_TimeLimit");
+	CLASS_REDIRECT("BTFunctionLibrary", "/Script/AIModule.BTFunctionLibrary");
+	CLASS_REDIRECT("BTNode", "/Script/AIModule.BTNode");
+	CLASS_REDIRECT("BTService", "/Script/AIModule.BTService");
+	CLASS_REDIRECT("BTService_BlackboardBase", "/Script/AIModule.BTService_BlackboardBase");
+	CLASS_REDIRECT("BTService_BlueprintBase", "/Script/AIModule.BTService_BlueprintBase");
+	CLASS_REDIRECT("BTService_DefaultFocus", "/Script/AIModule.BTService_DefaultFocus");
+	CLASS_REDIRECT("BTTaskNode", "/Script/AIModule.BTTaskNode");
+	CLASS_REDIRECT("BTTask_BlackboardBase", "/Script/AIModule.BTTask_BlackboardBase");
+	CLASS_REDIRECT("BTTask_BlueprintBase", "/Script/AIModule.BTTask_BlueprintBase");
+	CLASS_REDIRECT("BTTask_MakeNoise", "/Script/AIModule.BTTask_MakeNoise");
+	CLASS_REDIRECT("BTTask_MoveDirectlyToward", "/Script/AIModule.BTTask_MoveDirectlyToward");
+	CLASS_REDIRECT("BTTask_MoveTo", "/Script/AIModule.BTTask_MoveTo");
+	CLASS_REDIRECT("BTTask_PlaySound", "/Script/AIModule.BTTask_PlaySound");
+	CLASS_REDIRECT("BTTask_RunBehavior", "/Script/AIModule.BTTask_RunBehavior");
+	CLASS_REDIRECT("BTTask_RunEQSQuery", "/Script/AIModule.BTTask_RunEQSQuery");
+	CLASS_REDIRECT("BTTask_Wait", "/Script/AIModule.BTTask_Wait");
+	CLASS_REDIRECT("BehaviorTree", "/Script/AIModule.BehaviorTree");
+	CLASS_REDIRECT("BehaviorTreeComponent", "/Script/AIModule.BehaviorTreeComponent");
+	CLASS_REDIRECT("BehaviorTreeManager", "/Script/AIModule.BehaviorTreeManager");
+	CLASS_REDIRECT("BehaviorTreeTypes", "/Script/AIModule.BehaviorTreeTypes");
+	CLASS_REDIRECT("BlackboardComponent", "/Script/AIModule.BlackboardComponent");
+	CLASS_REDIRECT("BlackboardData", "/Script/AIModule.BlackboardData");
+	CLASS_REDIRECT("BlackboardKeyAllTypes", "/Script/AIModule.BlackboardKeyAllTypes");
+	CLASS_REDIRECT("BlackboardKeyType", "/Script/AIModule.BlackboardKeyType");
+	CLASS_REDIRECT("BlackboardKeyType_Bool", "/Script/AIModule.BlackboardKeyType_Bool");
+	CLASS_REDIRECT("BlackboardKeyType_Class", "/Script/AIModule.BlackboardKeyType_Class");
+	CLASS_REDIRECT("BlackboardKeyType_Enum", "/Script/AIModule.BlackboardKeyType_Enum");
+	CLASS_REDIRECT("BlackboardKeyType_Float", "/Script/AIModule.BlackboardKeyType_Float");
+	CLASS_REDIRECT("BlackboardKeyType_Int", "/Script/AIModule.BlackboardKeyType_Int");
+	CLASS_REDIRECT("BlackboardKeyType_Name", "/Script/AIModule.BlackboardKeyType_Name");
+	CLASS_REDIRECT("BlackboardKeyType_NativeEnum", "/Script/AIModule.BlackboardKeyType_NativeEnum");
+	CLASS_REDIRECT("BlackboardKeyType_Object", "/Script/AIModule.BlackboardKeyType_Object");
+	CLASS_REDIRECT("BlackboardKeyType_String", "/Script/AIModule.BlackboardKeyType_String");
+	CLASS_REDIRECT("BlackboardKeyType_Vector", "/Script/AIModule.BlackboardKeyType_Vector");
+	CLASS_REDIRECT("BrainComponent", "/Script/AIModule.BrainComponent");
+	CLASS_REDIRECT("CrowdAgentInterface", "/Script/AIModule.CrowdAgentInterface");
+	CLASS_REDIRECT("CrowdFollowingComponent", "/Script/AIModule.CrowdFollowingComponent");
+	CLASS_REDIRECT("CrowdManager", "/Script/AIModule.CrowdManager");
+	CLASS_REDIRECT("EQSQueryResultSourceInterface", "/Script/AIModule.EQSQueryResultSourceInterface");
+	CLASS_REDIRECT("EQSRenderingComponent", "/Script/AIModule.EQSRenderingComponent");
+	CLASS_REDIRECT("EQSTestingPawn", "/Script/AIModule.EQSTestingPawn");
+	CLASS_REDIRECT("EnvQuery", "/Script/AIModule.EnvQuery");
+	CLASS_REDIRECT("EnvQueryAllItemTypes", "/Script/AIModule.EnvQueryAllItemTypes");
+	CLASS_REDIRECT("EnvQueryContext", "/Script/AIModule.EnvQueryContext");
+	CLASS_REDIRECT("EnvQueryContext_BlueprintBase", "/Script/AIModule.EnvQueryContext_BlueprintBase");
+	CLASS_REDIRECT("EnvQueryContext_Item", "/Script/AIModule.EnvQueryContext_Item");
+	CLASS_REDIRECT("EnvQueryContext_Querier", "/Script/AIModule.EnvQueryContext_Querier");
+	CLASS_REDIRECT("EnvQueryGenerator", "/Script/AIModule.EnvQueryGenerator");
+	CLASS_REDIRECT("EnvQueryGenerator_Composite", "/Script/AIModule.EnvQueryGenerator_Composite");
+	CLASS_REDIRECT("EnvQueryGenerator_OnCircle", "/Script/AIModule.EnvQueryGenerator_OnCircle");
+	CLASS_REDIRECT("EnvQueryGenerator_PathingGrid", "/Script/AIModule.EnvQueryGenerator_PathingGrid");
+	CLASS_REDIRECT("EnvQueryGenerator_ProjectedPoints", "/Script/AIModule.EnvQueryGenerator_ProjectedPoints");
+	CLASS_REDIRECT("EnvQueryGenerator_SimpleGrid", "/Script/AIModule.EnvQueryGenerator_SimpleGrid");
+	CLASS_REDIRECT("EnvQueryItemType", "/Script/AIModule.EnvQueryItemType");
+	CLASS_REDIRECT("EnvQueryItemType_Actor", "/Script/AIModule.EnvQueryItemType_Actor");
+	CLASS_REDIRECT("EnvQueryItemType_ActorBase", "/Script/AIModule.EnvQueryItemType_ActorBase");
+	CLASS_REDIRECT("EnvQueryItemType_Direction", "/Script/AIModule.EnvQueryItemType_Direction");
+	CLASS_REDIRECT("EnvQueryItemType_Point", "/Script/AIModule.EnvQueryItemType_Point");
+	CLASS_REDIRECT("EnvQueryItemType_VectorBase", "/Script/AIModule.EnvQueryItemType_VectorBase");
+	CLASS_REDIRECT("EnvQueryManager", "/Script/AIModule.EnvQueryManager");
+	CLASS_REDIRECT("EnvQueryOption", "/Script/AIModule.EnvQueryOption");
+	CLASS_REDIRECT("EnvQueryTest", "/Script/AIModule.EnvQueryTest");
+	CLASS_REDIRECT("EnvQueryTest_Distance", "/Script/AIModule.EnvQueryTest_Distance");
+	CLASS_REDIRECT("EnvQueryTest_Dot", "/Script/AIModule.EnvQueryTest_Dot");
+	CLASS_REDIRECT("EnvQueryTest_Pathfinding", "/Script/AIModule.EnvQueryTest_Pathfinding");
+	CLASS_REDIRECT("EnvQueryTest_Trace", "/Script/AIModule.EnvQueryTest_Trace");
+	CLASS_REDIRECT("EnvQueryTypes", "/Script/AIModule.EnvQueryTypes");
+	CLASS_REDIRECT("KismetAIAsyncTaskProxy", "/Script/AIModule.AIAsyncTaskBlueprintProxy");
+	CLASS_REDIRECT("KismetAIHelperLibrary", "/Script/AIModule.AIBlueprintHelperLibrary");
+	CLASS_REDIRECT("PathFollowingComponent", "/Script/AIModule.PathFollowingComponent");
+	CLASS_REDIRECT("PawnSensingComponent", "/Script/AIModule.PawnSensingComponent");
+
+	STRUCT_REDIRECT("SReply", "EventReply");
+
+	PROPERTY_REDIRECT("Actor.AddTickPrerequisiteActor.DependentActor", "PrerequisiteActor");
+	FUNCTION_REDIRECT("Actor.AttachRootComponentTo", "Actor.K2_AttachRootComponentTo");
+	FUNCTION_REDIRECT("Actor.AttachRootComponentToActor", "Actor.K2_AttachRootComponentToActor");
+	FUNCTION_REDIRECT("Actor.SetTickPrerequisite", "Actor.AddTickPrerequisiteActor");
+	PROPERTY_REDIRECT("BTTask_MoveDirectlyToward.bForceMoveToLocation", "bDisablePathUpdateOnGoalLocationChange");
+	PROPERTY_REDIRECT("KismetSystemLibrary.DrawDebugPlane.Loc", "Location");
+	PROPERTY_REDIRECT("KismetSystemLibrary.DrawDebugPlane.P", "PlaneCoordinates");
+	FUNCTION_REDIRECT("KismetSystemLibrary.EXPERIMENTAL_CloseAdBanner", "KismetSystemLibrary.ForceCloseAdBanner");
+	FUNCTION_REDIRECT("KismetSystemLibrary.EXPERIMENTAL_HideAdBanner", "KismetSystemLibrary.HideAdBanner");
+	FUNCTION_REDIRECT("KismetSystemLibrary.EXPERIMENTAL_ShowAdBanner", "KismetSystemLibrary.ShowAdBanner");
+	FUNCTION_REDIRECT("LightComponent.SetBrightness", "LightComponent.SetIntensity");
+	FUNCTION_REDIRECT("NavigationPath.GetPathLenght", "NavigationPath.GetPathLength");
+	FUNCTION_REDIRECT("Pawn.GetMovementInputVector", "Pawn.K2_GetMovementInputVector");
+	FUNCTION_REDIRECT("PawnMovementComponent.GetInputVector", "PawnMovementComponent.K2_GetInputVector");
+	FUNCTION_REDIRECT("SceneComponent.AttachTo", "SceneComponent.K2_AttachTo");
+	FUNCTION_REDIRECT("SkyLightComponent.SetBrightness", "SkyLightComponent.SetIntensity");
+
+	PROPERTY_REDIRECT("AnimCurveBase.CurveName", "LastObservedName");
+	PROPERTY_REDIRECT("CameraComponent.bUsePawnViewRotation", "CameraComponent.bUsePawnControlRotation");
+	PROPERTY_REDIRECT("CharacterMovementComponent.bCrouchMovesCharacterDown", "CharacterMovementComponent.bCrouchMaintainsBaseLocation");
+	PROPERTY_REDIRECT("SpringArmComponent.bUseControllerViewRotation", "SpringArmComponent.bUsePawnControlRotation");
+	PROPERTY_REDIRECT("SpringArmComponent.bUsePawnViewRotation", "SpringArmComponent.bUsePawnControlRotation");
+
+	// 4.6
+
+	CLASS_REDIRECT("ControlPointMeshComponent", "/Script/Landscape.ControlPointMeshComponent");
+	CLASS_REDIRECT("Landscape", "/Script/Landscape.Landscape");
+	CLASS_REDIRECT("LandscapeComponent", "/Script/Landscape.LandscapeComponent");
+	CLASS_REDIRECT("LandscapeGizmoActiveActor", "/Script/Landscape.LandscapeGizmoActiveActor");
+	CLASS_REDIRECT("LandscapeGizmoActor", "/Script/Landscape.LandscapeGizmoActor");
+	CLASS_REDIRECT("LandscapeGizmoRenderComponent", "/Script/Landscape.LandscapeGizmoRenderComponent");
+	CLASS_REDIRECT("LandscapeHeightfieldCollisionComponent", "/Script/Landscape.LandscapeHeightfieldCollisionComponent");
+	CLASS_REDIRECT("LandscapeInfo", "/Script/Landscape.LandscapeInfo");
+	CLASS_REDIRECT("LandscapeInfoMap", "/Script/Landscape.LandscapeInfoMap");
+	CLASS_REDIRECT("LandscapeLayerInfoObject", "/Script/Landscape.LandscapeLayerInfoObject");
+	CLASS_REDIRECT("LandscapeMaterialInstanceConstant", "/Script/Landscape.LandscapeMaterialInstanceConstant");
+	CLASS_REDIRECT("LandscapeMeshCollisionComponent", "/Script/Landscape.LandscapeMeshCollisionComponent");
+	CLASS_REDIRECT("LandscapeProxy", "/Script/Landscape.LandscapeProxy");
+	CLASS_REDIRECT("LandscapeSplineControlPoint", "/Script/Landscape.LandscapeSplineControlPoint");
+	CLASS_REDIRECT("LandscapeSplineSegment", "/Script/Landscape.LandscapeSplineSegment");
+	CLASS_REDIRECT("LandscapeSplinesComponent", "/Script/Landscape.LandscapeSplinesComponent");
+	CLASS_REDIRECT("MaterialExpressionLandscapeLayerBlend", "/Script/Landscape.MaterialExpressionLandscapeLayerBlend");
+	CLASS_REDIRECT("MaterialExpressionLandscapeLayerCoords", "/Script/Landscape.MaterialExpressionLandscapeLayerCoords");
+	CLASS_REDIRECT("MaterialExpressionLandscapeLayerSwitch", "/Script/Landscape.MaterialExpressionLandscapeLayerSwitch");
+	CLASS_REDIRECT("MaterialExpressionLandscapeLayerWeight", "/Script/Landscape.MaterialExpressionLandscapeLayerWeight");
+	CLASS_REDIRECT("MaterialExpressionLandscapeVisibilityMask", "/Script/Landscape.MaterialExpressionLandscapeVisibilityMask");
+	CLASS_REDIRECT("MaterialExpressionTerrainLayerCoords", "/Script/Landscape.MaterialExpressionLandscapeLayerCoords");
+	CLASS_REDIRECT("MaterialExpressionTerrainLayerSwitch", "/Script/Landscape.MaterialExpressionLandscapeLayerSwitch");
+	CLASS_REDIRECT("MaterialExpressionTerrainLayerWeight", "/Script/Landscape.MaterialExpressionLandscapeLayerWeight");
+	CLASS_REDIRECT("ReverbVolume", "AudioVolume");
+	CLASS_REDIRECT("ReverbVolumeToggleable", "AudioVolume");
+
+	STRUCT_REDIRECT("KeyboardEvent", "KeyEvent");
+	STRUCT_REDIRECT("KeyboardFocusEvent", "FocusEvent");
+
+	FUNCTION_REDIRECT("Actor.AddActorLocalOffset", "Actor.K2_AddActorLocalOffset");
+	FUNCTION_REDIRECT("Actor.AddActorLocalRotation", "Actor.K2_AddActorLocalRotation");
+	FUNCTION_REDIRECT("Actor.AddActorLocalTransform", "Actor.K2_AddActorLocalTransform");
+	FUNCTION_REDIRECT("Actor.AddActorLocalTranslation", "Actor.K2_AddActorLocalOffset");
+	PROPERTY_REDIRECT("Actor.K2_AddActorLocalOffset.DeltaTranslation", "DeltaLocation");
+	FUNCTION_REDIRECT("Actor.AddActorWorldOffset", "Actor.K2_AddActorWorldOffset");
+	FUNCTION_REDIRECT("Actor.AddActorWorldRotation", "Actor.K2_AddActorWorldRotation");
+	FUNCTION_REDIRECT("Actor.AddActorWorldTransform", "Actor.K2_AddActorWorldTransform");
+	FUNCTION_REDIRECT("Actor.SetActorLocation", "Actor.K2_SetActorLocation");
+	FUNCTION_REDIRECT("Actor.SetActorLocationAndRotation", "Actor.K2_SetActorLocationAndRotation");
+	FUNCTION_REDIRECT("Actor.SetActorRelativeLocation", "Actor.K2_SetActorRelativeLocation");
+	PROPERTY_REDIRECT("Actor.K2_SetActorRelativeLocation.NewRelativeTranslation", "NewRelativeLocation");
+	FUNCTION_REDIRECT("Actor.SetActorRelativeRotation", "Actor.K2_SetActorRelativeRotation");
+	FUNCTION_REDIRECT("Actor.SetActorRelativeTransform", "Actor.K2_SetActorRelativeTransform");
+	FUNCTION_REDIRECT("Actor.SetActorRelativeTranslation", "Actor.K2_SetActorRelativeLocation");
+	FUNCTION_REDIRECT("Actor.SetActorTransform", "Actor.K2_SetActorTransform");
+	FUNCTION_REDIRECT("BTFunctionLibrary.GetBlackboard", "BTFunctionLibrary.GetOwnersBlackboard");
+	FUNCTION_REDIRECT("KismetMathLibrary.NearlyEqual_RotatorRotator", "EqualEqual_RotatorRotator");
+	FUNCTION_REDIRECT("KismetMathLibrary.NearlyEqual_VectorVector", "EqualEqual_VectorVector");
+	FUNCTION_REDIRECT("KismetMathLibrary.ProjectOnTo", "ProjectVectorOnToVector");
+	PROPERTY_REDIRECT("KismetMathLibrary.ProjectVectorOnToVector.X", "V");
+	PROPERTY_REDIRECT("KismetMathLibrary.ProjectVectorOnToVector.Y", "Target");
+	PROPERTY_REDIRECT("LightComponent.SetIntensity.NewBrightness", "NewIntensity");
+	FUNCTION_REDIRECT("SceneComponent.AddLocalOffset", "SceneComponent.K2_AddLocalOffset");
+	FUNCTION_REDIRECT("SceneComponent.AddLocalRotation", "SceneComponent.K2_AddLocalRotation");
+	FUNCTION_REDIRECT("SceneComponent.AddLocalTransform", "SceneComponent.K2_AddLocalTransform");
+	FUNCTION_REDIRECT("SceneComponent.AddLocalTranslation", "SceneComponent.K2_AddLocalOffset");
+	PROPERTY_REDIRECT("SceneComponent.K2_AddLocalOffset.DeltaTranslation", "DeltaLocation");
+	FUNCTION_REDIRECT("SceneComponent.AddRelativeLocation", "SceneComponent.K2_AddRelativeLocation");
+	PROPERTY_REDIRECT("SceneComponent.K2_AddRelativeLocation.DeltaTranslation", "DeltaLocation");
+	FUNCTION_REDIRECT("SceneComponent.AddRelativeRotation", "SceneComponent.K2_AddRelativeRotation");
+	FUNCTION_REDIRECT("SceneComponent.AddRelativeTranslation", "SceneComponent.K2_AddRelativeLocation");
+	FUNCTION_REDIRECT("SceneComponent.AddWorldOffset", "SceneComponent.K2_AddWorldOffset");
+	FUNCTION_REDIRECT("SceneComponent.AddWorldRotation", "SceneComponent.K2_AddWorldRotation");
+	FUNCTION_REDIRECT("SceneComponent.AddWorldTransform", "SceneComponent.K2_AddWorldTransform");
+	FUNCTION_REDIRECT("SceneComponent.SetRelativeLocation", "SceneComponent.K2_SetRelativeLocation");
+	PROPERTY_REDIRECT("SceneComponent.K2_SetRelativeLocation.NewTranslation", "NewLocation");
+	FUNCTION_REDIRECT("SceneComponent.SetRelativeRotation", "SceneComponent.K2_SetRelativeRotation");
+	FUNCTION_REDIRECT("SceneComponent.SetRelativeTransform", "SceneComponent.K2_SetRelativeTransform");
+	FUNCTION_REDIRECT("SceneComponent.SetRelativeTranslation", "SceneComponent.K2_SetRelativeLocation");
+	FUNCTION_REDIRECT("SceneComponent.SetWorldLocation", "SceneComponent.K2_SetWorldLocation");
+	PROPERTY_REDIRECT("SceneComponent.K2_SetWorldLocation.NewTranslation", "NewLocation");
+	FUNCTION_REDIRECT("SceneComponent.SetWorldRotation", "SceneComponent.K2_SetWorldRotation");
+	FUNCTION_REDIRECT("SceneComponent.SetWorldTransform", "SceneComponent.K2_SetWorldTransform");
+	FUNCTION_REDIRECT("SceneComponent.SetWorldTranslation", "SceneComponent.K2_SetWorldLocation");
+	PROPERTY_REDIRECT("SkyLightComponent.SetIntensity.NewBrightness", "NewIntensity");
+}
+
+static void RegisterNativeRedirects49(TArray<FCoreRedirect>& Redirects)
+{
+	// 4.7
+
+	CLASS_REDIRECT("EdGraphNode_Comment", "/Script/UnrealEd.EdGraphNode_Comment");
+	CLASS_REDIRECT("K2Node_Comment", "/Script/UnrealEd.EdGraphNode_Comment");
+	CLASS_REDIRECT("VimBlueprintFactory", "AnimBlueprintFactory");
+
+	FUNCTION_REDIRECT("Actor.SetTickEnabled", "Actor.SetActorTickEnabled");
+	PROPERTY_REDIRECT("UserWidget.OnKeyboardFocusLost.InKeyboardFocusEvent", "InFocusEvent");
+	PROPERTY_REDIRECT("UserWidget.OnControllerAnalogValueChanged.ControllerEvent", "InAnalogInputEvent");
+	PROPERTY_REDIRECT("UserWidget.OnControllerButtonPressed.ControllerEvent", "InKeyEvent");
+	PROPERTY_REDIRECT("UserWidget.OnControllerButtonReleased.ControllerEvent", "InKeyEvent");
+	PROPERTY_REDIRECT("UserWidget.OnKeyDown.InKeyboardEvent", "InKeyEvent");
+	PROPERTY_REDIRECT("UserWidget.OnKeyUp.InKeyboardEvent", "InKeyEvent");
+	PROPERTY_REDIRECT("UserWidget.OnKeyboardFocusReceived.InKeyboardFocusEvent", "InFocusEvent");
+	PROPERTY_REDIRECT("UserWidget.OnPreviewKeyDown.InKeyboardEvent", "InKeyEvent");
+	
+	PROPERTY_REDIRECT("MeshComponent.Materials", "MeshComponent.OverrideMaterials");
+	PROPERTY_REDIRECT("Pawn.AutoPossess", "Pawn.AutoPossessPlayer");
+
+	FCoreRedirect* ECollisionChannel = ENUM_REDIRECT("ECollisionChannel", "ECollisionChannel");
+	ECollisionChannel->ValueChanges.Add(TEXT("ECC_Default"), TEXT("ECC_Visibility"));
+	ECollisionChannel->ValueChanges.Add(TEXT("ECC_Dynamic"), TEXT("ECC_WorldDynamic"));
+	ECollisionChannel->ValueChanges.Add(TEXT("ECC_OverlapAll"), TEXT("ECC_OverlapAll_Deprecated"));
+	ECollisionChannel->ValueChanges.Add(TEXT("ECC_OverlapAllDynamic"), TEXT("ECC_OverlapAll_Deprecated"));
+	ECollisionChannel->ValueChanges.Add(TEXT("ECC_OverlapAllDynamic_Deprecated"), TEXT("ECC_OverlapAll_Deprecated"));
+	ECollisionChannel->ValueChanges.Add(TEXT("ECC_OverlapAllStatic"), TEXT("ECC_OverlapAll_Deprecated"));
+	ECollisionChannel->ValueChanges.Add(TEXT("ECC_OverlapAllStatic_Deprecated"), TEXT("ECC_OverlapAll_Deprecated"));
+	ECollisionChannel->ValueChanges.Add(TEXT("ECC_PawnMovement"), TEXT("ECC_Pawn"));
+	ECollisionChannel->ValueChanges.Add(TEXT("ECC_RigidBody"), TEXT("ECC_PhysicsBody"));
+	ECollisionChannel->ValueChanges.Add(TEXT("ECC_RigidBodyInteractable"), TEXT("ECC_PhysicsBody"));
+	ECollisionChannel->ValueChanges.Add(TEXT("ECC_TouchAll"), TEXT("ECC_OverlapAll_Deprecated"));
+	ECollisionChannel->ValueChanges.Add(TEXT("ECC_TouchAllDynamic"), TEXT("ECC_OverlapAll_Deprecated"));
+	ECollisionChannel->ValueChanges.Add(TEXT("ECC_TouchAllStatic"), TEXT("ECC_OverlapAll_Deprecated"));
+	ECollisionChannel->ValueChanges.Add(TEXT("ECC_VehicleMovement"), TEXT("ECC_Vehicle"));
+	ECollisionChannel->ValueChanges.Add(TEXT("ECC_WorldTrace"), TEXT("ECC_WorldStatic"));
+
+	// 4.8
+
+	CLASS_REDIRECT("EditorGameAgnosticSettings", "/Script/UnrealEd.EditorSettings");
+	CLASS_REDIRECT("FoliageType", "/Script/Foliage.FoliageType");
+	CLASS_REDIRECT("FoliageType_InstancedStaticMesh", "/Script/Foliage.FoliageType_InstancedStaticMesh");
+	CLASS_REDIRECT("FoliageVertexColorMask", "/Script/Foliage.FoliageVertexColorMask");
+	CLASS_REDIRECT("InstancedFoliageActor", "/Script/Foliage.InstancedFoliageActor");
+	CLASS_REDIRECT("InstancedFoliageSettings", "/Script/Foliage.FoliageType_InstancedStaticMesh");
+	CLASS_REDIRECT("InteractiveFoliageComponent", "/Script/Foliage.InteractiveFoliageComponent");
+	CLASS_REDIRECT("ProceduralFoliage", "/Script/Foliage.ProceduralFoliageSpawner");
+	CLASS_REDIRECT("ProceduralFoliageActor", "/Script/Foliage.ProceduralFoliageVolume");
+	
+	STRUCT_REDIRECT("ProceduralFoliageTypeData", "/Script/Foliage.FoliageTypeObject");
+
+	FCoreRedirect* EComponentCreationMethod = ENUM_REDIRECT("EComponentCreationMethod", "EComponentCreationMethod");
+	EComponentCreationMethod->ValueChanges.Add(TEXT("EComponentCreationMethod::ConstructionScript"), TEXT("EComponentCreationMethod::SimpleConstructionScript"));
+
+	FCoreRedirect* EConstraintTransform = ENUM_REDIRECT("EConstraintTransform", "EConstraintTransform");
+	EConstraintTransform->ValueChanges.Add(TEXT("EConstraintTransform::Absoluate"), TEXT("EConstraintTransform::Absolute"));
+
+	FCoreRedirect* ELockedAxis = ENUM_REDIRECT("ELockedAxis", "EDOFMode");
+	ELockedAxis->ValueChanges.Add(TEXT("Custom"), TEXT("EDOFMode::CustomPlane"));
+	ELockedAxis->ValueChanges.Add(TEXT("X"), TEXT("EDOFMode::YZPlane"));
+	ELockedAxis->ValueChanges.Add(TEXT("Y"), TEXT("EDOFMode::XZPlane"));
+	ELockedAxis->ValueChanges.Add(TEXT("Z"), TEXT("EDOFMode::XYPlane"));
+
+	FCoreRedirect* EEndPlayReason = ENUM_REDIRECT("EEndPlayReason", "EEndPlayReason");
+	EEndPlayReason->ValueChanges.Add(TEXT("EEndPlayReason::ActorDestroyed"), TEXT("EEndPlayReason::Destroyed"));
+
+	FUNCTION_REDIRECT("ActorComponent.ReceiveInitializeComponent", "ActorComponent.ReceiveBeginPlay");
+	FUNCTION_REDIRECT("ActorComponent.ReceiveUninitializeComponent", "ActorComponent.ReceiveEndPlay");
+
+	PROPERTY_REDIRECT("CameraComponent.bUseControllerViewRotation", "CameraComponent.bUsePawnControlRotation");
+	PROPERTY_REDIRECT("CameraComponent.bUsePawnViewRotation", "CameraComponent.bUsePawnControlRotation");
+	PROPERTY_REDIRECT("CharacterMovementComponent.AirSpeed", "CharacterMovementComponent.MaxFlySpeed");
+	PROPERTY_REDIRECT("CharacterMovementComponent.CrouchedSpeedPercent", "CharacterMovementComponent.CrouchedSpeedMultiplier");
+	PROPERTY_REDIRECT("CharacterMovementComponent.GroundSpeed", "CharacterMovementComponent.MaxWalkSpeed");
+	PROPERTY_REDIRECT("CharacterMovementComponent.JumpZ", "CharacterMovementComponent.JumpZVelocity");
+	PROPERTY_REDIRECT("CharacterMovementComponent.WaterSpeed", "CharacterMovementComponent.MaxSwimSpeed");
+	PROPERTY_REDIRECT("CharacterMovementComponent.bCrouchMovesCharacterDown", "CharacterMovementComponent.bCrouchMaintainsBaseLocation");
+	PROPERTY_REDIRECT("CharacterMovementComponent.bOrientToMovement", "CharacterMovementComponent.bOrientRotationToMovement");
+	PROPERTY_REDIRECT("FunctionalTest.GetAdditionalTestFinishedMessage", "FunctionalTest.OnAdditionalTestFinishedMessageRequest");
+	PROPERTY_REDIRECT("FunctionalTest.WantsToRunAgain", "FunctionalTest.OnWantsReRunCheck");
+	PROPERTY_REDIRECT("ProjectileMovementComponent.Speed", "ProjectileMovementComponent.InitialSpeed");
+	PROPERTY_REDIRECT("SpringArmComponent.bUseControllerViewRotation", "SpringArmComponent.bUsePawnControlRotation");
+	PROPERTY_REDIRECT("SpringArmComponent.bUsePawnViewRotation", "SpringArmComponent.bUsePawnControlRotation");
+	PROPERTY_REDIRECT("BodyInstance.CustomLockedAxis", "BodyInstance.CustomDOFPlaneNormal");
+	PROPERTY_REDIRECT("BodyInstance.LockedAxisMode", "BodyInstance.DOFMode");
+	PROPERTY_REDIRECT("CharacterMovementComponent.NavMeshProjectionCapsuleHeightScaleDown", "CharacterMovementComponent.NavMeshProjectionHeightScaleDown");
+	PROPERTY_REDIRECT("CharacterMovementComponent.NavMeshProjectionCapsuleHeightScaleUp", "CharacterMovementComponent.NavMeshProjectionHeightScaleUp");
+	PROPERTY_REDIRECT("LandscapeSplineControlPoint.MeshComponent", "LandscapeSplineControlPoint.LocalMeshComponent");
+	PROPERTY_REDIRECT("LandscapeSplineSegment.MeshComponents", "LandscapeSplineSegment.LocalMeshComponents");
+	PROPERTY_REDIRECT("ProceduralFoliageComponent.Overlap", "ProceduralFoliageComponent.TileOverlap");
+	PROPERTY_REDIRECT("ProceduralFoliageComponent.ProceduralFoliage", "ProceduralFoliageComponent.FoliageSpawner");
+	PROPERTY_REDIRECT("ProceduralFoliageSpawner.Types", "ProceduralFoliageSpawner.FoliageTypes");
+	PROPERTY_REDIRECT("SpriteGeometryCollection.Polygons", "SpriteGeometryCollection.Shapes");
+
+	// 4.9
+
+	CLASS_REDIRECT("EditorUserSettings", "/Script/UnrealEd.EditorPerProjectUserSettings");	
+	CLASS_REDIRECT("MovieScene", "/Script/MovieScene.MovieScene");
+	CLASS_REDIRECT("MovieScene3DTransformSection", "/Script/MovieSceneTracks.MovieScene3DTransformSection");
+	CLASS_REDIRECT("MovieScene3DTransformTrack", "/Script/MovieSceneTracks.MovieScene3DTransformTrack");
+	CLASS_REDIRECT("MovieSceneAnimationSection", "/Script/MovieSceneTracks.MovieSceneAnimationSection");
+	CLASS_REDIRECT("MovieSceneAnimationTrack", "/Script/MovieSceneTracks.MovieSceneAnimationTrack");
+	CLASS_REDIRECT("MovieSceneAudioSection", "/Script/MovieSceneTracks.MovieSceneAudioSection");
+	CLASS_REDIRECT("MovieSceneAudioTrack", "/Script/MovieSceneTracks.MovieSceneAudioTrack");
+	CLASS_REDIRECT("MovieSceneBindings", "/Script/MovieScene.MovieSceneBindings");
+	CLASS_REDIRECT("MovieSceneBoolSection", "/Script/MovieSceneTracks.MovieSceneBoolSection");
+	CLASS_REDIRECT("MovieSceneBoolTrack", "/Script/MovieSceneTracks.MovieSceneBoolTrack");
+	CLASS_REDIRECT("MovieSceneByteSection", "/Script/MovieSceneTracks.MovieSceneByteSection");
+	CLASS_REDIRECT("MovieSceneByteTrack", "/Script/MovieSceneTracks.MovieSceneByteTrack");
+	CLASS_REDIRECT("MovieSceneColorSection", "/Script/MovieSceneTracks.MovieSceneColorSection");
+	CLASS_REDIRECT("MovieSceneColorTrack", "/Script/MovieSceneTracks.MovieSceneColorTrack");
+	CLASS_REDIRECT("MovieSceneDirectorTrack", "/Script/MovieSceneTracks.MovieSceneDirectorTrack");
+	CLASS_REDIRECT("MovieSceneFloatSection", "/Script/MovieSceneTracks.MovieSceneFloatSection");
+	CLASS_REDIRECT("MovieSceneFloatTrack", "/Script/MovieSceneTracks.MovieSceneFloatTrack");
+	CLASS_REDIRECT("MovieSceneParticleSection", "/Script/MovieSceneTracks.MovieSceneParticleSection");
+	CLASS_REDIRECT("MovieSceneParticleTrack", "/Script/MovieSceneTracks.MovieSceneParticleTrack");
+	CLASS_REDIRECT("MovieScenePropertyTrack", "/Script/MovieScene.MovieScenePropertyTrack");
+	CLASS_REDIRECT("MovieSceneSection", "/Script/MovieScene.MovieSceneSection");
+	CLASS_REDIRECT("MovieSceneTrack", "/Script/MovieScene.MovieSceneTrack");
+	CLASS_REDIRECT("MovieSceneVectorSection", "/Script/MovieSceneTracks.MovieSceneVectorSection");
+	CLASS_REDIRECT("MovieSceneVectorTrack", "/Script/MovieSceneTracks.MovieSceneVectorTrack");
+	CLASS_REDIRECT("RuntimeMovieScenePlayer", "/Script/MovieScene.RuntimeMovieScenePlayer");
+	CLASS_REDIRECT("SubMovieSceneSection", "/Script/MovieSceneTracks.SubMovieSceneSection");
+	CLASS_REDIRECT("SubMovieSceneTrack", "/Script/MovieSceneTracks.SubMovieSceneTrack");
+
+	PACKAGE_REDIRECT("/Script/MovieSceneCore", "/Script/MovieScene");
+	PACKAGE_REDIRECT("/Script/MovieSceneCoreTypes", "/Script/MovieSceneTracks");
+
+	STRUCT_REDIRECT("Anchors", "/Script/Slate.Anchors");
+	STRUCT_REDIRECT("AnimNode_BoneDrivenController", "/Script/AnimGraphRuntime.AnimNode_BoneDrivenController");
+	STRUCT_REDIRECT("AnimNode_CopyBone", "/Script/AnimGraphRuntime.AnimNode_CopyBone");
+	STRUCT_REDIRECT("AnimNode_HandIKRetargeting", "/Script/AnimGraphRuntime.AnimNode_HandIKRetargeting");
+	STRUCT_REDIRECT("AnimNode_LookAt", "/Script/AnimGraphRuntime.AnimNode_LookAt");
+	STRUCT_REDIRECT("AnimNode_ModifyBone", "/Script/AnimGraphRuntime.AnimNode_ModifyBone");
+	STRUCT_REDIRECT("AnimNode_RotationMultiplier", "/Script/AnimGraphRuntime.AnimNode_RotationMultiplier");
+	STRUCT_REDIRECT("AnimNode_SkeletalControlBase", "/Script/AnimGraphRuntime.AnimNode_SkeletalControlBase");
+	STRUCT_REDIRECT("AnimNode_SpringBone", "/Script/AnimGraphRuntime.AnimNode_SpringBone");
+	STRUCT_REDIRECT("AnimNode_Trail", "/Script/AnimGraphRuntime.AnimNode_Trail");
+	STRUCT_REDIRECT("AnimNode_TwoBoneIK", "/Script/AnimGraphRuntime.AnimNode_TwoBoneIK");
+	STRUCT_REDIRECT("MovieSceneBoundObject", "/Script/MovieScene.MovieSceneBoundObject");
+	STRUCT_REDIRECT("MovieSceneEditorData", "/Script/MovieScene.MovieSceneEditorData");
+	STRUCT_REDIRECT("MovieSceneObjectBinding", "/Script/MovieScene.MovieSceneBinding");
+	STRUCT_REDIRECT("MovieScenePossessable", "/Script/MovieScene.MovieScenePossessable");
+	STRUCT_REDIRECT("MovieSceneSpawnable", "/Script/MovieScene.MovieSceneSpawnable");
+	STRUCT_REDIRECT("SpritePolygon", "SpriteGeometryShape");
+	STRUCT_REDIRECT("SpritePolygonCollection", "SpriteGeometryCollection");
+
+	FUNCTION_REDIRECT("GameplayStatics.PlayDialogueAttached", "GameplayStatics.SpawnDialogueAttached");
+	FUNCTION_REDIRECT("GameplayStatics.PlaySoundAttached", "GameplayStatics.SpawnSoundAttached");
+	FUNCTION_REDIRECT("KismetMathLibrary.BreakRot", "KismetMathLibrary.BreakRotator");
+	FUNCTION_REDIRECT("KismetMathLibrary.MakeRot", "KismetMathLibrary.MakeRotator");
+	FUNCTION_REDIRECT("KismetMathLibrary.MapRange", "KismetMathLibrary.MapRangeUnclamped");
+	FUNCTION_REDIRECT("PrimitiveComponent.GetMoveIgnoreActors", "PrimitiveComponent.CopyArrayOfMoveIgnoreActors");
+	FUNCTION_REDIRECT("SplineComponent.GetNumSplinePoints", "SplineComponent.GetNumberOfSplinePoints");
+	FUNCTION_REDIRECT("VerticalBox.AddChildVerticalBox", "VerticalBox.AddChildToVerticalBox");
+	
+	PROPERTY_REDIRECT("ComponentKey.VariableGuid", "ComponentKey.AssociatedGuid");
+	PROPERTY_REDIRECT("ComponentKey.VariableName", "ComponentKey.SCSVariableName");
+	PROPERTY_REDIRECT("FoliageType.InitialMaxAge", "FoliageType.MaxInitialAge");
+	PROPERTY_REDIRECT("FoliageType.bGrowsInShade", "FoliageType.bSpawnsInShade");
+	PROPERTY_REDIRECT("MemberReference.MemberParentClass", "MemberReference.MemberParent");
+	PROPERTY_REDIRECT("SimpleMemberReference.MemberParentClass", "SimpleMemberReference.MemberParent");
+	PROPERTY_REDIRECT("SoundNodeModPlayer.SoundMod", "SoundNodeModPlayer.SoundModAssetPtr");
+	PROPERTY_REDIRECT("SoundNodeWavePlayer.SoundWave", "SoundNodeWavePlayer.SoundWaveAssetPtr");
+
+	ENUM_REDIRECT("ECheckBoxState", "/Script/SlateCore.ECheckBoxState");
+	ENUM_REDIRECT("ESlateCheckBoxState", "/Script/SlateCore.ECheckBoxState");
+	ENUM_REDIRECT("EAxisOption", "/Script/AnimGraphRuntime.EAxisOption");
+	ENUM_REDIRECT("EBoneAxis", "/Script/AnimGraphRuntime.EBoneAxis");
+	ENUM_REDIRECT("EBoneModificationMode", "/Script/AnimGraphRuntime.EBoneModificationMode");
+	ENUM_REDIRECT("EComponentType", "/Script/AnimGraphRuntime.EComponentType");
+	ENUM_REDIRECT("EInterpolationBlend", "/Script/AnimGraphRuntime.EInterpolationBlend");
+}
+
+PRAGMA_ENABLE_OPTIMIZATION
+
+#endif //WITH_EDITORONLY_DATA
+
+void FCoreRedirects::RegisterNativeRedirects()
+{
+#if WITH_EDITORONLY_DATA
+
+	// Registering redirects here instead of in baseengine.ini is faster to parse and can clean up the ini, but is not required
+	TArray<FCoreRedirect> Redirects;
+
+	RegisterNativeRedirects40(Redirects);
+	RegisterNativeRedirects46(Redirects);
+	RegisterNativeRedirects49(Redirects);
+
+	// 4.10 and later are in baseengine.ini
+
+	AddRedirectList(Redirects, TEXT("RegisterNativeRedirects"));
+
+#endif //WITH_EDITORONLY_DATA
+}
\ No newline at end of file
