Copying //UE4/Dev-Core to //UE4/Dev-Main (Source: //UE4/Dev-Core @ 3283640)

diff --git a/Engine/Source/Runtime/CoreUObject/Private/UObject/Enum.cpp b/Engine/Source/Runtime/CoreUObject/Private/UObject/Enum.cpp
--- a/Engine/Source/Runtime/CoreUObject/Private/UObject/Enum.cpp
+++ b/Engine/Source/Runtime/CoreUObject/Private/UObject/Enum.cpp
@@ -11,6 +11,7 @@
 #include "UObject/UObjectThreadContext.h"
 #include "UObject/DevObjectVersion.h"
 #include "UObject/CoreObjectVersion.h"
+#include "UObject/CoreRedirects.h"
 
 DEFINE_LOG_CATEGORY_STATIC(LogEnum, Log, All);
 
@@ -19,8 +20,6 @@ DEFINE_LOG_CATEGORY_STATIC(LogEnum, Log, All);
 -----------------------------------------------------------------------------*/
 
 TMap<FName, UEnum*> UEnum::AllEnumNames;
-TMap<FName,TMap<FName,FName> > UEnum::EnumRedirects;
-TMap<FName,TMap<FString,FString> > UEnum::EnumSubstringRedirects;
 
 UEnum::UEnum(const FObjectInitializer& ObjectInitializer)
 	: UField(ObjectInitializer)
@@ -445,80 +444,42 @@ int32 UEnum::FindEnumRedirects(const UEnum* Enum, FName EnumEntryName)
 {
 	check (Enum);
 
-	// Init redirect map from ini file
-	static bool bAlreadyInitialized_EnumRedirectsMap = false;
-	if( !bAlreadyInitialized_EnumRedirectsMap )
-	{
-		InitEnumRedirectsMap();
-		bAlreadyInitialized_EnumRedirectsMap = true;
-	}
+	const TMap<FString, FString>* ValueChanges = FCoreRedirects::GetValueRedirects(ECoreRedirectFlags::Type_Enum, Enum);
 
-	// See if we have an entry for this enum
-	const TMap<FName, FName>* ThisEnumRedirects = EnumRedirects.Find(Enum->GetFName());
-	if (ThisEnumRedirects != nullptr)
+	if (ValueChanges)
 	{
 		// first look for default name
-		const FName* NewEntryName = ThisEnumRedirects->Find(EnumEntryName);
-		if (NewEntryName != nullptr)
+		const FString EnumEntryString = EnumEntryName.ToString();
+		const FString* NewEntryString = ValueChanges->Find(EnumEntryString);
+		if (NewEntryString != nullptr)
 		{
-			return Enum->GetIndexByName(*NewEntryName);
+			return Enum->GetIndexByName(FName(**NewEntryString));
 		}
 		// if not found, look for long name if short or short name if long
 		else
 		{
 			// Note: This can pollute the name table if the ini is set up incorrectly
-			const FString EnumEntryString = EnumEntryName.ToString();
-			const int32 DoubleColonIndex = EnumEntryString.Find(TEXT("::"),ESearchCase::CaseSensitive);
-			FName ModifiedName;
+			const int32 DoubleColonIndex = EnumEntryString.Find(TEXT("::"), ESearchCase::CaseSensitive);
+			FString ModifiedName;
 			if (DoubleColonIndex == INDEX_NONE)
 			{
-				ModifiedName = FName(*GenerateFullEnumName(Enum, *EnumEntryName.ToString()));
+				ModifiedName = GenerateFullEnumName(Enum, *EnumEntryName.ToString());
 			}
 			else
 			{
-				ModifiedName = FName(*EnumEntryString.RightChop(DoubleColonIndex+2));
-			}
-			NewEntryName = ThisEnumRedirects->Find(ModifiedName);
-			if (NewEntryName != nullptr)
-			{
-				return Enum->GetIndexByName(*NewEntryName);
+				ModifiedName = EnumEntryString.RightChop(DoubleColonIndex + 2);
 			}
-		}
-	}
-
-	// Now try the substring replacement
-	const TMap<FString, FString>* ThisEnumSubstringRedirects = EnumSubstringRedirects.Find(Enum->GetFName());
-	if (ThisEnumSubstringRedirects != NULL)
-	{
-		FString EntryString = EnumEntryName.ToString();
-		for (TMap<FString, FString>::TConstIterator It(*ThisEnumSubstringRedirects); It; ++It)
-		{
-			if (EntryString.Contains(*It.Key()))
+			NewEntryString = ValueChanges->Find(ModifiedName);
+			if (NewEntryString != nullptr)
 			{
-				// Note: This can pollute the name table if the ini is set up incorrectly
-				FName NewName = FName(*EntryString.Replace(*It.Key(), *It.Value()));
-
-				int32 FoundIndex = Enum->GetIndexByName(NewName);
-
-				if (FoundIndex == INDEX_NONE)
-				{
-					UE_LOG(LogEnum, Warning, TEXT("Matched substring redirect %s in Enum %s failed to resolve %s to valid enum"), *It.Key(), *Enum->GetName(), *NewName.ToString());
-				}
-				else
-				{
-					return FoundIndex;
-				}
+				return Enum->GetIndexByName(FName(**NewEntryString));
 			}
 		}
 	}
 
 	return INDEX_NONE;
 }
-/**
- * Sets the array of enums.
- *
- * @return	true unless the MAX enum already exists.
- */
+
 bool UEnum::SetEnums(TArray<TPair<FName, int64>>& InNames, UEnum::ECppForm InCppForm, bool bAddMaxKeyIfMissing)
 {
 	if (Names.Num() > 0)
@@ -785,42 +746,3 @@ IMPLEMENT_CORE_INTRINSIC_CLASS(UEnum, UField,
 	{
 	}
 );
-
-void UEnum::InitEnumRedirectsMap()
-{
-	if( GConfig )
-	{
-		FConfigSection* EngineSection = GConfig->GetSectionPrivate( TEXT("/Script/Engine.Engine"), false, true, GEngineIni );
-		for( FConfigSection::TIterator It(*EngineSection); It; ++It )
-		{
-			if( It.Key() == TEXT("EnumRedirects") )
-			{
-				const FString& ConfigValue = It.Value().GetValue();
-				FName EnumName = NAME_None;
-				FName OldEnumEntry = NAME_None;
-				FName NewEnumEntry = NAME_None;
-
-				FString OldEnumSubstring;
-				FString NewEnumSubstring;
-
-				FParse::Value( *ConfigValue, TEXT("EnumName="), EnumName );
-				if (FParse::Value( *ConfigValue, TEXT("OldEnumEntry="), OldEnumEntry ) )
-				{
-					FParse::Value( *ConfigValue, TEXT("NewEnumEntry="), NewEnumEntry );
-					check(EnumName != NAME_None && OldEnumEntry != NAME_None && NewEnumEntry != NAME_None );
-					EnumRedirects.FindOrAdd(EnumName).Add(OldEnumEntry, NewEnumEntry);
-				}
-				else if (FParse::Value( *ConfigValue, TEXT("OldEnumSubstring="), OldEnumSubstring ))
-				{
-					FParse::Value( *ConfigValue, TEXT("NewEnumSubstring="), NewEnumSubstring );
-					check(EnumName != NAME_None && !OldEnumSubstring.IsEmpty() && !NewEnumSubstring.IsEmpty());
-					EnumSubstringRedirects.FindOrAdd(EnumName).Add(OldEnumSubstring, NewEnumSubstring);
-				}
-			}			
-		}
-	}
-	else
-	{
-		UE_LOG(LogEnum, Warning, TEXT(" **** ENUM REDIRECTS UNABLE TO INITIALIZE! **** "));
-	}
-}
diff --git a/Engine/Source/Runtime/CoreUObject/Private/UObject/GarbageCollection.cpp b/Engine/Source/Runtime/CoreUObject/Private/UObject/GarbageCollection.cpp
--- a/Engine/Source/Runtime/CoreUObject/Private/UObject/GarbageCollection.cpp
+++ b/Engine/Source/Runtime/CoreUObject/Private/UObject/GarbageCollection.cpp
@@ -1933,6 +1933,8 @@ void UClass::AssembleReferenceTokenStream(bool bForce)
 		}
 	} ReferenceTokenStreamLock(ReferenceTokenStreamCritical, !(ClassFlags & CLASS_Native));
 
+	UE_CLOG(!IsInGameThread() && !IsGarbageCollectionLocked(), LogGarbage, Fatal, TEXT("AssembleReferenceTokenStream for %s called on a non-game thread while GC is not locked."), *GetFullName());
+
 	if (!HasAnyClassFlags(CLASS_TokenStreamAssembled) || bForce)
 	{
 		if (bForce)
diff --git a/Engine/Source/Runtime/CoreUObject/Private/UObject/LinkerLoad.cpp b/Engine/Source/Runtime/CoreUObject/Private/UObject/LinkerLoad.cpp
--- a/Engine/Source/Runtime/CoreUObject/Private/UObject/LinkerLoad.cpp
+++ b/Engine/Source/Runtime/CoreUObject/Private/UObject/LinkerLoad.cpp
@@ -33,6 +33,7 @@
 #include "HAL/ThreadHeartBeat.h"
 #include "Serialization/BulkData.h"
 #include "Serialization/AsyncLoadingPrivate.h"
+#include "UObject/CoreRedirects.h"
 
 class FTexture2DResourceMem;
 
@@ -66,35 +67,6 @@ FName FLinkerLoad::NAME_LoadErrors("LoadErrors");
 
 TMap<FString, FLinkerLoad::FPackagePrecacheInfo> FLinkerLoad::PackagePrecacheMap;
 
-/**
- * Here is the format for the ClassRedirection:
- * 
- *  ; Basic redirects
- *  ;ActiveClassRedirects=(OldClassName="MyClass",NewClassName="NewNativePackage.MyClass")
- *	ActiveClassRedirects=(OldClassName="CylinderComponent",NewClassName="CapsuleComponent")
- *  Note: For class name redirects, the OldClassName must be the plain OldClassName, it cannot be OldPackage.OldClassName
- *
- *	; Keep both classes around, but convert any existing instances of that object to a particular class (insert into the inheritance hierarchy
- *	;ActiveClassRedirects=(OldClassName="MyClass",NewClassName="MyClassParent",InstanceOnly="true")
- *
- */
-
-
-TMap<FName, FName> FLinkerLoad::ObjectNameRedirects;			    // OldClassName to NewClassName for ImportMap
-TMap<FName, TPair<FName, FName>> FLinkerLoad::ObjectNameClassRedirects; // OldClassName to NewClassClass and NewClassPackage for ImportMap
-TMap<FName, FName> FLinkerLoad::ObjectNameRedirectsInstanceOnly;	// OldClassName to NewClassName for ExportMap
-TMap<FName, FName> FLinkerLoad::ObjectNameRedirectsObjectOnly;		// Object name to NewClassName for export map
-TMap<FName, FName> FLinkerLoad::GameNameRedirects;					// Game package name to new game package name
-TMap<FName, FName> FLinkerLoad::StructNameRedirects;				// Old struct name to new struct name mapping
-TMap<FString, FString> FLinkerLoad::PluginNameRedirects;			// Old plugin name to new plugin name mapping
-TMap<FName, FLinkerLoad::FSubobjectRedirect> FLinkerLoad::SubobjectNameRedirects;	
-TSet<FName> FLinkerLoad::KnownMissingPackages;
-bool FLinkerLoad::bActiveRedirectsMapInitialized = false;
-
-void FLinkerLoad::AddGameNameRedirect(const FName OldName, const FName NewName)
-{
-	GameNameRedirects.Add(OldName, NewName);
-}
 
 /*----------------------------------------------------------------------------
 Helpers
@@ -109,11 +81,29 @@ bool FLinkerLoad::IsValidPackageIndex(FPackageIndex InIndex)
 		|| (InIndex.IsExport() && ExportMap.IsValidIndex(InIndex.ToExport()));
 }
 
+bool FLinkerLoad::bActiveRedirectsMapInitialized = false;
+
+
+
 /**
- * Add redirects to FLinkerLoad static map
- */
+* DEPRECATED: Replace with FCoreRedirects format for newly added ini entries
+*
+* Here is the format for the ClassRedirection:
+*
+*  ; Basic redirects
+*  ;ActiveClassRedirects=(OldClassName="MyClass",NewClassName="NewNativePackage.MyClass")
+*	ActiveClassRedirects=(OldClassName="CylinderComponent",NewClassName="CapsuleComponent")
+*  Note: For class name redirects, the OldClassName must be the plain OldClassName, it cannot be OldPackage.OldClassName
+*
+*	; Keep both classes around, but convert any existing instances of that object to a particular class (insert into the inheritance hierarchy
+*	;ActiveClassRedirects=(OldClassName="MyClass",NewClassName="MyClassParent",InstanceOnly="true")
+*
+*/
+
 void FLinkerLoad::CreateActiveRedirectsMap(const FString& GEngineIniName)
 {
+	// Soft deprecated, replaced by FCoreRedirects, but it will still read the old format for the foreseeable future
+
 	// mark that this has been done at least once
 	bActiveRedirectsMapInitialized = true;
 
@@ -122,6 +112,7 @@ void FLinkerLoad::CreateActiveRedirectsMap(const FString& GEngineIniName)
 		FConfigSection* PackageRedirects = GConfig->GetSectionPrivate( TEXT("/Script/Engine.Engine"), false, true, GEngineIniName );
 		if (PackageRedirects)
 		{
+			TArray<FCoreRedirect> NewRedirects;
 			FDeferredMessageLog RedirectErrors(NAME_LoadErrors);
 
 			static FName ActiveClassRedirectsKey(TEXT("ActiveClassRedirects"));
@@ -154,33 +145,26 @@ void FLinkerLoad::CreateActiveRedirectsMap(const FString& GEngineIniName)
 					if (NewSubobjName != NAME_None || OldSubobjName != NAME_None)
 					{
 						check(OldSubobjName != NAME_None && OldClassName != NAME_None );
-
-						if (SubobjectNameRedirects.Contains(OldSubobjName))
-						{
-							RedirectErrors.Error(FText::Format(LOCTEXT("SubobjRedirectorCollision", "{0} contains a collision with multiple redirectors for old suboject named {1} (old class {2})"), FText::FromName(ActiveClassRedirectsKey), FText::FromName(OldSubobjName), FText::FromName(OldClassName)));
-						}
-
-						SubobjectNameRedirects.Add(OldSubobjName, FSubobjectRedirect(OldClassName, NewSubobjName));
+						FCoreRedirect* Redirect = new (NewRedirects) FCoreRedirect(ECoreRedirectFlags::Type_Class, OldClassName.ToString(), OldClassName.ToString());
+						Redirect->ValueChanges.Add(OldSubobjName.ToString(), NewSubobjName.ToString());
 					}
 					//instances only
 					else if( bInstanceOnly )
 					{
-						if (ObjectNameRedirectsInstanceOnly.Contains(OldClassName))
+						// If NewClassName is none, register as removed instead
+						if (NewClassName == NAME_None)
 						{
-							RedirectErrors.Error(FText::Format(LOCTEXT("InstanceOnlyRedirectorCollision", "{0} contains a collision with multiple instance-only redirectors for old class named {1}"), FText::FromName(ActiveClassRedirectsKey), FText::FromName(OldClassName)));
+							FCoreRedirect* Redirect = new (NewRedirects) FCoreRedirect(ECoreRedirectFlags::Type_Class | ECoreRedirectFlags::Option_InstanceOnly | ECoreRedirectFlags::Option_Removed, OldClassName.ToString(), NewClassName.ToString());
+						}
+						else
+						{
+							FCoreRedirect* Redirect = new (NewRedirects) FCoreRedirect(ECoreRedirectFlags::Type_Class | ECoreRedirectFlags::Option_InstanceOnly, OldClassName.ToString(), NewClassName.ToString());
 						}
-
-						ObjectNameRedirectsInstanceOnly.Add(OldClassName,NewClassName);
 					}
 					//objects only on a per-object basis
 					else if( ObjectName != NAME_None )
 					{
-						if (ObjectNameRedirectsObjectOnly.Contains(OldClassName))
-						{
-							RedirectErrors.Error(FText::Format(LOCTEXT("ObjectOnlyRedirectorCollision", "{0} contains a collision with multiple object-only redirectors for old class named {1}"), FText::FromName(ActiveClassRedirectsKey), FText::FromName(OldClassName)));
-						}
-
-						ObjectNameRedirectsObjectOnly.Add(ObjectName, NewClassName);
+						UE_LOG(LogLinker, Warning, TEXT("Generic Object redirects are not supported with ActiveClassRedirects and never worked, move to new CoreRedirects system"));
 					}
 					//full redirect
 					else
@@ -191,16 +175,21 @@ void FLinkerLoad::CreateActiveRedirectsMap(const FString& GEngineIniName)
 						}
 						else
 						{
-							if (ObjectNameRedirects.Contains(OldClassName))
-							{
-								RedirectErrors.Error(FText::Format(LOCTEXT("RedirectorCollision", "{0} contains a collision with multiple redirectors for old class named {1}"), FText::FromName(ActiveClassRedirectsKey), FText::FromName(OldClassName)));
-							}
-
-							ObjectNameRedirects.Add(OldClassName,NewClassName);
+							FCoreRedirect* Redirect = new(NewRedirects) FCoreRedirect(ECoreRedirectFlags::Type_Class, OldClassName.ToString(), NewClassName.ToString());
 
 							if (!NewClassClass.IsNone() || !NewClassPackage.IsNone())
 							{
-								ObjectNameClassRedirects.Add(OldClassName,TPair<FName,FName>(TPairInitializer<FName, FName>(NewClassClass,NewClassPackage)));
+								Redirect->OverrideClassName = FCoreRedirectObjectName(NewClassClass, NAME_None, NewClassPackage);
+							}
+							else if (Redirect->NewName.ObjectName.ToString().StartsWith(TEXT("E"), ESearchCase::CaseSensitive))
+							{
+								// This might be an enum, so we have to register it
+								FCoreRedirect* EnumRedirect = new(NewRedirects) FCoreRedirect(ECoreRedirectFlags::Type_Enum, OldClassName.ToString(), NewClassName.ToString());
+							}
+							else
+							{
+								// This might be a struct redirect because many of them were registered incorrectly
+								FCoreRedirect* StructRedirect = new(NewRedirects) FCoreRedirect(ECoreRedirectFlags::Type_Struct, OldClassName.ToString(), NewClassName.ToString());
 							}
 						}
 					}
@@ -213,7 +202,7 @@ void FLinkerLoad::CreateActiveRedirectsMap(const FString& GEngineIniName)
 					FParse::Value( *It.Value().GetValue(), TEXT("OldGameName="), OldGameName );
 					FParse::Value( *It.Value().GetValue(), TEXT("NewGameName="), NewGameName );
 
-					GameNameRedirects.Add(OldGameName, NewGameName);
+					FCoreRedirect* Redirect = new(NewRedirects) FCoreRedirect(ECoreRedirectFlags::Type_Package, OldGameName.ToString(), NewGameName.ToString());
 				}
 				else if ( It.Key() == TEXT("ActiveStructRedirects") )
 				{
@@ -223,7 +212,7 @@ void FLinkerLoad::CreateActiveRedirectsMap(const FString& GEngineIniName)
 					FParse::Value( *It.Value().GetValue(), TEXT("OldStructName="), OldStructName );
 					FParse::Value( *It.Value().GetValue(), TEXT("NewStructName="), NewStructName );
 
-					StructNameRedirects.Add(OldStructName, NewStructName);
+					FCoreRedirect* Redirect = new(NewRedirects) FCoreRedirect(ECoreRedirectFlags::Type_Struct, OldStructName.ToString(), NewStructName.ToString());
 				}
 				else if ( It.Key() == TEXT("ActivePluginRedirects") )
 				{
@@ -236,7 +225,7 @@ void FLinkerLoad::CreateActiveRedirectsMap(const FString& GEngineIniName)
 					OldPluginName = FString(TEXT("/")) + OldPluginName + FString(TEXT("/"));
 					NewPluginName = FString(TEXT("/")) + NewPluginName + FString(TEXT("/"));
 
-					PluginNameRedirects.Add(OldPluginName, NewPluginName);
+					FCoreRedirect* Redirect = new(NewRedirects) FCoreRedirect(ECoreRedirectFlags::Type_Package | ECoreRedirectFlags::Option_MatchSubstring, OldPluginName, NewPluginName);
 				}
 				else if ( It.Key() == TEXT("KnownMissingPackages") )
 				{
@@ -244,9 +233,48 @@ void FLinkerLoad::CreateActiveRedirectsMap(const FString& GEngineIniName)
 
 					FParse::Value( *It.Value().GetValue(), TEXT("PackageName="), KnownMissingPackage );
 
-					KnownMissingPackages.Add(KnownMissingPackage);
+					FCoreRedirect* Redirect = new(NewRedirects) FCoreRedirect(ECoreRedirectFlags::Type_Package | ECoreRedirectFlags::Option_Removed, KnownMissingPackage.ToString(), FString());
+				}
+				else if (It.Key() == TEXT("TaggedPropertyRedirects"))
+				{
+					FName ClassName = NAME_None;
+					FName OldPropertyName = NAME_None;
+					FName NewPropertyName = NAME_None;
+
+					FParse::Value(*It.Value().GetValue(), TEXT("ClassName="), ClassName);
+					FParse::Value(*It.Value().GetValue(), TEXT("OldPropertyName="), OldPropertyName);
+					FParse::Value(*It.Value().GetValue(), TEXT("NewPropertyName="), NewPropertyName);
+
+					check(ClassName != NAME_None && OldPropertyName != NAME_None && NewPropertyName != NAME_None);
+
+					FCoreRedirect* Redirect = new (NewRedirects) FCoreRedirect(ECoreRedirectFlags::Type_Property, FCoreRedirectObjectName(OldPropertyName, ClassName, NAME_None), FCoreRedirectObjectName(NewPropertyName, ClassName, NAME_None));
+				}
+				else if (It.Key() == TEXT("EnumRedirects"))
+				{
+					const FString& ConfigValue = It.Value().GetValue();
+					FName EnumName = NAME_None;
+					FName OldEnumEntry = NAME_None;
+					FName NewEnumEntry = NAME_None;
+
+					FString OldEnumSubstring;
+					FString NewEnumSubstring;
+
+					FParse::Value(*ConfigValue, TEXT("EnumName="), EnumName);
+					if (FParse::Value(*ConfigValue, TEXT("OldEnumEntry="), OldEnumEntry))
+					{
+						FParse::Value(*ConfigValue, TEXT("NewEnumEntry="), NewEnumEntry);
+						check(EnumName != NAME_None && OldEnumEntry != NAME_None && NewEnumEntry != NAME_None);
+						FCoreRedirect* Redirect = new(NewRedirects) FCoreRedirect(ECoreRedirectFlags::Type_Enum, EnumName.ToString(), EnumName.ToString());
+						Redirect->ValueChanges.Add(OldEnumEntry.ToString(), NewEnumEntry.ToString());
+					}
+					else if (FParse::Value(*ConfigValue, TEXT("OldEnumSubstring="), OldEnumSubstring))
+					{
+						UE_LOG(LogLinker, Warning, TEXT("OldEnumSubstring no longer supported! Replace with multiple entries or use the better syntax in the CoreRedirects section "));
+					}
 				}
 			}
+
+			FCoreRedirects::AddRedirectList(NewRedirects, GEngineIniName);
 		}
 	}
 	else
@@ -255,7 +283,6 @@ void FLinkerLoad::CreateActiveRedirectsMap(const FString& GEngineIniName)
 	}
 }
 
-
 FScopedCreateImportCounter::FScopedCreateImportCounter(FLinkerLoad* Linker, int32 Index)
 {
 	FUObjectThreadContext& ThreadContext = FUObjectThreadContext::Get();
@@ -525,12 +552,17 @@ void FLinkerLoad::InvalidateExport(UObject* OldObject)
 	}
 }
 
-FName FLinkerLoad::FindSubobjectRedirectName(const FName& Name)
+FName FLinkerLoad::FindSubobjectRedirectName(const FName& Name, UClass* Class)
 {
-	FSubobjectRedirect *Redirect = SubobjectNameRedirects.Find(Name);
-	if (Redirect)
+	const TMap<FString, FString>* ValueChanges = FCoreRedirects::GetValueRedirects(ECoreRedirectFlags::Type_Class, Class);
+
+	if (ValueChanges)
 	{
-		return Redirect->NewName;
+		const FString* NewInstanceName = ValueChanges->Find(Name.ToString());
+		if (NewInstanceName)
+		{
+			return FName(**NewInstanceName);
+		}
 	}
 
 	return FName();
@@ -669,12 +701,6 @@ FLinkerLoad::ELinkerStatus FLinkerLoad::Tick( float InTimeLimit, bool bInUseTime
 				Status = FixupImportMap();
 			}
 
-			if ( Status == LINKER_Loaded )
-			{
-				SCOPED_LOADTIMER(LinkerLoad_RemapImports);
-				Status = RemapImports();
-			}
-
 			// Fix up export map for object class conversion 
 			if( Status == LINKER_Loaded )
 			{	
@@ -1484,205 +1510,144 @@ FLinkerLoad::ELinkerStatus FLinkerLoad::FixupImportMap()
 #if WITH_EDITOR
 		LoadProgressScope->EnterProgressFrame(1);
 #endif
-
 		// Fix up imports, not required if everything is cooked.
 		if (!FPlatformProperties::RequiresCookedData())
 		{
-			static const FName NAME_ScriptStruct(TEXT("ScriptStruct"));
 			static const FName NAME_BlueprintGeneratedClass(TEXT("BlueprintGeneratedClass"));
 
 			bool bDone = false;
 			while (!bDone)
 			{
+				TArray<FName> NewPackageImports;
+
 				bDone = true;
 				for( int32 i=0; i<ImportMap.Num(); i++ )
 				{
 					FObjectImport& Import = ImportMap[i];
+
+					// Compute class name first, as instance can override it
+					const FCoreRedirect* ClassValueRedirect = nullptr;
+					FCoreRedirectObjectName OldClassName(Import.ClassName, NAME_None, Import.ClassPackage), NewClassName;
+
+					FCoreRedirects::RedirectNameAndValues(ECoreRedirectFlags::Type_Class, OldClassName, NewClassName, &ClassValueRedirect);
+
+					if (ClassValueRedirect)
 					{
-						FSubobjectRedirect* Redirect = SubobjectNameRedirects.Find(Import.ObjectName);
-						if (Redirect)
+						// Apply class value redirects before other redirects, to mirror old subobject order
+						const FString* NewInstanceName = ClassValueRedirect->ValueChanges.Find(Import.ObjectName.ToString());
+						if (NewInstanceName)
 						{
-							if (Import.ClassName == Redirect->MatchClass)
-							{
-								if (!Import.OuterIndex.IsNull())
-								{
-									FString Was = GetImportFullName(i);
-									Import.ObjectName = Redirect->NewName;
+							// Rename this import directly
+							FString Was = GetImportFullName(i);
+							Import.ObjectName = FName(**NewInstanceName);
 
-									if (Import.ObjectName != NAME_None)
-									{
-										FString Now = GetImportFullName(i);
-										UE_LOG(LogLinker, Verbose, TEXT("FLinkerLoad::FixupImportMap() - Renamed component from %s   to   %s"), *Was, *Now);
-									}
-									else
-									{
-										UE_LOG(LogLinker, Verbose, TEXT("FLinkerLoad::FixupImportMap() - Removed component %s"), *Was);
-									}
-									
-									bDone = false;
-									continue;
-								}
+							if (Import.ObjectName != NAME_None)
+							{
+								FString Now = GetImportFullName(i);
+								UE_LOG(LogLinker, Verbose, TEXT("FLinkerLoad::FixupImportMap() - Renamed object from %s   to   %s"), *Was, *Now);
+							}
+							else
+							{
+								UE_LOG(LogLinker, Verbose, TEXT("FLinkerLoad::FixupImportMap() - Removed object %s"), *Was);
 							}
 						}
 					}
 
+					FCoreRedirectObjectName OldObjectName(GetImportPathName(i)), NewObjectName;
+					ECoreRedirectFlags ObjectRedirectFlags = FCoreRedirects::GetFlagsForTypeName(Import.ClassPackage, Import.ClassName);
+					const FCoreRedirect* ValueRedirect = nullptr;
+					
+					FCoreRedirects::RedirectNameAndValues(ObjectRedirectFlags, OldObjectName, NewObjectName, &ValueRedirect);
 
-					const bool bIsClass = (Import.ClassName == NAME_Class) || (Import.ClassName == NAME_BlueprintGeneratedClass);
-					const bool bIsStruct = (Import.ClassName == NAME_ScriptStruct);
-					const bool bIsEnum = (Import.ClassName == NAME_Enum);
-					const bool bIsClassOrStructOrEnum = bIsClass || bIsStruct || bIsEnum;
-
-					FString RedirectName, ResultPackage, ResultClass;
-					const FName* RedirectNameObj = ObjectNameRedirects.Find(Import.ObjectName);
-					const FName* RedirectNameClass = ObjectNameRedirects.Find(Import.ClassName);
-
-					if (bIsStruct && !RedirectNameObj)
+					if (ValueRedirect && ValueRedirect->OverrideClassName.IsValid())
 					{
-						// Check struct redirects as well
-						RedirectNameObj = StructNameRedirects.Find(Import.ObjectName);
+						// Override class name if found, even if the name didn't actually change
+						NewClassName = ValueRedirect->OverrideClassName;
 					}
 
-					int32 OldOuterIndex = 0;
-					if ( (RedirectNameObj && bIsClassOrStructOrEnum) || RedirectNameClass )
+					if (NewObjectName != OldObjectName)
 					{
-						FString NewDefaultObjectName = Import.ObjectName.ToString();
-						FObjectImport OldImport = Import;
-						bool bUpdateOuterIndex = false;
-						int32 ImportPackage = -1;
-
-						// We are dealing with an object that needs to be redirected to a new classname (possibly a new package as well)
-
-						FString stringObjectName(FString(Import.ObjectName.ToString()));
-						if ( RedirectNameClass )
-						{
-							// This is an object instance
-							RedirectName = RedirectNameClass->ToString();
-						}
-						else if ( RedirectNameObj && bIsClassOrStructOrEnum )
+						// If right below package and package has changed, need to swap outer
+						if (NewObjectName.OuterName == NAME_None && NewObjectName.PackageName != OldObjectName.PackageName)
 						{
-							// This is a class object (needs to have its OuterIndex changed if the package is different)
-							bUpdateOuterIndex = true;
-							RedirectName = RedirectNameObj->ToString();
+							FPackageIndex NewPackageIndex;
 
-							if (TPair<FName, FName>* RedirectNameNewClassPair = ObjectNameClassRedirects.Find(Import.ObjectName))
+							if (FindImportPackage(NewObjectName.PackageName, NewPackageIndex))
 							{
-								if (!RedirectNameNewClassPair->Key.IsNone())
-								{
-									Import.ClassName = RedirectNameNewClassPair->Key;
-								}
-								if (!RedirectNameNewClassPair->Value.IsNone())
-								{
-									Import.ClassPackage = RedirectNameNewClassPair->Value;
-								}
+								// Already in import table, set it
+								Import.OuterIndex = NewPackageIndex;
 							}
-						}
-
-						// Accepts either "PackageName.ClassName" or just "ClassName"
-						int32 Offset = RedirectName.Find(TEXT("."), ESearchCase::CaseSensitive);
-						if ( Offset >= 0 )
-						{
-							// A package class name redirect
-							ResultPackage = RedirectName.Left(Offset);
-							ResultClass = RedirectName.Right(RedirectName.Len() - Offset - 1);
-						}
-						else
-						{
-							// Just a class name change within the same package
-							ResultPackage = Import.ClassPackage.ToString();
-							ResultClass = RedirectName;
-							bUpdateOuterIndex = false;
-						}
-
-						// Find the OuterIndex of the current package for the Import
-						for ( int32 ImportIndex = 0; ImportIndex < ImportMap.Num(); ImportIndex++ )
-						{
-							if ( ImportMap[ImportIndex].ClassName == NAME_Package && ImportMap[ImportIndex].ObjectName == Import.ClassPackage )
+							else
 							{
-								OldOuterIndex = ImportIndex;
+								// Need to add package import and try again
+								NewPackageImports.AddUnique(NewObjectName.PackageName);
+								bDone = false;
 								break;
-							}
+							}							
 						}
-						if ( !Import.OuterIndex.IsNull() && Import.OuterIndex == FPackageIndex::FromImport(OldOuterIndex) )
+#if WITH_EDITOR
+						// If this is a class, set old name here 
+						if (ObjectRedirectFlags == ECoreRedirectFlags::Type_Class)
 						{
-							// This is a object instance that is owned by a specific package (default class instance or an archetype etc)
-							// (needs its OuterIndex changed if the package is different)
-							if(ResultPackage != Import.ClassPackage.ToString())
-							{
-								bUpdateOuterIndex = true;
-							}					
+							Import.OldClassName = Import.ObjectName;
 						}
+						
+#endif
+						// Change object name
+						Import.ObjectName = NewObjectName.ObjectName;
 
-						if ( bUpdateOuterIndex && ResultPackage.Len() > 0 )
+						UE_LOG(LogLinker, Verbose, TEXT("FLinkerLoad::FixupImportMap() - Renamed Object %s -> %s"), *LinkerRoot->GetName(), *OldObjectName.ToString(), *NewObjectName.ToString());
+					}
+
+					if (NewClassName != OldClassName)
+					{
+						// Swap class if needed
+						if (Import.ClassPackage != NewClassName.PackageName && !IsCoreUObjectPackage(NewClassName.PackageName))
 						{
-							// Reset the Import.OuterIndex to the package it is intended to be in
-							for ( int32 ImportIndex = 0; ImportIndex < ImportMap.Num(); ImportIndex++ )
-							{
-								if ( ImportMap[ImportIndex].ClassName == NAME_Package && ImportMap[ImportIndex].ObjectName == FName(*ResultPackage) )
-								{
-									ImportPackage = ImportIndex;
-									break;
-								}
-							}
-							if (ImportPackage == -1 && !IsCoreUObjectPackage(FName(*ResultPackage)))
+							FPackageIndex NewPackageIndex;
+
+							if (!FindImportPackage(NewClassName.PackageName, NewPackageIndex))
 							{
-								// We are adding a new import to the map as we need the new package dependency added to the works
-								ImportMap.AddUninitialized();
-								ImportMap[ImportMap.Num()-1].ClassName = NAME_Package;
-								ImportMap[ImportMap.Num()-1].ClassPackage = GLongCoreUObjectPackageName;
-								ImportMap[ImportMap.Num()-1].ObjectName = FName(*ResultPackage);
-								ImportMap[ImportMap.Num()-1].OuterIndex = FPackageIndex();
-								ImportMap[ImportMap.Num()-1].XObject = 0;
-								ImportMap[ImportMap.Num()-1].SourceLinker = 0;
-								ImportMap[ImportMap.Num()-1].SourceIndex = -1;
-								ImportPackage = ImportMap.Num() - 1;
-
-								// Since this destroys the array, the current Import object is invalid and we must restart the whole process again
+								// Need to add package import and try again
+								NewPackageImports.AddUnique(NewClassName.PackageName);
 								bDone = false;
 								break;
 							}
-
-							// Assign the new OuterIndex for a default object instance or a class itself
-							if ( ImportPackage != -1 )
-							{
-								Import.OuterIndex = FPackageIndex::FromImport(ImportPackage);
-							}
 						}
-
-						if ( RedirectNameClass )
-						{
-							// Changing the package and class name of an object instance
-							Import.ClassPackage = *ResultPackage;
 #if WITH_EDITOR
-							Import.OldClassName = Import.ClassName;
+						Import.OldClassName = Import.ClassName;
 #endif
-							Import.ClassName = *ResultClass;
-						}
+						// Change class name/package
+						Import.ClassPackage = NewClassName.PackageName;
+						Import.ClassName = NewClassName.ObjectName;
 
-						if ( RedirectNameObj && bIsClassOrStructOrEnum )
-						{
-							// Changing the object name of a class object
-#if WITH_EDITOR
-							Import.OldClassName = Import.ObjectName;
-#endif
-							Import.ObjectName = *ResultClass;							
-						}
+						// Also change CDO name if needed
+						FString NewDefaultObjectName = Import.ObjectName.ToString();
 
-						// Default objects should be converted by name as well
-						if ( NewDefaultObjectName.Left(9) == FString("Default__") )
+						if (NewDefaultObjectName.StartsWith(DEFAULT_OBJECT_PREFIX))
 						{
-							NewDefaultObjectName = FString("Default__");
-							NewDefaultObjectName += *ResultClass;
-							Import.ObjectName = *NewDefaultObjectName;
+							NewDefaultObjectName = FString(DEFAULT_OBJECT_PREFIX);
+							NewDefaultObjectName += NewClassName.ObjectName.ToString();
+							Import.ObjectName = FName(*NewDefaultObjectName);
 						}
 
-						// Log the object redirection to the console for review
-						if ( OldImport.ObjectName != Import.ObjectName || OldImport.ClassName != Import.ClassName || OldImport.ClassPackage != Import.ClassPackage || OldImport.OuterIndex != Import.OuterIndex )
-						{
-							UE_LOG(LogLinker, Verbose, TEXT("FLinkerLoad::FixupImportMap() - Pkg<%s> [Obj<%s> Cls<%s> Pkg<%s> Out<%s>] -> [Obj<%s> Cls<%s> Pkg<%s> Out<%s>]"), *LinkerRoot->GetName(),
-								*OldImport.ObjectName.ToString(), *OldImport.ClassName.ToString(), *OldImport.ClassPackage.ToString(), OldImport.OuterIndex.IsImport() ? *Imp(OldImport.OuterIndex).ObjectName.ToString() : TEXT("None"),
-								*Import.ObjectName.ToString(), *Import.ClassName.ToString(), *Import.ClassPackage.ToString(),	Import.OuterIndex.IsImport() ? *Imp(Import.OuterIndex).ObjectName.ToString() : TEXT("None"));
-						}
-					}
+						UE_LOG(LogLinker, Verbose, TEXT("FLinkerLoad::FixupImportMap() - Renamed Class %s -> %s"), *LinkerRoot->GetName(), *OldClassName.ToString(), *NewClassName.ToString());
+					}	
+				}
+
+				// Add new packages, after loop iteration for safety
+				for (FName NewPackage : NewPackageImports)
+				{
+					// We are adding a new import to the map as we need the new package dependency added to the works
+					FObjectImport* NewImport = new (ImportMap) FObjectImport();
+
+					NewImport->ClassName = NAME_Package;
+					NewImport->ClassPackage = GLongCoreUObjectPackageName;
+					NewImport->ObjectName = NewPackage;
+					NewImport->OuterIndex = FPackageIndex();
+					NewImport->XObject = 0;
+					NewImport->SourceLinker = 0;
+					NewImport->SourceIndex = -1;
 				}
 			}
 		}
@@ -1717,45 +1682,6 @@ FLinkerLoad::ELinkerStatus FLinkerLoad::SerializeExportMap()
 	return ((ExportMapIndex == Summary.ExportCount) && !IsTimeLimitExceeded( TEXT("serializing export map") )) ? LINKER_Loaded : LINKER_TimedOut;
 }
 
-FLinkerLoad::ELinkerStatus FLinkerLoad::RemapImports()
-{
-	DECLARE_SCOPE_CYCLE_COUNTER( TEXT( "FLinkerLoad::RemapImports" ), STAT_LinkerLoad_RemapImports, STATGROUP_LinkerLoad );
-
-#if !(UE_BUILD_SHIPPING || UE_BUILD_TEST)
-
-	for (int32 ImportIndex = 0; ImportIndex < ImportMap.Num(); ImportIndex++)
-	{
-		FObjectImport& Import = ImportMap[ImportIndex];
-
-		FName* RedirectPackageName = GameNameRedirects.Find(Import.ClassPackage);
-		if (RedirectPackageName != NULL)
-		{
-			Import.ClassPackage = *RedirectPackageName;
-		}
-		if (Import.ClassName == NAME_Package)
-		{
-			RedirectPackageName = GameNameRedirects.Find(Import.ObjectName);
-			if (RedirectPackageName != NULL)
-			{
-				Import.ObjectName = *RedirectPackageName;
-			}
-
-			for ( TMap<FString, FString>::TConstIterator PluginRedirectIt = PluginNameRedirects.CreateConstIterator(); PluginRedirectIt; ++PluginRedirectIt )
-			{
-				if ( Import.ObjectName.ToString().StartsWith( PluginRedirectIt.Key() ) )
-				{
-					const FString NewPath = PluginRedirectIt.Value() + Import.ObjectName.ToString().RightChop(PluginRedirectIt.Key().Len());
-					Import.ObjectName = FName(*NewPath);
-				}
-			}
-		}
-	}
-
-#endif // !(UE_BUILD_SHIPPING || UE_BUILD_TEST)
-
-	return LINKER_Loaded;
-}
-
 /**
  * Serializes the depends map.
  */
@@ -2386,6 +2312,7 @@ FLinkerLoad::EVerifyResult FLinkerLoad::VerifyImport(int32 ImportIndex)
 				// serialize in the properties of the redirector (to get the object the redirector point to)
 				// Always load redirectors in case there was a circular dependency. This will allow inner redirector
 				// references to always serialize fully here before accessing the DestinationObject
+				check(!GEventDrivenLoaderEnabled);
 				Redir->SetFlags(RF_NeedLoad);
 				Preload(Redir);
 
@@ -3409,6 +3336,7 @@ void FLinkerLoad::Preload( UObject* Object )
 
 							// reset the flag and return (don't worry, we make
 							// sure to force load this later)
+							check(!GEventDrivenLoaderEnabled);
 							Object->SetFlags(RF_NeedLoad);
 							return;
 						}
@@ -3753,6 +3681,7 @@ UObject* FLinkerLoad::CreateExport( int32 Index )
 					const UClass* AsClass = dynamic_cast<UClass*>(SuperStruct);
 					if (AsClass && !AsClass->ClassDefaultObject)
 					{
+						check(!GEventDrivenLoaderEnabled);
 						SuperStruct->SetFlags(RF_NeedLoad);
 						Preload(SuperStruct);
 					}
@@ -3951,6 +3880,7 @@ UObject* FLinkerLoad::CreateExport( int32 Index )
 					if (!Export.Object->HasAnyFlags(RF_LoadCompleted) &&
 						(Export.Object->HasAnyFlags(RF_DefaultSubObject) || (ThisParent && ThisParent->IsTemplate(RF_ClassDefaultObject))))
 					{
+						check(!GEventDrivenLoaderEnabled);
 						Export.Object->SetFlags(RF_NeedLoad | RF_NeedPostLoad | RF_NeedPostLoadSubobjects | RF_WasLoaded);
 					}
 				}
@@ -4716,6 +4646,20 @@ void FLinkerLoad::MarkScriptSerializationEnd( const UObject* Obj )
 	}
 }
 
+bool FLinkerLoad::FindImportPackage(FName PackageName, FPackageIndex& PackageIdx)
+{
+	for (int32 ImportMapIdx = 0; ImportMapIdx < ImportMap.Num(); ImportMapIdx++)
+	{
+		if (ImportMap[ImportMapIdx].ObjectName == PackageName && ImportMap[ImportMapIdx].ClassName == NAME_Package)
+		{
+			PackageIdx = FPackageIndex::FromImport(ImportMapIdx);
+			return true;
+		}
+	}
+
+	return false;
+}
+
 /**
  * Locates the class adjusted index and its package adjusted index for a given class name in the import map
  */
@@ -4806,21 +4750,24 @@ bool FLinkerLoad::CreateImportClassAndPackage( FName ClassName, FName PackageNam
 TArray<FName> FLinkerLoad::FindPreviousNamesForClass(FString CurrentClassPath, bool bIsInstance)
 {
 	TArray<FName> OldNames;
-	for (auto It = ObjectNameRedirects.CreateConstIterator(); It; ++It)
+	TArray<FCoreRedirectObjectName> OldObjectNames;
+
+	if (FCoreRedirects::FindPreviousNames(ECoreRedirectFlags::Type_Class, FCoreRedirectObjectName(CurrentClassPath), OldObjectNames))
 	{
-		if (It.Value().ToString() == CurrentClassPath)
+		for (FCoreRedirectObjectName& OldObjectName : OldObjectNames)
 		{
-			OldNames.Add(It.Key());
+			OldNames.AddUnique(OldObjectName.ObjectName);
 		}
 	}
 
 	if (bIsInstance)
 	{
-		for (auto It = ObjectNameRedirectsInstanceOnly.CreateConstIterator(); It; ++It)
+		OldObjectNames.Empty();
+		if (FCoreRedirects::FindPreviousNames(ECoreRedirectFlags::Type_Class | ECoreRedirectFlags::Option_InstanceOnly, FCoreRedirectObjectName(CurrentClassPath), OldObjectNames))
 		{
-			if (It.Value().ToString() == CurrentClassPath)
+			for (FCoreRedirectObjectName& OldObjectName : OldObjectNames)
 			{
-				OldNames.Add(It.Key());
+				OldNames.AddUnique(OldObjectName.ObjectName);
 			}
 		}
 	}
@@ -4830,58 +4777,66 @@ TArray<FName> FLinkerLoad::FindPreviousNamesForClass(FString CurrentClassPath, b
 
 FName FLinkerLoad::FindNewNameForEnum(const FName OldEnumName)
 {
-	if (FName* RedirectName = ObjectNameRedirects.Find(OldEnumName))
+	FCoreRedirectObjectName OldName = FCoreRedirectObjectName(OldEnumName, NAME_None, NAME_None);
+	FCoreRedirectObjectName NewName = FCoreRedirects::GetRedirectedName(ECoreRedirectFlags::Type_Enum, OldName);
+
+	if (NewName != OldName)
 	{
-		return *RedirectName;
+		return NewName.ObjectName;
 	}
-
 	return NAME_None;
 }
 
 FName FLinkerLoad::FindNewNameForStruct(const FName OldStructName)
 {
-	FName* RedirectName = StructNameRedirects.Find(OldStructName);
+	FCoreRedirectObjectName OldName = FCoreRedirectObjectName(OldStructName, NAME_None, NAME_None);
+	FCoreRedirectObjectName NewName = FCoreRedirects::GetRedirectedName(ECoreRedirectFlags::Type_Struct, OldName);
 
-	if (RedirectName)
+	if (NewName != OldName)
 	{
-		return *RedirectName;
+		return NewName.ObjectName;
 	}
-	
-	RedirectName = ObjectNameRedirects.Find(OldStructName);
-
-	return (RedirectName ? *RedirectName : NAME_None);
+	return NAME_None;
 }
 
 FName FLinkerLoad::FindNewNameForClass(FName OldClassName, bool bIsInstance)
 {
-	FName* RedirectName = ObjectNameRedirects.Find(OldClassName);
+	FCoreRedirectObjectName OldName = FCoreRedirectObjectName(OldClassName, NAME_None, NAME_None);
+	FCoreRedirectObjectName NewName = FCoreRedirects::GetRedirectedName(ECoreRedirectFlags::Type_Class, OldName);
 
-	if (RedirectName)
+	if (NewName != OldName)
 	{
-		return *RedirectName;
+		return NewName.ObjectName;
 	}
 
 	if (bIsInstance)
 	{
-		RedirectName = ObjectNameRedirectsInstanceOnly.Find(OldClassName);
-	}
+		// Also check instance types
+		NewName = FCoreRedirects::GetRedirectedName(ECoreRedirectFlags::Type_Class | ECoreRedirectFlags::Option_InstanceOnly, OldName);
 
-	if (RedirectName)
-	{
-		return *RedirectName;
+		if (NewName != OldName)
+		{
+			return NewName.ObjectName;
+		}
 	}
-
 	return NAME_None;
 }
 
 bool FLinkerLoad::IsKnownMissingPackage(FName PackageName)
 {
-	return KnownMissingPackages.Contains(PackageName);
+	return FCoreRedirects::IsKnownMissing(ECoreRedirectFlags::Type_Package, FCoreRedirectObjectName(NAME_None, NAME_None, PackageName));
 }
 
 void FLinkerLoad::AddKnownMissingPackage(FName PackageName)
 {
-	KnownMissingPackages.Add(PackageName);
+	FCoreRedirects::AddKnownMissing(ECoreRedirectFlags::Type_Package, FCoreRedirectObjectName(NAME_None, NAME_None, PackageName));
+}
+
+void FLinkerLoad::AddGameNameRedirect(const FName OldName, const FName NewName)
+{
+	TArray<FCoreRedirect> NewRedirects;
+	NewRedirects.Emplace(ECoreRedirectFlags::Type_Package, FCoreRedirectObjectName(NAME_None, NAME_None, OldName), FCoreRedirectObjectName(NAME_None, NAME_None, NewName));
+	FCoreRedirects::AddRedirectList(NewRedirects, TEXT("AddGameNameRedirect"));
 }
 
 #if WITH_EDITOR
@@ -5013,6 +4968,7 @@ FLinkerLoad::ELinkerStatus FLinkerLoad::FixupExportMap()
 			}
 			FName NameClass = GetExportClassName(ExportMapIdx);
 			FName NamePackage = GetExportClassPackage(ExportMapIdx);
+			FString StrObjectName = Export.ObjectName.ToString();
 
 			// ActorComponents outered to a BlueprintGeneratedClass (or even older ones that are outered to Blueprint) need to be marked RF_Public, but older content was 
 			// not created as such.  This updates the ExportTable such that they are correctly flagged when created and when other packages validate their imports.
@@ -5042,83 +4998,66 @@ FLinkerLoad::ELinkerStatus FLinkerLoad::FixupExportMap()
 				}
 			}
 
+			// Look for subobject redirects and instance redirects
+			FCoreRedirectObjectName OldClassName(NameClass, NAME_None, NamePackage);
+				
+			const TMap<FString, FString>* ValueChanges = FCoreRedirects::GetValueRedirects(ECoreRedirectFlags::Type_Class, OldClassName);
+
+			if (ValueChanges)
 			{
-				FSubobjectRedirect *Redirect = SubobjectNameRedirects.Find(Export.ObjectName);
-				if (Redirect)
+				// Apply class value redirects before other redirects, to mirror old subobject order
+				const FString* NewInstanceName = ValueChanges->Find(Export.ObjectName.ToString());
+				if (NewInstanceName)
 				{
-					if (NameClass == Redirect->MatchClass)
-					{
-						if (!Export.OuterIndex.IsNull())
-						{
-							FString Was = GetExportFullName(ExportMapIdx);
-							Export.ObjectName = Redirect->NewName;
+					// Rename this import directly
+					FString Was = GetExportFullName(ExportMapIdx);
+					Export.ObjectName = FName(**NewInstanceName);
 
-							if (Export.ObjectName != NAME_None)
-							{
-								FString Now = GetExportFullName(ExportMapIdx);
-								UE_LOG(LogLinker, Log, TEXT("FLinkerLoad::FixupExportMap() - Renamed component from %s   to   %s"), *Was, *Now);
-							}
-							else
-							{
-								Export.bExportLoadFailed = true;
-								UE_LOG(LogLinker, Log, TEXT("FLinkerLoad::FixupExportMap() - Removed component %s"), *Was);
-							}
-							continue;
-						}
+					if (Export.ObjectName != NAME_None)
+					{
+						FString Now = GetExportFullName(ExportMapIdx);
+						UE_LOG(LogLinker, Verbose, TEXT("FLinkerLoad::FixupExportMap() - Renamed object from %s   to   %s"), *Was, *Now);
+					}
+					else
+					{
+						Export.bExportLoadFailed = true;
+						UE_LOG(LogLinker, Verbose, TEXT("FLinkerLoad::FixupExportMap() - Removed object %s"), *Was);
 					}
 				}
 			}
-			FName *RedirectName = ObjectNameRedirectsInstanceOnly.Find(NameClass);
-			if ( RedirectName )
-			{
-				FString StrRedirectName, ResultPackage, ResultClass;
-				FString StrObjectName = Export.ObjectName.ToString();
 
-				StrRedirectName = RedirectName->ToString();
+			// Never modify the default object instances
+			if (!StrObjectName.StartsWith(DEFAULT_OBJECT_PREFIX))
+			{
+				FCoreRedirectObjectName NewClassInstanceName = FCoreRedirects::GetRedirectedName(ECoreRedirectFlags::Type_Class | ECoreRedirectFlags::Option_InstanceOnly, OldClassName);
 
-				// Accepts either "PackageName.ClassName" or just "ClassName"
-				int32 Offset = StrRedirectName.Find(TEXT("."), ESearchCase::CaseSensitive);
-				if ( Offset >= 0 )
-				{
-					// A package class name redirect
-					ResultPackage = StrRedirectName.Left(Offset);
-					ResultClass = StrRedirectName.Right(StrRedirectName.Len() - Offset - 1);
-				}
-				else
+				bool bClassInstanceDeleted = FCoreRedirects::IsKnownMissing(ECoreRedirectFlags::Type_Class | ECoreRedirectFlags::Option_InstanceOnly, OldClassName);
+				if (bClassInstanceDeleted)
 				{
-					// Just a class name change within the same package
-					ResultPackage = NamePackage.ToString();
-					ResultClass = StrRedirectName;
-				}
+					UE_LOG(LogLinker, Log, TEXT("FLinkerLoad::FixupExportMap() - Pkg<%s> [Obj<%s> Cls<%s> ClsPkg<%s>] -> removed"), *LinkerRoot->GetName(),
+						*Export.ObjectName.ToString(), *NameClass.ToString(), *NamePackage.ToString());
 
-				// Never modify the default object instances
-				if ( StrObjectName.Left(9) != TEXT("Default__") )
+					Export.ClassIndex = FPackageIndex();
+					Export.OuterIndex = FPackageIndex();
+					Export.ObjectName = NAME_None;
+#if WITH_EDITOR
+					Export.OldClassName = NameClass;
+#endif
+				}
+				else if (NewClassInstanceName != OldClassName)
 				{
 					FPackageIndex NewClassIndex;
 					FPackageIndex NewPackageIndex;
-					if ( ResultClass == TEXT("None") )
-					{
-						UE_LOG(LogLinker, Log, TEXT("FLinkerLoad::FixupExportMap() - Pkg<%s> [Obj<%s> Cls<%s> ClsPkg<%s>] -> removed"), *LinkerRoot->GetName(),
-							*Export.ObjectName.ToString(), *NameClass.ToString(), *NamePackage.ToString());
 
-						Export.ClassIndex = NewClassIndex;
-						Export.OuterIndex = NewClassIndex;
-						Export.ObjectName = NAME_None;
-#if WITH_EDITOR
-						Export.OldClassName = NameClass;
-#endif
-					}
-					else if ( CreateImportClassAndPackage(*ResultClass, *ResultPackage, NewClassIndex, NewPackageIndex) )
+					if (CreateImportClassAndPackage(NewClassInstanceName.ObjectName, NewClassInstanceName.PackageName, NewClassIndex, NewPackageIndex))
 					{
 						Export.ClassIndex = NewClassIndex;
 #if WITH_EDITOR
 						Export.OldClassName = NameClass;
 #endif
-						//Export.OuterIndex = newPackageIndex;
-
 						UE_LOG(LogLinker, Log, TEXT("FLinkerLoad::FixupExportMap() - Pkg<%s> [Obj<%s> Cls<%s> ClsPkg<%s>] -> [Obj<%s> Cls<%s> ClsPkg<%s>]"), *LinkerRoot->GetName(),
 							*Export.ObjectName.ToString(), *NameClass.ToString(), *NamePackage.ToString(),
-							*Export.ObjectName.ToString(), *ResultClass, *ResultPackage);
+							*Export.ObjectName.ToString(), *NewClassInstanceName.ObjectName.ToString(), *NewClassInstanceName.PackageName.ToString());
 					}
 					else
 					{
@@ -5126,54 +5065,7 @@ FLinkerLoad::ELinkerStatus FLinkerLoad::FixupExportMap()
 					}
 				}
 			}
-			else
-			{
-				//UE_LOG(LogLinker, Log, TEXT("Export: <%s>"), *( LinkerRoot->Name.ToString() + TEXT(".") + Export.ObjectName.ToString() ));
-				RedirectName = ObjectNameRedirectsObjectOnly.Find(*( LinkerRoot->GetName() + TEXT(".") + *Export.ObjectName.ToString() ) );
-				if ( RedirectName )
-				{
-					FString StrRedirectName, ResultPackage, ResultClass;
-					FString StrObjectName = Export.ObjectName.ToString();
-
-					StrRedirectName = RedirectName->ToString();
-
-					// Accepts either "PackageName.ClassName" or just "ClassName"
-					int32 Offset = StrRedirectName.Find(TEXT("."), ESearchCase::CaseSensitive);
-					if ( Offset >= 0 )
-					{
-						// A package class name redirect
-						ResultPackage = StrRedirectName.Left(Offset);
-						ResultClass = StrRedirectName.Right(StrRedirectName.Len() - Offset - 1);
-					}
-					else
-					{
-						ResultClass = StrRedirectName;
-					}
-
-					// Never modify the default object instances
-					if ( StrObjectName.Left(9) != TEXT("Default__") )
-					{
-						FPackageIndex NewClassIndex;
-						FPackageIndex NewPackageIndex;
-						if ( CreateImportClassAndPackage(*ResultClass, *ResultPackage, NewClassIndex, NewPackageIndex) )
-						{
-							Export.ClassIndex = NewClassIndex;
-#if WITH_EDITOR
-							Export.OldClassName = NameClass;
-#endif
-							UE_LOG(LogLinker, Log, TEXT("FLinkerLoad::FixupExportMap() - Pkg<%s> [Obj<%s> Cls<%s> ClsPkg<%s>] -> [Obj<%s> Cls<%s> ClsPkg<%s>]"), *LinkerRoot->GetName(),
-								*Export.ObjectName.ToString(), *NameClass.ToString(), *NamePackage.ToString(),
-								*Export.ObjectName.ToString(), *ResultClass, *ResultPackage);
-						}
-						else
-						{
-							UE_LOG(LogLinker, Log, TEXT("FLinkerLoad::FixupExportMap() - object redirection failed at %s"), *Export.ObjectName.ToString());
-						}
-					}
-				}
-			}	
 		}
-
 		bFixupExportMapDone = true;
 		return !IsTimeLimitExceeded( TEXT("fixing up export map") ) ? LINKER_Loaded : LINKER_TimedOut;
 	}
diff --git a/Engine/Source/Runtime/CoreUObject/Private/UObject/LinkerPlaceholderBase.inl b/Engine/Source/Runtime/CoreUObject/Private/UObject/LinkerPlaceholderBase.inl
--- a/Engine/Source/Runtime/CoreUObject/Private/UObject/LinkerPlaceholderBase.inl
+++ b/Engine/Source/Runtime/CoreUObject/Private/UObject/LinkerPlaceholderBase.inl
@@ -35,6 +35,7 @@ int32 TLinkerImportPlaceholder<PlaceholderType>::ResolveAllPlaceholderReferences
 				// here, needing to use the ClassB redirector, but we haven't 
 				// returned up the stack for it to be set yet... here we force 
 				// it to finish preloading (like we do in VerifyImport): 
+				check(!GEventDrivenLoaderEnabled);
 				ReplacementRedirector->SetFlags(RF_NeedLoad);
 				ReplacementLinker->Preload(ReplacementRedirector);
 			}
diff --git a/Engine/Source/Runtime/CoreUObject/Private/UObject/LinkerSave.cpp b/Engine/Source/Runtime/CoreUObject/Private/UObject/LinkerSave.cpp
--- a/Engine/Source/Runtime/CoreUObject/Private/UObject/LinkerSave.cpp
+++ b/Engine/Source/Runtime/CoreUObject/Private/UObject/LinkerSave.cpp
@@ -181,7 +181,7 @@ FPackageIndex FLinkerSave::MapObject( const UObject* Object ) const
 		{
 			if (IsEventDrivenLoaderEnabledInCookedBuilds() &&
 				IsCooking() && CurrentlySavingExport.IsExport() &&
-				!Object->GetOutermost()->HasAllPackagesFlags(PKG_CompiledIn) && 
+				Object->GetOutermost()->GetFName() != GLongCoreUObjectPackageName && // We assume nothing in coreuobject ever loads assets in a constructor
 				*Found != CurrentlySavingExport) // would be weird, but I can't be a dependency on myself
 			{
 				const FObjectExport& SavingExport = Exp(CurrentlySavingExport);
diff --git a/Engine/Source/Runtime/CoreUObject/Private/UObject/Obj.cpp b/Engine/Source/Runtime/CoreUObject/Private/UObject/Obj.cpp
--- a/Engine/Source/Runtime/CoreUObject/Private/UObject/Obj.cpp
+++ b/Engine/Source/Runtime/CoreUObject/Private/UObject/Obj.cpp
@@ -49,6 +49,7 @@
 #include "UObject/UObjectThreadContext.h"
 #include "Misc/ExclusiveLoadPackageTimeTracker.h"
 #include "Serialization/DeferredMessageLog.h"
+#include "UObject/CoreRedirects.h"
 
 DEFINE_LOG_CATEGORY(LogObj);
 
@@ -1006,7 +1007,7 @@ bool UObject::IsSelected() const
 void UObject::GetPreloadDependencies(TArray<UObject*>& OutDeps)
 {
 	UClass *ObjClass = GetClass();
-	if (!ObjClass->HasAnyClassFlags(CLASS_Intrinsic | CLASS_Native))
+	if (!ObjClass->HasAnyClassFlags(CLASS_Intrinsic))
 	{
 		OutDeps.Add(ObjClass);
 
@@ -3304,11 +3305,11 @@ bool StaticExec( UWorld* InWorld, const TCHAR* Cmd, FOutputDevice& Ar )
 
 				if (Errors.Num() > 0)
 				{
-					Ar.Logf(*FString::Printf(TEXT("Errors for %s"), *Target->GetName()));
+					Ar.Logf(TEXT("Errors for %s"), *Target->GetName());
 
 					for (auto ErrorStr : Errors)
 					{
-						Ar.Logf(*(FString(TEXT("  - ") + ErrorStr)));
+						Ar.Logf(TEXT("  - %s"), *ErrorStr);
 					}
 				}
 			}
@@ -3810,6 +3811,7 @@ void InitUObject()
 	// Initialize redirects map
 	for (const auto& It : *GConfig)
 	{
+		FCoreRedirects::ReadRedirectsFromIni(It.Key);
 		FLinkerLoad::CreateActiveRedirectsMap(It.Key);
 	}
 
diff --git a/Engine/Source/Runtime/CoreUObject/Private/UObject/ObjectResource.cpp b/Engine/Source/Runtime/CoreUObject/Private/UObject/ObjectResource.cpp
--- a/Engine/Source/Runtime/CoreUObject/Private/UObject/ObjectResource.cpp
+++ b/Engine/Source/Runtime/CoreUObject/Private/UObject/ObjectResource.cpp
@@ -110,8 +110,21 @@ FArchive& operator<<( FArchive& Ar, FObjectExport& E )
 		E.ObjectFlags = EObjectFlags(Save & RF_Load);
 	}
 
-	Ar << E.SerialSize;
-	Ar << E.SerialOffset;
+	if (Ar.UE4Ver() < VER_UE4_64BIT_EXPORTMAP_SERIALSIZES)
+	{
+		int32 SerialSize = E.SerialSize;
+		Ar << SerialSize;
+		E.SerialSize = (int64)SerialSize;
+
+		int32 SerialOffset = E.SerialOffset;
+		Ar << SerialOffset;
+		E.SerialOffset = SerialOffset;
+	}
+	else
+	{
+		Ar << E.SerialSize;
+		Ar << E.SerialOffset;
+	}
 
 	Ar << E.bForcedExport;
 	Ar << E.bNotForClient;
diff --git a/Engine/Source/Runtime/CoreUObject/Private/UObject/Property.cpp b/Engine/Source/Runtime/CoreUObject/Private/UObject/Property.cpp
--- a/Engine/Source/Runtime/CoreUObject/Private/UObject/Property.cpp
+++ b/Engine/Source/Runtime/CoreUObject/Private/UObject/Property.cpp
@@ -14,6 +14,7 @@
 #include "Templates/Casts.h"
 #include "UObject/UnrealType.h"
 #include "UObject/PropertyHelper.h"
+#include "UObject/CoreRedirects.h"
 #include "Misc/StringClassReference.h"
 
 DEFINE_LOG_CATEGORY(LogProperty);
@@ -416,7 +417,7 @@ bool UProperty::ValidateImportFlags( uint32 PortFlags, FOutputDevice* ErrorHandl
 
 		if (ErrorHandler)
 		{
-			ErrorHandler->Logf(*ErrorMsg);
+			ErrorHandler->Logf(TEXT("%s"), *ErrorMsg);
 		}
 		else
 		{
@@ -867,17 +868,14 @@ const TCHAR* UProperty::ImportSingleProperty( const TCHAR* Str, void* DestData,
 		const FName PropertyName(Token);
 		UProperty* Property = FindField<UProperty>(ObjectStruct, PropertyName);
 
-		if (Property == NULL)
+		if (Property == nullptr)
 		{
 			// Check for redirects
-			const TMap<FName, FName>* const ClassTaggedPropertyRedirects = UStruct::TaggedPropertyRedirects.Find(ObjectStruct->GetFName());
-			if (ClassTaggedPropertyRedirects)
+			FName NewPropertyName = FindRedirectedPropertyName(ObjectStruct, PropertyName);
+
+			if (NewPropertyName != NAME_None)
 			{
-				const FName* const NewPropertyName = ClassTaggedPropertyRedirects->Find(PropertyName);
-				if (NewPropertyName)
-				{
-					Property = FindField<UProperty>(ObjectStruct, *NewPropertyName);
-				}
+				Property = FindField<UProperty>(ObjectStruct, NewPropertyName);
 			}
 #if WITH_EDITOR
 			if (!Property)
@@ -978,7 +976,7 @@ const TCHAR* UProperty::ImportSingleProperty( const TCHAR* Str, void* DestData,
 
 						for ( int32 ErrorIndex = 0; ErrorIndex < ImportErrors.Num(); ErrorIndex++ )
 						{
-							Warn->Logf(ELogVerbosity::Warning, *ImportErrors[ErrorIndex]);
+							Warn->Logf(ELogVerbosity::Warning, TEXT("%s"), *ImportErrors[ErrorIndex]);
 						}
 					}
 					else if (Result == NULL || Result == Str)
@@ -1181,7 +1179,7 @@ const TCHAR* UProperty::ImportSingleProperty( const TCHAR* Str, void* DestData,
 
 						for ( int32 ErrorIndex = 0; ErrorIndex < ImportErrors.Num(); ErrorIndex++ )
 						{
-							Warn->Logf(ELogVerbosity::Warning,*ImportErrors[ErrorIndex]);
+							Warn->Logf(ELogVerbosity::Warning, TEXT("%s"), *ImportErrors[ErrorIndex]);
 						}
 					}
 					else if ((Result == NULL && ArrayProperty == nullptr) || Result == Str)
@@ -1213,7 +1211,7 @@ const TCHAR* UProperty::ImportSingleProperty( const TCHAR* Str, void* DestData,
 
 						for ( int32 ErrorIndex = 0; ErrorIndex < ImportErrors.Num(); ErrorIndex++ )
 						{
-							Warn->Logf(ELogVerbosity::Warning, *ImportErrors[ErrorIndex]);
+							Warn->Logf(ELogVerbosity::Warning, TEXT("%s"), *ImportErrors[ErrorIndex]);
 						}
 					}
 					else if ((Result == NULL && ArrayProperty == nullptr) || Result == Str)
@@ -1232,6 +1230,25 @@ const TCHAR* UProperty::ImportSingleProperty( const TCHAR* Str, void* DestData,
 	return Str;
 }
 
+FName UProperty::FindRedirectedPropertyName(UStruct* ObjectStruct, FName OldName)
+{
+	DECLARE_SCOPE_CYCLE_COUNTER(TEXT("UProperty::FindRedirectedPropertyName"), STAT_LinkerLoad_FindRedirectedPropertyName, STATGROUP_LoadTimeVerbose);
+
+	// ObjectStruct may be a nested struct, so extract path
+	UPackage* StructPackage = ObjectStruct->GetOutermost();
+	FString OuterPath = ObjectStruct->GetPathName(StructPackage);
+
+	FCoreRedirectObjectName OldRedirectName(OldName, FName(*OuterPath), StructPackage->GetFName());
+	FCoreRedirectObjectName NewRedirectName = FCoreRedirects::GetRedirectedName(ECoreRedirectFlags::Type_Property, OldRedirectName);
+
+	if (NewRedirectName != OldRedirectName)
+	{
+		return NewRedirectName.ObjectName;
+	}
+
+	return NAME_None;
+}
+
 /**
  * Returns the hash value for an element of this property.
  */
diff --git a/Engine/Source/Runtime/CoreUObject/Private/UObject/PropertyArray.cpp b/Engine/Source/Runtime/CoreUObject/Private/UObject/PropertyArray.cpp
--- a/Engine/Source/Runtime/CoreUObject/Private/UObject/PropertyArray.cpp
+++ b/Engine/Source/Runtime/CoreUObject/Private/UObject/PropertyArray.cpp
@@ -112,11 +112,11 @@ void UArrayProperty::SerializeItem( FArchive& Ar, void* Value, void const* Defau
 			// Check if the Inner property can successfully serialize, the type may have changed
 			UStructProperty* StructProperty = CastChecked<UStructProperty>(Inner);
 			// if check redirector to make sure if the name has changed
-			FName* NewName = FLinkerLoad::StructNameRedirects.Find(InnerTag.StructName);
+			FName NewName = FLinkerLoad::FindNewNameForStruct(InnerTag.StructName);
 			FName StructName = CastChecked<UStructProperty>(StructProperty)->Struct->GetFName();
-			if (NewName != nullptr && *NewName == StructName)
+			if (NewName != NAME_None && NewName == StructName)
 			{
-				InnerTag.StructName = *NewName;
+				InnerTag.StructName = NewName;
 			}
 
 			if (InnerTag.StructName != StructProperty->Struct->GetFName()
diff --git a/Engine/Source/Runtime/CoreUObject/Private/UObject/ReferenceChainSearch.cpp b/Engine/Source/Runtime/CoreUObject/Private/UObject/ReferenceChainSearch.cpp
--- a/Engine/Source/Runtime/CoreUObject/Private/UObject/ReferenceChainSearch.cpp
+++ b/Engine/Source/Runtime/CoreUObject/Private/UObject/ReferenceChainSearch.cpp
@@ -439,21 +439,24 @@ struct FStackEntry
 	int32 LoopStartIndex;
 };
 
-// Local helper function to add a reference chain to the temporary referecnce chain list
-void AddToReferenceList(TArray<FReferenceChainSearch::FReferenceChainLink>* ReferenceList, const FReferenceChainSearch::FReferenceChainLink& RefToAdd)
+// Local helper function to add a reference chain to the temporary reference chain list
+void AddToReferenceList(TArray<FReferenceChainSearch::FReferenceChainLink>& ReferenceList, const FReferenceChainSearch::FReferenceChainLink& RefToAdd)
 {
-	if (RefToAdd.ReferencedObj == NULL || RefToAdd.ReferencedBy == RefToAdd.ReferencedObj) { return; }
+	if (RefToAdd.ReferencedObj == NULL || RefToAdd.ReferencedBy == RefToAdd.ReferencedObj)
+	{
+		return;
+	}
 
 	bool bAdded = false;
 
-	for (int32 i=0; i<ReferenceList->Num(); ++i)
+	for (FReferenceChainSearch::FReferenceChainLink& Link : ReferenceList)
 	{
-		if ((*ReferenceList)[i].ReferencedObj == RefToAdd.ReferencedObj)
+		if (Link.ReferencedObj == RefToAdd.ReferencedObj)
 		{
 			bAdded = true;
 			if (RefToAdd.IsProperty() && RefToAdd.ReferencedThrough != NULL)
 			{
-				(*ReferenceList)[i] = RefToAdd;
+				Link = RefToAdd;
 			}
 			break;
 		}
@@ -461,21 +464,23 @@ void AddToReferenceList(TArray<FReferenceChainSearch::FReferenceChainLink>* Refe
 
 	if (!bAdded)
 	{
-		ReferenceList->Add(RefToAdd);
+		ReferenceList.Add(RefToAdd);
 	}
 }
 
-bool FReferenceChainSearch::ProcessObject( UObject* CurrentObject )
+void FReferenceChainSearch::ProcessObject( UObject* CurrentObject )
 {
+	UClass* ObjectClass = CurrentObject->GetClass();
+
 	// Make sure that token stream has been assembled at this point as the below code relies on it.
-	if ( !CurrentObject->GetClass()->HasAnyClassFlags(CLASS_TokenStreamAssembled) )
+	if ( !ObjectClass->HasAnyClassFlags(CLASS_TokenStreamAssembled) )
 	{
-		CurrentObject->GetClass()->AssembleReferenceTokenStream();
-		check(CurrentObject->GetClass()->HasAnyClassFlags(CLASS_TokenStreamAssembled));
+		ObjectClass->AssembleReferenceTokenStream();
+		check(ObjectClass->HasAnyClassFlags(CLASS_TokenStreamAssembled));
 	}
 
 	// Get pointer to token stream and jump to the start.
-	FGCReferenceTokenStream* RESTRICT TokenStream = &CurrentObject->GetClass()->ReferenceTokenStream;
+	FGCReferenceTokenStream* RESTRICT TokenStream = &ObjectClass->ReferenceTokenStream;
 	uint32 TokenStreamIndex			= 0;
 	// Keep track of index to reference info. Used to avoid LHSs.
 	uint32 ReferenceTokenStreamIndex	= 0;
@@ -483,7 +488,7 @@ bool FReferenceChainSearch::ProcessObject( UObject* CurrentObject )
 	TArray<FStackEntry> Stack;
 	Stack.AddUninitialized( 128 );
 
-	// Create strack entry and initialize sane values.
+	// Create stack entry and initialize sane values.
 	FStackEntry* RESTRICT StackEntry = Stack.GetData();
 	uint8* StackEntryData		= (uint8*) CurrentObject;
 	StackEntry->Data			= StackEntryData;
@@ -492,16 +497,12 @@ bool FReferenceChainSearch::ProcessObject( UObject* CurrentObject )
 	StackEntry->LoopStartIndex	= -1;
 
 	// Keep track of token return count in separate integer as arrays need to fiddle with it.
-	int32 TokenReturnCount		= 0;
-
-	bool bRetVal = false;
+	int32 TokenReturnCount = 0;
 
 	UProperty* InArrayProp = NULL;
 
-	ReferenceMap.Add(CurrentObject, TArray<FReferenceChainLink>());
-	TArray<FReferenceChainLink>* ReferenceList = ReferenceMap.Find(CurrentObject);
-	
-	
+	TArray<FReferenceChainLink>& ReferenceList = ReferenceMap.Emplace(CurrentObject);
+
 	// Parse the token stream.
 	while( true )
 	{
@@ -541,9 +542,23 @@ bool FReferenceChainSearch::ProcessObject( UObject* CurrentObject )
 		// the reference info means we need to manually increment the token index to skip to the next one.
 		TokenStreamIndex++;
 		// Helper to make code more readable and hide the ugliness that is avoiding LHSs from caching.
-#define	REFERENCE_INFO TokenStream->AccessReferenceInfo( ReferenceTokenStreamIndex )
+		FGCReferenceInfo ReferenceInfo = TokenStream->AccessReferenceInfo( ReferenceTokenStreamIndex );
+
+		if( ReferenceInfo.Type == GCRT_EndOfStream )
+		{
+			check(StackEntry == Stack.GetData());
+			return;
+		}
 
-		uint32 Offset = REFERENCE_INFO.Offset;
+		if (ReferenceInfo.Type == GCRT_EndOfPointer)
+		{
+			TokenReturnCount = ReferenceInfo.ReturnCount;
+			continue;
+		}
+
+		uint32 Offset = ReferenceInfo.Offset;
+		
+		void* StackEntryPtr = StackEntryData + Offset;
 		
 		// Get the property from token stream
 		UProperty* Prop = Internal::FindPropertyForOffset(CurrentObject->GetClass(), Offset);
@@ -553,155 +568,151 @@ bool FReferenceChainSearch::ProcessObject( UObject* CurrentObject )
 			Prop = InArrayProp;
 		}
 
-		if( REFERENCE_INFO.Type == GCRT_Object )
-		{	
-			// We're dealing with an object reference.
-			UObject**	ObjectPtr	= (UObject**)(StackEntryData + Offset);
-			UObject*&	Object		= *ObjectPtr;
-			TokenReturnCount		= REFERENCE_INFO.ReturnCount;
-
-			FReferenceChainLink TopRef(ReferenceTokenStreamIndex, InArrayProp != NULL ? EReferenceType::ArrayProperty : EReferenceType::Property, CurrentObject, Prop, Object);
-			AddToReferenceList(ReferenceList, TopRef);
-		}
-		else if( REFERENCE_INFO.Type == GCRT_ArrayObject )
+		switch (ReferenceInfo.Type)
 		{
-			// We're dealing with an array of object references.
-			TArray<UObject*>& ObjectArray = *((TArray<UObject*>*)(StackEntryData + Offset));
-			TokenReturnCount = REFERENCE_INFO.ReturnCount;
-
-			for( int32 ObjectIndex=0; ObjectIndex<ObjectArray.Num(); ObjectIndex++ )
+			case GCRT_Object:
+			case GCRT_PersistentObject:
 			{
-				UObject*& Object = ObjectArray[ObjectIndex];
-				
-				FReferenceChainLink TopRef(ReferenceTokenStreamIndex, EReferenceType::ArrayProperty, CurrentObject, Prop, Object, ObjectIndex);
+				// We're dealing with an object reference.
+				UObject* Object = *(UObject**)StackEntryPtr;
+				TokenReturnCount = ReferenceInfo.ReturnCount;
+
+				FReferenceChainLink TopRef(ReferenceTokenStreamIndex, InArrayProp != NULL ? EReferenceType::ArrayProperty : EReferenceType::Property, CurrentObject, Prop, Object);
 				AddToReferenceList(ReferenceList, TopRef);
 			}
-		}
-		else if( REFERENCE_INFO.Type == GCRT_ArrayStruct )
-		{
-			InArrayProp = Prop;
+			break;
 
-			// We're dealing with a dynamic array of structs.
-			const FScriptArray& Array = *((FScriptArray*)(StackEntryData + Offset));
-			StackEntry++;
-			StackEntryData				= (uint8*) Array.GetData();
-			StackEntry->Data			= StackEntryData;
-			StackEntry->Stride			= TokenStream->ReadStride( TokenStreamIndex );
-			StackEntry->Count			= Array.Num();
+			case GCRT_ArrayObject:
+			{
+				// We're dealing with an array of object references.
+				TArray<UObject*>& ObjectArray = *(TArray<UObject*>*)StackEntryPtr;
+				TokenReturnCount = ReferenceInfo.ReturnCount;
 
-			const FGCSkipInfo SkipInfo	= TokenStream->ReadSkipInfo( TokenStreamIndex );
-			StackEntry->LoopStartIndex	= TokenStreamIndex;
+				for( int32 ObjectIndex=0; ObjectIndex<ObjectArray.Num(); ObjectIndex++ )
+				{
+					UObject*& Object = ObjectArray[ObjectIndex];
+				
+					FReferenceChainLink TopRef(ReferenceTokenStreamIndex, EReferenceType::ArrayProperty, CurrentObject, Prop, Object, ObjectIndex);
+					AddToReferenceList(ReferenceList, TopRef);
+				}
+			}
+			break;
 
-			if( StackEntry->Count == 0 )
+			case GCRT_ArrayStruct:
 			{
-				// Skip empty array by jumping to skip index and set return count to the one about to be read in.
-				TokenStreamIndex		= SkipInfo.SkipIndex;
-				TokenReturnCount		= TokenStream->GetSkipReturnCount( SkipInfo );
-			}
-			else
-			{	
-				// Loop again.
-				check( StackEntry->Data );
-				TokenReturnCount		= 0;
+				InArrayProp = Prop;
+
+				// We're dealing with a dynamic array of structs.
+				const FScriptArray& Array = *(FScriptArray*)StackEntryPtr;
+				StackEntry++;
+				StackEntryData				= (uint8*) Array.GetData();
+				StackEntry->Data			= StackEntryData;
+				StackEntry->Stride			= TokenStream->ReadStride( TokenStreamIndex );
+				StackEntry->Count			= Array.Num();
+
+				const FGCSkipInfo SkipInfo	= TokenStream->ReadSkipInfo( TokenStreamIndex );
+				StackEntry->LoopStartIndex	= TokenStreamIndex;
+
+				if( StackEntry->Count == 0 )
+				{
+					// Skip empty array by jumping to skip index and set return count to the one about to be read in.
+					TokenStreamIndex		= SkipInfo.SkipIndex;
+					TokenReturnCount		= TokenStream->GetSkipReturnCount( SkipInfo );
+				}
+				else
+				{	
+					// Loop again.
+					check( StackEntry->Data );
+					TokenReturnCount		= 0;
+				}
 			}
-		}
-		else if( REFERENCE_INFO.Type == GCRT_PersistentObject )
-		{
-			// We're dealing with an object reference.
-			UObject**	ObjectPtr	= (UObject**)(StackEntryData + Offset);
-			UObject*&	Object		= *ObjectPtr;
-			TokenReturnCount		= REFERENCE_INFO.ReturnCount;
+			break;
 
-			FReferenceChainLink TopRef(ReferenceTokenStreamIndex, InArrayProp != NULL ? EReferenceType::ArrayProperty : EReferenceType::Property, CurrentObject, Prop, Object);
-			AddToReferenceList(ReferenceList, TopRef);
-		}
-		else if( REFERENCE_INFO.Type == GCRT_FixedArray )
-		{
-			InArrayProp = Prop;
+			case GCRT_FixedArray:
+			{
+				InArrayProp = Prop;
 			
-			// We're dealing with a fixed size array
-			uint8* PreviousData	= StackEntryData;
-			StackEntry++;
-			StackEntryData				= PreviousData;
-			StackEntry->Data			= PreviousData;
-			StackEntry->Stride			= TokenStream->ReadStride( TokenStreamIndex );
-			StackEntry->Count			= TokenStream->ReadCount( TokenStreamIndex );
-			StackEntry->LoopStartIndex	= TokenStreamIndex;
-			TokenReturnCount			= 0;
-		}
-		else if( REFERENCE_INFO.Type == GCRT_AddStructReferencedObjects )
-		{
-			// We're dealing with a function call
-			void const*	StructPtr	= (void*)(StackEntryData + Offset);
-			TokenReturnCount		= REFERENCE_INFO.ReturnCount;
-			UScriptStruct::ICppStructOps::TPointerToAddStructReferencedObjects Func = (UScriptStruct::ICppStructOps::TPointerToAddStructReferencedObjects) TokenStream->ReadPointer( TokenStreamIndex );
-
-			FFindReferencerCollector ReferenceCollector(this, EReferenceType::StructARO, (void*)Func, CurrentObject);
-			Func(StructPtr, ReferenceCollector);
-
-			for (int32 i=0; i < ReferenceCollector.References.Num(); ++i)
+				// We're dealing with a fixed size array
+				uint8* PreviousData	= StackEntryData;
+				StackEntry++;
+				StackEntryData				= PreviousData;
+				StackEntry->Data			= PreviousData;
+				StackEntry->Stride			= TokenStream->ReadStride( TokenStreamIndex );
+				StackEntry->Count			= TokenStream->ReadCount( TokenStreamIndex );
+				StackEntry->LoopStartIndex	= TokenStreamIndex;
+				TokenReturnCount			= 0;
+			}
+			break;
+
+			case GCRT_AddStructReferencedObjects:
 			{
-				AddToReferenceList(ReferenceList, ReferenceCollector.References[i]);
+				// We're dealing with a function call
+				TokenReturnCount		= ReferenceInfo.ReturnCount;
+				UScriptStruct::ICppStructOps::TPointerToAddStructReferencedObjects Func = (UScriptStruct::ICppStructOps::TPointerToAddStructReferencedObjects) TokenStream->ReadPointer( TokenStreamIndex );
+
+				FFindReferencerCollector ReferenceCollector(this, EReferenceType::StructARO, (void*)Func, CurrentObject);
+				Func(StackEntryPtr, ReferenceCollector);
+
+				for (int32 i=0; i < ReferenceCollector.References.Num(); ++i)
+				{
+					AddToReferenceList(ReferenceList, ReferenceCollector.References[i]);
+				}
 			}
-		}
-		else if( REFERENCE_INFO.Type == GCRT_AddReferencedObjects )
-		{
-			// Static AddReferencedObjects function call.
-			void (*AddReferencedObjects)(UObject*, FReferenceCollector&) = (void(*)(UObject*, FReferenceCollector&))TokenStream->ReadPointer( TokenStreamIndex );
-			TokenReturnCount = REFERENCE_INFO.ReturnCount;
+			break;
+
+			case GCRT_AddReferencedObjects:
+			{
+				// Static AddReferencedObjects function call.
+				void (*AddReferencedObjects)(UObject*, FReferenceCollector&) = (void(*)(UObject*, FReferenceCollector&))TokenStream->ReadPointer( TokenStreamIndex );
+				TokenReturnCount = ReferenceInfo.ReturnCount;
 
-			FFindReferencerCollector ReferenceCollector(this, EReferenceType::ARO, (void*)AddReferencedObjects, CurrentObject);
-			AddReferencedObjects(CurrentObject, ReferenceCollector);
+				FFindReferencerCollector ReferenceCollector(this, EReferenceType::ARO, (void*)AddReferencedObjects, CurrentObject);
+				AddReferencedObjects(CurrentObject, ReferenceCollector);
+
+				for (int32 i=0; i < ReferenceCollector.References.Num(); ++i)
+				{
+					AddToReferenceList(ReferenceList, ReferenceCollector.References[i]);
+				}
+			}
+			break;
 
-			for (int32 i=0; i < ReferenceCollector.References.Num(); ++i)
+			case GCRT_AddTMapReferencedObjects:
 			{
-				AddToReferenceList(ReferenceList, ReferenceCollector.References[i]);
+				UMapProperty* MapProperty = (UMapProperty*)TokenStream->ReadPointer( TokenStreamIndex );
+				TokenReturnCount = ReferenceInfo.ReturnCount;
+				FFindReferencerCollector ReferenceCollector(this, EReferenceType::MapProperty, (void*)MapProperty, CurrentObject);
+				FSimpleObjectReferenceCollectorArchive CollectorArchive(CurrentObject, ReferenceCollector);
+				MapProperty->SerializeItem(CollectorArchive, StackEntryPtr, nullptr);
+
+				for (const FReferenceChainLink& Ref : ReferenceCollector.References)
+				{
+					AddToReferenceList(ReferenceList, Ref);
+				}
 			}
-		}
-		else if( REFERENCE_INFO.Type == GCRT_AddTMapReferencedObjects )
-		{
-			void*         Map         = StackEntryData + REFERENCE_INFO.Offset;
-			UMapProperty* MapProperty = (UMapProperty*)TokenStream->ReadPointer( TokenStreamIndex );
-			TokenReturnCount = REFERENCE_INFO.ReturnCount;
-			FFindReferencerCollector ReferenceCollector(this, EReferenceType::MapProperty, (void*)MapProperty, CurrentObject);
-			FSimpleObjectReferenceCollectorArchive CollectorArchive(CurrentObject, ReferenceCollector);
-			MapProperty->SerializeItem(CollectorArchive, Map, nullptr);
+			break;
 
-			for (const FReferenceChainLink& Ref : ReferenceCollector.References)
+			case GCRT_AddTSetReferencedObjects:
 			{
-				AddToReferenceList(ReferenceList, Ref);
+				USetProperty* SetProperty = (USetProperty*)TokenStream->ReadPointer( TokenStreamIndex );
+				TokenReturnCount = ReferenceInfo.ReturnCount;
+				FFindReferencerCollector ReferenceCollector(this, EReferenceType::SetProperty, (void*)SetProperty, CurrentObject);
+				FSimpleObjectReferenceCollectorArchive CollectorArchive(CurrentObject, ReferenceCollector);
+				SetProperty->SerializeItem(CollectorArchive, StackEntryPtr, nullptr);
+
+				for (const FReferenceChainLink& Ref : ReferenceCollector.References)
+				{
+					AddToReferenceList(ReferenceList, Ref);
+				}
 			}
-		}
-		else if( REFERENCE_INFO.Type == GCRT_AddTSetReferencedObjects )
-		{
-			void*         Set         = StackEntryData + REFERENCE_INFO.Offset;
-			USetProperty* SetProperty = (USetProperty*)TokenStream->ReadPointer( TokenStreamIndex );
-			TokenReturnCount = REFERENCE_INFO.ReturnCount;
-			FFindReferencerCollector ReferenceCollector(this, EReferenceType::SetProperty, (void*)SetProperty, CurrentObject);
-			FSimpleObjectReferenceCollectorArchive CollectorArchive(CurrentObject, ReferenceCollector);
-			SetProperty->SerializeItem(CollectorArchive, Set, nullptr);
+			break;
 
-			for (const FReferenceChainLink& Ref : ReferenceCollector.References)
+			default:
 			{
-				AddToReferenceList(ReferenceList, Ref);
+				UE_CLOG(ShouldOutputToLog(), LogReferenceChain, Fatal,TEXT("Unknown token"));
 			}
-		}
-		else if (REFERENCE_INFO.Type == GCRT_EndOfPointer)
-		{
-			TokenReturnCount = REFERENCE_INFO.ReturnCount;
-		}
-		else if( REFERENCE_INFO.Type == GCRT_EndOfStream )
-		{
-			// Break out of loop.
 			break;
 		}
-		else
-		{
-			UE_CLOG(ShouldOutputToLog(), LogReferenceChain, Fatal,TEXT("Unknown token"));
-		}
 	}
-	check(StackEntry == Stack.GetData());
-	return bRetVal;
 }
 
 FReferenceChainSearch::FReferenceChainSearch( UObject* InObjectToFind, uint32 Mode ) 
diff --git a/Engine/Source/Runtime/CoreUObject/Private/UObject/SavePackage.cpp b/Engine/Source/Runtime/CoreUObject/Private/UObject/SavePackage.cpp
--- a/Engine/Source/Runtime/CoreUObject/Private/UObject/SavePackage.cpp
+++ b/Engine/Source/Runtime/CoreUObject/Private/UObject/SavePackage.cpp
@@ -735,14 +735,14 @@ bool IsEditorOnlyObject(const UObject* InObject)
 	{
 		return false;
 	}
-	
+	check(InObject);
+
 	if (InObject->HasAnyMarks(OBJECTMARK_EditorOnly) || InObject->IsEditorOnly())
 	{
 		return true;
 	}
 
 	bool bResult = false;
-	check(InObject);
 	// If this is a package that is editor only or the object is in editor-only package,
 	// the object is editor-only too.
 	const bool bIsAPackage = InObject->IsA<UPackage>();
@@ -1121,6 +1121,7 @@ class FArchiveSaveTagImports : public FArchiveUObject
 public:
 	FLinkerSave* Linker;
 	TArray<UObject*> Dependencies;
+	TArray<UObject*> NativeDependencies;
 
 	FArchiveSaveTagImports(FLinkerSave* InLinker)
 		: Linker(InLinker)
@@ -1243,11 +1244,11 @@ FArchive& FArchiveSaveTagImports::operator<<( UObject*& Obj )
 	// Skip PendingKill objects and objects that are both not for client and not for server when cooking.
 	if (Obj && !Obj->IsPendingKill() && (!IsCooking() || !Obj->HasAllMarks(ObjectMarks)) && !Obj->HasAnyMarks(OBJECTMARK_EditorOnly))
 	{
-		if( !Obj->HasAnyFlags(RF_Transient) || Obj->IsNative() )
+		bool bIsNative = Obj->IsNative();
+		if( !Obj->HasAnyFlags(RF_Transient) || bIsNative)
 		{
 			// remember it as a dependency, unless it's a top level package or native
 			const bool bIsTopLevelPackage = Obj->GetOuter() == NULL && dynamic_cast<UPackage*>(Obj);
-			bool bIsNative = Obj->IsNative();
 			UObject* Outer = Obj->GetOuter();
 
 			const bool bIsEditorOnly = IsCookingForNoEditorDataPlatform(*this) && IsEditorOnlyObject(Obj);
@@ -1272,6 +1273,10 @@ FArchive& FArchiveSaveTagImports::operator<<( UObject*& Obj )
 			{
 				Dependencies.AddUnique(Obj);
 			}
+			else if (!bIsTopLevelPackage)
+			{
+				NativeDependencies.AddUnique(Obj);
+			}
 
 			if( !Obj->HasAnyMarks(OBJECTMARK_TagExp) )  
 			{
@@ -1298,6 +1303,9 @@ FArchive& FArchiveSaveTagImports::operator<<( UObject*& Obj )
 							TArray<UObject*> CurrentSubobjects;
 							TArray<UObject*> NextSubobjects;
 
+							UE_LOG(LogSavePackage, Verbose, TEXT("******Class import %s"), *Obj->GetFullName());
+
+
 							// Recursively search for subobjects. Only care about ones that have a full subobject chain as some nested objects are set wrong
 							GetObjectsWithOuter(ClassObj, NextSubobjects, false);
 							GetObjectsWithOuter(CDO, NextSubobjects, false);
@@ -1313,6 +1321,10 @@ FArchive& FArchiveSaveTagImports::operator<<( UObject*& Obj )
 										ObjectTemplates.Add(SubObj);
 										GetObjectsWithOuter(SubObj, NextSubobjects, false);
 									}
+									else
+									{
+										UE_LOG(LogSavePackage, Verbose, TEXT("      rejected subobj because it was not RF_DefaultSubObject | RF_ArchetypeObject %s"), *SubObj->GetFullName());
+									}
 								}
 							}
 
@@ -1322,13 +1334,20 @@ FArchive& FArchiveSaveTagImports::operator<<( UObject*& Obj )
 								{
 									// Explicitly mark as editor only so it will be skipped later
 									ObjTemplate->Mark(OBJECTMARK_EditorOnly);
+									UE_LOG(LogSavePackage, Verbose, TEXT("      rejected subobj because it IsEditorOnlyObject %s"), *ObjTemplate->GetFullName());
 								}
 								else if (!ObjTemplate->HasAnyFlags(RF_Transient) && !ObjTemplate->IsPendingKill())
 								{
 									// Correctly mark objects that shouldn't be cooked
 									ConditionallyExcludeObjectForTarget(ObjTemplate, *this);
+									UE_CLOG(ObjTemplate->HasAnyMarks(OBJECTMARK_NotForClient), LogSavePackage, Verbose, TEXT("      will reject subobj because it is OBJECTMARK_NotForClient %s"), *ObjTemplate->GetFullName());
 
 									ObjTemplate->Mark(OBJECTMARK_TagImp);
+									UE_LOG(LogSavePackage, Verbose, TEXT("      OBJECTMARK_TagImp %s"), *ObjTemplate->GetFullName());
+								}
+								else
+								{
+									UE_LOG(LogSavePackage, Verbose, TEXT("      rejected subobj because it transient or pendingkill %s"), *ObjTemplate->GetFullName());
 								}
 							}
 						}
@@ -2850,7 +2869,7 @@ public:
 									}
 									else
 									{
-										UByteProperty* ByteProp = dynamic_cast<UByteProperty*>(Object);
+								UByteProperty* ByteProp = dynamic_cast<UByteProperty*>(Object);
 										if (ByteProp && ByteProp->Enum)
 										{
 											Enum = ByteProp->Enum;
@@ -2858,9 +2877,9 @@ public:
 									}
 
 									if (Enum)
-									{
+								{
 										HandleDependency(Enum, /*bProcessObject =*/true);
-									}
+								}
 								}
 
 								// a normal field - property, enum, const; just insert it into the list and keep going
@@ -3945,7 +3964,7 @@ ESavePackageResult UPackage::Save(UPackage* InOuter, UObject* Base, EObjectFlags
 					Arguments.Add(TEXT("Name"), FText::FromString(NewPath));
 					ErrorText = FText::Format(NSLOCTEXT("SavePackage", "CannotSaveAssetPartiallyLoaded", "Asset '{Name}' cannot be saved as it has only been partially loaded"), Arguments);
 				}
-				Error->Logf(ELogVerbosity::Warning, *ErrorText.ToString());
+				Error->Logf(ELogVerbosity::Warning, TEXT("%s"), *ErrorText.ToString());
 			}
 			return ESavePackageResult::Error;
 		}
@@ -3971,7 +3990,7 @@ ESavePackageResult UPackage::Save(UPackage* InOuter, UObject* Base, EObjectFlags
 						Arguments.Add(TEXT("Name"), FText::FromString(NewPath));
 						ErrorText = FText::Format(NSLOCTEXT("SavePackage", "AssetSaveNotAllowed", "Asset '{Name}' is not allowed to save (see log for reason)"), Arguments);
 					}
-					Error->Logf(ELogVerbosity::Warning, *ErrorText.ToString());
+					Error->Logf(ELogVerbosity::Warning, TEXT("%s"), *ErrorText.ToString());
 				}
 				return ESavePackageResult::Error;
 			}
@@ -3995,7 +4014,7 @@ ESavePackageResult UPackage::Save(UPackage* InOuter, UObject* Base, EObjectFlags
 					Arguments.Add(TEXT("Name"), FText::FromString(NewPath));
 					ErrorText = FText::Format(NSLOCTEXT("SavePackage", "CannotSaveAssetConformIncompatibility", "Conformed Asset '{Name}' cannot be saved as it is incompatible with the original"), Arguments);
 				}
-				Error->Logf(ELogVerbosity::Error, *ErrorText.ToString());
+				Error->Logf(ELogVerbosity::Error, TEXT("%s"), *ErrorText.ToString());
 			}
 			return ESavePackageResult::Error;
 		}
@@ -4046,6 +4065,9 @@ ESavePackageResult UPackage::Save(UPackage* InOuter, UObject* Base, EObjectFlags
 
 		TArray<UObject*> CachedObjects;
 
+		// structure to track what every export needs to import (native only)
+		TMap<UObject*, TArray<UObject*> > NativeObjectDependencies;
+
 		// Size of serialized out package in bytes. This is before compression.
 		int32 PackageSize = INDEX_NONE;
 		{
@@ -4280,8 +4302,7 @@ ESavePackageResult UPackage::Save(UPackage* InOuter, UObject* Base, EObjectFlags
 						ImportTagger << Class;
 
 						UObject* Template = Obj->GetArchetype();
-						if ( Template 
-							)
+						if (Template)
 						{
 							// If we're not cooking for the event driven loader, exclude the CDO
 							if (Template != Class->GetDefaultObject() || (IsEventDrivenLoaderEnabledInCookedBuilds() && TargetPlatform))
@@ -4290,6 +4311,34 @@ ESavePackageResult UPackage::Save(UPackage* InOuter, UObject* Base, EObjectFlags
 							}
 						}
 
+						if (IsEventDrivenLoaderEnabledInCookedBuilds() && TargetPlatform)
+						{
+							TArray<UObject*> Deps;
+							Obj->GetPreloadDependencies(Deps);
+							for (UObject* Dep : Deps)
+							{
+								if (!Dep)
+								{
+									continue;
+								}
+								bool bIsNative = Dep->IsNative();
+
+								if (
+									Dep->GetOutermost()->GetFName() != GLongCoreUObjectPackageName && // We assume nothing in coreuobject ever loads assets in a constructor
+									(!Dep->HasAllFlags(RF_Transient) || bIsNative) &&
+									!Dep->IsPendingKill() &&
+									!Dep->HasAnyMarks(OBJECTMARK_TagExp))
+								{
+									bool bNotFiltered = !Dep->HasAllMarks(ObjectMarks) && (!(Linker->Summary.PackageFlags & PKG_FilterEditorOnly) || !IsEditorOnlyObject(Dep));
+									if (bNotFiltered)
+									{
+										Dep->Mark(OBJECTMARK_TagImp);
+									}
+								}
+							}
+						}
+
+
 						if( Obj->IsIn(GetTransientPackage()) )
 						{
 							UE_LOG(LogSavePackage, Fatal, TEXT("%s"), *FString::Printf( TEXT("Transient object imported: %s"), *Obj->GetFullName() ) );
@@ -4307,10 +4356,20 @@ ESavePackageResult UPackage::Save(UPackage* InOuter, UObject* Base, EObjectFlags
 									ImportTagger.Dependencies.RemoveAtSwap(DependencyIndex);
 								}								
 							}
+							for (int32 DependencyIndex = ImportTagger.NativeDependencies.Num() - 1; DependencyIndex >= 0; DependencyIndex--)
+							{
+								UObject* DependencyObj = ImportTagger.NativeDependencies[DependencyIndex];
+								if (DependencyObj->HasAllMarks(ObjectMarks))
+								{
+									DependencyObj->UnMark(OBJECTMARK_TagImp);
+									ImportTagger.NativeDependencies.RemoveAtSwap(DependencyIndex);
+								}
+							}
 						}
 
 						// add the list of dependencies to the dependency map
 						ObjectDependencies.Add(Obj, ImportTagger.Dependencies);
+						NativeObjectDependencies.Add(Obj, ImportTagger.NativeDependencies);
 
 						if (Obj->GetClass() != UObjectRedirector::StaticClass())
 						{
@@ -4396,13 +4455,16 @@ ESavePackageResult UPackage::Save(UPackage* InOuter, UObject* Base, EObjectFlags
 							// When Startup is cooked, it will pull in C and A. When M is cooked, it will pull in B, but not A, because
 							// A was already marked by the cooker. M.xxx now has a private import to A, which is normally illegal, hence
 							// the OBJECTMARK_MarkedByCooker check below
+							UPackage* ObjPackage = Obj->GetOutermost();
 							if( !Obj->HasAnyFlags(RF_Public) && !Obj->HasAnyFlags(RF_Transient))
 							{
-								PrivateObjects.Add(Obj);
+								if (!IsEventDrivenLoaderEnabledInCookedBuilds() || !TargetPlatform || !ObjPackage->HasAnyPackageFlags(PKG_CompiledIn))
+								{
+									PrivateObjects.Add(Obj);
+								}
 							}
 
 							// See whether the object we are referencing is in another map package.
-							UPackage* ObjPackage = dynamic_cast<UPackage*>(Obj->GetOutermost());
 							if( ObjPackage && ObjPackage->ContainsMap() )
 							{
 								if ( ObjPackage != Obj && Obj->GetFName() != NAME_PersistentLevel && Obj->GetClass()->GetFName() != WorldClassName )
@@ -4415,7 +4477,7 @@ ESavePackageResult UPackage::Save(UPackage* InOuter, UObject* Base, EObjectFlags
 
 										if (!(SaveFlags & SAVE_NoError))
 										{
-											Error->Logf(ELogVerbosity::Warning, *FText::Format( NSLOCTEXT( "Core", "SavePackageObjInAnotherMap", "Object '{0}' is in another map" ), FText::FromString( *Obj->GetFullName() ) ).ToString() );
+											Error->Logf(ELogVerbosity::Warning, TEXT("%s"), *FText::Format( NSLOCTEXT( "Core", "SavePackageObjInAnotherMap", "Object '{0}' is in another map" ), FText::FromString( *Obj->GetFullName() ) ).ToString() );
 										}
 									}
 								}
@@ -4829,7 +4891,6 @@ ESavePackageResult UPackage::Save(UPackage* InOuter, UObject* Base, EObjectFlags
 					ExportToIndexMap.Add(Linker->ExportMap[ExpIndex].Object, FPackageIndex::FromExport(ExpIndex));
 				}
 
-				
 				// go back over the (now sorted) exports and fill out the DependsMap
 				for (int32 ExpIndex = 0; ExpIndex < Linker->ExportMap.Num(); ExpIndex++)
 				{
@@ -4840,7 +4901,7 @@ ESavePackageResult UPackage::Save(UPackage* InOuter, UObject* Base, EObjectFlags
 						UE_LOG(LogSavePackage, Warning, TEXT("Object is missing for an export, unable to save dependency map. Most likely this is caused my conforming against a package that is missing this object. See log for more info"));
 						if (!(SaveFlags & SAVE_NoError))
 						{
-							Error->Logf(ELogVerbosity::Warning, *FText::Format( NSLOCTEXT( "Core", "SavePackageObjectIsMissingExport", "Object is missing for an export, unable to save dependency map for asset '{0}'. Most likely this is caused my conforming against a asset that is missing this object. See log for more info" ), FText::FromString( FString( Filename ) ) ).ToString() );
+							Error->Logf(ELogVerbosity::Warning, TEXT("%s"), *FText::Format( NSLOCTEXT( "Core", "SavePackageObjectIsMissingExport", "Object is missing for an export, unable to save dependency map for asset '{0}'. Most likely this is caused my conforming against a asset that is missing this object. See log for more info" ), FText::FromString( FString( Filename ) ) ).ToString() );
 						}
 						continue;
 					}
@@ -4953,7 +5014,7 @@ ESavePackageResult UPackage::Save(UPackage* InOuter, UObject* Base, EObjectFlags
 
 						if (!(SaveFlags & SAVE_NoError))
 						{
-							Error->Logf(ELogVerbosity::Warning, *FText::Format( NSLOCTEXT( "Core", "SavePackageNoMap", "Attempting to save a map asset '{0}' that does not contain a map object" ), FText::FromString( FString( Filename ) ) ).ToString() );
+							Error->Logf(ELogVerbosity::Warning, TEXT("%s"), *FText::Format( NSLOCTEXT( "Core", "SavePackageNoMap", "Attempting to save a map asset '{0}' that does not contain a map object" ), FText::FromString( FString( Filename ) ) ).ToString() );
 						}
 						Success = false;
 					}
@@ -5200,26 +5261,38 @@ ESavePackageResult UPackage::Save(UPackage* InOuter, UObject* Base, EObjectFlags
 					const EObjectMark ObjectMarks = UPackage::GetObjectMarksForTargetPlatform(Linker->CookingTarget(), Linker->IsCooking());
 					Linker->Summary.PreloadDependencyCount = 0;
 
-					auto IncludeObjectAsDependency = [Linker,ObjectMarks](TSet<FPackageIndex>& AddTo, UObject* ToTest, UObject* ForObj, bool bMandatory)
+					auto IncludeObjectAsDependency = [Linker,ObjectMarks](int32 CallSite, TSet<FPackageIndex>& AddTo, UObject* ToTest, UObject* ForObj, bool bMandatory, bool bOnlyIfInLinkerTable = false)
 					{
 						// Skip transient, editor only, and excluded client/server objects
-						if (ToTest && !ToTest->HasAllFlags(RF_Transient) && !ToTest->IsPendingKill())
+						if (ToTest)
 						{
 							UPackage* Outermost = ToTest->GetOutermost();
 							check(Outermost);
-							if (Outermost->HasAnyPackageFlags(EPackageFlags::PKG_CompiledIn))
+							if (Outermost->GetFName() == GLongCoreUObjectPackageName)
 							{
-								return; // we never add dependencies for things that are compiled in
+								return; // We assume nothing in coreuobject ever loads assets in a constructor
 							}
-							bool bNotFiltered = !ToTest->HasAllMarks(ObjectMarks) && (!(Linker->Summary.PackageFlags & PKG_FilterEditorOnly) || !ToTest->HasAnyMarks(OBJECTMARK_EditorOnly));
+							FPackageIndex Index = Linker->MapObject(ToTest);
+							if (Index.IsNull() && bOnlyIfInLinkerTable)
+							{
+								return;
+							}
+							if (!Index.IsNull() && (ToTest->HasAllFlags(RF_Transient) && !ToTest->IsNative()))
+							{
+								UE_LOG(LogSavePackage, Warning, TEXT("A dependency '%s' of '%s' is in the linker table, but is transient. We will keep the dependency anyway (%d)."), *ToTest->GetFullName(), *ForObj->GetFullName(), CallSite);
+							}
+							if (!Index.IsNull() && ToTest->IsPendingKill())
+							{
+								UE_LOG(LogSavePackage, Warning, TEXT("A dependency '%s' of '%s' is in the linker table, but is pending kill. We will keep the dependency anyway (%d)."), *ToTest->GetFullName(), *ForObj->GetFullName(), CallSite);
+							}
+							bool bNotFiltered = !ToTest->HasAllMarks(ObjectMarks) && (!(Linker->Summary.PackageFlags & PKG_FilterEditorOnly) || !IsEditorOnlyObject(ToTest));
 							if (bMandatory && !bNotFiltered)
 							{
-								UE_LOG(LogSavePackage, Warning, TEXT("A dependency '%s' of '%s' was filtered, but is mandatory. This indicates a problem with editor only stripping. We will keep the dependency anyway."), *ToTest->GetFullName(), *ForObj->GetFullName());
+								UE_LOG(LogSavePackage, Warning, TEXT("A dependency '%s' of '%s' was filtered, but is mandatory. This indicates a problem with editor only stripping. We will keep the dependency anyway (%d)."), *ToTest->GetFullName(), *ForObj->GetFullName(), CallSite);
 								bNotFiltered = true;
 							}
 							if (bNotFiltered)
 							{
-								FPackageIndex Index = Linker->MapObject(ToTest);
 								if (!Index.IsNull())
 								{
 									AddTo.Add(Index);
@@ -5227,7 +5300,8 @@ ESavePackageResult UPackage::Save(UPackage* InOuter, UObject* Base, EObjectFlags
 								}
 								else
 								{
-									UE_LOG(LogSavePackage, Fatal, TEXT("A dependency '%s' of '%s' was not actually in the linker tables and so will be ignored."), *ToTest->GetFullName(), *ForObj->GetFullName());
+									UE_CLOG(Outermost->HasAnyPackageFlags(PKG_CompiledIn), LogSavePackage, Verbose, TEXT("A compiled in dependency '%s' of '%s' was not actually in the linker tables and so will be ignored (%d)."), *ToTest->GetFullName(), *ForObj->GetFullName(), CallSite);
+									UE_CLOG(!Outermost->HasAnyPackageFlags(PKG_CompiledIn), LogSavePackage, Fatal, TEXT("A dependency '%s' of '%s' was not actually in the linker tables and so will be ignored (%d)."), *ToTest->GetFullName(), *ForObj->GetFullName(), CallSite);
 								}
 							}
 							check(!bMandatory);
@@ -5242,7 +5316,7 @@ ESavePackageResult UPackage::Save(UPackage* InOuter, UObject* Base, EObjectFlags
 							if (ToTest)
 							{
 								UPackage* Outermost = ToTest->GetOutermost();
-								if (Outermost && !Outermost->HasAnyPackageFlags(EPackageFlags::PKG_CompiledIn))
+								if (Outermost && Outermost->GetFName() != GLongCoreUObjectPackageName) // We assume nothing in coreuobject ever loads assets in a constructor
 								{
 									AddTo.Add(Dep);
 								}
@@ -5252,6 +5326,7 @@ ESavePackageResult UPackage::Save(UPackage* InOuter, UObject* Base, EObjectFlags
 
 
 					TArray<UObject*> Subobjects;
+					TArray<UObject*> Deps;
 					for (int32 i = 0; i < Linker->ExportMap.Num(); i++)
 					{
 						FObjectExport& Export = Linker->ExportMap[i];
@@ -5262,7 +5337,7 @@ ESavePackageResult UPackage::Save(UPackage* InOuter, UObject* Base, EObjectFlags
 							{
 								IncludeIndexAsDependency(SerializationBeforeCreateDependencies, Export.ClassIndex);
 								UObject* CDO = Export.Object->GetArchetype();
-								IncludeObjectAsDependency(SerializationBeforeCreateDependencies, CDO, Export.Object, true);
+								IncludeObjectAsDependency(1, SerializationBeforeCreateDependencies, CDO, Export.Object, true);
 								Subobjects.Reset();
 								GetObjectsWithOuter(CDO, Subobjects);
 								for (UObject* SubObj : Subobjects)
@@ -5270,37 +5345,93 @@ ESavePackageResult UPackage::Save(UPackage* InOuter, UObject* Base, EObjectFlags
 									// Only include subobject archetypes
 									if (SubObj->HasAnyFlags(RF_DefaultSubObject | RF_ArchetypeObject))
 									{
-										IncludeObjectAsDependency(SerializationBeforeCreateDependencies, SubObj, Export.Object, false);
+										while (SubObj->HasAnyFlags(RF_Transient)) // transient components are stripped by the ICH, so find the one it will really use at runtime
+										{
+											UObject* SubObjArch = SubObj->GetArchetype();
+											if (SubObjArch->GetClass()->HasAnyClassFlags(CLASS_Native | CLASS_Intrinsic))
+											{
+												break;
+											}
+											SubObj = SubObjArch;
+										}
+										if (!SubObj->IsPendingKill())
+										{
+											IncludeObjectAsDependency(2, SerializationBeforeCreateDependencies, SubObj, Export.Object, false);
+										}
 									}
 								}
 							}
 							TSet<FPackageIndex> SerializationBeforeSerializationDependencies;
 							{
-								TArray<UObject*> Deps;
+								Deps.Reset();
 								Export.Object->GetPreloadDependencies(Deps);
 
 								for (UObject* Obj : Deps)
 								{
-									IncludeObjectAsDependency(SerializationBeforeSerializationDependencies, Obj, Export.Object, false);
+									IncludeObjectAsDependency(3, SerializationBeforeSerializationDependencies, Obj, Export.Object, false, true);
+								}
+								if (Export.Object->HasAnyFlags(RF_ArchetypeObject | RF_ClassDefaultObject))
+								{
+									UObject *Outer = Export.Object->GetOuter();
+									if (!Outer->IsA(UPackage::StaticClass()))
+									{
+										IncludeObjectAsDependency(4, SerializationBeforeSerializationDependencies, Outer, Export.Object, true);
+									}
 								}
-								if (Export.Object->HasAnyFlags(RF_ArchetypeObject | RF_ClassDefaultObject) && !Export.Object->GetOuter()->IsA(UPackage::StaticClass()))
+								if (Export.Object->IsA(UClass::StaticClass()))
 								{
-									IncludeObjectAsDependency(SerializationBeforeSerializationDependencies, Export.Object->GetOuter(), Export.Object, true);
+									// we need to load archetypes of our subobjects before we load the class
+									UObject* CDO = CastChecked<UClass>(Export.Object)->GetDefaultObject();
+									Subobjects.Reset();
+									GetObjectsWithOuter(CDO, Subobjects);
+									for (UObject* SubObj : Subobjects)
+									{
+										// Only include subobject archetypes
+										if (SubObj->HasAnyFlags(RF_DefaultSubObject | RF_ArchetypeObject))
+										{
+											SubObj = SubObj->GetArchetype();
+											while (SubObj->HasAnyFlags(RF_Transient)) // transient components are stripped by the ICH, so find the one it will really use at runtime
+											{
+												UObject* SubObjArch = SubObj->GetArchetype();
+												if (SubObjArch->GetClass()->HasAnyClassFlags(CLASS_Native | CLASS_Intrinsic))
+												{
+													break;
+												}
+												SubObj = SubObjArch;
+											}
+											if (!SubObj->IsPendingKill())
+											{
+												IncludeObjectAsDependency(5, SerializationBeforeSerializationDependencies, SubObj, Export.Object, false);
+											}
+										}
+									}
 								}
 							}
 
 							TSet<FPackageIndex> CreateBeforeSerializationDependencies;
 							{
-								TArray<FPackageIndex>& Depends = Linker->DependsMap[i];
-								for (FPackageIndex Dep : Depends)
+								UClass* Class = Cast<UClass>(Export.Object);
+								UObject* ClassCDO = Class ? Class->GetDefaultObject() : nullptr;
 								{
-									UClass* Class = Cast<UClass>(Export.Object);
-									UObject* ToTest = Dep.IsExport() ? Linker->Exp(Dep).Object : Linker->Imp(Dep).XObject;
-									if (Class && ToTest == Class->GetDefaultObject())
+									TArray<FPackageIndex>& Depends = Linker->DependsMap[i];
+									for (FPackageIndex Dep : Depends)
 									{
-										continue;
+										UObject* ToTest = Dep.IsExport() ? Linker->Exp(Dep).Object : Linker->Imp(Dep).XObject;
+										if (ToTest != ClassCDO)
+										{
+											IncludeIndexAsDependency(CreateBeforeSerializationDependencies, Dep);
+										}
+									}
+								}
+								{
+									TArray<UObject*>& NativeDeps = NativeObjectDependencies[Export.Object];
+									for (UObject* ToTest : NativeDeps)
+									{
+										if (ToTest != ClassCDO)
+										{
+											IncludeObjectAsDependency(6, CreateBeforeSerializationDependencies, ToTest, Export.Object, false, true);
+										}
 									}
-									IncludeIndexAsDependency(CreateBeforeSerializationDependencies, Dep);
 								}
 							}
 							TSet<FPackageIndex> CreateBeforeCreateDependencies;
@@ -5387,7 +5518,7 @@ ESavePackageResult UPackage::Save(UPackage* InOuter, UObject* Base, EObjectFlags
 							}
 						}
 					}
-					UE_LOG(LogSavePackage, Log, TEXT("Saved %d dependencies for %d exports."), Linker->Summary.PreloadDependencyCount, Linker->ExportMap.Num());
+					UE_LOG(LogSavePackage, Verbose, TEXT("Saved %d dependencies for %d exports."), Linker->Summary.PreloadDependencyCount, Linker->ExportMap.Num());
 				}
 
 				
@@ -5832,7 +5963,7 @@ ESavePackageResult UPackage::Save(UPackage* InOuter, UObject* Base, EObjectFlags
 						else
 						{
 							UE_LOG(LogSavePackage, Error, TEXT("%s"), *FString::Printf( TEXT("Error saving '%s'"), Filename ) );
-							Error->Logf(ELogVerbosity::Warning, *FText::Format( NSLOCTEXT( "Core", "SaveWarning", "Error saving '{0}'" ), FText::FromString( FString( Filename ) ) ).ToString() );
+							Error->Logf(ELogVerbosity::Warning, TEXT("%s"), *FText::Format( NSLOCTEXT( "Core", "SaveWarning", "Error saving '{0}'" ), FText::FromString( FString( Filename ) ) ).ToString() );
 						}
 					}
 					else
@@ -5885,7 +6016,7 @@ ESavePackageResult UPackage::Save(UPackage* InOuter, UObject* Base, EObjectFlags
 									FFormatNamedArguments Arguments;
 									Arguments.Add(TEXT("FileName"), FText::FromString( BaseFilename ));
 									Arguments.Add(TEXT("MaxLength"), FText::AsNumber( MaxFilenameLength ));
-									Error->Logf(ELogVerbosity::Warning, *FText::Format( NSLOCTEXT( "Core", "Error_FilenameIsTooLongForCooking", "Filename '{FileName}' is too long; this may interfere with cooking for consoles.  Unreal filenames should be no longer than {MaxLength} characters." ), Arguments ).ToString() );
+									Error->Logf(ELogVerbosity::Warning, TEXT("%s"), *FText::Format( NSLOCTEXT( "Core", "Error_FilenameIsTooLongForCooking", "Filename '{FileName}' is too long; this may interfere with cooking for consoles.  Unreal filenames should be no longer than {MaxLength} characters." ), Arguments ).ToString() );
 								}
 							}
 						}
