UE4 - Fix VS2013 issue with WEX and core merge to main.

diff --git a/Engine/Source/Runtime/CoreUObject/Private/UObject/SavePackage.cpp b/Engine/Source/Runtime/CoreUObject/Private/UObject/SavePackage.cpp
--- a/Engine/Source/Runtime/CoreUObject/Private/UObject/SavePackage.cpp
+++ b/Engine/Source/Runtime/CoreUObject/Private/UObject/SavePackage.cpp
@@ -5261,7 +5261,7 @@ ESavePackageResult UPackage::Save(UPackage* InOuter, UObject* Base, EObjectFlags
 					const EObjectMark ObjectMarks = UPackage::GetObjectMarksForTargetPlatform(Linker->CookingTarget(), Linker->IsCooking());
 					Linker->Summary.PreloadDependencyCount = 0;
 
-					auto IncludeObjectAsDependency = [Linker,ObjectMarks](int32 CallSite, TSet<FPackageIndex>& AddTo, UObject* ToTest, UObject* ForObj, bool bMandatory, bool bOnlyIfInLinkerTable = false)
+					auto IncludeObjectAsDependency = [Linker,ObjectMarks](int32 CallSite, TSet<FPackageIndex>& AddTo, UObject* ToTest, UObject* ForObj, bool bMandatory, bool bOnlyIfInLinkerTable)
 					{
 						// Skip transient, editor only, and excluded client/server objects
 						if (ToTest)
@@ -5337,7 +5337,7 @@ ESavePackageResult UPackage::Save(UPackage* InOuter, UObject* Base, EObjectFlags
 							{
 								IncludeIndexAsDependency(SerializationBeforeCreateDependencies, Export.ClassIndex);
 								UObject* CDO = Export.Object->GetArchetype();
-								IncludeObjectAsDependency(1, SerializationBeforeCreateDependencies, CDO, Export.Object, true);
+								IncludeObjectAsDependency(1, SerializationBeforeCreateDependencies, CDO, Export.Object, true, false);
 								Subobjects.Reset();
 								GetObjectsWithOuter(CDO, Subobjects);
 								for (UObject* SubObj : Subobjects)
@@ -5356,7 +5356,7 @@ ESavePackageResult UPackage::Save(UPackage* InOuter, UObject* Base, EObjectFlags
 										}
 										if (!SubObj->IsPendingKill())
 										{
-											IncludeObjectAsDependency(2, SerializationBeforeCreateDependencies, SubObj, Export.Object, false);
+											IncludeObjectAsDependency(2, SerializationBeforeCreateDependencies, SubObj, Export.Object, false, false);
 										}
 									}
 								}
@@ -5375,7 +5375,7 @@ ESavePackageResult UPackage::Save(UPackage* InOuter, UObject* Base, EObjectFlags
 									UObject *Outer = Export.Object->GetOuter();
 									if (!Outer->IsA(UPackage::StaticClass()))
 									{
-										IncludeObjectAsDependency(4, SerializationBeforeSerializationDependencies, Outer, Export.Object, true);
+										IncludeObjectAsDependency(4, SerializationBeforeSerializationDependencies, Outer, Export.Object, true, false);
 									}
 								}
 								if (Export.Object->IsA(UClass::StaticClass()))
@@ -5401,7 +5401,7 @@ ESavePackageResult UPackage::Save(UPackage* InOuter, UObject* Base, EObjectFlags
 											}
 											if (!SubObj->IsPendingKill())
 											{
-												IncludeObjectAsDependency(5, SerializationBeforeSerializationDependencies, SubObj, Export.Object, false);
+												IncludeObjectAsDependency(5, SerializationBeforeSerializationDependencies, SubObj, Export.Object, false, false);
 											}
 										}
 									}
