Also fix it so redirecting to /package.object works correctly if the old path had an outer in it

diff --git a/Engine/Source/Runtime/CoreUObject/Private/UObject/CoreRedirects.cpp b/Engine/Source/Runtime/CoreUObject/Private/UObject/CoreRedirects.cpp
--- a/Engine/Source/Runtime/CoreUObject/Private/UObject/CoreRedirects.cpp
+++ b/Engine/Source/Runtime/CoreUObject/Private/UObject/CoreRedirects.cpp
@@ -149,7 +149,7 @@ int32 FCoreRedirectObjectName::MatchScore(const FCoreRedirectObjectName& Other)
 
 bool FCoreRedirectObjectName::HasValidCharacters() const
 {
-	static FString InvalidRedirectCharacters = TEXT("\"' ,:|&!~\n\r\t@#(){}[]=;^%$`");
+	static FString InvalidRedirectCharacters = TEXT("\"' ,|&!~\n\r\t@#(){}[]=;^%$`");
 
 	return ObjectName.IsValidXName(InvalidRedirectCharacters) && OuterName.IsValidXName(InvalidRedirectCharacters) && PackageName.IsValidXName(InvalidRedirectCharacters);
 }
@@ -161,13 +161,25 @@ bool FCoreRedirectObjectName::ExpandNames(const FString& InString, FName& OutNam
 	FullString.TrimTrailing();
 
 	// Parse (/path.)?(outerchain.)?(name) where path and outerchain are optional
+	// We also need to support (/path.)?(singleouter:)?(name) because the second delimiter in a chain is : for historical reasons
 	int32 SlashIndex = INDEX_NONE;
 	int32 FirstPeriodIndex = INDEX_NONE;
 	int32 LastPeriodIndex = INDEX_NONE;
+	int32 FirstColonIndex = INDEX_NONE;
+	int32 LastColonIndex = INDEX_NONE;
 
 	FullString.FindChar('/', SlashIndex);
 
-	if (!FullString.FindChar('.', FirstPeriodIndex))
+	FullString.FindChar('.', FirstPeriodIndex);
+	FullString.FindChar(':', FirstColonIndex);
+
+	if (FirstColonIndex != INDEX_NONE && (FirstPeriodIndex == INDEX_NONE || FirstColonIndex < FirstPeriodIndex))
+	{
+		// If : is before . treat it as the first period
+		FirstPeriodIndex = FirstColonIndex;
+	}
+
+	if (FirstPeriodIndex == INDEX_NONE)
 	{
 		// If start with /, fill in package name, otherwise name
 		if (SlashIndex != INDEX_NONE)
@@ -182,6 +194,13 @@ bool FCoreRedirectObjectName::ExpandNames(const FString& InString, FName& OutNam
 	}
 
 	FullString.FindLastChar('.', LastPeriodIndex);
+	FullString.FindLastChar(':', LastColonIndex);
+
+	if (LastColonIndex != INDEX_NONE && (LastPeriodIndex == INDEX_NONE || LastColonIndex > LastPeriodIndex))
+	{
+		// If : is after . treat it as the last period
+		LastPeriodIndex = LastColonIndex;
+	}
 
 	if (SlashIndex == INDEX_NONE)
 	{
@@ -209,13 +228,32 @@ FString FCoreRedirectObjectName::CombineNames(FName NewName, FName NewOuter, FNa
 
 	if (NewOuter != NAME_None)
 	{
-		if (NewPackage != NAME_None)
+		// If Outer is simple, need to use : instead of . because : is used for second delimiter only
+		FString OuterString = NewOuter.ToString();
+
+		int32 DelimIndex = INDEX_NONE;
+
+		if (OuterString.FindChar('.', DelimIndex) || OuterString.FindChar(':', DelimIndex))
 		{
-			return FString::Printf(TEXT("%s.%s.%s"), *NewPackage.ToString(), *NewOuter.ToString(), *NewName.ToString());
+			if (NewPackage != NAME_None)
+			{
+				return FString::Printf(TEXT("%s.%s.%s"), *NewPackage.ToString(), *NewOuter.ToString(), *NewName.ToString());
+			}
+			else
+			{
+				return FString::Printf(TEXT("%s.%s"), *NewOuter.ToString(), *NewName.ToString());
+			}
 		}
 		else
 		{
-			return FString::Printf(TEXT("%s.%s"), *NewOuter.ToString(), *NewName.ToString());
+			if (NewPackage != NAME_None)
+			{
+				return FString::Printf(TEXT("%s.%s:%s"), *NewPackage.ToString(), *NewOuter.ToString(), *NewName.ToString());
+			}
+			else
+			{
+				return FString::Printf(TEXT("%s:%s"), *NewOuter.ToString(), *NewName.ToString());
+			}
 		}
 	}
 	else if (NewPackage != NAME_None)
@@ -350,7 +388,8 @@ FCoreRedirectObjectName FCoreRedirect::RedirectName(const FCoreRedirectObjectNam
 			ModifyName.ObjectName = NewName.ObjectName;
 		}
 	}
-	if (OldName.OuterName != NewName.OuterName)
+	// If package name is specified, copy outer as it was set to null explicitly
+	if (OldName.OuterName != NewName.OuterName || NewName.PackageName != NAME_None)
 	{
 		if (IsSubstringMatch())
 		{
