UE4 - Fixed rare ensure cooking for the EDL.

diff --git a/Engine/Source/Runtime/CoreUObject/Private/UObject/SavePackage.cpp b/Engine/Source/Runtime/CoreUObject/Private/UObject/SavePackage.cpp
--- a/Engine/Source/Runtime/CoreUObject/Private/UObject/SavePackage.cpp
+++ b/Engine/Source/Runtime/CoreUObject/Private/UObject/SavePackage.cpp
@@ -1342,8 +1342,11 @@ FArchive& FArchiveSaveTagImports::operator<<( UObject*& Obj )
 									ConditionallyExcludeObjectForTarget(ObjTemplate, *this);
 									UE_CLOG(ObjTemplate->HasAnyMarks(OBJECTMARK_NotForClient), LogSavePackage, Verbose, TEXT("      will reject subobj because it is OBJECTMARK_NotForClient %s"), *ObjTemplate->GetFullName());
 
-									ObjTemplate->Mark(OBJECTMARK_TagImp);
-									UE_LOG(LogSavePackage, Verbose, TEXT("      OBJECTMARK_TagImp %s"), *ObjTemplate->GetFullName());
+									for (UObject* ObjToMark = ObjTemplate; ObjToMark != nullptr; ObjToMark = ObjToMark->GetOuter())
+									{
+										ObjToMark->Mark(OBJECTMARK_TagImp);
+										UE_LOG(LogSavePackage, Verbose, TEXT("      OBJECTMARK_TagImp %s"), *ObjToMark->GetFullName());
+									}
 								}
 								else
 								{
