Copying //UE4/Dev-Core to //UE4/Dev-Main (Source: //UE4/Dev-Core @ 3345728)

diff --git a/Engine/Source/Runtime/CoreUObject/Private/UObject/Class.cpp b/Engine/Source/Runtime/CoreUObject/Private/UObject/Class.cpp
--- a/Engine/Source/Runtime/CoreUObject/Private/UObject/Class.cpp
+++ b/Engine/Source/Runtime/CoreUObject/Private/UObject/Class.cpp
@@ -107,6 +107,22 @@ void FNativeFunctionRegistrar::RegisterFunction(class UClass* Class, const WIDEC
 	Class->AddNativeFunction(InName, InPointer);
 }
 
+void FNativeFunctionRegistrar::RegisterFunctions(class UClass* Class, const TNameNativePtrPair<ANSICHAR>* InArray, int32 NumFunctions)
+{
+	for (; NumFunctions; ++InArray, --NumFunctions)
+	{
+		Class->AddNativeFunction(InArray->Name, InArray->Pointer);
+	}
+}
+
+void FNativeFunctionRegistrar::RegisterFunctions(class UClass* Class, const TNameNativePtrPair<WIDECHAR>* InArray, int32 NumFunctions)
+{
+	for (; NumFunctions; ++InArray, --NumFunctions)
+	{
+		Class->AddNativeFunction(InArray->Name, InArray->Pointer);
+	}
+}
+
 /*-----------------------------------------------------------------------------
 	UField implementation.
 -----------------------------------------------------------------------------*/
@@ -2553,10 +2569,8 @@ class FRestoreClassInfo: public FRestoreForUObjectOverwrite
 	EClassCastFlags	CastFlags;
 	/** Saved ClassConstructor **/
 	UClass::ClassConstructorType Constructor;
-#if WITH_HOT_RELOAD_CTORS
 	/** Saved ClassVTableHelperCtorCaller **/
 	UClass::ClassVTableHelperCtorCallerType ClassVTableHelperCtorCaller;
-#endif // WITH_HOT_RELOAD_CTORS
 	/** Saved ClassConstructor **/
 	UClass::ClassAddReferencedObjectsType AddReferencedObjects;
 	/** Saved NativeFunctionLookupTable. */
@@ -2576,9 +2590,7 @@ public:
 		Flags(Save->ClassFlags & CLASS_Abstract),
 		CastFlags(Save->ClassCastFlags),
 		Constructor(Save->ClassConstructor),
-#if WITH_HOT_RELOAD_CTORS
 		ClassVTableHelperCtorCaller(Save->ClassVTableHelperCtorCaller),
-#endif // WITH_HOT_RELOAD_CTORS
 		AddReferencedObjects(Save->ClassAddReferencedObjects),
 		NativeFunctionLookupTable(Save->NativeFunctionLookupTable)
 	{
@@ -2593,9 +2605,7 @@ public:
 		Target->ClassFlags |= Flags;
 		Target->ClassCastFlags |= CastFlags;
 		Target->ClassConstructor = Constructor;
-#if WITH_HOT_RELOAD_CTORS
 		Target->ClassVTableHelperCtorCaller = ClassVTableHelperCtorCaller;
-#endif // WITH_HOT_RELOAD_CTORS
 		Target->ClassAddReferencedObjects = AddReferencedObjects;
 		Target->NativeFunctionLookupTable = NativeFunctionLookupTable;
 	}
@@ -2835,9 +2845,7 @@ void UClass::Bind()
 
 	UClass* SuperClass = GetSuperClass();
 	if (SuperClass && (ClassConstructor == nullptr || ClassAddReferencedObjects == nullptr
-#if WITH_HOT_RELOAD_CTORS
 		|| ClassVTableHelperCtorCaller == nullptr
-#endif // WITH_HOT_RELOAD_CTORS
 		))
 	{
 		// Chase down constructor in parent class.
@@ -2846,12 +2854,10 @@ void UClass::Bind()
 		{
 			ClassConstructor = SuperClass->ClassConstructor;
 		}
-#if WITH_HOT_RELOAD_CTORS
 		if (!ClassVTableHelperCtorCaller)
 		{
 			ClassVTableHelperCtorCaller = SuperClass->ClassVTableHelperCtorCaller;
 		}
-#endif // WITH_HOT_RELOAD_CTORS
 		if (!ClassAddReferencedObjects)
 		{
 			ClassAddReferencedObjects = SuperClass->ClassAddReferencedObjects;
@@ -3700,9 +3706,7 @@ UObject* UClass::GetArchetypeForCDO() const
 void UClass::PurgeClass(bool bRecompilingOnLoad)
 {
 	ClassConstructor = nullptr;
-#if WITH_HOT_RELOAD_CTORS
 	ClassVTableHelperCtorCaller = nullptr;
-#endif // WITH_HOT_RELOAD_CTORS
 	ClassFlags = 0;
 	ClassCastFlags = 0;
 	ClassUnique = 0;
@@ -3859,16 +3863,12 @@ UClass::UClass
 	const TCHAR*    InConfigName,
 	EObjectFlags	InFlags,
 	ClassConstructorType InClassConstructor,
-#if WITH_HOT_RELOAD_CTORS
 	ClassVTableHelperCtorCallerType InClassVTableHelperCtorCaller,
-#endif // WITH_HOT_RELOAD_CTORS
 	ClassAddReferencedObjectsType InClassAddReferencedObjects
 )
 :	UStruct					( EC_StaticConstructor, InSize, InFlags )
 ,	ClassConstructor		( InClassConstructor )
-#if WITH_HOT_RELOAD_CTORS
 ,	ClassVTableHelperCtorCaller(InClassVTableHelperCtorCaller)
-#endif // WITH_HOT_RELOAD_CTORS
 ,	ClassAddReferencedObjects( InClassAddReferencedObjects )
 ,	ClassUnique				( 0 )
 ,	ClassFlags				( InClassFlags | CLASS_Native )
@@ -3897,9 +3897,7 @@ bool UClass::HotReloadPrivateStaticClass(
 	EClassCastFlags	InClassCastFlags,
 	const TCHAR*    InConfigName,
 	ClassConstructorType InClassConstructor,
-#if WITH_HOT_RELOAD_CTORS
 	ClassVTableHelperCtorCallerType InClassVTableHelperCtorCaller,
-#endif // WITH_HOT_RELOAD_CTORS
 	ClassAddReferencedObjectsType InClassAddReferencedObjects,
 	class UClass* TClass_Super_StaticClass,
 	class UClass* TClass_WithinClass_StaticClass
@@ -3920,9 +3918,7 @@ bool UClass::HotReloadPrivateStaticClass(
 	//@todo safe? ClassConfigName = InConfigName;
 	ClassConstructorType OldClassConstructor = ClassConstructor;
 	ClassConstructor = InClassConstructor;
-#if WITH_HOT_RELOAD_CTORS
 	ClassVTableHelperCtorCaller = InClassVTableHelperCtorCaller;
-#endif // WITH_HOT_RELOAD_CTORS
 	ClassAddReferencedObjects = InClassAddReferencedObjects;
 	/* No recursive ::StaticClass calls allowed. Setup extras. */
 	/* @todo safe? 
@@ -3940,7 +3936,6 @@ bool UClass::HotReloadPrivateStaticClass(
 	UE_LOG(LogClass, Verbose, TEXT("Attempting to change VTable for class %s."),*GetName());
 	ClassWithin = UPackage::StaticClass();  // We are just avoiding error checks with this...we don't care about this temp object other than to get the vtable.
 
-#if WITH_HOT_RELOAD_CTORS
 	static struct FUseVTableConstructorsCache
 	{
 		FUseVTableConstructorsCache()
@@ -3953,19 +3948,11 @@ bool UClass::HotReloadPrivateStaticClass(
 	} UseVTableConstructorsCache;
 
 	UObject* TempObjectForVTable = nullptr;
-	if (UseVTableConstructorsCache.bUseVTableConstructors)
 	{
 		TGuardValue<bool> Guard(GIsRetrievingVTablePtr, true);
-		auto Helper = FVTableHelper();
+		FVTableHelper Helper;
 		TempObjectForVTable = ClassVTableHelperCtorCaller(Helper);
 	}
-	else
-	{
-		TempObjectForVTable = StaticConstructObject_Internal(this, GetTransientPackage(), NAME_None, RF_NeedLoad | RF_ClassDefaultObject | RF_TagGarbageTemp);
-	}
-#else // WITH_HOT_RELOAD_CTORS
-	UObject* TempObjectForVTable = StaticConstructObject_Internal(this, GetTransientPackage(), NAME_None, RF_NeedLoad | RF_ClassDefaultObject | RF_TagGarbageTemp);
-#endif // WITH_HOT_RELOAD_CTORS
 
 	if( !TempObjectForVTable->IsRooted() )
 	{
@@ -3997,9 +3984,7 @@ bool UClass::HotReloadPrivateStaticClass(
 				if (Class->ClassConstructor == OldClassConstructor)
 				{
 					Class->ClassConstructor = ClassConstructor;
-#if WITH_HOT_RELOAD_CTORS
 					Class->ClassVTableHelperCtorCaller = ClassVTableHelperCtorCaller;
-#endif // WITH_HOT_RELOAD_CTORS
 					Class->ClassAddReferencedObjects = ClassAddReferencedObjects;
 					CountClass++;
 				}
@@ -4341,9 +4326,7 @@ void GetPrivateStaticClassBody(
 				InClassCastFlags,
 				InConfigName,
 				InClassConstructor,
-#if WITH_HOT_RELOAD_CTORS
 				InClassVTableHelperCtorCaller,
-#endif // WITH_HOT_RELOAD_CTORS
 				InClassAddReferencedObjects,
 				InSuperClassFn(),
 				InWithinClassFn()
@@ -4375,9 +4358,7 @@ void GetPrivateStaticClassBody(
 			InConfigName,
 			EObjectFlags(RF_Public | RF_Standalone | RF_Transient | RF_MarkAsNative | RF_MarkAsRootSet),
 			InClassConstructor,
-#if WITH_HOT_RELOAD_CTORS
 			InClassVTableHelperCtorCaller,
-#endif // WITH_HOT_RELOAD_CTORS
 			InClassAddReferencedObjects
 			);
 		check(ReturnClass);
@@ -4396,9 +4377,7 @@ void GetPrivateStaticClassBody(
 			InConfigName,
 			EObjectFlags(RF_Public | RF_Standalone | RF_Transient | RF_Dynamic | (GIsInitialLoad ? RF_MarkAsRootSet : RF_NoFlags)),
 			InClassConstructor,
-#if WITH_HOT_RELOAD_CTORS
 			InClassVTableHelperCtorCaller,
-#endif // WITH_HOT_RELOAD_CTORS
 			InClassAddReferencedObjects
 			);
 		check(ReturnClass);
@@ -4897,9 +4876,7 @@ UDynamicClass::UDynamicClass(
 	const TCHAR*    InConfigName,
 	EObjectFlags	InFlags,
 	ClassConstructorType InClassConstructor,
-#if WITH_HOT_RELOAD_CTORS
 	ClassVTableHelperCtorCallerType InClassVTableHelperCtorCaller,
-#endif // WITH_HOT_RELOAD_CTORS
 	ClassAddReferencedObjectsType InClassAddReferencedObjects)
 : UClass(
   EC_StaticConstructor
@@ -4910,9 +4887,7 @@ UDynamicClass::UDynamicClass(
 , InConfigName
 , InFlags
 , InClassConstructor
-#if WITH_HOT_RELOAD_CTORS
 , InClassVTableHelperCtorCaller
-#endif // WITH_HOT_RELOAD_CTORS
 , InClassAddReferencedObjects)
 , AnimClassImplementation(nullptr)
 {
diff --git a/Engine/Source/Runtime/CoreUObject/Private/UObject/GarbageCollection.cpp b/Engine/Source/Runtime/CoreUObject/Private/UObject/GarbageCollection.cpp
--- a/Engine/Source/Runtime/CoreUObject/Private/UObject/GarbageCollection.cpp
+++ b/Engine/Source/Runtime/CoreUObject/Private/UObject/GarbageCollection.cpp
@@ -60,9 +60,6 @@ static FRawObjectIterator GObjCurrentPurgeObjectIndex;
 static bool GObjCurrentPurgeObjectIndexNeedsReset = true;
 static bool GObjCurrentPurgeObjectIndexResetPastPermanent = false;
 
-/** Currently running a parallel reachability test.											*/
-static volatile bool GIsRunningParallelReachability = false;
-
 /** Whether we are currently purging an object in the GC purge pass. */
 static bool GIsPurgingObject = false;
 
@@ -297,6 +294,7 @@ static void LogClassCountInfo( const TCHAR* LogText, TMap<const FName,uint32>& C
 /**
 * Handles UObject references found by TFastReferenceCollector
 */
+template <bool bParallel>
 class FGCReferenceProcessor
 {
 public:
@@ -310,16 +308,6 @@ public:
 		return GMinDesiredObjectsPerSubTask;
 	}
 
-	FORCEINLINE volatile bool IsRunningMultithreaded() const
-	{
-		return GIsRunningParallelReachability;
-	}
-
-	FORCEINLINE void SetIsRunningMultithreaded(bool bIsParallel)
-	{
-		GIsRunningParallelReachability = bIsParallel;
-	}
-
 	void UpdateDetailedStats(UObject* CurrentObject, uint32 DeltaCycles)
 	{
 #if PERF_DETAILED_PER_CLASS_GC_STATS
@@ -350,50 +338,105 @@ public:
 	}
 
 	/** Marks all objects that can't be directly in a cluster but are referenced by it as reachable */
-	template <bool bParallel>
-	static FORCEINLINE void MarkClusterMutableObjectsAsReachable(FUObjectCluster* Cluster, TArray<UObject*>& ObjectsToSerialize)
+	static FORCEINLINE bool MarkClusterMutableObjectsAsReachable(FUObjectCluster& Cluster, TArray<UObject*>& ObjectsToSerialize)
 	{
-		for (int32 ReferencedMutableObjectIndex : Cluster->MutableObjects)
+		// This is going to be the return value and basically means that we ran across some pending kill objects
+		bool bAddClusterObjectsToSerialize = false;
+		for (int32& ReferencedMutableObjectIndex : Cluster.MutableObjects)
 		{
-			FUObjectItem* ReferencedMutableObjectItem = GUObjectArray.IndexToObjectUnsafeForGC(ReferencedMutableObjectIndex);
-			if (bParallel)
+			if (ReferencedMutableObjectIndex >= 0) // Pending kill support
 			{
-				if (ReferencedMutableObjectItem->IsUnreachable() && ReferencedMutableObjectItem->ThisThreadAtomicallyClearedRFUnreachable())
+				FUObjectItem* ReferencedMutableObjectItem = GUObjectArray.IndexToObjectUnsafeForGC(ReferencedMutableObjectIndex);
+				if (bParallel)
 				{
-					ReferencedMutableObjectItem->ThisThreadAtomicallyClearedFlag(EInternalObjectFlags::NoStrongReference);
-					ObjectsToSerialize.Add(static_cast<UObject*>(ReferencedMutableObjectItem->Object));
+					if (!ReferencedMutableObjectItem->IsPendingKill())
+					{
+						if (ReferencedMutableObjectItem->IsUnreachable() && ReferencedMutableObjectItem->ThisThreadAtomicallyClearedRFUnreachable())
+						{
+							ReferencedMutableObjectItem->ThisThreadAtomicallyClearedFlag(EInternalObjectFlags::NoStrongReference);
+							ObjectsToSerialize.Add(static_cast<UObject*>(ReferencedMutableObjectItem->Object));
+						}
+					}
+					else
+					{
+						// Pending kill support for clusters
+						ReferencedMutableObjectIndex = -1;
+						bAddClusterObjectsToSerialize = true;
+					}
+				}
+				else if (ReferencedMutableObjectItem->IsUnreachable())
+				{
+					if (!ReferencedMutableObjectItem->IsPendingKill())
+					{
+						ReferencedMutableObjectItem->ClearFlags(EInternalObjectFlags::NoStrongReference | EInternalObjectFlags::Unreachable);
+						ObjectsToSerialize.Add(static_cast<UObject*>(ReferencedMutableObjectItem->Object));
+					}
+					else
+					{
+						// Pending kill support for clusters
+						ReferencedMutableObjectIndex = -1;
+						bAddClusterObjectsToSerialize = true;
+					}
 				}
-			}
-			else if (ReferencedMutableObjectItem->IsUnreachable())
-			{
-				ReferencedMutableObjectItem->ClearFlags(EInternalObjectFlags::NoStrongReference | EInternalObjectFlags::Unreachable);
-				ObjectsToSerialize.Add(static_cast<UObject*>(ReferencedMutableObjectItem->Object));
 			}
 		}
+		return bAddClusterObjectsToSerialize;
 	}
 
-	/** Marks all clusters referenced by another cluster as rechable */
-	template <bool bParallel>
-	static FORCEINLINE void MarkReferencedClustersAsReachable(int32 ClusterRootIndex, TArray<UObject*>& ObjectsToSerialize)
+	/** Marks all clusters referenced by another cluster as reachable */
+	static FORCEINLINE void MarkReferencedClustersAsReachable(int32 ClusterIndex, TArray<UObject*>& ObjectsToSerialize)
 	{
-		FUObjectCluster* Cluster = GUObjectClusters.FindChecked(ClusterRootIndex);
+		// If we run across some PendingKill objects we need to add all objects from this cluster
+		// to ObjectsToSerialize so that we can properly null out all the references.
+		// It also means this cluster will have to be dissolved because we may no longer guarantee all cross-cluster references are correct.
+
+		bool bAddClusterObjectsToSerialize = false;
+		FUObjectCluster& Cluster = GUObjectClusters[ClusterIndex];
 		// Also mark all referenced objects from outside of the cluster as reachable
-		MarkClusterMutableObjectsAsReachable<bParallel>(Cluster, ObjectsToSerialize);
-		for (int32 ReferncedClusterIndex : Cluster->ReferencedClusters)
+		for (int32& ReferncedClusterIndex : Cluster.ReferencedClusters)
 		{
-			FUObjectItem* ReferencedClusterRootObjectItem = GUObjectArray.IndexToObjectUnsafeForGC(ReferncedClusterIndex);
-			// This condition should get collapsed by the compiler based on the template argument
-			if (bParallel)
+			if (ReferncedClusterIndex >= 0) // Pending Kill support
 			{
-				ReferencedClusterRootObjectItem->ThisThreadAtomicallyClearedFlag(EInternalObjectFlags::NoStrongReference | EInternalObjectFlags::Unreachable);
+				FUObjectItem* ReferencedClusterRootObjectItem = GUObjectArray.IndexToObjectUnsafeForGC(ReferncedClusterIndex);
+				if (!ReferencedClusterRootObjectItem->IsPendingKill())
+				{
+					// This condition should get collapsed by the compiler based on the template argument
+					if (bParallel)
+					{
+						if (ReferencedClusterRootObjectItem->IsUnreachable())
+						{
+							ReferencedClusterRootObjectItem->ThisThreadAtomicallyClearedFlag(EInternalObjectFlags::NoStrongReference | EInternalObjectFlags::Unreachable);
+						}
+					}
+					else
+					{
+						ReferencedClusterRootObjectItem->ClearFlags(EInternalObjectFlags::NoStrongReference | EInternalObjectFlags::Unreachable);
+					}
+				}
+				else
+				{
+					// Pending kill support for clusters
+					ReferncedClusterIndex = -1;
+					bAddClusterObjectsToSerialize = true;
+				}
 			}
-			else
+		}
+		if (MarkClusterMutableObjectsAsReachable(Cluster, ObjectsToSerialize))
+		{
+			bAddClusterObjectsToSerialize = true;
+		}
+		if (bAddClusterObjectsToSerialize)
+		{
+			// We need to process all cluster objects to handle PendingKill objects we nulled out (-1) from the cluster.
+			for (int32 ClusterObjectIndex : Cluster.Objects)
 			{
-				ReferencedClusterRootObjectItem->ClearFlags(EInternalObjectFlags::NoStrongReference | EInternalObjectFlags::Unreachable);
+				FUObjectItem* ClusterObjectItem = GUObjectArray.IndexToObjectUnsafeForGC(ClusterObjectIndex);
+				UObject* ClusterObject = static_cast<UObject*>(ClusterObjectItem->Object);
+				ObjectsToSerialize.Add(ClusterObject);
 			}
-			FUObjectCluster* ReferencedCluster = GUObjectClusters.FindChecked(ReferncedClusterIndex);
-			MarkClusterMutableObjectsAsReachable<bParallel>(ReferencedCluster, ObjectsToSerialize);
-		}		
+			Cluster.bNeedsDissolving = true;
+			GUObjectClusters.SetClustersNeedDissolving();
+		}
 	}
 
 	/**
@@ -421,25 +464,26 @@ public:
 			return;
 		}
 
-		FUObjectItem* ObjectItem = GUObjectArray.ObjectToObjectItem(Object);
+		const int32 ObjectIndex = GUObjectArray.ObjectToIndex(Object);
+		FUObjectItem* ObjectItem = GUObjectArray.IndexToObjectUnsafeForGC(ObjectIndex);
 		// Remove references to pending kill objects if we're allowed to do so.
 		if (ObjectItem->IsPendingKill() && bAllowReferenceElimination)
 		{
-			checkSlow(ObjectItem->HasAnyFlags(EInternalObjectFlags::ClusterRoot) == false);
-			checkSlow(ObjectItem->GetOwnerIndex() == 0)
+			//checkSlow(ObjectItem->HasAnyFlags(EInternalObjectFlags::ClusterRoot) == false);
+			checkSlow(ObjectItem->GetOwnerIndex() <= 0)
 			// Null out reference.
 			Object = NULL;
 		}
 		// Add encountered object reference to list of to be serialized objects if it hasn't already been added.
 		else if (ObjectItem->IsUnreachable())
 		{
-			if (GIsRunningParallelReachability)
+			if (bParallel)
 			{
 				// Mark it as reachable.
 				if (ObjectItem->ThisThreadAtomicallyClearedRFUnreachable())
 				{
 					// Objects that are part of a GC cluster should never have the unreachable flag set!
-					checkSlow(ObjectItem->GetOwnerIndex() == 0);
+					checkSlow(ObjectItem->GetOwnerIndex() <= 0);
 
 					if (!ObjectItem->HasAnyFlags(EInternalObjectFlags::ClusterRoot))
 					{
@@ -449,8 +493,7 @@ public:
 					else
 					{
 						// This is a cluster root reference so mark all referenced clusters as reachable
-						const int32 ObjectIndex = GUObjectArray.ObjectToIndex(Object);
-						MarkReferencedClustersAsReachable<true>(ObjectIndex, ObjectsToSerialize);
+						MarkReferencedClustersAsReachable(ObjectItem->GetClusterIndex(), ObjectsToSerialize);
 					}
 				}
 			}
@@ -470,7 +513,7 @@ public:
 				ObjectItem->ClearUnreachable();
 
 				// Objects that are part of a GC cluster should never have the unreachable flag set!
-				checkSlow(ObjectItem->GetOwnerIndex() == 0);
+				checkSlow(ObjectItem->GetOwnerIndex() <= 0);
 
 				if (!ObjectItem->HasAnyFlags(EInternalObjectFlags::ClusterRoot))
 				{
@@ -480,15 +523,14 @@ public:
 				else
 				{
 					// This is a cluster root reference so mark all referenced clusters as reachable
-					const int32 ObjectIndex = GUObjectArray.ObjectToIndex(Object);
-					MarkReferencedClustersAsReachable<false>(ObjectIndex, ObjectsToSerialize);
+					MarkReferencedClustersAsReachable(ObjectItem->GetClusterIndex(), ObjectsToSerialize);
 				}
 			}
 		}
-		else if (ObjectItem->GetOwnerIndex() && !ObjectItem->HasAnyFlags(EInternalObjectFlags::ReachableInCluster))
+		else if (ObjectItem->GetOwnerIndex() > 0 && !ObjectItem->HasAnyFlags(EInternalObjectFlags::ReachableInCluster))
 		{
 			bool bNeedsDoing = true;
-			if (GIsRunningParallelReachability)
+			if (bParallel)
 			{
 				bNeedsDoing = ObjectItem->ThisThreadAtomicallySetFlag(EInternalObjectFlags::ReachableInCluster);
 			}
@@ -498,32 +540,32 @@ public:
 			}
 			if (bNeedsDoing)
 			{
-			    // Make sure cluster root object is reachable too
-			    const int32 OwnerIndex = ObjectItem->GetOwnerIndex();
-			    FUObjectItem* RootObjectItem = GUObjectArray.IndexToObjectUnsafeForGC(OwnerIndex);
+				// Make sure cluster root object is reachable too
+				const int32 OwnerIndex = ObjectItem->GetOwnerIndex();
+				FUObjectItem* RootObjectItem = GUObjectArray.IndexToObjectUnsafeForGC(OwnerIndex);
 				checkSlow(RootObjectItem->HasAnyFlags(EInternalObjectFlags::ClusterRoot));
-			    if (GIsRunningParallelReachability)
-			    {
-				    if (RootObjectItem->ThisThreadAtomicallyClearedRFUnreachable())
-				    {
+				if (bParallel)
+				{
+					if (RootObjectItem->ThisThreadAtomicallyClearedRFUnreachable())
+					{
 						if (bStrongReference)
 						{
-					    	RootObjectItem->ThisThreadAtomicallyClearedFlag(EInternalObjectFlags::NoStrongReference);
+							RootObjectItem->ThisThreadAtomicallyClearedFlag(EInternalObjectFlags::NoStrongReference);
 						}
-					    // Make sure all referenced clusters are marked as reachable too
-					    MarkReferencedClustersAsReachable<true>(OwnerIndex, ObjectsToSerialize);
-				    }
-			    }
-			    else if (RootObjectItem->IsUnreachable())
-			    {
+						// Make sure all referenced clusters are marked as reachable too
+						MarkReferencedClustersAsReachable(RootObjectItem->GetClusterIndex(), ObjectsToSerialize);
+					}
+				}
+				else if (RootObjectItem->IsUnreachable())
+				{
 					RootObjectItem->ClearFlags(EInternalObjectFlags::Unreachable);
 					if (bStrongReference)
 					{
 						RootObjectItem->ClearFlags(EInternalObjectFlags::NoStrongReference);
 					}
+					// Make sure all referenced clusters are marked as reachable too
+					MarkReferencedClustersAsReachable(RootObjectItem->GetClusterIndex(), ObjectsToSerialize);
 				}
-				// Make sure all referenced clusters are marked as reachable too
-				MarkReferencedClustersAsReachable<false>(OwnerIndex, ObjectsToSerialize);
 			}
 		}
 
@@ -580,19 +622,23 @@ public:
 	}
 };
 
+typedef FGCReferenceProcessor<true> FGCReferenceProcessorMultithreaded;
+typedef FGCReferenceProcessor<false> FGCReferenceProcessorSinglethreaded;
+
 /**
 * Specialized FReferenceCollector that uses FGCReferenceProcessor to mark objects as reachable.
 */
+template <bool bParallel>
 class FGCCollector : public FReferenceCollector
 {
-	FGCReferenceProcessor& ReferenceProcessor;
+	FGCReferenceProcessor<bParallel>& ReferenceProcessor;
 	TArray<UObject*>& ObjectArray;
 	bool bAllowEliminatingReferences;
 	bool bShouldHandleAsWeakRef;
 
 public:
 
-	FGCCollector(FGCReferenceProcessor& InProcessor, TArray<UObject*>& InObjectArray)
+	FGCCollector(FGCReferenceProcessor<bParallel>& InProcessor, TArray<UObject*>& InObjectArray)
 		: ReferenceProcessor(InProcessor)
 		, ObjectArray(InObjectArray)
 		, bAllowEliminatingReferences(true)
@@ -600,7 +646,7 @@ public:
 	{
 	}
 
-	virtual void HandleObjectReference(UObject*& Object, const UObject* ReferencingObject, const UProperty* ReferencingProperty) override
+	FORCEINLINE void InternalHandleObjectReference(UObject*& Object, const UObject* ReferencingObject, const UProperty* ReferencingProperty)
 	{
 #if !(UE_BUILD_TEST || UE_BUILD_SHIPPING)
 		if (Object && !Object->IsValidLowLevelFast())
@@ -613,21 +659,17 @@ public:
 #endif
 		ReferenceProcessor.HandleObjectReference(ObjectArray, const_cast<UObject*>(ReferencingObject), Object, bAllowEliminatingReferences, !bShouldHandleAsWeakRef);
 	}
+
+	virtual void HandleObjectReference(UObject*& Object, const UObject* ReferencingObject, const UProperty* ReferencingProperty) override
+	{
+		InternalHandleObjectReference(Object, ReferencingObject, ReferencingProperty);
+	}
 	virtual void HandleObjectReferences(UObject** InObjects, const int32 ObjectNum, const UObject* InReferencingObject, const UProperty* InReferencingProperty) override
 	{
 		for (int32 ObjectIndex = 0; ObjectIndex < ObjectNum; ++ObjectIndex)
 		{
 			UObject*& Object = InObjects[ObjectIndex];
-#if !(UE_BUILD_TEST || UE_BUILD_SHIPPING)
-			if (Object && !Object->IsValidLowLevelFast())
-			{
-				UE_LOG(LogGarbage, Fatal, TEXT("Invalid object in GC: 0x%016llx, ReferencingObject: %s, ReferencingProperty: %s"),
-					(int64)(PTRINT)Object,
-					InReferencingObject ? *InReferencingObject->GetFullName() : TEXT("NULL"),
-					InReferencingProperty ? *InReferencingProperty->GetFullName() : TEXT("NULL"));
-			}
-#endif
-			ReferenceProcessor.HandleObjectReference(ObjectArray, const_cast<UObject*>(InReferencingObject), Object, bAllowEliminatingReferences, !bShouldHandleAsWeakRef);
+			InternalHandleObjectReference(Object, InReferencingObject, InReferencingProperty);
 		}
 	}
 
@@ -650,6 +692,8 @@ public:
 	}
 };
 
+typedef FGCCollector<true> FGCCollectorMultithreaded;
+typedef FGCCollector<false> FGCCollectorSinglethreaded;
 
 /*----------------------------------------------------------------------------
 	FReferenceFinder.
@@ -741,7 +785,7 @@ public:
 
 		// Iterate over all objects. Note that we iterate over the UObjectArray and usually check only internal flags which
 		// are part of the array so we don't suffer from cache misses as much as we would if we were to check ObjectFlags.
-		TArray<UObjectBase*> KeepClusterRefs;
+		TArray<FUObjectItem*> KeepClusterRefs;
 		for (FRawObjectIterator It(true); It; ++It)
 		{
 			FUObjectItem* ObjectItem = *It;
@@ -760,11 +804,13 @@ public:
 				// IsValidLowLevel is extremely slow in this loop so only do it in debug
 				checkSlow(Object->IsValidLowLevel());
 				// We cannot use RF_PendingKill on objects that are part of the root set.
+#if DO_GUARD_SLOW
 				checkCode(if (ObjectItem->IsPendingKill()) { UE_LOG(LogGarbage, Fatal, TEXT("Object %s is part of root set though has been marked RF_PendingKill!"), *Object->GetFullName()); });
+#endif
 				ObjectsToSerialize.Add(Object);
 			}
-			// Regular objects.
-			else if (ObjectItem->GetOwnerIndex() == 0)
+			// Regular objects or cluster root objects
+			else if (ObjectItem->GetOwnerIndex() <= 0)
 			{
 				bool bMarkAsUnreachable = true;
 				if (!ObjectItem->IsPendingKill())
@@ -780,6 +826,11 @@ public:
 						bMarkAsUnreachable = false;
 					}
 				}
+				else if (ObjectItem->HasAnyFlags(EInternalObjectFlags::ClusterRoot))
+				{
+					GUObjectClusters.DissolveClusterAndMarkObjectsAsUnreachable(It.GetIndex(), ObjectItem);
+					GUObjectClusters.SetClustersNeedDissolving();					
+				}
 
 				// Mark objects as unreachable unless they have any of the passed in KeepFlags set and it's not marked for elimination..
 				if (!bMarkAsUnreachable)
@@ -788,9 +839,9 @@ public:
 					checkSlow(Object->IsValidLowLevel());
 					ObjectsToSerialize.Add(Object);
 
-					if (ObjectItem->HasAnyFlags(EInternalObjectFlags::ClusterRoot) || ObjectItem->GetOwnerIndex())
+					if (ObjectItem->HasAnyFlags(EInternalObjectFlags::ClusterRoot) || ObjectItem->GetOwnerIndex() > 0)
 					{
-						KeepClusterRefs.Add(Object);
+						KeepClusterRefs.Add(ObjectItem);
 					}
 				}
 				else
@@ -799,13 +850,12 @@ public:
 				}
 			}
 		}
-		for (UObjectBase* Object : KeepClusterRefs)
+
+		for (FUObjectItem* ObjectItem : KeepClusterRefs)
 		{
-			const int32 ObjectIndex = GUObjectArray.ObjectToIndex(Object);
-			FUObjectItem* ObjectItem = GUObjectArray.IndexToObjectUnsafeForGC(ObjectIndex);
-			if (ObjectItem->GetOwnerIndex())
+			if (ObjectItem->GetOwnerIndex() > 0)
 			{
-				check(!ObjectItem->HasAnyFlags(EInternalObjectFlags::ClusterRoot));
+				checkSlow(!ObjectItem->HasAnyFlags(EInternalObjectFlags::ClusterRoot));
 				bool bNeedsDoing = !ObjectItem->HasAnyFlags(EInternalObjectFlags::ReachableInCluster);
 				if (bNeedsDoing)
 				{
@@ -813,22 +863,22 @@ public:
 					// Make sure cluster root object is reachable too
 					const int32 OwnerIndex = ObjectItem->GetOwnerIndex();
 					FUObjectItem* RootObjectItem = GUObjectArray.IndexToObjectUnsafeForGC(OwnerIndex);
-					check(RootObjectItem->HasAnyFlags(EInternalObjectFlags::ClusterRoot));
+					checkSlow(RootObjectItem->HasAnyFlags(EInternalObjectFlags::ClusterRoot));
 					// if it is reachable via keep flags we will do this below (or maybe already have)
 					if (RootObjectItem->IsUnreachable()) 
 					{
 						RootObjectItem->ClearFlags(EInternalObjectFlags::Unreachable | EInternalObjectFlags::NoStrongReference);
 						// Make sure all referenced clusters are marked as reachable too
-						FGCReferenceProcessor::MarkReferencedClustersAsReachable<false>(OwnerIndex, ObjectsToSerialize);
+						FGCReferenceProcessorSinglethreaded::MarkReferencedClustersAsReachable(RootObjectItem->GetClusterIndex(), ObjectsToSerialize);
 					}
 				}
 			}
 			else
 			{
-				check(ObjectItem->HasAnyFlags(EInternalObjectFlags::ClusterRoot));
+				checkSlow(ObjectItem->HasAnyFlags(EInternalObjectFlags::ClusterRoot));
 				// this thing is definitely not marked unreachable, so don't test it here
 				// Make sure all referenced clusters are marked as reachable too
-				FGCReferenceProcessor::MarkReferencedClustersAsReachable<false>(ObjectIndex, ObjectsToSerialize);
+				FGCReferenceProcessorSinglethreaded::MarkReferencedClustersAsReachable(ObjectItem->GetClusterIndex(), ObjectsToSerialize);
 			}
 		}
 	}
@@ -858,12 +908,18 @@ public:
 
 		MarkObjectsAsUnreachable(ObjectsToSerialize, KeepFlags);
 
+		if (!bForceSingleThreaded)
 		{
-			FGCReferenceProcessor ReferenceProcessor;
-			TFastReferenceCollector<FGCReferenceProcessor, FGCCollector, FGCArrayPool> ReferenceCollector(ReferenceProcessor, FGCArrayPool::Get());
-			ReferenceCollector.CollectReferences(ObjectsToSerialize, bForceSingleThreaded);
+			FGCReferenceProcessorMultithreaded ReferenceProcessor;
+			TFastReferenceCollector<true, FGCReferenceProcessorMultithreaded, FGCCollectorMultithreaded, FGCArrayPool> ReferenceCollector(ReferenceProcessor, FGCArrayPool::Get());
+			ReferenceCollector.CollectReferences(ObjectsToSerialize);
+		}
+		else
+		{
+			FGCReferenceProcessorSinglethreaded ReferenceProcessor;
+			TFastReferenceCollector<false, FGCReferenceProcessorSinglethreaded, FGCCollectorSinglethreaded, FGCArrayPool> ReferenceCollector(ReferenceProcessor, FGCArrayPool::Get());
+			ReferenceCollector.CollectReferences(ObjectsToSerialize);
 		}
-
 		FGCArrayPool::Get().ReturnToPool(&ObjectsToSerialize);
 
 #if UE_BUILD_DEBUG
@@ -1259,7 +1315,7 @@ void CollectGarbageInternal(EObjectFlags KeepFlags, bool bPerformFullPurge)
 #if VERIFY_DISREGARD_GC_ASSUMPTIONS
 		FUObjectArray& UObjectArray = GUObjectArray;
 		// Only verify assumptions if option is enabled. This avoids false positives in the Editor or commandlets.
-		if ((UObjectArray.DisregardForGCEnabled() || GUObjectClusters.Num()) && GShouldVerifyGCAssumptions)
+		if ((UObjectArray.DisregardForGCEnabled() || GUObjectClusters.GetNumAllocatedClusters()) && GShouldVerifyGCAssumptions)
 		{
 			DECLARE_SCOPE_CYCLE_COUNTER(TEXT("CollectGarbageInternal.VerifyGCAssumptions"), STAT_CollectGarbageInternal_VerifyGCAssumptions, STATGROUP_GC);
 			bool bShouldAssert = false;
@@ -1286,7 +1342,7 @@ void CollectGarbageInternal(EObjectFlags KeepFlags, bool bPerformFullPurge)
 						if (ReferencedObject &&
 							!(ReferencedObject->IsRooted() ||
 								UObjectArray.IsDisregardForGC(ReferencedObject) ||
-								UObjectArray.ObjectToObjectItem(ReferencedObject)->GetOwnerIndex() ||
+								UObjectArray.ObjectToObjectItem(ReferencedObject)->GetOwnerIndex() > 0 ||
 								UObjectArray.ObjectToObjectItem(ReferencedObject)->HasAnyFlags(EInternalObjectFlags::ClusterRoot)))
 						{
 							UE_LOG(LogGarbage, Warning, TEXT("Disregard for GC object %s referencing %s which is not part of root set"),
@@ -1330,6 +1386,14 @@ void CollectGarbageInternal(EObjectFlags KeepFlags, bool bPerformFullPurge)
 			UE_LOG(LogGarbage, Log, TEXT("%f ms for GC"), (FPlatformTime::Seconds() - StartTime) * 1000);
 		}
 
+		// Reconstruct clusters if needed
+		if (GUObjectClusters.ClustersNeedDissolving())
+		{
+			const double StartTime = FPlatformTime::Seconds();
+			GUObjectClusters.DissolveClusters();
+			UE_LOG(LogGarbage, Log, TEXT("%f ms for dissolving GC clusters"), (FPlatformTime::Seconds() - StartTime) * 1000);
+		}
+
 #if WITH_EDITOR
 		if (GIsEditor && EditorPostReachabilityAnalysisCallback)
 		{
@@ -1342,9 +1406,13 @@ void CollectGarbageInternal(EObjectFlags KeepFlags, bool bPerformFullPurge)
 
 			TGuardValue<bool> GuardObjUnhashUnreachableIsInProgress(GObjUnhashUnreachableIsInProgress, true);
 
+			FCoreUObjectDelegates::PreGarbageCollectConditionalBeginDestroy.Broadcast();
+
 			// Unhash all unreachable objects.
 			const double StartTime = FPlatformTime::Seconds();
 			int32 ClustersRemoved = 0;
+			int32 Items = 0;
+			int32 ClusterItems = 0;
 			for (FRawObjectIterator It(true); It; ++It)
 			{
 				//@todo UE4 - A prefetch was removed here. Re-add it. It wasn't right anyway, since it was ten items ahead and the consoles on have 8 prefetch slots
@@ -1353,14 +1421,14 @@ void CollectGarbageInternal(EObjectFlags KeepFlags, bool bPerformFullPurge)
 				checkSlow(ObjectItem);
 				if (ObjectItem->IsUnreachable())
 				{
+					Items++;
 					if ((ObjectItem->GetFlags() & EInternalObjectFlags::ClusterRoot) == EInternalObjectFlags::ClusterRoot)
 					{
 						// Nuke the entire cluster
 						ObjectItem->ClearFlags(EInternalObjectFlags::ClusterRoot | EInternalObjectFlags::NoStrongReference);
 						const int32 ClusterRootIndex = It.GetIndex();
-						FUObjectCluster* Cluster = GUObjectClusters.FindChecked(ClusterRootIndex);
-						checkSlow(Cluster);
-						for (int32 ClusterObjectIndex : Cluster->Objects)
+						FUObjectCluster& Cluster = GUObjectClusters[ObjectItem->GetClusterIndex()];
+						for (int32 ClusterObjectIndex : Cluster.Objects)
 						{
 							FUObjectItem* ClusterObjectItem = GUObjectArray.IndexToObjectUnsafeForGC(ClusterObjectIndex);
 							ClusterObjectItem->ClearFlags(EInternalObjectFlags::NoStrongReference);
@@ -1373,11 +1441,11 @@ void CollectGarbageInternal(EObjectFlags KeepFlags, bool bPerformFullPurge)
 								{
 									UObject* ClusterObject = (UObject*)ClusterObjectItem->Object;
 									ClusterObject->ConditionalBeginDestroy();
+									ClusterItems++;
 								}
 							}
 						}
-						delete Cluster;
-						GUObjectClusters.Remove(ClusterRootIndex);
+						GUObjectClusters.FreeCluster(ObjectItem->GetClusterIndex());
 						ClustersRemoved++;
 					}
 
@@ -1391,7 +1459,8 @@ void CollectGarbageInternal(EObjectFlags KeepFlags, bool bPerformFullPurge)
 					ObjectItem->SetPendingKill();
 				}
 			}
-			UE_LOG(LogGarbage, Log, TEXT("%f ms for unhashing unreachable objects. Clusters removed: %d."), (FPlatformTime::Seconds() - StartTime) * 1000, ClustersRemoved);
+			UE_LOG(LogGarbage, Log, TEXT("%f ms for unhashing unreachable objects. Clusters removed: %d.   Items %d Cluster Items %d"), (FPlatformTime::Seconds() - StartTime) * 1000, ClustersRemoved, Items, ClusterItems);
+			FCoreUObjectDelegates::PostGarbageCollectConditionalBeginDestroy.Broadcast();
 		}
 
 		// Set flag to indicate that we are relying on a purge to be performed.
diff --git a/Engine/Source/Runtime/CoreUObject/Private/UObject/Obj.cpp b/Engine/Source/Runtime/CoreUObject/Private/UObject/Obj.cpp
--- a/Engine/Source/Runtime/CoreUObject/Private/UObject/Obj.cpp
+++ b/Engine/Source/Runtime/CoreUObject/Private/UObject/Obj.cpp
@@ -82,35 +82,18 @@ static UPackage*			GObjTransientPkg								= NULL;
 UObject::UObject( EStaticConstructor, EObjectFlags InFlags )
 : UObjectBaseUtility(InFlags | (!(InFlags & RF_Dynamic) ? (RF_MarkAsNative | RF_MarkAsRootSet) : RF_NoFlags))
 {
-#if WITH_HOT_RELOAD_CTORS
 	EnsureNotRetrievingVTablePtr();
-#endif // WITH_HOT_RELOAD_CTORS
 }
 
-#if WITH_HOT_RELOAD_CTORS
 UObject::UObject(FVTableHelper& Helper)
 {
 	EnsureRetrievingVTablePtrDuringCtor(TEXT("UObject(FVTableHelper& Helper)"));
-
-	static struct FUseVTableConstructorsCache
-	{
-		FUseVTableConstructorsCache()
-		{
-			bUseVTableConstructors = false;
-			GConfig->GetBool(TEXT("Core.System"), TEXT("UseVTableConstructors"), bUseVTableConstructors, GEngineIni);
-		}
-
-		bool bUseVTableConstructors;
-	} UseVTableConstructorsCache;
-
-	UE_CLOG(!UseVTableConstructorsCache.bUseVTableConstructors, LogCore, Fatal, TEXT("This constructor is disabled."));
 }
 
 void UObject::EnsureNotRetrievingVTablePtr() const
 {
 	UE_CLOG(GIsRetrievingVTablePtr, LogCore, Fatal, TEXT("We are currently retrieving VTable ptr. Please use FVTableHelper constructor instead."));
 }
-#endif // WITH_HOT_RELOAD_CTORS
 
 UObject* UObject::CreateDefaultSubobject(FName SubobjectFName, UClass* ReturnType, UClass* ClassToCreateByDefault, bool bIsRequired, bool bAbstract, bool bIsTransient)
 {
@@ -561,7 +544,7 @@ struct FClassExclusionData
 
 		auto ModuleShortNameFromClass = [](const UClass* Class) -> FName
 		{
-			return FName(*FPackageName::GetLongPackageAssetName(Class->GetOutermost()->GetName()));
+			return FName(*FPackageName::GetShortName(Class->GetOutermost()));
 		};
 
 		while (InClass != nullptr)
@@ -802,6 +785,30 @@ bool UObject::ImplementsGetWorld() const
 }
 #endif
 
+#define PROFILE_ConditionalBeginDestroy (0)
+
+#if PROFILE_ConditionalBeginDestroy
+
+struct FTimeCnt
+{
+	float TotalTime;
+	int32 Count;
+
+	FTimeCnt()
+		: TotalTime(0.0f)
+		, Count(0)
+	{
+	}
+
+	bool operator<(const FTimeCnt& Other) const 
+	{
+		return TotalTime > Other.TotalTime;
+	}
+};
+
+static TMap<FName, FTimeCnt> MyProfile;
+#endif
+
 bool UObject::ConditionalBeginDestroy()
 {
 	check(IsValidLowLevel());
@@ -812,7 +819,43 @@ bool UObject::ConditionalBeginDestroy()
 		checkSlow(!DebugBeginDestroyed.Contains(this));
 		DebugBeginDestroyed.Add(this);
 #endif
+
+#if PROFILE_ConditionalBeginDestroy
+		double StartTime = FPlatformTime::Seconds();
+#endif
+
 		BeginDestroy();
+
+#if PROFILE_ConditionalBeginDestroy
+		float ThisTime = float(FPlatformTime::Seconds() - StartTime);
+
+		FTimeCnt& TimeCnt = MyProfile.FindOrAdd(GetClass()->GetFName());
+		TimeCnt.Count++;
+		TimeCnt.TotalTime += ThisTime;
+
+		static float TotalTime = 0.0f;
+		static int32 TotalCnt = 0;
+
+		TotalTime += ThisTime;
+		if ((++TotalCnt) % 1000 == 0)
+		{
+			UE_LOG(LogTemp, Log, TEXT("ConditionalBeginDestroy %d cnt %fus"), TotalCnt, 1000.0f * 1000.0f * TotalTime / float(TotalCnt));
+
+			MyProfile.ValueSort(TLess<FTimeCnt>());
+
+			int32 NumPrint = 0;
+			for (auto& Item : MyProfile)
+			{
+				UE_LOG(LogTemp, Log, TEXT("    %6d cnt %6.2fus per   %6.2fms total  %s"), Item.Value.Count, 1000.0f * 1000.0f * Item.Value.TotalTime / float(Item.Value.Count), 1000.0f * Item.Value.TotalTime, *Item.Key.ToString());
+				if (NumPrint++ > 30)
+				{
+					break;
+				}
+			}
+		}
+#endif
+
+
 #if !(UE_BUILD_SHIPPING || UE_BUILD_TEST)
 		if( DebugBeginDestroyed.Contains(this) )
 		{
diff --git a/Engine/Source/Runtime/CoreUObject/Private/UObject/PackageFileSummary.cpp b/Engine/Source/Runtime/CoreUObject/Private/UObject/PackageFileSummary.cpp
--- a/Engine/Source/Runtime/CoreUObject/Private/UObject/PackageFileSummary.cpp
+++ b/Engine/Source/Runtime/CoreUObject/Private/UObject/PackageFileSummary.cpp
@@ -252,16 +252,13 @@ FArchive& operator<<( FArchive& Ar, FPackageFileSummary& Sum )
 
 		Ar << Sum.AdditionalPackagesToCook;
 
-#if WITH_ENGINE
 		if (LegacyFileVersion > -7)
 		{
 			int32 NumTextureAllocations = 0;
 			Ar << NumTextureAllocations;
+			// We haven't used texture allocation info for ages and it's no longer supported anyway
 			check(NumTextureAllocations == 0);
 		}
-#else
-		check(!"this can't serialize successfully");
-#endif		// WITH_ENGINE
 
 		Ar << Sum.AssetRegistryDataOffset;
 		Ar << Sum.BulkDataStartOffset;
diff --git a/Engine/Source/Runtime/CoreUObject/Private/UObject/PropertyNumeric.cpp b/Engine/Source/Runtime/CoreUObject/Private/UObject/PropertyNumeric.cpp
--- a/Engine/Source/Runtime/CoreUObject/Private/UObject/PropertyNumeric.cpp
+++ b/Engine/Source/Runtime/CoreUObject/Private/UObject/PropertyNumeric.cpp
@@ -16,7 +16,7 @@ uint8 UNumericProperty::ReadEnumAsUint8(FArchive& Ar, UStruct* DefaultsStruct, c
 	UEnum* Enum = FindField<UEnum>(dynamic_cast<UClass*>(DefaultsStruct) ? static_cast<UClass*>(DefaultsStruct) : DefaultsStruct->GetTypedOuter<UClass>(), Tag.EnumName);
 	if (!Enum)
 	{
-		Enum = FindObject<UEnum>(ANY_PACKAGE, *Tag.EnumName.ToString(), true);
+		Enum = FindObject<UEnum>(ANY_PACKAGE, *Tag.EnumName.ToString());
 	}
 
 	if (!Enum)
diff --git a/Engine/Source/Runtime/CoreUObject/Private/UObject/ReferenceChainSearch.cpp b/Engine/Source/Runtime/CoreUObject/Private/UObject/ReferenceChainSearch.cpp
--- a/Engine/Source/Runtime/CoreUObject/Private/UObject/ReferenceChainSearch.cpp
+++ b/Engine/Source/Runtime/CoreUObject/Private/UObject/ReferenceChainSearch.cpp
@@ -188,7 +188,7 @@ void FReferenceChainSearch::PrintReferencers(FReferenceChain& Referencer)
 			ObjectReachability += TEXT("(ClusterRoot) ");
 			bClusterRoot = true;
 		}
-		if (ReferencedByObjectItem->GetOwnerIndex())
+		if (ReferencedByObjectItem->GetOwnerIndex() > 0)
 		{
 			ObjectReachability += TEXT("(Clustered) ");
 		}
diff --git a/Engine/Source/Runtime/CoreUObject/Private/UObject/SavePackage.cpp b/Engine/Source/Runtime/CoreUObject/Private/UObject/SavePackage.cpp
--- a/Engine/Source/Runtime/CoreUObject/Private/UObject/SavePackage.cpp
+++ b/Engine/Source/Runtime/CoreUObject/Private/UObject/SavePackage.cpp
@@ -50,6 +50,7 @@
 #include "ProfilingDebugging/CookStats.h"
 #include "UObject/DebugSerializationFlags.h"
 #include "EnumProperty.h"
+#include "BlueprintSupport.h"
 
 DEFINE_LOG_CATEGORY_STATIC(LogSavePackage, Log, All);
 
@@ -236,6 +237,15 @@ namespace SavePackageStats
 }
 #endif
 
+#if WITH_EDITOR
+static const FCompilerNativizationOptions& GetNativizationOptionsForPlatform(const ITargetPlatform* TargetPlatform, const IBlueprintNativeCodeGenCore* Coordinator)
+{
+	check(Coordinator);
+	const FString PlatformName = ensure(TargetPlatform) ? TargetPlatform->PlatformName() : FString();
+	return Coordinator->GetNativizationOptionsForPlatform(PlatformName);
+}
+#endif //WITH_EDITOR
+
 static bool HasUnsaveableOuter(UObject* InObj, UPackage* InSavingPackage)
 {
 	UObject* Obj = InObj;
@@ -617,17 +627,18 @@ public:
 	/**
 	 * Names are cached before we replace objects for imports. So the names of replacements must be ealier.
 	 */
-	void AddReplacementsNames(UObject* Obj)
+	void AddReplacementsNames(UObject* Obj, const ITargetPlatform* TargetPlatform)
 	{
 		if (const IBlueprintNativeCodeGenCore* Coordinator = IBlueprintNativeCodeGenCore::Get())
 		{
-			if (const UClass* ReplObjClass = Coordinator->FindReplacedClassForObject(Obj))
+			const FCompilerNativizationOptions& NativizationOptions = GetNativizationOptionsForPlatform(TargetPlatform, Coordinator);
+			if (const UClass* ReplObjClass = Coordinator->FindReplacedClassForObject(Obj, NativizationOptions))
 			{
 				MarkNameAsReferenced(ReplObjClass->GetFName());
 			}
 
 			FName ReplacedName;
-			Coordinator->FindReplacedNameAndOuter(Obj, ReplacedName); //TODO: should we care about replaced outer ?
+			Coordinator->FindReplacedNameAndOuter(Obj, ReplacedName, NativizationOptions); //TODO: should we care about replaced outer ?
 			if (ReplacedName != NAME_None)
 			{
 				MarkNameAsReferenced(ReplacedName);
@@ -724,7 +735,7 @@ bool IsEditorOnlyObject(const UObject* InObject)
 /**
  * Marks object as not for client, not for server, or editor only. Recurses up outer/class chain as necessary
  */
-static void ConditionallyExcludeObjectForTarget(UObject* Obj, EObjectMark ExcludedObjectMarks)
+static void ConditionallyExcludeObjectForTarget(UObject* Obj, EObjectMark ExcludedObjectMarks, const ITargetPlatform* TargetPlatform)
 {
 #if WITH_EDITOR
 	if (!Obj || (ExcludedObjectMarks != OBJECTMARK_NOMARKS && Obj->HasAnyMarks(ExcludedObjectMarks)))
@@ -745,12 +756,13 @@ static void ConditionallyExcludeObjectForTarget(UObject* Obj, EObjectMark Exclud
 	// Check for nativization replacement
 	if (const IBlueprintNativeCodeGenCore* Coordinator = IBlueprintNativeCodeGenCore::Get())
 	{
+		const FCompilerNativizationOptions& NativizationOptions = GetNativizationOptionsForPlatform(TargetPlatform, Coordinator);
 		FName UnusedName;
-		if (UClass* ReplacedClass = Coordinator->FindReplacedClassForObject(Obj))
+		if (UClass* ReplacedClass = Coordinator->FindReplacedClassForObject(Obj, NativizationOptions))
 		{
 			ObjClass = ReplacedClass;
 		}
-		if (UObject* ReplacedOuter = Coordinator->FindReplacedNameAndOuter(Obj, /*out*/UnusedName))
+		if (UObject* ReplacedOuter = Coordinator->FindReplacedNameAndOuter(Obj, /*out*/UnusedName, NativizationOptions))
 		{
 			ObjOuter = ReplacedOuter;
 		}
@@ -768,14 +780,14 @@ static void ConditionallyExcludeObjectForTarget(UObject* Obj, EObjectMark Exclud
 
 	// Recurse into parents, then compute inherited marks
 	
-	ConditionallyExcludeObjectForTarget(ObjClass, ExcludedObjectMarks);
+	ConditionallyExcludeObjectForTarget(ObjClass, ExcludedObjectMarks, TargetPlatform);
 	InheritMark(ObjClass, OBJECTMARK_EditorOnly);
 	InheritMark(ObjClass, OBJECTMARK_NotForClient);
 	InheritMark(ObjClass, OBJECTMARK_NotForServer);
 
 	if (ObjOuter)
 	{
-		ConditionallyExcludeObjectForTarget(ObjOuter, ExcludedObjectMarks);
+		ConditionallyExcludeObjectForTarget(ObjOuter, ExcludedObjectMarks, TargetPlatform);
 		InheritMark(ObjOuter, OBJECTMARK_EditorOnly);
 		InheritMark(ObjOuter, OBJECTMARK_NotForClient);
 		InheritMark(ObjOuter, OBJECTMARK_NotForServer);
@@ -786,7 +798,7 @@ static void ConditionallyExcludeObjectForTarget(UObject* Obj, EObjectMark Exclud
 	if (ThisStruct && ThisStruct->GetSuperStruct())
 	{
 		UObject* SuperStruct = ThisStruct->GetSuperStruct();
-		ConditionallyExcludeObjectForTarget(SuperStruct, ExcludedObjectMarks);
+		ConditionallyExcludeObjectForTarget(SuperStruct, ExcludedObjectMarks, TargetPlatform);
 		InheritMark(SuperStruct, OBJECTMARK_EditorOnly);
 		InheritMark(SuperStruct, OBJECTMARK_NotForClient);
 		InheritMark(SuperStruct, OBJECTMARK_NotForServer);
@@ -933,7 +945,7 @@ FArchive& FArchiveSaveTagExports::operator<<(UObject*& Obj)
 
 	// Check outer chain for any exlcuded object marks
 	const EObjectMark ExcludedObjectMarks = UPackage::GetExcludedObjectMarksForTargetPlatform(CookingTarget(), IsCooking());
-	ConditionallyExcludeObjectForTarget(Obj, ExcludedObjectMarks);
+	ConditionallyExcludeObjectForTarget(Obj, ExcludedObjectMarks, CookingTarget());
 
 	if (Obj && Obj->IsIn(Outer) && !Obj->HasAnyFlags(RF_Transient) && !Obj->HasAnyMarks((EObjectMark)(OBJECTMARK_TagExp | ExcludedObjectMarks)))
 	{
@@ -1126,7 +1138,7 @@ FArchive& FArchiveSaveTagImports::operator<<( UObject*& Obj )
 	CheckObjectPriorToSave(*this, Obj, nullptr);
 
 	const EObjectMark ExcludedObjectMarks = UPackage::GetExcludedObjectMarksForTargetPlatform( CookingTarget(), IsCooking() );
-	ConditionallyExcludeObjectForTarget(Obj, ExcludedObjectMarks);
+	ConditionallyExcludeObjectForTarget(Obj, ExcludedObjectMarks, CookingTarget());
 	
 	// Skip PendingKill objects and objects that don't pass the platform mark filter
 	if (Obj && (ExcludedObjectMarks == OBJECTMARK_NOMARKS || !Obj->HasAnyMarks(ExcludedObjectMarks)))
@@ -1197,17 +1209,19 @@ FArchive& FArchiveSaveTagImports::operator<<( UObject*& Obj )
 						}
 					}
 #if WITH_EDITOR
-					SavePackageState->AddReplacementsNames(Obj);
+					SavePackageState->AddReplacementsNames(Obj, CookingTarget());
 #endif //WITH_EDITOR
 				}
 
 				// Recurse into parent
 				UObject* Parent = Obj->GetOuter();
 #if WITH_EDITOR
-				const IBlueprintNativeCodeGenCore* Coordinator = IBlueprintNativeCodeGenCore::Get();
-				FName UnusedName;
-				UObject* ReplacedOuter = Coordinator ? Coordinator->FindReplacedNameAndOuter(Obj, /*out*/UnusedName) : nullptr;
-				Parent = ReplacedOuter ? ReplacedOuter : Obj->GetOuter();
+				if(const IBlueprintNativeCodeGenCore* Coordinator = IBlueprintNativeCodeGenCore::Get())
+				{
+					FName UnusedName;
+					UObject* ReplacedOuter = Coordinator->FindReplacedNameAndOuter(Obj, /*out*/UnusedName, GetNativizationOptionsForPlatform(CookingTarget(), Coordinator));
+					Parent = ReplacedOuter ? ReplacedOuter : Obj->GetOuter();
+				}
 #endif //WITH_EDITOR
 				if( Parent )
 				{
@@ -3339,9 +3353,7 @@ static bool ValidateConformCompatibility(UPackage* NewPackage, FLinkerLoad* OldL
 		if (OldClass != nullptr && NewClass != nullptr && OldClass->IsNative() && NewClass->IsNative())
 		{
 			OldClass->ClassConstructor = NewClass->ClassConstructor;
-#if WITH_HOT_RELOAD_CTORS
 			OldClass->ClassVTableHelperCtorCaller = NewClass->ClassVTableHelperCtorCaller;
-#endif // WITH_HOT_RELOAD_CTORS
 			OldClass->ClassAddReferencedObjects = NewClass->ClassAddReferencedObjects;
 		}
 	}
@@ -4154,9 +4166,9 @@ ESavePackageResult UPackage::Save(UPackage* InOuter, UObject* Base, EObjectFlags
 					}
 
 #if WITH_EDITOR
-					if (IBlueprintNativeCodeGenCore::Get())
+					if (const IBlueprintNativeCodeGenCore* Coordinator = IBlueprintNativeCodeGenCore::Get())
 					{
-						EReplacementResult ReplacmentResult = IBlueprintNativeCodeGenCore::Get()->IsTargetedForReplacement(InOuter);
+						EReplacementResult ReplacmentResult = Coordinator->IsTargetedForReplacement(InOuter, GetNativizationOptionsForPlatform(TargetPlatform, Coordinator));
 						if (ReplacmentResult == EReplacementResult::ReplaceCompletely)
 						{
 							if (IsEventDrivenLoaderEnabledInCookedBuilds() && TargetPlatform)
@@ -4300,7 +4312,7 @@ ESavePackageResult UPackage::Save(UPackage* InOuter, UObject* Base, EObjectFlags
 
 						SavePackageState->MarkNameAsReferenced(Obj->GetFName());
 #if WITH_EDITOR
-						SavePackageState->AddReplacementsNames(Obj);
+						SavePackageState->AddReplacementsNames(Obj, TargetPlatform);
 #endif //WITH_EDITOR
 						if( Obj->GetOuter() )
 						{
@@ -4640,11 +4652,12 @@ ESavePackageResult UPackage::Save(UPackage* InOuter, UObject* Base, EObjectFlags
 						FName ReplacedName = NAME_None;
 						if (const IBlueprintNativeCodeGenCore* Coordinator = IBlueprintNativeCodeGenCore::Get())
 						{
-							if (UClass* ReplacedClass = Coordinator->FindReplacedClassForObject(Obj))
+							const FCompilerNativizationOptions& NativizationOptions = GetNativizationOptionsForPlatform(TargetPlatform, Coordinator);
+							if (UClass* ReplacedClass = Coordinator->FindReplacedClassForObject(Obj, NativizationOptions))
 							{
 								ObjClass = ReplacedClass;
 							}
-							if (UObject* ReplacedOuter = Coordinator->FindReplacedNameAndOuter(Obj, /*out*/ReplacedName))
+							if (UObject* ReplacedOuter = Coordinator->FindReplacedNameAndOuter(Obj, /*out*/ReplacedName, NativizationOptions))
 							{
 								ReplacedImportOuters.Add(Obj, ReplacedOuter);
 							}
diff --git a/Engine/Source/Runtime/CoreUObject/Private/UObject/UObjectArray.cpp b/Engine/Source/Runtime/CoreUObject/Private/UObject/UObjectArray.cpp
--- a/Engine/Source/Runtime/CoreUObject/Private/UObject/UObjectArray.cpp
+++ b/Engine/Source/Runtime/CoreUObject/Private/UObject/UObjectArray.cpp
@@ -12,7 +12,7 @@
 
 DEFINE_LOG_CATEGORY_STATIC(LogUObjectArray, Log, All);
 
-TMap<int32, FUObjectCluster*> GUObjectClusters;
+FUObjectClusterContainer GUObjectClusters;
 
 FUObjectArray::FUObjectArray()
 : ObjFirstGCIndex(0)
@@ -145,7 +145,7 @@ void FUObjectArray::AllocateUObjectIndex(UObjectBase* Object, bool bMergingThrea
 	}
 	// Regular pool/ range.
 	else
-	{		
+	{
 		int32* AvailableIndex = ObjAvailableList.Pop();
 		if (AvailableIndex)
 		{
diff --git a/Engine/Source/Runtime/CoreUObject/Private/UObject/UObjectBase.cpp b/Engine/Source/Runtime/CoreUObject/Private/UObject/UObjectBase.cpp
--- a/Engine/Source/Runtime/CoreUObject/Private/UObject/UObjectBase.cpp
+++ b/Engine/Source/Runtime/CoreUObject/Private/UObject/UObjectBase.cpp
@@ -598,7 +598,8 @@ void UClassCompiledInDefer(FFieldCompiledInInfo* ClassInfo, const TCHAR* Name, S
 	const FName CPPClassName(Name);
 #if WITH_HOT_RELOAD
 	// Check for existing classes
-	FFieldCompiledInInfo** ExistingClassInfo = GetDeferRegisterClassMap().Find(CPPClassName);
+	TMap<FName, FFieldCompiledInInfo*>& DeferMap = GetDeferRegisterClassMap();
+	FFieldCompiledInInfo** ExistingClassInfo = DeferMap.Find(CPPClassName);
 	ClassInfo->bHasChanged = !ExistingClassInfo || (*ExistingClassInfo)->Size != ClassInfo->Size || (*ExistingClassInfo)->Crc != ClassInfo->Crc;
 	if (ExistingClassInfo)
 	{
@@ -639,7 +640,7 @@ void UClassCompiledInDefer(FFieldCompiledInInfo* ClassInfo, const TCHAR* Name, S
 	}
 	else
 	{
-		GetDeferRegisterClassMap().Add(CPPClassName, ClassInfo);
+		DeferMap.Add(CPPClassName, ClassInfo);
 	}
 #endif
 	// We will either create a new class or update the static class pointer of the existing one
@@ -658,7 +659,8 @@ void UObjectCompiledInDefer(UClass *(*InRegister)(), UClass *(*InStaticClass)(),
 	{
 #if WITH_HOT_RELOAD
 		// Either add all classes if not hot-reloading, or those which have changed
-		if (!GIsHotReload || GetDeferRegisterClassMap().FindChecked(Name)->bHasChanged)
+		TMap<FName, FFieldCompiledInInfo*>& DeferMap = GetDeferRegisterClassMap();
+		if (!GIsHotReload || DeferMap.FindChecked(Name)->bHasChanged)
 #endif
 		{
 			FString NoPrefix(RemoveClassPrefix(Name));
diff --git a/Engine/Source/Runtime/CoreUObject/Private/UObject/UObjectClusters.cpp b/Engine/Source/Runtime/CoreUObject/Private/UObject/UObjectClusters.cpp
--- a/Engine/Source/Runtime/CoreUObject/Private/UObject/UObjectClusters.cpp
+++ b/Engine/Source/Runtime/CoreUObject/Private/UObject/UObjectClusters.cpp
@@ -35,56 +35,226 @@ static FAutoConsoleVariableRef CMergeGCClusters(
 	ECVF_Default
 	);
 
-#if !UE_BUILD_SHIPPING
 
-// Dumps all clusters to log.
-static void ListClusters(const TArray<FString>& Args)
+FUObjectClusterContainer::FUObjectClusterContainer()
+	: NumAllocatedClusters(0)
+	, bClustersNeedDissolving(false)
 {
-	const bool bHierarchy = Args.Num() && Args[0] == TEXT("Hierarchy");
-	int32 MaxInterClusterReferences = 0;
-	int32 TotalInterClusterReferences = 0;
-	int32 MaxClusterSize = 0;
-	int32 TotalClusterObjects = 0;	
-	for (TPair<int32, FUObjectCluster*>& Pair : GUObjectClusters)
-	{
-		FUObjectItem* RootItem = GUObjectArray.IndexToObjectUnsafeForGC(Pair.Key);
-		UObject* RootObject = static_cast<UObject*>(RootItem->Object);
-		UE_LOG(LogObj, Display, TEXT("%s (Index: %d), Size: %d, ReferencedClusters: %d"), *RootObject->GetFullName(), Pair.Key, Pair.Value->Objects.Num(), Pair.Value->ReferencedClusters.Num());
-		MaxInterClusterReferences = FMath::Max(MaxInterClusterReferences, Pair.Value->ReferencedClusters.Num());
-		TotalInterClusterReferences += Pair.Value->ReferencedClusters.Num();
-		MaxClusterSize = FMath::Max(MaxClusterSize, Pair.Value->Objects.Num());
-		TotalClusterObjects += Pair.Value->Objects.Num();
-		if (bHierarchy)
+
+}
+
+int32 FUObjectClusterContainer::AllocateCluster(int32 InRootObjectIndex)
+{
+	int32 ClusterIndex = INDEX_NONE;
+	if (FreeClusterIndices.Num())
+	{
+		ClusterIndex = FreeClusterIndices.Pop(false);
+	}
+	else
+	{
+		ClusterIndex = Clusters.Add(FUObjectCluster());
+	}
+	FUObjectCluster& NewCluster = Clusters[ClusterIndex];
+	check(NewCluster.RootIndex == INDEX_NONE);
+	NewCluster.RootIndex = InRootObjectIndex;
+	NumAllocatedClusters++;
+	return ClusterIndex;
+}
+
+void FUObjectClusterContainer::FreeCluster(int32 InClusterIndex)
+{
+	FUObjectCluster& Cluster = Clusters[InClusterIndex];
+	check(Cluster.RootIndex != INDEX_NONE);
+	FUObjectItem* RootItem = GUObjectArray.IndexToObject(Cluster.RootIndex);
+	check(RootItem->GetClusterIndex() == InClusterIndex);
+	RootItem->SetOwnerIndex(0);
+	RootItem->ClearFlags(EInternalObjectFlags::ClusterRoot);
+
+	for (int32 ReferencedClusterRootIndex : Cluster.ReferencedClusters)
+	{
+		if (ReferencedClusterRootIndex >= 0)
 		{
-			int32 Index = 0;
-			for (int32 ObjectIndex : Pair.Value->Objects)
+			FUObjectItem* ReferencedClusterRootItem = GUObjectArray.IndexToObjectUnsafeForGC(ReferencedClusterRootIndex);
+			if (ReferencedClusterRootItem->GetOwnerIndex() < 0)
 			{
-				FUObjectItem* ObjectItem = GUObjectArray.IndexToObjectUnsafeForGC(ObjectIndex);
-				UObject* Object = static_cast<UObject*>(ObjectItem->Object);
-				UE_LOG(LogObj, Display, TEXT("    [%.4d]: %s (Index: %d)"), Index++, *Object->GetFullName(), ObjectIndex);
+				FUObjectCluster& ReferencedCluster = Clusters[ReferencedClusterRootItem->GetClusterIndex()];
+				ReferencedCluster.ReferencedByClusters.Remove(Cluster.RootIndex);
 			}
-			UE_LOG(LogObj, Display, TEXT("  Referenced clusters: %d"), Pair.Value->ReferencedClusters.Num());
-			for (int32 ClusterIndex : Pair.Value->ReferencedClusters)
+		}
+	}
+
+	Cluster.RootIndex = INDEX_NONE;
+	Cluster.Objects.Reset();
+	Cluster.MutableObjects.Reset();
+	Cluster.ReferencedClusters.Reset();
+	Cluster.ReferencedByClusters.Reset();
+	Cluster.bNeedsDissolving = false;
+	FreeClusterIndices.Add(InClusterIndex);
+	NumAllocatedClusters--;
+	check(NumAllocatedClusters >= 0);
+}
+
+void FUObjectClusterContainer::DissolveCluster(FUObjectCluster& Cluster)
+{
+	FUObjectItem* RootObjectItem = GUObjectArray.IndexToObjectUnsafeForGC(Cluster.RootIndex);
+
+	// Unreachable or not, we won't need this array later
+	TArray<int32> ReferencedByClusters = MoveTemp(Cluster.ReferencedByClusters);
+
+	// Unreachable clusters will be removed by GC during BeginDestroy phase (unhashing)
+	if (!RootObjectItem->IsUnreachable())
+	{
+		const int32 OldClusterIndex = RootObjectItem->GetClusterIndex();
+		for (int32 ClusterObjectIndex : Cluster.Objects)
+		{
+			FUObjectItem* ClusterObjectItem = GUObjectArray.IndexToObjectUnsafeForGC(ClusterObjectIndex);
+			ClusterObjectItem->SetOwnerIndex(0);
+		}		
+
+		FreeCluster(OldClusterIndex);
+	}
+
+	// Recursively dissolve all clusters this cluster is directly referenced by
+	for (int32 ReferencedByClusterRootIndex : ReferencedByClusters)
+	{
+		FUObjectItem* ReferencedByClusterRootObjectItem = GUObjectArray.IndexToObjectUnsafeForGC(ReferencedByClusterRootIndex);
+		if (ReferencedByClusterRootObjectItem->GetOwnerIndex())
+		{
+			DissolveCluster(Clusters[ReferencedByClusterRootObjectItem->GetClusterIndex()]);
+		}
+	}
+}
+
+void FUObjectClusterContainer::DissolveClusterAndMarkObjectsAsUnreachable(const int32 CurrentIndex, FUObjectItem* RootObjectItem)
+{
+	const int32 OldClusterIndex = RootObjectItem->GetClusterIndex();
+	FUObjectCluster& Cluster = Clusters[OldClusterIndex];
+
+	// Unreachable or not, we won't need this array later
+	TArray<int32> ReferencedByClusters = MoveTemp(Cluster.ReferencedByClusters);	
+
+	for (int32 ClusterObjectIndex : Cluster.Objects)
+	{
+		FUObjectItem* ClusterObjectItem = GUObjectArray.IndexToObjectUnsafeForGC(ClusterObjectIndex);
+		ClusterObjectItem->SetOwnerIndex(0);
+		if (ClusterObjectIndex < CurrentIndex)
+		{
+			ClusterObjectItem->SetFlags(EInternalObjectFlags::Unreachable | EInternalObjectFlags::NoStrongReference);
+		}
+	}
+
+	UObject* ClusterRootObject = static_cast<UObject*>(RootObjectItem->Object);
+	ClusterRootObject->OnClusterMarkedAsPendingKill();
+
+	FreeCluster(OldClusterIndex);
+
+	// Recursively dissolve all clusters this cluster is directly referenced by
+	for (int32 ReferencedByClusterRootIndex : ReferencedByClusters)
+	{
+		FUObjectItem* ReferencedByClusterRootItem = GUObjectArray.IndexToObjectUnsafeForGC(ReferencedByClusterRootIndex);
+		if (ReferencedByClusterRootItem->HasAnyFlags(EInternalObjectFlags::ClusterRoot))
+		{
+			if (ReferencedByClusterRootIndex < CurrentIndex)
+			{
+				ReferencedByClusterRootItem->SetFlags(EInternalObjectFlags::Unreachable | EInternalObjectFlags::NoStrongReference);
+			}
+			DissolveClusterAndMarkObjectsAsUnreachable(CurrentIndex, ReferencedByClusterRootItem);
+		}
+	}
+}
+
+void FUObjectClusterContainer::DissolveClusters()
+{
+	for (FUObjectCluster& Cluster : Clusters)
+	{
+		if (Cluster.RootIndex >= 0 && Cluster.bNeedsDissolving)
+		{
+			DissolveCluster(Cluster);
+		}
+	}
+	bClustersNeedDissolving = false;
+}
+
+#if !UE_BUILD_SHIPPING
+
+void DumpClusterToLog(const FUObjectCluster& Cluster, bool bHierarchy)
+{
+	FUObjectItem* RootItem = GUObjectArray.IndexToObjectUnsafeForGC(Cluster.RootIndex);
+	UObject* RootObject = static_cast<UObject*>(RootItem->Object);
+	UE_LOG(LogObj, Display, TEXT("%s (Index: %d), Size: %d, ReferencedClusters: %d"), *RootObject->GetFullName(), Cluster.RootIndex, Cluster.Objects.Num(), Cluster.ReferencedClusters.Num());
+	if (bHierarchy)
+	{
+		int32 Index = 0;
+		for (int32 ObjectIndex : Cluster.Objects)
+		{
+			FUObjectItem* ObjectItem = GUObjectArray.IndexToObjectUnsafeForGC(ObjectIndex);
+			UObject* Object = static_cast<UObject*>(ObjectItem->Object);
+			UE_LOG(LogObj, Display, TEXT("    [%.4d]: %s (Index: %d)"), Index++, *Object->GetFullName(), ObjectIndex);
+		}
+		UE_LOG(LogObj, Display, TEXT("  Referenced clusters: %d"), Cluster.ReferencedClusters.Num());
+		for (int32 ClusterRootIndex : Cluster.ReferencedClusters)
+		{
+			if (ClusterRootIndex >= 0)
 			{
-				FUObjectItem* ClusterRootItem = GUObjectArray.IndexToObjectUnsafeForGC(ClusterIndex);
+				FUObjectItem* ClusterRootItem = GUObjectArray.IndexToObjectUnsafeForGC(ClusterRootIndex);
 				UObject* ClusterRootObject = static_cast<UObject*>(ClusterRootItem->Object);
-				UE_LOG(LogObj, Display, TEXT("    -> %s (Index: %d)"), *ClusterRootObject->GetFullName(), ClusterIndex);
+				UE_LOG(LogObj, Display, TEXT("    -> %s (Index: %d)"), *ClusterRootObject->GetFullName(), ClusterRootIndex);
+			}
+			else
+			{
+				UE_LOG(LogObj, Display, TEXT("    -> nullptr"));
 			}
-			UE_LOG(LogObj, Display, TEXT("  External (mutable) objects: %d"), Pair.Value->MutableObjects.Num());
-			for (int32 ObjectIndex : Pair.Value->MutableObjects)
+		}
+		UE_LOG(LogObj, Display, TEXT("  External (mutable) objects: %d"), Cluster.MutableObjects.Num());
+		for (int32 ObjectIndex : Cluster.MutableObjects)
+		{
+			if (ObjectIndex >= 0)
 			{
 				FUObjectItem* ObjectItem = GUObjectArray.IndexToObjectUnsafeForGC(ObjectIndex);
 				UObject* Object = static_cast<UObject*>(ObjectItem->Object);
 				UE_LOG(LogObj, Display, TEXT("    -> %s (Index: %d)"), *Object->GetFullName(), ObjectIndex);
 			}
+			else
+			{
+				UE_LOG(LogObj, Display, TEXT("    -> nullptr"));
+			}
 		}
 	}
-	UE_LOG(LogObj, Display, TEXT("Number of clusters: %d"), GUObjectClusters.Num());
+}
+
+
+// Dumps all clusters to log.
+static void ListClusters(const TArray<FString>& Args)
+{
+	const bool bHierarchy = Args.Num() && Args[0] == TEXT("Hierarchy");
+	int32 MaxInterClusterReferences = 0;
+	int32 TotalInterClusterReferences = 0;
+	int32 MaxClusterSize = 0;
+	int32 TotalClusterObjects = 0;	
+	int32 NumClusters = 0;
+
+	for (const FUObjectCluster& Cluster : GUObjectClusters.GetClustersUnsafe())
+	{
+		if (Cluster.RootIndex == INDEX_NONE)
+		{
+			continue;
+		}
+
+		NumClusters++;
+
+		MaxInterClusterReferences = FMath::Max(MaxInterClusterReferences, Cluster.ReferencedClusters.Num());
+		TotalInterClusterReferences += Cluster.ReferencedClusters.Num();
+		MaxClusterSize = FMath::Max(MaxClusterSize, Cluster.Objects.Num());
+		TotalClusterObjects += Cluster.Objects.Num();
+
+		DumpClusterToLog(Cluster, bHierarchy);
+	}
+	UE_LOG(LogObj, Display, TEXT("Number of clusters: %d"), NumClusters);
 	UE_LOG(LogObj, Display, TEXT("Maximum cluster size: %d"), MaxClusterSize);
-	UE_LOG(LogObj, Display, TEXT("Average cluster size: %d"), GUObjectClusters.Num() ? (TotalClusterObjects / GUObjectClusters.Num()) : 0);
+	UE_LOG(LogObj, Display, TEXT("Average cluster size: %d"), NumClusters ? (TotalClusterObjects / NumClusters) : 0);
 	UE_LOG(LogObj, Display, TEXT("Number of objects in GC clusters: %d"), TotalClusterObjects);
 	UE_LOG(LogObj, Display, TEXT("Maximum number of custer-to-cluster references: %d"), MaxInterClusterReferences);
-	UE_LOG(LogObj, Display, TEXT("Average number of custer-to-cluster references: %d"), GUObjectClusters.Num() ? (TotalInterClusterReferences / GUObjectClusters.Num()) : 0);
+	UE_LOG(LogObj, Display, TEXT("Average number of custer-to-cluster references: %d"), NumClusters ? (TotalInterClusterReferences / NumClusters) : 0);
 }
 
 static void FindStaleClusters(const TArray<FString>& Args)
@@ -343,18 +513,18 @@ public:
 		else
 		{
 			OtherClusterRootIndex = ObjectItem->GetOwnerIndex();
+			check(OtherClusterRootIndex > 0);
 			OtherClusterRootItem = GUObjectArray.IndexToObjectUnsafeForGC(OtherClusterRootIndex);
 		}
 		// This is another cluster, merge it with this one
-		FUObjectCluster* ClusterToMerge = GUObjectClusters.FindChecked(OtherClusterRootIndex);
-		for (int32 OtherClusterObjectIndex : ClusterToMerge->Objects)
+		FUObjectCluster& ClusterToMerge = GUObjectClusters[OtherClusterRootItem->GetClusterIndex()];
+		for (int32 OtherClusterObjectIndex : ClusterToMerge.Objects)
 		{
 			FUObjectItem* OtherClusterObjectItem = GUObjectArray.IndexToObjectUnsafeForGC(OtherClusterObjectIndex);
 			OtherClusterObjectItem->SetOwnerIndex(0);
 			AddObjectToCluster(OtherClusterObjectIndex, OtherClusterObjectItem, static_cast<UObject*>(OtherClusterObjectItem->Object), ObjectsToSerialize, true);
 		}		
-		GUObjectClusters.Remove(OtherClusterRootIndex);
-		delete ClusterToMerge;
+		GUObjectClusters.FreeCluster(OtherClusterRootItem->GetClusterIndex());
 
 		// Make sure the root object is also added to the current cluster
 		OtherClusterRootItem->ClearFlags(EInternalObjectFlags::ClusterRoot);
@@ -396,14 +566,18 @@ public:
 					{
 						// Simply reference this cluster and all clusters it's referencing
 						const int32 OtherClusterRootIndex = ObjectItem->HasAnyFlags(EInternalObjectFlags::ClusterRoot) ? GUObjectArray.ObjectToIndex(Object) : ObjectItem->GetOwnerIndex();
-						FUObjectCluster* OtherCluster = GUObjectClusters.FindChecked(OtherClusterRootIndex);
+						FUObjectItem* OtherClusterRootItem = GUObjectArray.IndexToObject(OtherClusterRootIndex);
+						const int32 OtherClusterIndex = OtherClusterRootItem->GetClusterIndex();
+						FUObjectCluster& OtherCluster = GUObjectClusters[OtherClusterIndex];
 						Cluster.ReferencedClusters.AddUnique(OtherClusterRootIndex);
-						for (int32 OtherClusterReferencedCluster : OtherCluster->ReferencedClusters)
+						OtherCluster.ReferencedByClusters.AddUnique(ClusterRootIndex);
+
+						for (int32 OtherClusterReferencedCluster : OtherCluster.ReferencedClusters)
 						{
 							check(OtherClusterReferencedCluster != ClusterRootIndex);
 							Cluster.ReferencedClusters.AddUnique(OtherClusterReferencedCluster);
 						}
-						for (int32 OtherClusterReferencedMutableObjectIndex : OtherCluster->MutableObjects)
+						for (int32 OtherClusterReferencedMutableObjectIndex : OtherCluster.MutableObjects)
 						{
 							Cluster.MutableObjects.AddUnique(OtherClusterReferencedMutableObjectIndex);
 						}
@@ -481,7 +655,7 @@ void CreateClustersFromPackage(FLinkerLoad* PackageLinker)
 	}
 }
 
-void UObjectBaseUtility::AddToCluster(UObjectBaseUtility* ClusterRootOrObjectFromCluster)
+void UObjectBaseUtility::AddToCluster(UObjectBaseUtility* ClusterRootOrObjectFromCluster, bool bAddAsMutableObject /* = false */)
 {
 	check(ClusterRootOrObjectFromCluster);
 
@@ -498,15 +672,40 @@ void UObjectBaseUtility::AddToCluster(UObjectBaseUtility* ClusterRootOrObjectFro
 	}
 	if (ClusterRootIndex != 0)
 	{
-		FUObjectCluster* Cluster = GUObjectClusters.FindChecked(ClusterRootIndex);
-		FClusterReferenceProcessor Processor(ClusterRootIndex, *Cluster);			
-		TFastReferenceCollector<FClusterReferenceProcessor, TClusterCollector<FClusterReferenceProcessor>, FClusterArrayPool, true> ReferenceCollector(Processor, FClusterArrayPool::Get());
-		TArray<UObject*> ObjectsToProcess;
-		UObject* ThisObject = static_cast<UObject*>(this);
-		Processor.HandleTokenStreamObjectReference(ObjectsToProcess, static_cast<UObject*>(ClusterRootOrObjectFromCluster), ThisObject, INDEX_NONE, true);
-		if (ObjectsToProcess.Num())
+		const int32 ClusterIndex = ClusterRootIndex > 0 ? GUObjectArray.IndexToObject(ClusterRootIndex)->GetClusterIndex() : OuterItem->GetClusterIndex();
+		FUObjectCluster& Cluster = GUObjectClusters[ClusterIndex];
+		if (!bAddAsMutableObject)
 		{
-			ReferenceCollector.CollectReferences(ObjectsToProcess, true);
+			FClusterReferenceProcessor Processor(ClusterRootIndex, Cluster);
+			TFastReferenceCollector<false, FClusterReferenceProcessor, TClusterCollector<FClusterReferenceProcessor>, FClusterArrayPool, true> ReferenceCollector(Processor, FClusterArrayPool::Get());
+			TArray<UObject*> ObjectsToProcess;
+			UObject* ThisObject = static_cast<UObject*>(this);
+			Processor.HandleTokenStreamObjectReference(ObjectsToProcess, static_cast<UObject*>(ClusterRootOrObjectFromCluster), ThisObject, INDEX_NONE, true);
+			if (ObjectsToProcess.Num())
+			{
+				ReferenceCollector.CollectReferences(ObjectsToProcess);
+			}
+		}
+		else
+		{
+			// Adds this object's index to the MutableObjects array keeping it sorted and unique
+			const int32 ThisObjectIndex = GUObjectArray.ObjectToIndex(this);
+			int32 InsertedAt = INDEX_NONE;
+			for (int32 MutableObjectIndex = 0; MutableObjectIndex < Cluster.MutableObjects.Num() && InsertedAt == INDEX_NONE; ++MutableObjectIndex)
+			{
+				if (Cluster.MutableObjects[MutableObjectIndex] > ThisObjectIndex)
+				{
+					InsertedAt = Cluster.MutableObjects.Insert(ThisObjectIndex, MutableObjectIndex);
+				}
+				else if (Cluster.MutableObjects[MutableObjectIndex] == ThisObjectIndex)
+				{
+					InsertedAt = MutableObjectIndex;
+				}
+			}
+			if (InsertedAt == INDEX_NONE)
+			{
+				Cluster.MutableObjects.Add(ThisObjectIndex);
+			}
 		}
 	}
 }
@@ -530,29 +729,35 @@ void UObjectBaseUtility::CreateCluster()
 	check(!HasAnyFlags(RF_NeedLoad | RF_NeedPostLoad));
 
 	// Create a new cluster, reserve an arbitrary amount of memory for it.
-	FUObjectCluster* Cluster = new FUObjectCluster;
-	Cluster->Objects.Reserve(64);
+	const int32 ClusterIndex = GUObjectClusters.AllocateCluster(InternalIndex);
+	FUObjectCluster& Cluster = GUObjectClusters[ClusterIndex];
+	Cluster.Objects.Reserve(64);
 
 	// Collect all objects referenced by cluster root and by all objects it's referencing
-	FClusterReferenceProcessor Processor(InternalIndex, *Cluster);
-	TFastReferenceCollector<FClusterReferenceProcessor, TClusterCollector<FClusterReferenceProcessor>, FClusterArrayPool, true> ReferenceCollector(Processor, FClusterArrayPool::Get());
+	FClusterReferenceProcessor Processor(InternalIndex, Cluster);
+	TFastReferenceCollector<false, FClusterReferenceProcessor, TClusterCollector<FClusterReferenceProcessor>, FClusterArrayPool, true> ReferenceCollector(Processor, FClusterArrayPool::Get());
 	TArray<UObject*> ObjectsToProcess;
 	ObjectsToProcess.Add(static_cast<UObject*>(this));
-	ReferenceCollector.CollectReferences(ObjectsToProcess, true);
+	ReferenceCollector.CollectReferences(ObjectsToProcess);
 #if UE_BUILD_DEBUG
 	FClusterArrayPool::Get().CheckLeaks();
 #endif
 
-	if (Cluster->Objects.Num())
+	if (Cluster.Objects.Num())
 	{
 		// Add new cluster to the global cluster map.
-		GUObjectClusters.Add(InternalIndex, Cluster);
+		Cluster.Objects.Sort();
+		Cluster.ReferencedClusters.Sort();		
+		Cluster.MutableObjects.Sort();		
 		check(RootItem->GetOwnerIndex() == 0);
+		RootItem->SetClusterIndex(ClusterIndex);
 		RootItem->SetFlags(EInternalObjectFlags::ClusterRoot);
 	}
 	else
 	{
-		delete Cluster;
+		check(RootItem->GetOwnerIndex() == 0);
+		RootItem->SetClusterIndex(ClusterIndex);
+		GUObjectClusters.FreeCluster(ClusterIndex);
 	}
 }
 
@@ -574,7 +779,7 @@ public:
 	FClusterVerifyReferenceProcessor(UObject* InClusterRootObject)		
 		: ClusterRootObject(InClusterRootObject)
 		, ClusterRootIndex(GUObjectArray.ObjectToIndex(InClusterRootObject))
-		, Cluster(*GUObjectClusters.FindChecked(ClusterRootIndex)) // This can't fail otherwise there's something wrong with cluster creation code
+		, Cluster(GUObjectClusters[GUObjectArray.IndexToObject(ClusterRootIndex)->GetClusterIndex()]) // This can't fail otherwise there's something wrong with cluster creation code
 		, bIsRunningMultithreaded(false)
 		, bFailed(false)
 	{
@@ -657,14 +862,15 @@ public:
 		{
 			ProcessedObjects.Add(Object);
 			FUObjectItem* ObjectItem = GUObjectArray.ObjectToObjectItem(Object);
-			if (ObjectItem->GetOwnerIndex() == 0)
+			if (ObjectItem->GetOwnerIndex() <= 0)
 			{
 				// We are allowed to reference other clusters, root set objects and objects from diregard for GC pool
 					if (!ObjectItem->HasAnyFlags(EInternalObjectFlags::ClusterRoot|EInternalObjectFlags::RootSet) 
 						&& !GUObjectArray.IsDisregardForGC(Object) && Object->CanBeInCluster() &&
 						!Cluster.MutableObjects.Contains(GUObjectArray.ObjectToIndex(Object))) // This is for objects that had RF_NeedLoad|RF_NeedPostLoad set when creating the cluster
 				{
-					UE_LOG(LogObj, Warning, TEXT("Object %s from cluster %s is referencing 0x%016llx %s which is not part of root set or cluster."),
+					UE_LOG(LogObj, Warning, TEXT("Object 0x%016llx %s from cluster %s is referencing 0x%016llx %s which is not part of root set or cluster."),
+						(int64)(PTRINT)ReferencingObject,
 						*ReferencingObject->GetFullName(),
 						*ClusterRootObject->GetFullName(),
 						(int64)(PTRINT)Object,
@@ -699,6 +905,7 @@ public:
 			{
 				// If we're referencing an object from another cluster, make sure the other cluster is actually referenced by this cluster
 				const int32 OtherClusterRootIndex = ObjectItem->GetOwnerIndex();
+				check(OtherClusterRootIndex > 0);
 				const FUObjectItem* OtherClusterRootItem = GUObjectArray.IndexToObjectUnsafeForGC(OtherClusterRootIndex);				
 				check(OtherClusterRootItem && OtherClusterRootItem->Object);
 				UObject* OtherClusterRootObject = static_cast<UObject*>(OtherClusterRootItem->Object);
@@ -719,9 +926,9 @@ bool VerifyClusterAssumptions(UObject* ClusterRootObject)
 {
 	// Collect all objects referenced by cluster root and by all objects it's referencing
 	FClusterVerifyReferenceProcessor Processor(ClusterRootObject);
-	TFastReferenceCollector<FClusterVerifyReferenceProcessor, TClusterCollector<FClusterVerifyReferenceProcessor>, FClusterArrayPool> ReferenceCollector(Processor, FClusterArrayPool::Get());
+	TFastReferenceCollector<false, FClusterVerifyReferenceProcessor, TClusterCollector<FClusterVerifyReferenceProcessor>, FClusterArrayPool> ReferenceCollector(Processor, FClusterArrayPool::Get());
 	TArray<UObject*> ObjectsToProcess;
 	ObjectsToProcess.Add(ClusterRootObject);
-	ReferenceCollector.CollectReferences(ObjectsToProcess, true);
+	ReferenceCollector.CollectReferences(ObjectsToProcess);
 	return Processor.NoExternalReferencesFound();
 }
diff --git a/Engine/Source/Runtime/CoreUObject/Private/UObject/UObjectGlobals.cpp b/Engine/Source/Runtime/CoreUObject/Private/UObject/UObjectGlobals.cpp
--- a/Engine/Source/Runtime/CoreUObject/Private/UObject/UObjectGlobals.cpp
+++ b/Engine/Source/Runtime/CoreUObject/Private/UObject/UObjectGlobals.cpp
@@ -117,6 +117,9 @@ FCoreUObjectDelegates::FOnRedirectorFollowed FCoreUObjectDelegates::RedirectorFo
 FSimpleMulticastDelegate FCoreUObjectDelegates::PreGarbageCollect;
 FSimpleMulticastDelegate FCoreUObjectDelegates::PostGarbageCollect;
 
+FSimpleMulticastDelegate FCoreUObjectDelegates::PreGarbageCollectConditionalBeginDestroy;
+FSimpleMulticastDelegate FCoreUObjectDelegates::PostGarbageCollectConditionalBeginDestroy;
+
 FCoreUObjectDelegates::FPreLoadMapDelegate FCoreUObjectDelegates::PreLoadMap;
 FSimpleMulticastDelegate FCoreUObjectDelegates::PostLoadMap;
 FSimpleMulticastDelegate FCoreUObjectDelegates::PostDemoPlay;
@@ -851,7 +854,7 @@ UObject* StaticLoadObjectInternal(UClass* ObjectClass, UObject* InOuter, const T
 			))
 		{
 			Result = StaticFindObjectFast(ObjectClass, InOuter, *StrName);
-			if (Result && Result->HasAnyFlags(RF_NeedLoad | RF_NeedPostLoad | RF_NeedPostLoadSubobjects))
+			if (Result && Result->HasAnyFlags(RF_NeedLoad | RF_NeedPostLoad | RF_NeedPostLoadSubobjects | RF_WillBeLoaded))
 			{
 				// Object needs loading so load it before returning
 				Result = nullptr;
@@ -865,6 +868,10 @@ UObject* StaticLoadObjectInternal(UClass* ObjectClass, UObject* InOuter, const T
 
 			// now, find the object in the package
 			Result = StaticFindObjectFast(ObjectClass, InOuter, *StrName);
+			if (GEventDrivenLoaderEnabled && Result && Result->HasAnyFlags(RF_NeedLoad | RF_NeedPostLoad | RF_NeedPostLoadSubobjects | RF_WillBeLoaded))
+			{
+				UE_LOG(LogUObjectGlobals, Fatal, TEXT("Return an object still needing load from StaticLoadObjectInternal %s"), *GetFullNameSafe(Result));
+			}
 
 			// If the object was not found, check for a redirector and follow it if the class matches
 			if (!Result && !(LoadFlags & LOAD_NoRedirects))
@@ -1143,17 +1150,20 @@ static UPackage* LoadPackageInternalInner(UPackage* InOuter, const TCHAR* InLong
 		FName PackageFName(*InPackageName);
 
 		Result = FindObjectFast<UPackage>(nullptr, PackageFName);
-		if (!Result || !Result->IsFullyLoaded())
+		if (!Result || Result->LinkerLoad || !Result->IsFullyLoaded())
 		{
 			int32 RequestID = LoadPackageAsync(InName, nullptr, *InPackageName);
 			FlushAsyncLoading(RequestID);
 		}
 
 		if (InOuter)
-			{
+		{
 			return InOuter;
-			}
-				Result = FindObjectFast<UPackage>(nullptr, PackageFName);
+		}
+		if (!Result)
+		{
+			Result = FindObjectFast<UPackage>(nullptr, PackageFName);
+		}
 		return Result;
 }
 
@@ -2543,9 +2553,7 @@ void UObject::PostInitProperties()
 
 UObject::UObject()
 {
-#if WITH_HOT_RELOAD_CTORS
 	EnsureNotRetrievingVTablePtr();
-#endif // WITH_HOT_RELOAD_CTORS
 
 	FObjectInitializer* ObjectInitializerPtr = FUObjectThreadContext::Get().TopInitializer();
 	UE_CLOG(!ObjectInitializerPtr, LogUObjectGlobals, Fatal, TEXT("%s is not being constructed with either NewObject, NewNamedObject or ConstructObject."), *GetName());
@@ -2557,9 +2565,7 @@ UObject::UObject()
 
 UObject::UObject(const FObjectInitializer& ObjectInitializer)
 {
-#if WITH_HOT_RELOAD_CTORS
 	EnsureNotRetrievingVTablePtr();
-#endif // WITH_HOT_RELOAD_CTORS
 
 	UE_CLOG(ObjectInitializer.Obj != nullptr && ObjectInitializer.Obj != this, LogUObjectGlobals, Fatal, TEXT("UObject(const FObjectInitializer&) constructor called but it's not the object that's currently being constructed with NewObject. Maybe you trying to construct it on the stack which is not supported."));
 	const_cast<FObjectInitializer&>(ObjectInitializer).Obj = this;
diff --git a/Engine/Source/Runtime/CoreUObject/Public/Blueprint/BlueprintSupport.h b/Engine/Source/Runtime/CoreUObject/Public/Blueprint/BlueprintSupport.h
--- a/Engine/Source/Runtime/CoreUObject/Public/Blueprint/BlueprintSupport.h
+++ b/Engine/Source/Runtime/CoreUObject/Public/Blueprint/BlueprintSupport.h
@@ -7,6 +7,7 @@
 #include "UObject/UObjectGlobals.h"
 
 class UDynamicClass;
+struct FCompilerNativizationOptions;
 
 struct FBlueprintWarningDeclaration
 {
@@ -138,7 +139,7 @@ struct IBlueprintNativeCodeGenCore
 	 * @param Package	The package in question
 	 * @return Whether the package should be converted
 	 */
-	virtual EReplacementResult IsTargetedForReplacement(const UPackage* Package) const = 0;
+	virtual EReplacementResult IsTargetedForReplacement(const UPackage* Package, const FCompilerNativizationOptions& NativizationOptions) const = 0;
 	
 	/**
 	 * Determines whether the provided object needs to be replaced (in part or completely).
@@ -149,7 +150,7 @@ struct IBlueprintNativeCodeGenCore
 	 * @param Object	The package in question
 	 * @return Whether the object should be converted
 	 */
-	virtual EReplacementResult IsTargetedForReplacement(const UObject* Object) const = 0;
+	virtual EReplacementResult IsTargetedForReplacement(const UObject* Object, const FCompilerNativizationOptions& NativizationOptions) const = 0;
 
 	/** 
 	 * Function used to change the type of a class from, say, UBlueprintGeneratedClass to 
@@ -159,7 +160,7 @@ struct IBlueprintNativeCodeGenCore
 	 * @param Object whose class will be replaced
 	 * @return A replacement class ptr, null if none
 	 */
-	virtual UClass* FindReplacedClassForObject(const UObject* Object) const = 0;
+	virtual UClass* FindReplacedClassForObject(const UObject* Object, const FCompilerNativizationOptions& NativizationOptions) const = 0;
 	
 	/** 
 	 * Function used to change the path of subobject from a nativized class.
@@ -168,7 +169,12 @@ struct IBlueprintNativeCodeGenCore
 	 * @param OutName Referenced to name, that will be saved in import table.
 	 * @return An Outer object that should be saved in import table.
 	 */
-	virtual UObject* FindReplacedNameAndOuter(UObject* Object, FName& OutName) const = 0;
+	virtual UObject* FindReplacedNameAndOuter(UObject* Object, FName& OutName, const FCompilerNativizationOptions& NativizationOptions) const = 0;
+
+	/*
+	 * Return nativization options for given platform.
+	 */
+	virtual const FCompilerNativizationOptions& GetNativizationOptionsForPlatform(const FString& PlatformName) const = 0;
 };
 
 #endif // WITH_EDITOR
diff --git a/Engine/Source/Runtime/CoreUObject/Public/Misc/PackageName.h b/Engine/Source/Runtime/CoreUObject/Public/Misc/PackageName.h
--- a/Engine/Source/Runtime/CoreUObject/Public/Misc/PackageName.h
+++ b/Engine/Source/Runtime/CoreUObject/Public/Misc/PackageName.h
@@ -391,6 +391,13 @@ public:
 	 */
 	static bool IsScriptPackage(const FString& InPackageName);
 
+	/**
+	 * Checks the root of hte package's path to see if it's a memory package
+	 *		This should be set for packages that reside in memory and not on disk, we treat them similar to a script package
+	 * @return true if the root of the patch matches the memory path
+	 */
+	static bool IsMemoryPackage(const FString& InPackageName);
+
 	/**
 	 * Checks the root of the package's path to see if it is a localized package
 	 * @return true if the root of the path matches any localized root path
diff --git a/Engine/Source/Runtime/CoreUObject/Public/Serialization/ObjectAndNameAsStringProxyArchive.h b/Engine/Source/Runtime/CoreUObject/Public/Serialization/ObjectAndNameAsStringProxyArchive.h
--- a/Engine/Source/Runtime/CoreUObject/Public/Serialization/ObjectAndNameAsStringProxyArchive.h
+++ b/Engine/Source/Runtime/CoreUObject/Public/Serialization/ObjectAndNameAsStringProxyArchive.h
@@ -5,6 +5,8 @@
 #include "CoreMinimal.h"
 #include "Serialization/NameAsStringProxyArchive.h"
 
+class UObject;
+
 /**
  * Implements a proxy archive that serializes UObjects and FNames as string data.
  *
@@ -32,6 +34,7 @@ struct FObjectAndNameAsStringProxyArchive : public FNameAsStringProxyArchive
 	/** If we fail to find an object during loading, try and load it. */
 	bool bLoadIfFindFails;
 
-	COREUOBJECT_API virtual FArchive& operator<<(class UObject*& Obj);
+	COREUOBJECT_API virtual FArchive& operator<<(UObject*& Obj) override;
+	COREUOBJECT_API virtual FArchive& operator<<(FWeakObjectPtr& Obj) override;
 };
 
diff --git a/Engine/Source/Runtime/CoreUObject/Public/Templates/SubclassOf.h b/Engine/Source/Runtime/CoreUObject/Public/Templates/SubclassOf.h
--- a/Engine/Source/Runtime/CoreUObject/Public/Templates/SubclassOf.h
+++ b/Engine/Source/Runtime/CoreUObject/Public/Templates/SubclassOf.h
@@ -92,6 +92,12 @@ public:
 		Ar << SubclassOf.Class;
 		return Ar;
 	}
+
+	friend uint32 GetTypeHash(const TSubclassOf& SubclassOf)
+	{
+		return GetTypeHash(SubclassOf.Class);
+	}
+
 private:
 	UClass* Class;
 };
diff --git a/Engine/Source/Runtime/CoreUObject/Public/UObject/CoreNative.h b/Engine/Source/Runtime/CoreUObject/Public/UObject/CoreNative.h
--- a/Engine/Source/Runtime/CoreUObject/Public/UObject/CoreNative.h
+++ b/Engine/Source/Runtime/CoreUObject/Public/UObject/CoreNative.h
@@ -15,6 +15,14 @@ struct FFrame;
 /** The type of a native function callable by script */
 typedef void (UObject::*Native)(FFrame& TheStack, RESULT_DECL);
 
+// This class is deliberately simple (i.e. POD) to keep generated code size down.
+template <typename CharType>
+struct TNameNativePtrPair
+{
+	const CharType* Name;
+	Native          Pointer;
+};
+
 extern COREUOBJECT_API Native GCasts[];
 uint8 COREUOBJECT_API GRegisterCast( int32 CastCode, const Native& Func );
 
@@ -29,4 +37,7 @@ struct FNativeFunctionRegistrar
 	static COREUOBJECT_API void RegisterFunction(class UClass* Class, const ANSICHAR* InName, Native InPointer);
 	// overload for types generated from blueprints, which can have unicode names:
 	static COREUOBJECT_API void RegisterFunction(class UClass* Class, const WIDECHAR* InName, Native InPointer);
+
+	static COREUOBJECT_API void RegisterFunctions(class UClass* Class, const TNameNativePtrPair<ANSICHAR>* InArray, int32 NumFunctions);
+	static COREUOBJECT_API void RegisterFunctions(class UClass* Class, const TNameNativePtrPair<WIDECHAR>* InArray, int32 NumFunctions);
 };
