Copying //UE4/Dev-Core to //UE4/Dev-Main (Source: //UE4/Dev-Core @ 3345728)

diff --git a/Engine/Source/Runtime/CoreUObject/Public/UObject/FastReferenceCollector.h b/Engine/Source/Runtime/CoreUObject/Public/UObject/FastReferenceCollector.h
new file mode 100644
--- /dev/null
+++ b/Engine/Source/Runtime/CoreUObject/Public/UObject/FastReferenceCollector.h
@@ -0,0 +1,667 @@
+// Copyright 1998-2017 Epic Games, Inc. All Rights Reserved.
+
+#pragma once
+
+#include "CoreMinimal.h"
+#include "Stats/Stats.h"
+#include "UObject/GarbageCollection.h"
+#include "UObject/Class.h"
+#include "Async/TaskGraphInterfaces.h"
+#include "UObject/UnrealType.h"
+#include "Misc/ScopeLock.h"
+#include "HAL/PlatformProcess.h"
+
+struct FStackEntry;
+
+/*=============================================================================
+	FastReferenceCollector.h: Unreal realtime garbage collection helpers
+=============================================================================*/
+
+/**
+ * Helper class that looks for UObject references by traversing UClass token stream and calls AddReferencedObjects.
+ * Provides a generic way of processing references that is used by Unreal Engine garbage collection.
+ * Can be used for fast (does not use serialization) reference collection purposes.
+ * 
+ * IT IS CRITICAL THIS CLASS DOES NOT CHANGE WITHOUT CONSIDERING PERFORMANCE IMPACT OF SAID CHANGES
+ *
+ * This class depends on three components: ReferenceProcessor, ReferenceCollector and ArrayPool.
+ * The assumptions for each of those components are as follows:
+ *
+   class FSampleReferenceProcessor
+   {
+   public:
+     int32 GetMinDesiredObjectsPerSubTask() const;
+		 void HandleTokenStreamObjectReference(TArray<UObject*>& ObjectsToSerialize, UObject* ReferencingObject, UObject*& Object, const int32 TokenIndex, bool bAllowReferenceElimination);
+		 void UpdateDetailedStats(UObject* CurrentObject, uint32 DeltaCycles);
+		 void LogDetailedStatsSummary();
+	 };
+
+	 class FSampleCollctor : public FReferenceCollector 
+	 {
+	   // Needs to implement FReferenceCollector pure virtual functions
+	 };
+   
+	 class FSampleArrayPool
+	 {
+	   static FSampleArrayPool& Get();
+		 TArray<UObject*>* GetArrayFromPool();
+		 void ReturnToPool(TArray<UObject*>* Array);
+	 };
+ */
+template <bool bParallel, typename ReferenceProcessorType, typename CollectorType, typename ArrayPoolType, bool bAutoGenerateTokenStream = false>
+class TFastReferenceCollector
+{
+private:
+
+	class FCollectorTaskQueue
+	{
+		TFastReferenceCollector*	Owner;
+		ArrayPoolType& ArrayPool;
+		TLockFreePointerListUnordered<TArray<UObject*>, PLATFORM_CACHE_LINE_SIZE> Tasks;
+
+		FCriticalSection WaitingThreadsLock;
+		TArray<FEvent*> WaitingThreads;
+		bool bDone;
+		int32 NumThreadsStarted;
+	public:
+
+		FCollectorTaskQueue(TFastReferenceCollector* InOwner, ArrayPoolType& InArrayPool)
+			: Owner(InOwner)
+			, ArrayPool(InArrayPool)
+			, bDone(false)
+			, NumThreadsStarted(0)
+		{
+		}
+
+		void CheckDone()
+		{
+			FScopeLock Lock(&WaitingThreadsLock);
+			check(bDone);
+			check(!Tasks.Pop());
+			check(!WaitingThreads.Num());
+			check(NumThreadsStarted);
+		}
+
+		FORCENOINLINE void AddTask(const TArray<UObject*>* InObjectsToSerialize, int32 StartIndex, int32 NumObjects)
+		{
+			TArray<UObject*>* ObjectsToSerialize(ArrayPool.GetArrayFromPool());
+			ObjectsToSerialize->AddUninitialized(NumObjects);
+			FMemory::Memcpy(ObjectsToSerialize->GetData(), InObjectsToSerialize->GetData() + StartIndex, NumObjects * sizeof(UObject*));
+			Tasks.Push(ObjectsToSerialize);
+
+			FEvent* WaitingThread = nullptr;
+			{
+				FScopeLock Lock(&WaitingThreadsLock);
+				check(!bDone);
+				if (WaitingThreads.Num())
+				{
+					WaitingThread = WaitingThreads.Pop();
+				}
+			}
+			if (WaitingThread)
+			{
+				WaitingThread->Trigger();
+			}
+		}
+
+		FORCENOINLINE void DoTask()
+		{
+			{
+				FScopeLock Lock(&WaitingThreadsLock);
+				if (bDone)
+				{
+					return;
+				}
+				NumThreadsStarted++;
+			}
+			while (true)
+			{
+				TArray<UObject*>* ObjectsToSerialize = Tasks.Pop();
+				while (!ObjectsToSerialize)
+				{
+					if (bDone)
+					{
+						return;
+					}
+					FEvent* WaitEvent = nullptr;
+					{
+						FScopeLock Lock(&WaitingThreadsLock);
+						if (bDone)
+						{
+							return;
+						}
+						ObjectsToSerialize = Tasks.Pop();
+						if (!ObjectsToSerialize)
+						{
+							if (WaitingThreads.Num() + 1 == NumThreadsStarted)
+							{
+								bDone = true;
+								FPlatformMisc::MemoryBarrier();
+								for (FEvent* WaitingThread : WaitingThreads)
+								{
+									WaitingThread->Trigger();
+								}
+								WaitingThreads.Empty();
+								return;
+							}
+							else
+							{
+								WaitEvent = FPlatformProcess::GetSynchEventFromPool(false);
+								WaitingThreads.Push(WaitEvent);
+							}
+						}
+					}
+					if (ObjectsToSerialize)
+					{
+						check(!WaitEvent);
+					}
+					else
+					{
+						check(WaitEvent);
+						WaitEvent->Wait();
+						FPlatformProcess::ReturnSynchEventToPool(WaitEvent);
+						ObjectsToSerialize = Tasks.Pop();
+						check(!ObjectsToSerialize || !bDone);
+					}
+				}
+				Owner->ProcessObjectArray(*ObjectsToSerialize, FGraphEventRef());
+				ArrayPool.ReturnToPool(ObjectsToSerialize);
+			}
+		}
+	};
+
+	/** Task graph task responsible for processing UObject array */
+	class FCollectorTaskProcessorTask
+	{
+		FCollectorTaskQueue& TaskQueue;
+		ENamedThreads::Type DesiredThread;
+	public:
+		FCollectorTaskProcessorTask(FCollectorTaskQueue& InTaskQueue, ENamedThreads::Type InDesiredThread)
+			: TaskQueue(InTaskQueue)
+			, DesiredThread(InDesiredThread)
+		{
+		}
+		FORCEINLINE TStatId GetStatId() const
+		{
+			RETURN_QUICK_DECLARE_CYCLE_STAT(FCollectorTaskProcessorTask, STATGROUP_TaskGraphTasks);
+		}
+		ENamedThreads::Type GetDesiredThread()
+		{
+			return DesiredThread;
+		}
+		static ESubsequentsMode::Type GetSubsequentsMode()
+		{
+			return ESubsequentsMode::TrackSubsequents;
+		}
+		void DoTask(ENamedThreads::Type CurrentThread, FGraphEventRef& MyCompletionGraphEvent)
+		{
+			TaskQueue.DoTask();
+		}
+	};
+
+	/** Task graph task responsible for processing UObject array */
+	class FCollectorTask
+	{
+		TFastReferenceCollector*	Owner;
+		TArray<UObject*>*	ObjectsToSerialize;
+		ArrayPoolType& ArrayPool;
+
+	public:
+		FCollectorTask(TFastReferenceCollector* InOwner, const TArray<UObject*>* InObjectsToSerialize, int32 StartIndex, int32 NumObjects, ArrayPoolType& InArrayPool)
+			: Owner(InOwner)
+			, ObjectsToSerialize(InArrayPool.GetArrayFromPool())
+			, ArrayPool(InArrayPool)
+		{
+			ObjectsToSerialize->AddUninitialized(NumObjects);
+			FMemory::Memcpy(ObjectsToSerialize->GetData(), InObjectsToSerialize->GetData() + StartIndex, NumObjects * sizeof(UObject*));
+		}
+		~FCollectorTask()
+		{
+			ArrayPool.ReturnToPool(ObjectsToSerialize);
+		}
+		FORCEINLINE TStatId GetStatId() const
+		{
+			RETURN_QUICK_DECLARE_CYCLE_STAT(FCollectorTask, STATGROUP_TaskGraphTasks);
+		}
+		static ENamedThreads::Type GetDesiredThread()
+		{
+			if ((PLATFORM_XBOXONE || PLATFORM_PS4) && ENamedThreads::bHasHighPriorityThreads)
+			{
+				if (PLATFORM_PS4 && ENamedThreads::bHasBackgroundThreads) // on the PS4, background threads can use the 7th core, so lets put it to work.
+				{
+					int32 CoreRand = FMath::RandRange(0, 6);
+					if (CoreRand < 2)
+					{
+						return ENamedThreads::AnyBackgroundThreadNormalTask;
+					}
+					else if (CoreRand < 4)
+					{
+						return ENamedThreads::AnyHiPriThreadNormalTask;
+					}
+				}
+				else
+				{
+					if (FMath::RandRange(0, 1))
+					{
+						return ENamedThreads::AnyHiPriThreadNormalTask;
+					}
+				}
+			}
+			return ENamedThreads::AnyThread;
+		}
+		static ESubsequentsMode::Type GetSubsequentsMode()
+		{
+			return ESubsequentsMode::TrackSubsequents;
+		}
+		void DoTask(ENamedThreads::Type CurrentThread, FGraphEventRef& MyCompletionGraphEvent)
+		{
+			Owner->ProcessObjectArray(*ObjectsToSerialize, MyCompletionGraphEvent);
+		}
+	};
+
+	/** Object that handles all UObject references */
+	ReferenceProcessorType& ReferenceProcessor;
+	/** Custom TArray allocator */
+	ArrayPoolType& ArrayPool;
+
+	FCollectorTaskQueue TaskQueue;
+
+	/** Helper struct for stack based approach */
+	struct FStackEntry
+	{
+		/** Current data pointer, incremented by stride */
+		uint8*	Data;
+		/** Current stride */
+		int32		Stride;
+		/** Current loop count, decremented each iteration */
+		int32		Count;
+		/** First token index in loop */
+		int32		LoopStartIndex;
+	};
+
+
+
+public:
+	/** Default constructor, initializing all members. */
+	TFastReferenceCollector(ReferenceProcessorType& InReferenceProcessor, ArrayPoolType& InArrayPool)
+		: ReferenceProcessor(InReferenceProcessor)
+		, ArrayPool(InArrayPool)
+		, TaskQueue(this, InArrayPool)
+	{}
+
+	/**
+	* Performs reachability analysis.
+	*
+	* @param ObjectsToCollectReferencesFor List of objects which references should be collected
+	* @param bForceSingleThreaded Collect references on a single thread
+	*/
+	void CollectReferences(TArray<UObject*>& ObjectsToCollectReferencesFor)
+	{
+		if (ObjectsToCollectReferencesFor.Num())
+		{
+			if (!bParallel)
+			{
+				FGraphEventRef InvalidRef;
+				ProcessObjectArray(ObjectsToCollectReferencesFor, InvalidRef);
+			}
+			else
+			{
+				FGraphEventArray ChunkTasks;
+
+				int32 NumThreads = FTaskGraphInterface::Get().GetNumWorkerThreads();
+				int32 NumBackgroundThreads = ENamedThreads::bHasBackgroundThreads ? NumThreads : 0;
+#if ((PLATFORM_PS4 && USE_7TH_CORE) || PLATFORM_XBOXONE)
+				if (NumBackgroundThreads)
+				{
+					NumBackgroundThreads = 7 - NumThreads;
+				}
+#elif PLATFORM_PS4
+				if (NumBackgroundThreads)
+				{
+					NumBackgroundThreads = 6 - NumThreads;
+				}
+#endif
+				int32 NumTasks = NumThreads + NumBackgroundThreads;
+
+				check(NumTasks > 0);
+				ChunkTasks.Empty(NumTasks);
+				int32 NumPerChunk = ObjectsToCollectReferencesFor.Num() / NumTasks;
+				int32 StartIndex = 0;
+				for (int32 Chunk = 0; Chunk < NumTasks; Chunk++)
+				{
+					if (Chunk + 1 == NumTasks)
+					{
+						NumPerChunk = ObjectsToCollectReferencesFor.Num() - StartIndex; // last chunk takes all remaining items
+					}
+					TaskQueue.AddTask(&ObjectsToCollectReferencesFor, StartIndex, NumPerChunk);
+					StartIndex += NumPerChunk;
+				}
+				for (int32 Chunk = 0; Chunk < NumTasks; Chunk++)
+				{
+					ChunkTasks.Add(TGraphTask< FCollectorTaskProcessorTask >::CreateTask().ConstructAndDispatchWhenReady(TaskQueue, Chunk >= NumThreads ? ENamedThreads::AnyBackgroundThreadNormalTask : ENamedThreads::AnyNormalThreadNormalTask));
+				}
+
+				QUICK_SCOPE_CYCLE_COUNTER(STAT_GC_Subtask_Wait);
+				FTaskGraphInterface::Get().WaitUntilTasksComplete(ChunkTasks, ENamedThreads::GameThread_Local);
+				TaskQueue.CheckDone();
+			}
+		}
+	}
+
+private:
+
+	/**
+	 * Traverses UObject token stream to find existing references
+	 *
+	 * @param InObjectsToSerializeArray Objects to process
+	 * @param MyCompletionGraphEvent Task graph event
+	 */
+	void ProcessObjectArray(TArray<UObject*>& InObjectsToSerializeArray, const FGraphEventRef& MyCompletionGraphEvent)
+	{
+		DECLARE_SCOPE_CYCLE_COUNTER(TEXT("TFastReferenceCollector::ProcessObjectArray"), STAT_FFastReferenceCollector_ProcessObjectArray, STATGROUP_GC);
+
+		UObject* CurrentObject = NULL;
+
+		const int32 MinDesiredObjectsPerSubTask = ReferenceProcessor.GetMinDesiredObjectsPerSubTask(); // sometimes there will be less, a lot less
+
+		/** Growing array of objects that require serialization */
+		TArray<UObject*>&	NewObjectsToSerializeArray = *ArrayPool.GetArrayFromPool();
+
+		// Ping-pong between these two arrays if there's not enough objects to spawn a new task
+		TArray<UObject*>& ObjectsToSerialize = InObjectsToSerializeArray;
+		TArray<UObject*>& NewObjectsToSerialize = NewObjectsToSerializeArray;
+
+		// Presized "recursion" stack for handling arrays and structs.
+		TArray<FStackEntry> Stack;
+		Stack.AddUninitialized(128); //@todo rtgc: need to add code handling more than 128 layers of recursion or at least assert
+
+		// it is necessary to have at least one extra item in the array memory block for the iffy prefetch code, below
+		ObjectsToSerialize.Reserve(ObjectsToSerialize.Num() + 1);
+
+		// Keep serializing objects till we reach the end of the growing array at which point
+		// we are done.
+		int32 CurrentIndex = 0;
+		do
+		{
+			CollectorType ReferenceCollector(ReferenceProcessor, NewObjectsToSerialize);
+			while (CurrentIndex < ObjectsToSerialize.Num())
+			{
+#if PERF_DETAILED_PER_CLASS_GC_STATS
+				uint32 StartCycles = FPlatformTime::Cycles();
+#endif
+				CurrentObject = ObjectsToSerialize[CurrentIndex++];
+
+				// GetData() used to avoiding bounds checking (min and max)
+				// FMath::Min used to avoid out of bounds (without branching) on last iteration. Though anything can be passed into PrefetchBlock, 
+				// reading ObjectsToSerialize out of bounds is not safe since ObjectsToSerialize[Num()] may be an unallocated/unsafe address.
+				const UObject * const NextObject = ObjectsToSerialize.GetData()[FMath::Min<int32>(CurrentIndex, ObjectsToSerialize.Num() - 1)];
+
+				// Prefetch the next object assuming that the property size of the next object is the same as the current one.
+				// This allows us to avoid a branch here.
+				FPlatformMisc::PrefetchBlock(NextObject, CurrentObject->GetClass()->GetPropertiesSize());
+
+				//@todo rtgc: we need to handle object references in struct defaults
+
+				// Make sure that token stream has been assembled at this point as the below code relies on it.
+				if (!bParallel && bAutoGenerateTokenStream)
+				{
+					UClass* ObjectClass = CurrentObject->GetClass();
+					if (!ObjectClass->HasAnyClassFlags(CLASS_TokenStreamAssembled))
+					{
+						ObjectClass->AssembleReferenceTokenStream();
+					}
+				}
+#if DO_CHECK
+				if (!CurrentObject->GetClass()->HasAnyClassFlags(CLASS_TokenStreamAssembled))
+				{
+					UE_LOG(LogGarbage, Fatal, TEXT("%s does not yet have a token stream assembled."), *GetFullNameSafe(CurrentObject->GetClass()));
+				}
+#endif
+
+				// Get pointer to token stream and jump to the start.
+				FGCReferenceTokenStream* RESTRICT TokenStream = &CurrentObject->GetClass()->ReferenceTokenStream;
+				uint32 TokenStreamIndex = 0;
+				// Keep track of index to reference info. Used to avoid LHSs.
+				uint32 ReferenceTokenStreamIndex = 0;
+
+				// Create stack entry and initialize sane values.
+				FStackEntry* RESTRICT StackEntry = Stack.GetData();
+				uint8* StackEntryData = (uint8*)CurrentObject;
+				StackEntry->Data = StackEntryData;
+				StackEntry->Stride = 0;
+				StackEntry->Count = -1;
+				StackEntry->LoopStartIndex = -1;
+
+				// Keep track of token return count in separate integer as arrays need to fiddle with it.
+				int32 TokenReturnCount = 0;
+
+				// Parse the token stream.
+				while (true)
+				{
+					// Cache current token index as it is the one pointing to the reference info.
+					ReferenceTokenStreamIndex = TokenStreamIndex;
+
+					// Handle returning from an array of structs, array of structs of arrays of ... (yadda yadda)
+					for (int32 ReturnCount = 0; ReturnCount<TokenReturnCount; ReturnCount++)
+					{
+						// Make sure there's no stack underflow.
+						check(StackEntry->Count != -1);
+
+						// We pre-decrement as we're already through the loop once at this point.
+						if (--StackEntry->Count > 0)
+						{
+							// Point data to next entry.
+							StackEntryData = StackEntry->Data + StackEntry->Stride;
+							StackEntry->Data = StackEntryData;
+
+							// Jump back to the beginning of the loop.
+							TokenStreamIndex = StackEntry->LoopStartIndex;
+							ReferenceTokenStreamIndex = StackEntry->LoopStartIndex;
+							// We're not done with this token loop so we need to early out instead of backing out further.
+							break;
+						}
+						else
+						{
+							StackEntry--;
+							StackEntryData = StackEntry->Data;
+						}
+					}
+
+					TokenStreamIndex++;
+					FGCReferenceInfo ReferenceInfo = TokenStream->AccessReferenceInfo(ReferenceTokenStreamIndex);
+
+					switch(ReferenceInfo.Type)
+					{
+					case GCRT_Object:
+					{
+						// We're dealing with an object reference.
+						UObject**	ObjectPtr = (UObject**)(StackEntryData + ReferenceInfo.Offset);
+						UObject*&	Object = *ObjectPtr;
+						TokenReturnCount = ReferenceInfo.ReturnCount;
+						ReferenceProcessor.HandleTokenStreamObjectReference(NewObjectsToSerialize, CurrentObject, Object, ReferenceTokenStreamIndex, true);
+					}
+					break;
+					case GCRT_ArrayObject:
+					{
+						// We're dealing with an array of object references.
+						TArray<UObject*>& ObjectArray = *((TArray<UObject*>*)(StackEntryData + ReferenceInfo.Offset));
+						TokenReturnCount = ReferenceInfo.ReturnCount;
+						for (int32 ObjectIndex = 0, ObjectNum = ObjectArray.Num(); ObjectIndex < ObjectNum; ++ObjectIndex)
+						{
+							ReferenceProcessor.HandleTokenStreamObjectReference(NewObjectsToSerialize, CurrentObject, ObjectArray[ObjectIndex], ReferenceTokenStreamIndex, true);
+						}
+					}
+					break;
+					case GCRT_ArrayStruct:
+					{
+						// We're dealing with a dynamic array of structs.
+						const FScriptArray& Array = *((FScriptArray*)(StackEntryData + ReferenceInfo.Offset));
+						StackEntry++;
+						StackEntryData = (uint8*)Array.GetData();
+						StackEntry->Data = StackEntryData;
+						StackEntry->Stride = TokenStream->ReadStride(TokenStreamIndex);
+						StackEntry->Count = Array.Num();
+
+						const FGCSkipInfo SkipInfo = TokenStream->ReadSkipInfo(TokenStreamIndex);
+						StackEntry->LoopStartIndex = TokenStreamIndex;
+
+						if (StackEntry->Count == 0)
+						{
+							// Skip empty array by jumping to skip index and set return count to the one about to be read in.
+							TokenStreamIndex = SkipInfo.SkipIndex;
+							TokenReturnCount = TokenStream->GetSkipReturnCount(SkipInfo);
+						}
+						else
+						{
+							// Loop again.
+							check(StackEntry->Data);
+							TokenReturnCount = 0;
+						}
+					}
+					break;
+					case GCRT_PersistentObject:
+					{
+						// We're dealing with an object reference.
+						UObject**	ObjectPtr = (UObject**)(StackEntryData + ReferenceInfo.Offset);
+						UObject*&	Object = *ObjectPtr;
+						TokenReturnCount = ReferenceInfo.ReturnCount;
+						ReferenceProcessor.HandleTokenStreamObjectReference(NewObjectsToSerialize, CurrentObject, Object, ReferenceTokenStreamIndex, false);
+					}
+					break;
+					case GCRT_FixedArray:
+					{
+						// We're dealing with a fixed size array
+						uint8* PreviousData = StackEntryData;
+						StackEntry++;
+						StackEntryData = PreviousData;
+						StackEntry->Data = PreviousData;
+						StackEntry->Stride = TokenStream->ReadStride(TokenStreamIndex);
+						StackEntry->Count = TokenStream->ReadCount(TokenStreamIndex);
+						StackEntry->LoopStartIndex = TokenStreamIndex;
+						TokenReturnCount = 0;
+					}
+					break;
+					case GCRT_AddStructReferencedObjects:
+					{
+						// We're dealing with a function call
+						void const*	StructPtr = (void*)(StackEntryData + ReferenceInfo.Offset);
+						TokenReturnCount = ReferenceInfo.ReturnCount;
+						UScriptStruct::ICppStructOps::TPointerToAddStructReferencedObjects Func = (UScriptStruct::ICppStructOps::TPointerToAddStructReferencedObjects) TokenStream->ReadPointer(TokenStreamIndex);
+						Func(StructPtr, ReferenceCollector);
+					}
+					break;
+					case GCRT_AddReferencedObjects:
+					{
+						// Static AddReferencedObjects function call.
+						void(*AddReferencedObjects)(UObject*, FReferenceCollector&) = (void(*)(UObject*, FReferenceCollector&))TokenStream->ReadPointer(TokenStreamIndex);
+						TokenReturnCount = ReferenceInfo.ReturnCount;
+						AddReferencedObjects(CurrentObject, ReferenceCollector);
+					}
+					break;
+					case GCRT_AddTMapReferencedObjects:
+					{
+						void*         Map = StackEntryData + ReferenceInfo.Offset;
+						UMapProperty* MapProperty = (UMapProperty*)TokenStream->ReadPointer(TokenStreamIndex);
+						TokenReturnCount = ReferenceInfo.ReturnCount;
+						FSimpleObjectReferenceCollectorArchive CollectorArchive(CurrentObject, ReferenceCollector);
+						MapProperty->SerializeItem(CollectorArchive, Map, nullptr);
+					}
+					break;
+					case GCRT_AddTSetReferencedObjects:
+					{
+						void*         Set = StackEntryData + ReferenceInfo.Offset;
+						USetProperty* SetProperty = (USetProperty*)TokenStream->ReadPointer(TokenStreamIndex);
+						TokenReturnCount = ReferenceInfo.ReturnCount;
+						FSimpleObjectReferenceCollectorArchive CollectorArchive(CurrentObject, ReferenceCollector);
+						SetProperty->SerializeItem(CollectorArchive, Set, nullptr);
+					}
+					break;
+					case GCRT_EndOfPointer:
+					{
+						TokenReturnCount = ReferenceInfo.ReturnCount;
+					}
+					break;
+					case GCRT_EndOfStream:
+					{
+						// Break out of loop.
+						goto EndLoop;
+					}
+					break;
+					default:
+					{
+						UE_LOG(LogGarbage, Fatal, TEXT("Unknown token"));
+						break;
+					}
+				}
+				}
+EndLoop:
+				check(StackEntry == Stack.GetData());
+
+				if (bParallel && NewObjectsToSerialize.Num() >= MinDesiredObjectsPerSubTask)
+				{
+					// This will start queueing task with objects from the end of array until there's less objects than worth to queue
+					const int32 ObjectsPerSubTask = FMath::Max<int32>(MinDesiredObjectsPerSubTask, NewObjectsToSerialize.Num() / FTaskGraphInterface::Get().GetNumWorkerThreads());
+					while (NewObjectsToSerialize.Num() >= MinDesiredObjectsPerSubTask)
+					{
+						const int32 StartIndex = FMath::Max(0, NewObjectsToSerialize.Num() - ObjectsPerSubTask);
+						const int32 NumThisTask = NewObjectsToSerialize.Num() - StartIndex;
+						if (MyCompletionGraphEvent.GetReference())
+						{
+							MyCompletionGraphEvent->DontCompleteUntil(TGraphTask< FCollectorTask >::CreateTask().ConstructAndDispatchWhenReady(this, &NewObjectsToSerialize, StartIndex, NumThisTask, ArrayPool));
+						}
+						else
+						{
+							TaskQueue.AddTask(&NewObjectsToSerialize, StartIndex, NumThisTask);
+						}
+						NewObjectsToSerialize.SetNumUnsafeInternal(StartIndex);
+					}
+				}
+
+#if PERF_DETAILED_PER_CLASS_GC_STATS
+				// Detailed per class stats should not be performed when parallel GC is running
+				check(!bParallel);
+				ReferenceProcessor.UpdateDetailedStats(CurrentObject, FPlatformTime::Cycles() - StartCycles);
+#endif
+			}
+
+			if (bParallel && NewObjectsToSerialize.Num() >= MinDesiredObjectsPerSubTask)
+			{
+				const int32 ObjectsPerSubTask = FMath::Max<int32>(MinDesiredObjectsPerSubTask, NewObjectsToSerialize.Num() / FTaskGraphInterface::Get().GetNumWorkerThreads());
+				int32 StartIndex = 0;
+				while (StartIndex < NewObjectsToSerialize.Num())
+				{
+					const int32 NumThisTask = FMath::Min<int32>(ObjectsPerSubTask, NewObjectsToSerialize.Num() - StartIndex);
+					if (MyCompletionGraphEvent.GetReference())
+					{
+						MyCompletionGraphEvent->DontCompleteUntil(TGraphTask< FCollectorTask >::CreateTask().ConstructAndDispatchWhenReady(this, &NewObjectsToSerialize, StartIndex, NumThisTask, ArrayPool));
+					}
+					else
+					{
+						TaskQueue.AddTask(&NewObjectsToSerialize, StartIndex, NumThisTask);
+					}
+					StartIndex += NumThisTask;
+				}
+				NewObjectsToSerialize.SetNumUnsafeInternal(0);
+			}
+			else if (NewObjectsToSerialize.Num())
+			{
+				// Don't spawn a new task, continue in the current one
+				// To avoid allocating and moving memory around swap ObjectsToSerialize and NewObjectsToSerialize arrays
+				Exchange(ObjectsToSerialize, NewObjectsToSerialize);
+				// Empty but don't free allocated memory
+				NewObjectsToSerialize.SetNumUnsafeInternal(0);
+
+				CurrentIndex = 0;
+			}
+		}
+		while (CurrentIndex < ObjectsToSerialize.Num());
+
+#if PERF_DETAILED_PER_CLASS_GC_STATS
+		// Detailed per class stats should not be performed when parallel GC is running
+		check(!bParallel);
+		ReferenceProcessor.LogDetailedStatsSummary();
+#endif
+
+		ArrayPool.ReturnToPool(&NewObjectsToSerializeArray);
+	}
+};
diff --git a/Engine/Source/Runtime/CoreUObject/Public/UObject/Object.h b/Engine/Source/Runtime/CoreUObject/Public/UObject/Object.h
--- a/Engine/Source/Runtime/CoreUObject/Public/UObject/Object.h
+++ b/Engine/Source/Runtime/CoreUObject/Public/UObject/Object.h
@@ -37,12 +37,10 @@ class COREUOBJECT_API UObject : public UObjectBaseUtility
 	// Declarations.
 	DECLARE_CLASS(UObject,UObject,CLASS_Abstract|CLASS_NoExport,CASTCLASS_None,TEXT("/Script/CoreUObject"),NO_API)
 	DEFINE_DEFAULT_OBJECT_INITIALIZER_CONSTRUCTOR_CALL(UObject)
-#if WITH_HOT_RELOAD_CTORS
 	static UObject* __VTableCtorCaller(FVTableHelper& Helper)
 	{
 		return new (EC_InternalUseOnlyConstructor, (UObject*)GetTransientPackage(), NAME_None, RF_NeedLoad | RF_ClassDefaultObject | RF_TagGarbageTemp) UObject(Helper);
 	}
-#endif // WITH_HOT_RELOAD_CTORS
 
 	typedef UObject WithinClass;
 	static const TCHAR* StaticConfigName() {return TEXT("Engine");}
@@ -51,10 +49,8 @@ class COREUOBJECT_API UObject : public UObjectBaseUtility
 	UObject();
 	UObject(const FObjectInitializer& ObjectInitializer);
 	UObject( EStaticConstructor, EObjectFlags InFlags );
-#if WITH_HOT_RELOAD_CTORS
 	/** DO NOT USE. This constructor is for internal usage only for hot-reload purposes. */
 	UObject(FVTableHelper& Helper);
-#endif // WITH_HOT_RELOAD_CTORS
 
 	static void StaticRegisterNativesUObject() 
 	{
@@ -880,13 +876,11 @@ private:
 	 */
 	void UpdateSingleSectionOfConfigFile(const FString& ConfigIniName);
 
-#if WITH_HOT_RELOAD_CTORS
 	/**
 	 * Ensures that current thread is NOT during vtable ptr retrieval process
 	 * of some UClass.
 	 */
 	void EnsureNotRetrievingVTablePtr() const;
-#endif // WITH_HOT_RELOAD_CTORS
 
 public:
 	
diff --git a/Engine/Source/Runtime/CoreUObject/Public/UObject/ObjectMacros.h b/Engine/Source/Runtime/CoreUObject/Public/UObject/ObjectMacros.h
--- a/Engine/Source/Runtime/CoreUObject/Public/UObject/ObjectMacros.h
+++ b/Engine/Source/Runtime/CoreUObject/Public/UObject/ObjectMacros.h
@@ -148,7 +148,6 @@ enum EStaticConstructor				{EC_StaticConstructor};
 enum EInternal						{EC_InternalUseOnlyConstructor};
 enum ECppProperty					{EC_CppProperty};
 
-#if WITH_HOT_RELOAD_CTORS
 /** DO NOT USE. Helper class to invoke specialized hot-reload constructor. */
 class FVTableHelper
 {
@@ -159,7 +158,6 @@ public:
 		EnsureRetrievingVTablePtrDuringCtor(TEXT("FVTableHelper()"));
 	}
 };
-#endif // WITH_HOT_RELOAD_CTORS
 
 /** Empty API definition.  Used as a placeholder parameter when no DLL export/import API is needed for a UObject class */
 #define NO_API
@@ -1374,141 +1372,95 @@ public: \
 #define DEFINE_DEFAULT_OBJECT_INITIALIZER_CONSTRUCTOR_CALL(TClass) \
 	static void __DefaultConstructor(const FObjectInitializer& X) { new((EInternal*)X.GetObj())TClass(X); }
 
-#if WITH_HOT_RELOAD_CTORS
-	#define DECLARE_VTABLE_PTR_HELPER_CTOR(API, TClass) \
-		/** DO NOT USE. This constructor is for internal usage only for hot-reload purposes. */ \
-		API TClass(FVTableHelper& Helper);
+#define DECLARE_VTABLE_PTR_HELPER_CTOR(API, TClass) \
+	/** DO NOT USE. This constructor is for internal usage only for hot-reload purposes. */ \
+	API TClass(FVTableHelper& Helper);
 
-	#define DEFINE_VTABLE_PTR_HELPER_CTOR(TClass) \
-		TClass::TClass(FVTableHelper& Helper) : Super(Helper) {};
+#define DEFINE_VTABLE_PTR_HELPER_CTOR(TClass) \
+	TClass::TClass(FVTableHelper& Helper) : Super(Helper) {};
 
-	#define DEFINE_VTABLE_PTR_HELPER_CTOR_CALLER_DUMMY() \
-		static UObject* __VTableCtorCaller(FVTableHelper& Helper) \
-		{ \
-			return nullptr; \
-		}
+#define DEFINE_VTABLE_PTR_HELPER_CTOR_CALLER_DUMMY() \
+	static UObject* __VTableCtorCaller(FVTableHelper& Helper) \
+	{ \
+		return nullptr; \
+	}
 
-	#if WITH_HOT_RELOAD
-		#define DEFINE_VTABLE_PTR_HELPER_CTOR_CALLER(TClass) \
-			static UObject* __VTableCtorCaller(FVTableHelper& Helper) \
-			{ \
-				return new (EC_InternalUseOnlyConstructor, (UObject*)GetTransientPackage(), NAME_None, RF_NeedLoad | RF_ClassDefaultObject | RF_TagGarbageTemp) TClass(Helper); \
-			}
-	#else // WITH_HOT_RELOAD
-		#define DEFINE_VTABLE_PTR_HELPER_CTOR_CALLER(TClass) \
-			DEFINE_VTABLE_PTR_HELPER_CTOR_CALLER_DUMMY()
-	#endif // WITH_HOT_RELOAD
-
-	#define DECLARE_CLASS_INTRINSIC_NO_CTOR(TClass,TSuperClass,TStaticFlags,TPackage) \
-		DECLARE_CLASS(TClass, TSuperClass, TStaticFlags | CLASS_Intrinsic, CASTCLASS_None, TPackage, NO_API) \
-		enum { IsIntrinsic = 1 }; \
-		static void StaticRegisterNatives##TClass() {} \
-		DECLARE_SERIALIZER(TClass) \
-		DEFINE_DEFAULT_OBJECT_INITIALIZER_CONSTRUCTOR_CALL(TClass) \
+#if WITH_HOT_RELOAD
+	#define DEFINE_VTABLE_PTR_HELPER_CTOR_CALLER(TClass) \
 		static UObject* __VTableCtorCaller(FVTableHelper& Helper) \
 		{ \
 			return new (EC_InternalUseOnlyConstructor, (UObject*)GetTransientPackage(), NAME_None, RF_NeedLoad | RF_ClassDefaultObject | RF_TagGarbageTemp) TClass(Helper); \
 		}
+#else // WITH_HOT_RELOAD
+	#define DEFINE_VTABLE_PTR_HELPER_CTOR_CALLER(TClass) \
+		DEFINE_VTABLE_PTR_HELPER_CTOR_CALLER_DUMMY()
+#endif // WITH_HOT_RELOAD
+
+#define DECLARE_CLASS_INTRINSIC_NO_CTOR(TClass,TSuperClass,TStaticFlags,TPackage) \
+	DECLARE_CLASS(TClass, TSuperClass, TStaticFlags | CLASS_Intrinsic, CASTCLASS_None, TPackage, NO_API) \
+	enum { IsIntrinsic = 1 }; \
+	static void StaticRegisterNatives##TClass() {} \
+	DECLARE_SERIALIZER(TClass) \
+	DEFINE_DEFAULT_OBJECT_INITIALIZER_CONSTRUCTOR_CALL(TClass) \
+	static UObject* __VTableCtorCaller(FVTableHelper& Helper) \
+	{ \
+		return new (EC_InternalUseOnlyConstructor, (UObject*)GetTransientPackage(), NAME_None, RF_NeedLoad | RF_ClassDefaultObject | RF_TagGarbageTemp) TClass(Helper); \
+	}
 
-	#define DECLARE_CLASS_INTRINSIC(TClass,TSuperClass,TStaticFlags,TPackage) \
-		DECLARE_CLASS(TClass,TSuperClass,TStaticFlags|CLASS_Intrinsic,CASTCLASS_None,TPackage,NO_API ) \
-		RELAY_CONSTRUCTOR(TClass, TSuperClass) \
-		/** DO NOT USE. This constructor is for internal usage only for hot-reload purposes. */ \
-		TClass(FVTableHelper& Helper) : Super(Helper) {}; \
-		enum {IsIntrinsic=1}; \
-		static void StaticRegisterNatives##TClass() {} \
-		DECLARE_SERIALIZER(TClass) \
-		DEFINE_DEFAULT_OBJECT_INITIALIZER_CONSTRUCTOR_CALL(TClass) \
-		static UObject* __VTableCtorCaller(FVTableHelper& Helper) \
-		{ \
-			return new (EC_InternalUseOnlyConstructor, (UObject*)GetTransientPackage(), NAME_None, RF_NeedLoad | RF_ClassDefaultObject | RF_TagGarbageTemp) TClass(Helper); \
-		}
+#define DECLARE_CLASS_INTRINSIC(TClass,TSuperClass,TStaticFlags,TPackage) \
+	DECLARE_CLASS(TClass,TSuperClass,TStaticFlags|CLASS_Intrinsic,CASTCLASS_None,TPackage,NO_API ) \
+	RELAY_CONSTRUCTOR(TClass, TSuperClass) \
+	/** DO NOT USE. This constructor is for internal usage only for hot-reload purposes. */ \
+	TClass(FVTableHelper& Helper) : Super(Helper) {}; \
+	enum {IsIntrinsic=1}; \
+	static void StaticRegisterNatives##TClass() {} \
+	DECLARE_SERIALIZER(TClass) \
+	DEFINE_DEFAULT_OBJECT_INITIALIZER_CONSTRUCTOR_CALL(TClass) \
+	static UObject* __VTableCtorCaller(FVTableHelper& Helper) \
+	{ \
+		return new (EC_InternalUseOnlyConstructor, (UObject*)GetTransientPackage(), NAME_None, RF_NeedLoad | RF_ClassDefaultObject | RF_TagGarbageTemp) TClass(Helper); \
+	}
 
-	#define DECLARE_CASTED_CLASS_INTRINSIC_WITH_API_NO_CTOR( TClass, TSuperClass, TStaticFlags, TPackage, TStaticCastFlags, TRequiredAPI ) \
-		DECLARE_CLASS(TClass, TSuperClass, TStaticFlags | CLASS_Intrinsic, TStaticCastFlags, TPackage, TRequiredAPI) \
-		enum { IsIntrinsic = 1 }; \
-		static void StaticRegisterNatives##TClass() {} \
-		DECLARE_SERIALIZER(TClass) \
-		DEFINE_DEFAULT_OBJECT_INITIALIZER_CONSTRUCTOR_CALL(TClass) \
-		static UObject* __VTableCtorCaller(FVTableHelper& Helper) \
-		{ \
-			return new (EC_InternalUseOnlyConstructor, (UObject*)GetTransientPackage(), NAME_None, RF_NeedLoad | RF_ClassDefaultObject | RF_TagGarbageTemp) TClass(Helper); \
-		}
+#define DECLARE_CASTED_CLASS_INTRINSIC_WITH_API_NO_CTOR( TClass, TSuperClass, TStaticFlags, TPackage, TStaticCastFlags, TRequiredAPI ) \
+	DECLARE_CLASS(TClass, TSuperClass, TStaticFlags | CLASS_Intrinsic, TStaticCastFlags, TPackage, TRequiredAPI) \
+	enum { IsIntrinsic = 1 }; \
+	static void StaticRegisterNatives##TClass() {} \
+	DECLARE_SERIALIZER(TClass) \
+	DEFINE_DEFAULT_OBJECT_INITIALIZER_CONSTRUCTOR_CALL(TClass) \
+	static UObject* __VTableCtorCaller(FVTableHelper& Helper) \
+	{ \
+		return new (EC_InternalUseOnlyConstructor, (UObject*)GetTransientPackage(), NAME_None, RF_NeedLoad | RF_ClassDefaultObject | RF_TagGarbageTemp) TClass(Helper); \
+	}
 
-	#define DECLARE_CASTED_CLASS_INTRINSIC_WITH_API( TClass, TSuperClass, TStaticFlags, TPackage, TStaticCastFlags, TRequiredAPI ) \
-		DECLARE_CLASS(TClass,TSuperClass,TStaticFlags|CLASS_Intrinsic,TStaticCastFlags,TPackage,TRequiredAPI ) \
-		RELAY_CONSTRUCTOR(TClass, TSuperClass) \
-		/** DO NOT USE. This constructor is for internal usage only for hot-reload purposes. */ \
-		TClass(FVTableHelper& Helper) : Super(Helper) {}; \
-		enum {IsIntrinsic=1}; \
-		static void StaticRegisterNatives##TClass() {} \
-		DECLARE_SERIALIZER(TClass) \
-		DEFINE_DEFAULT_OBJECT_INITIALIZER_CONSTRUCTOR_CALL(TClass) \
-		static UObject* __VTableCtorCaller(FVTableHelper& Helper) \
-		{ \
-			return new (EC_InternalUseOnlyConstructor, (UObject*)GetTransientPackage(), NAME_None, RF_NeedLoad | RF_ClassDefaultObject | RF_TagGarbageTemp) TClass(Helper); \
-		}
+#define DECLARE_CASTED_CLASS_INTRINSIC_WITH_API( TClass, TSuperClass, TStaticFlags, TPackage, TStaticCastFlags, TRequiredAPI ) \
+	DECLARE_CLASS(TClass,TSuperClass,TStaticFlags|CLASS_Intrinsic,TStaticCastFlags,TPackage,TRequiredAPI ) \
+	RELAY_CONSTRUCTOR(TClass, TSuperClass) \
+	/** DO NOT USE. This constructor is for internal usage only for hot-reload purposes. */ \
+	TClass(FVTableHelper& Helper) : Super(Helper) {}; \
+	enum {IsIntrinsic=1}; \
+	static void StaticRegisterNatives##TClass() {} \
+	DECLARE_SERIALIZER(TClass) \
+	DEFINE_DEFAULT_OBJECT_INITIALIZER_CONSTRUCTOR_CALL(TClass) \
+	static UObject* __VTableCtorCaller(FVTableHelper& Helper) \
+	{ \
+		return new (EC_InternalUseOnlyConstructor, (UObject*)GetTransientPackage(), NAME_None, RF_NeedLoad | RF_ClassDefaultObject | RF_TagGarbageTemp) TClass(Helper); \
+	}
 
-	#define DECLARE_CASTED_CLASS_INTRINSIC_NO_CTOR_NO_VTABLE_CTOR( TClass, TSuperClass, TStaticFlags, TPackage, TStaticCastFlags, TRequiredAPI ) \
-		DECLARE_CLASS(TClass,TSuperClass,TStaticFlags|CLASS_Intrinsic,TStaticCastFlags,TPackage, TRequiredAPI ) \
-		enum {IsIntrinsic=1}; \
-		static void StaticRegisterNatives##TClass() {} \
-		DECLARE_SERIALIZER(TClass) \
-		DEFINE_DEFAULT_OBJECT_INITIALIZER_CONSTRUCTOR_CALL(TClass) \
-		static UObject* __VTableCtorCaller(FVTableHelper& Helper) \
-		{ \
-			return new (EC_InternalUseOnlyConstructor, (UObject*)GetTransientPackage(), NAME_None, RF_NeedLoad | RF_ClassDefaultObject | RF_TagGarbageTemp) TClass(Helper); \
-		}
+#define DECLARE_CASTED_CLASS_INTRINSIC_NO_CTOR_NO_VTABLE_CTOR( TClass, TSuperClass, TStaticFlags, TPackage, TStaticCastFlags, TRequiredAPI ) \
+	DECLARE_CLASS(TClass,TSuperClass,TStaticFlags|CLASS_Intrinsic,TStaticCastFlags,TPackage, TRequiredAPI ) \
+	enum {IsIntrinsic=1}; \
+	static void StaticRegisterNatives##TClass() {} \
+	DECLARE_SERIALIZER(TClass) \
+	DEFINE_DEFAULT_OBJECT_INITIALIZER_CONSTRUCTOR_CALL(TClass) \
+	static UObject* __VTableCtorCaller(FVTableHelper& Helper) \
+	{ \
+		return new (EC_InternalUseOnlyConstructor, (UObject*)GetTransientPackage(), NAME_None, RF_NeedLoad | RF_ClassDefaultObject | RF_TagGarbageTemp) TClass(Helper); \
+	}
 
-	#define DECLARE_CASTED_CLASS_INTRINSIC_NO_CTOR( TClass, TSuperClass, TStaticFlags, TPackage, TStaticCastFlags, TRequiredAPI ) \
-		DECLARE_CASTED_CLASS_INTRINSIC_NO_CTOR_NO_VTABLE_CTOR( TClass, TSuperClass, TStaticFlags, TPackage, TStaticCastFlags, TRequiredAPI ) \
-		/** DO NOT USE. This constructor is for internal usage only for hot-reload purposes. */ \
-		TClass(FVTableHelper& Helper) : Super(Helper) {}; \
-
-#else // WITH_HOT_RELOAD_CTORS
-	#define DECLARE_VTABLE_PTR_HELPER_CTOR(API, TClass)
-	#define DEFINE_VTABLE_PTR_HELPER_CTOR(TClass)
-	#define DEFINE_VTABLE_PTR_HELPER_CTOR_CALLER(TClass)
-	#define DEFINE_VTABLE_PTR_HELPER_CTOR_CALLER_DUMMY()
-
-	#define DECLARE_CLASS_INTRINSIC_NO_CTOR(TClass,TSuperClass,TStaticFlags,TPackage) \
-		DECLARE_CLASS(TClass, TSuperClass, TStaticFlags | CLASS_Intrinsic, CASTCLASS_None, TPackage, NO_API) \
-		enum { IsIntrinsic = 1 }; \
-		static void StaticRegisterNatives##TClass() {} \
-		DECLARE_SERIALIZER(TClass) \
-		DEFINE_DEFAULT_OBJECT_INITIALIZER_CONSTRUCTOR_CALL(TClass) \
-
-	#define DECLARE_CLASS_INTRINSIC(TClass,TSuperClass,TStaticFlags,TPackage) \
-		DECLARE_CLASS(TClass,TSuperClass,TStaticFlags|CLASS_Intrinsic,CASTCLASS_None,TPackage,NO_API ) \
-		RELAY_CONSTRUCTOR(TClass, TSuperClass) \
-		enum {IsIntrinsic=1}; \
-		static void StaticRegisterNatives##TClass() {} \
-		DECLARE_SERIALIZER(TClass) \
-		DEFINE_DEFAULT_OBJECT_INITIALIZER_CONSTRUCTOR_CALL(TClass) \
-
-	#define DECLARE_CASTED_CLASS_INTRINSIC_WITH_API_NO_CTOR( TClass, TSuperClass, TStaticFlags, TPackage, TStaticCastFlags, TRequiredAPI ) \
-		DECLARE_CLASS(TClass, TSuperClass, TStaticFlags | CLASS_Intrinsic, TStaticCastFlags, TPackage, TRequiredAPI) \
-		enum { IsIntrinsic = 1 }; \
-		static void StaticRegisterNatives##TClass() {} \
-		DECLARE_SERIALIZER(TClass) \
-		DEFINE_DEFAULT_OBJECT_INITIALIZER_CONSTRUCTOR_CALL(TClass) \
-
-	#define DECLARE_CASTED_CLASS_INTRINSIC_WITH_API( TClass, TSuperClass, TStaticFlags, TPackage, TStaticCastFlags, TRequiredAPI ) \
-		DECLARE_CLASS(TClass,TSuperClass,TStaticFlags|CLASS_Intrinsic,TStaticCastFlags,TPackage,TRequiredAPI ) \
-		RELAY_CONSTRUCTOR(TClass, TSuperClass) \
-		enum {IsIntrinsic=1}; \
-		static void StaticRegisterNatives##TClass() {} \
-		DECLARE_SERIALIZER(TClass) \
-		DEFINE_DEFAULT_OBJECT_INITIALIZER_CONSTRUCTOR_CALL(TClass) \
-
-	#define DECLARE_CASTED_CLASS_INTRINSIC_NO_CTOR( TClass, TSuperClass, TStaticFlags, TPackage, TStaticCastFlags, TRequiredAPI ) \
-		DECLARE_CLASS(TClass,TSuperClass,TStaticFlags|CLASS_Intrinsic,TStaticCastFlags,TPackage, TRequiredAPI ) \
-		enum {IsIntrinsic=1}; \
-		static void StaticRegisterNatives##TClass() {} \
-		DECLARE_SERIALIZER(TClass) \
-		DEFINE_DEFAULT_OBJECT_INITIALIZER_CONSTRUCTOR_CALL(TClass) \
-
-#endif // WITH_HOT_RELOAD_CTORS
+#define DECLARE_CASTED_CLASS_INTRINSIC_NO_CTOR( TClass, TSuperClass, TStaticFlags, TPackage, TStaticCastFlags, TRequiredAPI ) \
+	DECLARE_CASTED_CLASS_INTRINSIC_NO_CTOR_NO_VTABLE_CTOR( TClass, TSuperClass, TStaticFlags, TPackage, TStaticCastFlags, TRequiredAPI ) \
+	/** DO NOT USE. This constructor is for internal usage only for hot-reload purposes. */ \
+	TClass(FVTableHelper& Helper) : Super(Helper) {}; \
 
 
 #define DECLARE_CASTED_CLASS_INTRINSIC( TClass, TSuperClass, TStaticFlags, TPackage, TStaticCastFlags ) \
diff --git a/Engine/Source/Runtime/CoreUObject/Public/UObject/Package.h b/Engine/Source/Runtime/CoreUObject/Public/UObject/Package.h
--- a/Engine/Source/Runtime/CoreUObject/Public/UObject/Package.h
+++ b/Engine/Source/Runtime/CoreUObject/Public/UObject/Package.h
@@ -49,7 +49,6 @@ class COREUOBJECT_API UPackage : public UObject
 	// Have to unwind this macro to support the reference variable, can go back to commented declaration when removing the deprecated variable
 	// DECLARE_CASTED_CLASS_INTRINSIC(UPackage, UObject, 0, TEXT("/Script/CoreUObject"), CASTCLASS_UPackage)
 
-#if WITH_HOT_RELOAD_CTORS
 	DECLARE_CASTED_CLASS_INTRINSIC_NO_CTOR_NO_VTABLE_CTOR( UPackage, UObject, 0, TEXT("/Script/CoreUObject"), CASTCLASS_UPackage, NO_API )
 	/** DO NOT USE. This constructor is for internal usage only for hot-reload purposes. */
 	UPackage(FVTableHelper& Helper)
@@ -57,9 +56,6 @@ class COREUOBJECT_API UPackage : public UObject
 		, PackageFlagsPrivate(PackageFlags) 
 	{
 	};
-#else
-	DECLARE_CASTED_CLASS_INTRINSIC_NO_CTOR(UPackage, UObject, 0, TEXT("/Script/CoreUObject"), CASTCLASS_UPackage, NO_API)
-#endif
 
 
 public:
diff --git a/Engine/Source/Runtime/CoreUObject/Public/UObject/UObjectArray.h b/Engine/Source/Runtime/CoreUObject/Public/UObject/UObjectArray.h
--- a/Engine/Source/Runtime/CoreUObject/Public/UObject/UObjectArray.h
+++ b/Engine/Source/Runtime/CoreUObject/Public/UObject/UObjectArray.h
@@ -28,18 +28,39 @@ struct FUObjectItem
 	// Internal flags
 	int32 Flags;
 	// UObject Owner Cluster Index
-	int32 ClusterIndex;	
+	int32 ClusterRootIndex;	
 	// Weak Object Pointer Serial number associated with the object
 	int32 SerialNumber;
 
+	FUObjectItem()
+		: Object(nullptr)
+		, Flags(0)
+		, ClusterRootIndex(0)
+		, SerialNumber(0)
+	{
+	}
+
 	FORCEINLINE void SetOwnerIndex(int32 OwnerIndex)
 	{
-		ClusterIndex = OwnerIndex;
+		ClusterRootIndex = OwnerIndex;
 	}
 
 	FORCEINLINE int32 GetOwnerIndex() const
 	{
-		return ClusterIndex;
+		return ClusterRootIndex;
+	}
+
+	/** Encodes the cluster index in the ClusterRootIndex variable */
+	FORCEINLINE void SetClusterIndex(int32 ClusterIndex)
+	{
+		ClusterRootIndex = -ClusterIndex - 1;
+	}
+
+	/** Decodes the cluster index from the ClusterRootIndex variable */
+	FORCEINLINE int32 GetClusterIndex() const
+	{
+		checkSlow(ClusterRootIndex < 0);
+		return -ClusterRootIndex - 1;
 	}
 
 	FORCEINLINE int32 GetSerialNumber() const
@@ -176,7 +197,7 @@ struct FUObjectItem
 	FORCEINLINE void ResetSerialNumberAndFlags()
 	{
 		Flags = 0;
-		ClusterIndex = 0;
+		ClusterRootIndex = 0;
 		SerialNumber = 0;
 	}
 };
@@ -204,7 +225,7 @@ public:
 
 	~FFixedUObjectArray()
 	{
-		FMemory::Free(Objects);
+		delete [] Objects;
 	}
 
 	/**
@@ -214,8 +235,7 @@ public:
 	void PreAllocate(int32 InMaxElements)
 	{
 		check(!Objects);
-		Objects = (FUObjectItem*)FMemory::Malloc(sizeof(FUObjectItem)* InMaxElements);
-		FMemory::Memzero(Objects, sizeof(FUObjectItem)* InMaxElements);
+		Objects = new FUObjectItem[InMaxElements];
 		MaxElements = InMaxElements;
 	}
 
@@ -247,6 +267,12 @@ public:
 		return &Objects[Index];
 	}
 
+	FORCEINLINE FUObjectItem* GetObjectPtr(int32 Index)
+	{
+		check(Index >= 0 && Index < NumElements);
+		return &Objects[Index];
+	}
+
 	/**
 	* Return the number of elements in the array
 	* Thread safe, but you know, someone might have added more elements before this even returns
@@ -279,6 +305,13 @@ public:
 		return *ItemPtr;
 	}
 
+	FORCEINLINE FUObjectItem& operator[](int32 Index)
+	{
+		FUObjectItem* ItemPtr = GetObjectPtr(Index);
+		check(ItemPtr);
+		return *ItemPtr;
+	}
+
 	/**
 	* Return a naked pointer to the fundamental data structure for debug visualizers.
 	**/
@@ -406,7 +439,7 @@ public:
 	 * @param Object object to get the index of
 	 * @return index of this object
 	 */
-	FORCEINLINE int32 ObjectToIndex(const class UObjectBase* Object)
+	FORCEINLINE int32 ObjectToIndex(const class UObjectBase* Object) const
 	{
 		return Object->InternalIndex;
 	}
@@ -762,17 +795,89 @@ private:
 /** UObject cluster. Groups UObjects into a single unit for GC. */
 struct FUObjectCluster
 {
+	FUObjectCluster()
+		: RootIndex(INDEX_NONE)
+		, bNeedsDissolving(false)
+	{}
+
+	/** Root object index */
+	int32 RootIndex;
 	/** Objects that belong to this cluster */
 	TArray<int32> Objects;
 	/** Other clusters referenced by this cluster */
 	TArray<int32> ReferencedClusters;
 	/** Objects that could not be added to the cluster but still need to be referenced by it */
 	TArray<int32> MutableObjects;
+	/** List of clusters that direcly reference this cluster. Used when dissolving a cluster. */
+	TArray<int32> ReferencedByClusters;
+
+	/** Cluster needs dissolving, probably due to PendingKill reference */
+	bool bNeedsDissolving;
+};
+
+class COREUOBJECT_API FUObjectClusterContainer
+{
+	/** List of all clusters */
+	TArray<FUObjectCluster> Clusters;
+	/** List of available cluster indices */
+	TArray<int32> FreeClusterIndices;
+	/** Number of allocated clusters */
+	int32 NumAllocatedClusters;
+	/** Clusters need dissolving, probably due to PendingKill reference */
+	bool bClustersNeedDissolving;
+
+	/** Dissolves a cluster */
+	void DissolveCluster(FUObjectCluster& Cluster);
+
+public:
+
+	FUObjectClusterContainer();
+
+	FORCEINLINE FUObjectCluster& operator[](int32 Index)
+	{
+		return Clusters[Index];
+	}
+
+	/** Returns an index to a new cluster */
+	int32 AllocateCluster(int32 InRootObjectIndex);
+
+	/** Frees the cluster at the specified index */
+	void FreeCluster(int32 InClusterIndex);
+
+	/** Dissolve all clusters marked for dissolving */
+	void DissolveClusters();
+
+	/** Dissolve the specified cluster and all clusters that reference it */
+	void DissolveClusterAndMarkObjectsAsUnreachable(const int32 CurrentIndex, FUObjectItem* RootObjectItem);
+
+	/** Gets the clusters array (for internal use only!) */
+	TArray<FUObjectCluster>& GetClustersUnsafe() 
+	{ 
+		return Clusters;  
+	}
+
+	/** Returns the number of currently allocated clusters */
+	int32 GetNumAllocatedClusters() const
+	{
+		return NumAllocatedClusters;
+	}
+
+	/** Lets the FUObjectClusterContainer know some clusters need dissolving */
+	void SetClustersNeedDissolving()
+	{
+		bClustersNeedDissolving = true;
+	}
+	
+	/** Checks if any clusters need dissolving */
+	bool ClustersNeedDissolving() const
+	{
+		return bClustersNeedDissolving;
+	}
 };
 
 /** Global UObject allocator							*/
 extern COREUOBJECT_API FUObjectArray GUObjectArray;
-extern COREUOBJECT_API TMap<int32, FUObjectCluster* > GUObjectClusters;
+extern COREUOBJECT_API FUObjectClusterContainer GUObjectClusters;
 
 /**
 	* Static version of IndexToObject for use with TWeakObjectPtr.
diff --git a/Engine/Source/Runtime/CoreUObject/Public/UObject/UObjectBaseUtility.h b/Engine/Source/Runtime/CoreUObject/Public/UObject/UObjectBaseUtility.h
--- a/Engine/Source/Runtime/CoreUObject/Public/UObject/UObjectBaseUtility.h
+++ b/Engine/Source/Runtime/CoreUObject/Public/UObject/UObjectBaseUtility.h
@@ -310,15 +310,21 @@ public:
 	virtual bool CanBeInCluster() const;
 
 	/**
-	* Called during PostLoad to create UObject cluster
+	* Called after PostLoad to create UObject cluster
 	*/
 	virtual void CreateCluster();
 
+	/**
+	* Called during Garbage Collection to perform additional cleanup when the cluster is about to be destroyed due to PendingKill flag being set on it.
+	*/
+	virtual void OnClusterMarkedAsPendingKill() {}
+
 	/** 
 	 * Adds this objects to a GC cluster that already exists 
 	 * @param ClusterRootOrObjectFromCluster Object that belongs to the cluster we want to add this object to.
+	 * @param Add this object to the target cluster as a mutable object without adding this object's references.
 	 */
-	virtual void AddToCluster(UObjectBaseUtility* ClusterRootOrObjectFromCluster);
+	virtual void AddToCluster(UObjectBaseUtility* ClusterRootOrObjectFromCluster, bool bAddAsMutableObject = false);
 
 protected:
 
diff --git a/Engine/Source/Runtime/CoreUObject/Public/UObject/UObjectGlobals.h b/Engine/Source/Runtime/CoreUObject/Public/UObject/UObjectGlobals.h
--- a/Engine/Source/Runtime/CoreUObject/Public/UObject/UObjectGlobals.h
+++ b/Engine/Source/Runtime/CoreUObject/Public/UObject/UObjectGlobals.h
@@ -1630,7 +1630,7 @@ public:
 	 * @param ReferencingProperty Referencing property (if available).
 	 */
 	template<class UObjectType>
-	void AddReferencedObject(UObjectType*& Object, const UObject* ReferencingObject = NULL, const UProperty* ReferencingProperty = NULL)
+	void AddReferencedObject(UObjectType*& Object, const UObject* ReferencingObject = nullptr, const UProperty* ReferencingProperty = nullptr)
 	{
 		HandleObjectReference(*(UObject**)&Object, ReferencingObject, ReferencingProperty);
 	}
@@ -1643,7 +1643,7 @@ public:
 	* @param ReferencingProperty Referencing property (if available).
 	*/
 	template<class UObjectType>
-	void AddReferencedObjects(TArray<UObjectType*>& ObjectArray, const UObject* ReferencingObject = NULL, const UProperty* ReferencingProperty = NULL)
+	void AddReferencedObjects(TArray<UObjectType*>& ObjectArray, const UObject* ReferencingObject = nullptr, const UProperty* ReferencingProperty = nullptr)
 	{
 		static_assert(TPointerIsConvertibleFromTo<UObjectType, const UObject>::Value, "'UObjectType' template parameter to AddReferencedObjects must be derived from UObject");
 		HandleObjectReferences(reinterpret_cast<UObject**>(ObjectArray.GetData()), ObjectArray.Num(), ReferencingObject, ReferencingProperty);
@@ -1657,7 +1657,7 @@ public:
 	* @param ReferencingProperty Referencing property (if available).
 	*/
 	template<class UObjectType>
-	void AddReferencedObjects(TSet<UObjectType>& ObjectSet, const UObject* ReferencingObject = NULL, const UProperty* ReferencingProperty = NULL)
+	void AddReferencedObjects(TSet<UObjectType>& ObjectSet, const UObject* ReferencingObject = nullptr, const UProperty* ReferencingProperty = nullptr)
 	{
 		for (auto& Object : ObjectSet)
 		{
@@ -1942,6 +1942,12 @@ struct COREUOBJECT_API FCoreUObjectDelegates
 	// Called after garbage collection
 	static FSimpleMulticastDelegate PostGarbageCollect;
 
+	// Called before ConditionalBeginDestroy phase of garbage collection
+	static FSimpleMulticastDelegate PreGarbageCollectConditionalBeginDestroy;
+
+	// Called after ConditionalBeginDestroy phase of garbage collection
+	static FSimpleMulticastDelegate PostGarbageCollectConditionalBeginDestroy;
+
 	/** delegate type for querying whether a loaded object should replace an already existing one */
 	DECLARE_DELEGATE_RetVal_OneParam(bool, FOnLoadObjectsOnTop, const FString&);
 
diff --git a/Engine/Source/Runtime/CoreUObject/Public/UObject/UnrealType.h b/Engine/Source/Runtime/CoreUObject/Public/UObject/UnrealType.h
--- a/Engine/Source/Runtime/CoreUObject/Public/UObject/UnrealType.h
+++ b/Engine/Source/Runtime/CoreUObject/Public/UObject/UnrealType.h
@@ -985,10 +985,8 @@ public:
 		SetElementSize();
 	}
 
-#if WITH_HOT_RELOAD_CTORS
 	/** DO NOT USE. This constructor is for internal usage only for hot-reload purposes. */
 	TProperty(FVTableHelper& Helper) : Super(Helper) {};
-#endif // WITH_HOT_RELOAD_CTORS
 
 	// UHT interface
 	virtual FString GetCPPType( FString* ExtendedTypeText=NULL, uint32 CPPExportFlags=0 ) const override
@@ -1095,10 +1093,8 @@ public:
 	{
 	}
 
-#if WITH_HOT_RELOAD_CTORS
 	/** DO NOT USE. This constructor is for internal usage only for hot-reload purposes. */
 	TProperty_WithEqualityAndSerializer(FVTableHelper& Helper) : Super(Helper) {};
-#endif // WITH_HOT_RELOAD_CTORS
 
 	// UProperty interface.
 	virtual bool Identical( const void* A, const void* B, uint32 PortFlags=0 ) const override
@@ -1255,10 +1251,8 @@ public:
 	{
 	}
 
-#if WITH_HOT_RELOAD_CTORS
 	/** DO NOT USE. This constructor is for internal usage only for hot-reload purposes. */
 	TProperty_Numeric(FVTableHelper& Helper) : Super(Helper) {};
-#endif // WITH_HOT_RELOAD_CTORS
 
 	virtual FString GetCPPTypeForwardDeclaration() const override
 	{
@@ -2000,10 +1994,8 @@ public:
 		this->PropertyClass = InClass;
 	}
 
-#if WITH_HOT_RELOAD_CTORS
 	/** DO NOT USE. This constructor is for internal usage only for hot-reload purposes. */
 	TUObjectPropertyBase(FVTableHelper& Helper) : Super(Helper) {};
-#endif // WITH_HOT_RELOAD_CTORS
 
 	// UProperty interface.
 	virtual bool ContainsObjectReference(TArray<const UStructProperty*>& EncounteredStructProps) const override
diff --git a/Engine/Source/Runtime/Engine/Classes/Engine/Blueprint.h b/Engine/Source/Runtime/Engine/Classes/Engine/Blueprint.h
--- a/Engine/Source/Runtime/Engine/Classes/Engine/Blueprint.h
+++ b/Engine/Source/Runtime/Engine/Classes/Engine/Blueprint.h
@@ -11,6 +11,7 @@
 #include "Engine/EngineTypes.h"
 #include "EdGraph/EdGraphPin.h"
 #include "Engine/BlueprintCore.h"
+#include "Misc/StringAssetReference.h"
 #include "Blueprint.generated.h"
 
 class FCompilerResultsLog;
@@ -89,12 +90,26 @@ struct FCompilerNativizationOptions
 {
 	GENERATED_USTRUCT_BODY()
 
+	UPROPERTY()
+	FString PlatformName;
+
 	UPROPERTY()
 	bool ServerOnlyPlatform;
 
 	UPROPERTY()
 	bool ClientOnlyPlatform;
 
+	UPROPERTY()
+	TArray<FName> ExcludedModules;
+
+	// Individually excluded assets
+	UPROPERTY()
+	TSet<FStringAssetReference> ExcludedAssets;
+
+	// Excluded folders. It excludes only BPGCs, enums and structures are still converted.
+	UPROPERTY()
+	TArray<FString> ExcludedFolderPaths;
+
 	FCompilerNativizationOptions()
 		: ServerOnlyPlatform(false)
 		, ClientOnlyPlatform(false)
diff --git a/Engine/Source/Runtime/Engine/Classes/Engine/BlueprintGeneratedClass.h b/Engine/Source/Runtime/Engine/Classes/Engine/BlueprintGeneratedClass.h
--- a/Engine/Source/Runtime/Engine/Classes/Engine/BlueprintGeneratedClass.h
+++ b/Engine/Source/Runtime/Engine/Classes/Engine/BlueprintGeneratedClass.h
@@ -698,6 +698,7 @@ public:
 	virtual FPrimaryAssetId GetPrimaryAssetId() const override;
 	virtual bool NeedsLoadForServer() const override;
 	virtual bool NeedsLoadForClient() const override;
+	virtual bool CanBeClusterRoot() const override;
 	// End UObject interface
 	
 	// UClass interface
diff --git a/Engine/Source/Runtime/Engine/Classes/Engine/CoreSettings.h b/Engine/Source/Runtime/Engine/Classes/Engine/CoreSettings.h
--- a/Engine/Source/Runtime/Engine/Classes/Engine/CoreSettings.h
+++ b/Engine/Source/Runtime/Engine/Classes/Engine/CoreSettings.h
@@ -167,7 +167,17 @@ protected:
 	UPROPERTY(EditAnywhere, config, Category = Optimization, meta = (
 		ConsoleVariable = "gc.MergeGCClusters", DisplayName = "Merge GC Clusters",
 		ToolTip = "If true, when creating clusters, the clusters referenced from another cluster will get merged into one big cluster."))
-		uint32 MergeGCClusters : 1;
+	uint32 MergeGCClusters : 1;
+
+	UPROPERTY(EditAnywhere, config, Category = Optimization, meta = (
+		ConsoleVariable = "gc.ActorClusteringEnabled", DisplayName = "Actor Clustering Enabled",
+		ToolTip = "Whether to allow levels to create actor clusters for GC."))
+	uint32 ActorClusteringEnabled : 1;
+
+	UPROPERTY(EditAnywhere, config, Category = Optimization, meta = (
+		ConsoleVariable = "gc.BlueprintClusteringEnabled", DisplayName = "Blueprint Clustering Enabled",
+		ToolTip = "Whether to allow Blueprint classes to create GC clusters."))
+	uint32 BlueprintClusteringEnabled : 1;
 
 	UPROPERTY(EditAnywhere, config, Category = General, meta = (
 		ConsoleVariable = "gc.NumRetriesBeforeForcingGC", DisplayName = "Number Of Retries Before Forcing GC",
diff --git a/Engine/Source/Runtime/Engine/Classes/Engine/GameInstance.h b/Engine/Source/Runtime/Engine/Classes/Engine/GameInstance.h
--- a/Engine/Source/Runtime/Engine/Classes/Engine/GameInstance.h
+++ b/Engine/Source/Runtime/Engine/Classes/Engine/GameInstance.h
@@ -275,8 +275,6 @@ public:
 	/** This gets called when the player scrubs in a replay to a different level */
 	virtual void OnSeamlessTravelDuringReplay() { }
 
-	static void AddReferencedObjects(UObject* InThis, FReferenceCollector& Collector);
-
 	inline FTimerManager& GetTimerManager() const { return *TimerManager; }
 
 	inline FLatentActionManager& GetLatentActionManager() const { return *LatentActionManager;  }
@@ -355,4 +353,4 @@ public:
 	void NotifyPreClientTravel(const FString& PendingURL, ETravelType TravelType, bool bIsSeamlessTravel);
 	/** @return delegate fired when client travel occurs */
 	FOnPreClientTravel& OnNotifyPreClientTravel() { return NotifyPreClientTravelDelegates; }
-};
+};
\ No newline at end of file
diff --git a/Engine/Source/Runtime/Engine/Classes/Engine/Level.h b/Engine/Source/Runtime/Engine/Classes/Engine/Level.h
--- a/Engine/Source/Runtime/Engine/Classes/Engine/Level.h
+++ b/Engine/Source/Runtime/Engine/Classes/Engine/Level.h
@@ -29,6 +29,7 @@ class UMapBuildDataRegistry;
 class UNavigationDataChunk;
 class UTexture2D;
 struct FLevelCollection;
+class ULevelActorContainer;
 
 /**
  * Structure containing all information needed for determining the screen space
@@ -365,6 +366,9 @@ public:
 	/** Array of all actors in this level, used by FActorIteratorBase and derived classes */
 	TArray<AActor*> Actors;
 
+	/** Array of actors to be exposed to GC in this level. All other actors will be referenced through ULevelActorContainer */
+	TArray<AActor*> ActorsForGC;
+
 	/** Set before calling LoadPackage for a streaming level to ensure that OwningWorld is correct on the Level */
 	ENGINE_API static TMap<FName, TWeakObjectPtr<UWorld> > StreamedLevelsOwningWorld;
 		
@@ -384,6 +388,9 @@ public:
 	UPROPERTY()
 	TArray<class UModelComponent*> ModelComponents;
 
+	UPROPERTY(Transient, DuplicateTransient, NonTransactional)
+	ULevelActorContainer* ActorCluster;
+
 #if WITH_EDITORONLY_DATA
 	/** Reference to the blueprint for level scripting */
 	UPROPERTY(NonTransactional)
@@ -600,13 +607,11 @@ public:
 	ENGINE_API void Initialize(const FURL& InURL);
 	ULevel(const FObjectInitializer& ObjectInitializer = FObjectInitializer::Get());
 
-#if WITH_HOT_RELOAD_CTORS
 	/** DO NOT USE. This constructor is for internal usage only for hot-reload purposes. */
 	ULevel(FVTableHelper& Helper)
 		: Super(Helper)
 		, Actors()
 	{}
-#endif // WITH_HOT_RELOAD_CTORS
 
 	~ULevel();
 
@@ -626,6 +631,8 @@ public:
 	virtual void PostLoad() override;
 	virtual void PreSave(const class ITargetPlatform* TargetPlatform) override;
 	virtual void PostDuplicate(bool bDuplicateForPIE) override;
+	virtual bool CanBeClusterRoot() const override;
+	virtual void CreateCluster() override;
 	static void AddReferencedObjects(UObject* InThis, FReferenceCollector& Collector);
 	//~ End UObject Interface.
 
diff --git a/Engine/Source/Runtime/Engine/Classes/Engine/StaticMesh.h b/Engine/Source/Runtime/Engine/Classes/Engine/StaticMesh.h
--- a/Engine/Source/Runtime/Engine/Classes/Engine/StaticMesh.h
+++ b/Engine/Source/Runtime/Engine/Classes/Engine/StaticMesh.h
@@ -623,6 +623,7 @@ public:
 	ENGINE_API virtual void GetAssetRegistryTags(TArray<FAssetRegistryTag>& OutTags) const override;
 	ENGINE_API virtual FString GetDesc() override;
 	ENGINE_API virtual void GetResourceSizeEx(FResourceSizeEx& CumulativeResourceSize) override;
+	ENGINE_API virtual bool CanBeClusterRoot() const override;
 	//~ End UObject Interface.
 
 	/**
diff --git a/Engine/Source/Runtime/Engine/Classes/GameFramework/Actor.h b/Engine/Source/Runtime/Engine/Classes/GameFramework/Actor.h
--- a/Engine/Source/Runtime/Engine/Classes/GameFramework/Actor.h
+++ b/Engine/Source/Runtime/Engine/Classes/GameFramework/Actor.h
@@ -426,6 +426,14 @@ public:
 	UPROPERTY(EditAnywhere, BlueprintReadWrite, Category = Actor)
 	uint8 bGenerateOverlapEventsDuringLevelStreaming:1;
 
+protected:
+
+	/** If true, this actor can be put inside of a GC Cluster to improve Garbage Collection performance */
+	UPROPERTY(Category = Actor, EditAnywhere, AdvancedDisplay)
+	uint8 bCanBeInCluster:1;
+
+public:
+
 	/** Controls how to handle spawning this actor in a situation where it's colliding with something else. "Default" means AlwaysSpawn here. */
 	UPROPERTY(EditAnywhere, BlueprintReadWrite, Category=Actor)
 	ESpawnActorCollisionHandlingMethod SpawnCollisionHandlingMethod;
@@ -1527,6 +1535,7 @@ public:
 	virtual void BeginDestroy() override;
 	virtual bool IsReadyForFinishDestroy() override;
 	virtual bool Rename( const TCHAR* NewName=NULL, UObject* NewOuter=NULL, ERenameFlags Flags=REN_None ) override;
+	virtual bool CanBeInCluster() const override;
 	static void AddReferencedObjects(UObject* InThis, FReferenceCollector& Collector);
 #if WITH_EDITOR
 	virtual void PreEditChange(UProperty* PropertyThatWillChange) override;
diff --git a/Engine/Source/Runtime/Engine/Private/AI/Navigation/NavigationSystem.cpp b/Engine/Source/Runtime/Engine/Private/AI/Navigation/NavigationSystem.cpp
--- a/Engine/Source/Runtime/Engine/Private/AI/Navigation/NavigationSystem.cpp
+++ b/Engine/Source/Runtime/Engine/Private/AI/Navigation/NavigationSystem.cpp
@@ -881,10 +881,7 @@ void UNavigationSystem::AddReferencedObjects(UObject* InThis, FReferenceCollecto
 	// don't reference NavAreaClasses in editor (unless PIE is active)
 	if (This->OperationMode != FNavigationSystemRunMode::EditorMode)
 	{
-		for (const UClass* NavAreaClass : This->NavAreaClasses)
-		{
-			Collector.AddReferencedObject(NavAreaClass, InThis);
-		}
+		Collector.AddReferencedObjects(This->NavAreaClasses, InThis);
 	}
 }
 
diff --git a/Engine/Source/Runtime/Engine/Private/Actor.cpp b/Engine/Source/Runtime/Engine/Private/Actor.cpp
--- a/Engine/Source/Runtime/Engine/Private/Actor.cpp
+++ b/Engine/Source/Runtime/Engine/Private/Actor.cpp
@@ -273,6 +273,11 @@ void AActor::PostInitProperties()
 	}
 }
 
+bool AActor::CanBeInCluster() const
+{
+	return bCanBeInCluster;
+}
+
 void AActor::AddReferencedObjects(UObject* InThis, FReferenceCollector& Collector)
 {
 	AActor* This = CastChecked<AActor>(InThis);
@@ -482,7 +487,12 @@ void AActor::RemoveControllingMatineeActor( AMatineeActor& InMatineeActor )
 
 void AActor::BeginDestroy()
 {
+	ULevel* OwnerLevel = Cast<ULevel>(GetOuter());
 	UnregisterAllComponents();
+	if (OwnerLevel)
+	{
+		OwnerLevel->Actors.Remove(this);
+	}
 	Super::BeginDestroy();
 }
 
diff --git a/Engine/Source/Runtime/Engine/Private/Blueprint.cpp b/Engine/Source/Runtime/Engine/Private/Blueprint.cpp
--- a/Engine/Source/Runtime/Engine/Private/Blueprint.cpp
+++ b/Engine/Source/Runtime/Engine/Private/Blueprint.cpp
@@ -1217,11 +1217,13 @@ void UBlueprint::BeginCacheForCookedPlatformData(const ITargetPlatform *TargetPl
 		const IBlueprintNativeCodeGenCore* NativeCodeGenCore = IBlueprintNativeCodeGenCore::Get();
 		if (GeneratedClass != nullptr && NativeCodeGenCore != nullptr && FParse::Param(FCommandLine::Get(), TEXT("NativizeAssets")))
 		{
+			ensure(TargetPlatform);
+			const FCompilerNativizationOptions& NativizationOptions = NativeCodeGenCore->GetNativizationOptionsForPlatform(TargetPlatform ? TargetPlatform->PlatformName() : FString{});
 			TArray<const UBlueprintGeneratedClass*> ParentBPClassStack;
 			UBlueprintGeneratedClass::GetGeneratedClassesHierarchy(GeneratedClass->GetSuperClass(), ParentBPClassStack);
 			for (const UBlueprintGeneratedClass *ParentBPClass : ParentBPClassStack)
 			{
-				if (NativeCodeGenCore->IsTargetedForReplacement(ParentBPClass) == EReplacementResult::ReplaceCompletely)
+				if (NativeCodeGenCore->IsTargetedForReplacement(ParentBPClass, NativizationOptions) == EReplacementResult::ReplaceCompletely)
 				{
 					if (UBlueprintGeneratedClass* BPGC = CastChecked<UBlueprintGeneratedClass>(*GeneratedClass))
 					{
@@ -1251,7 +1253,7 @@ void UBlueprint::BeginCacheForCookedPlatformData(const ITargetPlatform *TargetPl
 									}
 									else
 									{
-										bool bResult = (NativeCodeGenCore->IsTargetedForReplacement(RecordIt->ComponentKey.GetComponentOwner()) == EReplacementResult::ReplaceCompletely);
+										bool bResult = (NativeCodeGenCore->IsTargetedForReplacement(RecordIt->ComponentKey.GetComponentOwner(), NativizationOptions) == EReplacementResult::ReplaceCompletely);
 										bIsOwnerClassTargetedForReplacement = ParentBPClassNativizationResultMap.Add(ComponentTemplateOwnerClass, bResult);
 									}
 
diff --git a/Engine/Source/Runtime/Engine/Private/BlueprintGeneratedClass.cpp b/Engine/Source/Runtime/Engine/Private/BlueprintGeneratedClass.cpp
--- a/Engine/Source/Runtime/Engine/Private/BlueprintGeneratedClass.cpp
+++ b/Engine/Source/Runtime/Engine/Private/BlueprintGeneratedClass.cpp
@@ -28,6 +28,14 @@
 
 DEFINE_STAT(STAT_PersistentUberGraphFrameMemory);
 
+int32 GBlueprintClusteringEnabled = 0;
+static FAutoConsoleVariableRef CVarUseBackgroundLevelStreaming(
+	TEXT("gc.BlueprintClusteringEnabled"),
+	GBlueprintClusteringEnabled,
+	TEXT("Whether to allow Blueprint classes to create GC clusters."),
+	ECVF_Default
+);
+
 UBlueprintGeneratedClass::UBlueprintGeneratedClass(const FObjectInitializer& ObjectInitializer)
 	: Super(ObjectInitializer)
 {
@@ -1238,6 +1246,12 @@ bool UBlueprintGeneratedClass::NeedsLoadForClient() const
 	return Super::NeedsLoadForClient();
 }
 
+bool UBlueprintGeneratedClass::CanBeClusterRoot() const
+{
+	// Clustering level BPs doesn't work yet
+	return GBlueprintClusteringEnabled && !GetOutermost()->ContainsMap();
+}
+
 void UBlueprintGeneratedClass::Link(FArchive& Ar, bool bRelinkExistingProperties)
 {
 	// Ensure that function netflags equate to any super function in a parent BP prior to linking; it may have been changed by the user
@@ -1597,4 +1611,4 @@ FGuid UBlueprintGeneratedClass::FindPropertyGuidFromName(const FName InName) con
 	}
 #endif // WITH_EDITORONLY_DATA
 	return PropertyGuid;
-}
+}
\ No newline at end of file
diff --git a/Engine/Source/Runtime/Engine/Private/Components/MeshComponent.cpp b/Engine/Source/Runtime/Engine/Private/Components/MeshComponent.cpp
--- a/Engine/Source/Runtime/Engine/Private/Components/MeshComponent.cpp
+++ b/Engine/Source/Runtime/Engine/Private/Components/MeshComponent.cpp
@@ -68,7 +68,11 @@ void UMeshComponent::SetMaterial(int32 ElementIndex, UMaterialInterface* Materia
 
 			// Set the material and invalidate things
 			OverrideMaterials[ElementIndex] = Material;
-			MarkRenderStateDirty();
+			MarkRenderStateDirty();			
+			if (Material)
+			{
+				Material->AddToCluster(this);
+			}
 
 			FBodyInstance* BodyInst = GetBodyInstance();
 			if (BodyInst && BodyInst->IsValidBodyInstance())
diff --git a/Engine/Source/Runtime/Engine/Private/Components/ReflectionCaptureComponent.cpp b/Engine/Source/Runtime/Engine/Private/Components/ReflectionCaptureComponent.cpp
--- a/Engine/Source/Runtime/Engine/Private/Components/ReflectionCaptureComponent.cpp
+++ b/Engine/Source/Runtime/Engine/Private/Components/ReflectionCaptureComponent.cpp
@@ -136,6 +136,8 @@ AReflectionCapture::AReflectionCapture(const FObjectInitializer& ObjectInitializ
 {
 	CaptureComponent = CreateAbstractDefaultSubobject<UReflectionCaptureComponent>(TEXT("NewReflectionComponent"));
 
+	bCanBeInCluster = true;
+
 #if WITH_EDITORONLY_DATA
 	SpriteComponent = CreateEditorOnlyDefaultSubobject<UBillboardComponent>(TEXT("Sprite"));
 	if (!IsRunningCommandlet() && (SpriteComponent != nullptr))
diff --git a/Engine/Source/Runtime/Engine/Private/Components/SceneComponent.cpp b/Engine/Source/Runtime/Engine/Private/Components/SceneComponent.cpp
--- a/Engine/Source/Runtime/Engine/Private/Components/SceneComponent.cpp
+++ b/Engine/Source/Runtime/Engine/Private/Components/SceneComponent.cpp
@@ -1793,6 +1793,7 @@ bool USceneComponent::AttachToComponent(USceneComponent* Parent, const FAttachme
 		{
 			Parent->AttachChildren.Add(this);
 		}
+		AddToCluster(Parent, true);
 
 		if (Parent->IsNetSimulating() && !IsNetSimulating())
 		{
diff --git a/Engine/Source/Runtime/Engine/Private/Components/ShapeComponent.cpp b/Engine/Source/Runtime/Engine/Private/Components/ShapeComponent.cpp
--- a/Engine/Source/Runtime/Engine/Private/Components/ShapeComponent.cpp
+++ b/Engine/Source/Runtime/Engine/Private/Components/ShapeComponent.cpp
@@ -117,6 +117,9 @@ void UShapeComponent::CreateShapeBodySetupIfNeeded()
 			ShapeBodySetup->AddToRoot();
 		}
 
+		// If this component is in GC cluster, make sure we add the body setup to it to
+		ShapeBodySetup->AddToCluster(this);
+		
 		ShapeBodySetup->CollisionTraceFlag = CTF_UseSimpleAsComplex;
 		AddShapeToGeomArray<ShapeElemType>();
 		ShapeBodySetup->bNeverNeedsCookedCollisionData = true;
diff --git a/Engine/Source/Runtime/Engine/Private/CoreSettings.cpp b/Engine/Source/Runtime/Engine/Private/CoreSettings.cpp
--- a/Engine/Source/Runtime/Engine/Private/CoreSettings.cpp
+++ b/Engine/Source/Runtime/Engine/Private/CoreSettings.cpp
@@ -135,6 +135,8 @@ UGarbageCollectionSettings::UGarbageCollectionSettings()
 	MaxObjectsInGame = 2 * 1024 * 1024;
 	CreateGCClusters = true;
 	MergeGCClusters = false;
+	ActorClusteringEnabled = true;
+	BlueprintClusteringEnabled = false;
 }
 
 void UGarbageCollectionSettings::PostInitProperties()
diff --git a/Engine/Source/Runtime/Engine/Private/GameInstance.cpp b/Engine/Source/Runtime/Engine/Private/GameInstance.cpp
--- a/Engine/Source/Runtime/Engine/Private/GameInstance.cpp
+++ b/Engine/Source/Runtime/Engine/Private/GameInstance.cpp
@@ -795,18 +795,6 @@ const TArray<class ULocalPlayer*>& UGameInstance::GetLocalPlayers() const
 	return LocalPlayers;
 }
 
-void UGameInstance::AddReferencedObjects(UObject* InThis, FReferenceCollector& Collector)
-{
-	UGameInstance* This = CastChecked<UGameInstance>(InThis);
-
-	for (ULocalPlayer* LocalPlayer : This->LocalPlayers)
-	{
-		Collector.AddReferencedObject(LocalPlayer, This);
-	}
-
-	Super::AddReferencedObjects(This, Collector);
-}
-
 void UGameInstance::StartRecordingReplay(const FString& Name, const FString& FriendlyName, const TArray<FString>& AdditionalOptions)
 {
 	if ( FParse::Param( FCommandLine::Get(),TEXT( "NOREPLAYS" ) ) )
diff --git a/Engine/Source/Runtime/Engine/Private/InstancedReferenceSubobjectHelper.cpp b/Engine/Source/Runtime/Engine/Private/InstancedReferenceSubobjectHelper.cpp
new file mode 100644
--- /dev/null
+++ b/Engine/Source/Runtime/Engine/Private/InstancedReferenceSubobjectHelper.cpp
@@ -0,0 +1,252 @@
+// Copyright 1998-2017 Epic Games, Inc. All Rights Reserved.
+
+#include "InstancedReferenceSubobjectHelper.h"
+
+void FFindInstancedReferenceSubobjectHelper::GetInstancedSubObjects_Inner(FInstancedPropertyPath& PropertyPath, const uint8* ContainerAddress, TFunctionRef<void(const FInstancedSubObjRef& Ref)> OutObjects)
+{
+	check(ContainerAddress);
+	const UProperty* TargetProp = PropertyPath.Head();
+
+	if (!TargetProp->HasAnyPropertyFlags(CPF_PersistentInstance | CPF_ContainsInstancedReference))
+	{
+		return;
+	}
+
+	if (const UArrayProperty* ArrayProperty = Cast<const UArrayProperty>(TargetProp))
+	{
+		if (const UStructProperty* InnerStructProperty = Cast<const UStructProperty>(ArrayProperty->Inner))
+		{
+			if (const UStruct* Struct = InnerStructProperty->Struct)
+			{
+				FScriptArrayHelper_InContainer ArrayHelper(ArrayProperty, ContainerAddress);
+				for (int32 ElementIndex = 0; ElementIndex < ArrayHelper.Num(); ++ElementIndex)
+				{
+					const uint8* ValueAddress = ArrayHelper.GetRawPtr(ElementIndex);
+					for (UProperty* StructProp = Struct->RefLink; StructProp; StructProp = StructProp->NextRef)
+					{
+						PropertyPath.Push(InnerStructProperty, ElementIndex);
+						GetInstancedSubObjects_Inner(PropertyPath, ValueAddress, OutObjects);
+						PropertyPath.Pop();
+					}
+				}
+			}
+		}
+		else if (const UObjectProperty* InnerObjectProperty = Cast<const UObjectProperty>(ArrayProperty->Inner))
+		{
+			if (InnerObjectProperty->HasAllPropertyFlags(CPF_PersistentInstance))
+			{
+				ensure(InnerObjectProperty->HasAllPropertyFlags(CPF_InstancedReference));
+				FScriptArrayHelper_InContainer ArrayHelper(ArrayProperty, ContainerAddress);
+				for (int32 ElementIndex = 0; ElementIndex < ArrayHelper.Num(); ++ElementIndex)
+				{
+					if (UObject* ObjectValue = InnerObjectProperty->GetObjectPropertyValue(ArrayHelper.GetRawPtr(ElementIndex)))
+					{
+						PropertyPath.Push(InnerObjectProperty, ElementIndex);
+						OutObjects(FInstancedSubObjRef(ObjectValue, PropertyPath));
+						PropertyPath.Pop();
+					}
+				}
+			}
+		}
+	}
+	else if (const UMapProperty* MapProperty = Cast<const UMapProperty>(TargetProp))
+	{
+		if (const UStructProperty* KeyStructProperty = Cast<const UStructProperty>(MapProperty->KeyProp))
+		{
+			if (const UStruct* Struct = KeyStructProperty->Struct)
+			{
+				FScriptMapHelper MapHelper(MapProperty, MapProperty->ContainerPtrToValuePtr<void>(ContainerAddress));
+				for (int32 ElementIndex = 0; ElementIndex < MapHelper.Num(); ++ElementIndex)
+				{
+					if (MapHelper.IsValidIndex(ElementIndex))
+					{
+						const uint8* KeyAddress = MapHelper.GetKeyPtr(ElementIndex);
+						for (UProperty* StructProp = Struct->RefLink; StructProp; StructProp = StructProp->NextRef)
+						{
+							PropertyPath.Push(KeyStructProperty, ElementIndex);
+							GetInstancedSubObjects_Inner(PropertyPath, KeyAddress, OutObjects);
+							PropertyPath.Pop();
+						}
+					}
+				}
+			}
+		}
+		else if (const UObjectProperty* KeyObjectProperty = Cast<const UObjectProperty>(MapProperty->KeyProp))
+		{
+			if (KeyObjectProperty->HasAllPropertyFlags(CPF_PersistentInstance))
+			{
+				ensure(KeyObjectProperty->HasAllPropertyFlags(CPF_InstancedReference));
+				FScriptMapHelper MapHelper(MapProperty, MapProperty->ContainerPtrToValuePtr<void>(ContainerAddress));
+				for (int32 ElementIndex = 0; ElementIndex < MapHelper.Num(); ++ElementIndex)
+				{
+					if (MapHelper.IsValidIndex(ElementIndex))
+					{
+						const uint8* KeyAddress = MapHelper.GetKeyPtr(ElementIndex);
+						if (UObject* ObjectValue = KeyObjectProperty->GetObjectPropertyValue(KeyAddress))
+						{
+							PropertyPath.Push(KeyObjectProperty, ElementIndex);
+							OutObjects(FInstancedSubObjRef(ObjectValue, PropertyPath));
+							PropertyPath.Pop();
+						}
+					}
+				}
+			}
+		}
+
+		if (const UStructProperty* ValueStructProperty = Cast<const UStructProperty>(MapProperty->ValueProp))
+		{
+			if (const UStruct* Struct = ValueStructProperty->Struct)
+			{
+				FScriptMapHelper MapHelper(MapProperty, MapProperty->ContainerPtrToValuePtr<void>(ContainerAddress));
+				for (int32 ElementIndex = 0; ElementIndex < MapHelper.Num(); ++ElementIndex)
+				{
+					if (MapHelper.IsValidIndex(ElementIndex))
+					{
+						const uint8* ValueAddress = MapHelper.GetValuePtr(ElementIndex);
+						for (UProperty* StructProp = Struct->RefLink; StructProp; StructProp = StructProp->NextRef)
+						{
+							PropertyPath.Push(ValueStructProperty, ElementIndex);
+							GetInstancedSubObjects_Inner(PropertyPath, ValueAddress, OutObjects);
+							PropertyPath.Pop();
+						}
+					}
+				}
+			}
+		}
+		else if (const UObjectProperty* ValueObjectProperty = Cast<const UObjectProperty>(MapProperty->ValueProp))
+		{
+			if (ValueObjectProperty->HasAllPropertyFlags(CPF_PersistentInstance))
+			{
+				ensure(ValueObjectProperty->HasAllPropertyFlags(CPF_InstancedReference));
+				FScriptMapHelper MapHelper(MapProperty, MapProperty->ContainerPtrToValuePtr<void>(ContainerAddress));
+				for (int32 ElementIndex = 0; ElementIndex < MapHelper.Num(); ++ElementIndex)
+				{
+					if (MapHelper.IsValidIndex(ElementIndex))
+					{
+						const uint8* ValueAddress = MapHelper.GetValuePtr(ElementIndex);
+						if (UObject* ObjectValue = ValueObjectProperty->GetObjectPropertyValue(ValueAddress))
+						{
+							PropertyPath.Push(ValueObjectProperty, ElementIndex);
+							OutObjects(FInstancedSubObjRef(ObjectValue, PropertyPath));
+							PropertyPath.Pop();
+						}
+					}
+				}
+			}
+		}
+	}
+	else if (const USetProperty* SetProperty = Cast<const USetProperty>(TargetProp))
+	{
+		if (const UStructProperty* ElementStructProperty = Cast<const UStructProperty>(SetProperty->ElementProp))
+		{
+			if (const UStruct* Struct = ElementStructProperty->Struct)
+			{
+				FScriptSetHelper SetHelper(SetProperty, SetProperty->ContainerPtrToValuePtr<void>(ContainerAddress));
+				for (int32 ElementIndex = 0; ElementIndex < SetHelper.Num(); ++ElementIndex)
+				{
+					if (SetHelper.IsValidIndex(ElementIndex))
+					{
+						const uint8* ElementAddress = SetHelper.GetElementPtr(ElementIndex);
+						for (UProperty* StructProp = Struct->RefLink; StructProp; StructProp = StructProp->NextRef)
+						{
+							PropertyPath.Push(ElementStructProperty, ElementIndex);
+							GetInstancedSubObjects_Inner(PropertyPath, ElementAddress, OutObjects);
+							PropertyPath.Pop();
+						}
+					}
+				}
+			}
+		}
+		else if (const UObjectProperty* ElementObjectProperty = Cast<const UObjectProperty>(SetProperty->ElementProp))
+		{
+			if (ElementObjectProperty->HasAllPropertyFlags(CPF_PersistentInstance))
+			{
+				ensure(ElementObjectProperty->HasAllPropertyFlags(CPF_InstancedReference));
+				FScriptSetHelper SetHelper(SetProperty, SetProperty->ContainerPtrToValuePtr<void>(ContainerAddress));
+				for (int32 ElementIndex = 0; ElementIndex < SetHelper.Num(); ++ElementIndex)
+				{
+					if (SetHelper.IsValidIndex(ElementIndex))
+					{
+						const uint8* ElementAddress = SetHelper.GetElementPtr(ElementIndex);
+						if (UObject* ObjectValue = ElementObjectProperty->GetObjectPropertyValue(ElementAddress))
+						{
+							PropertyPath.Push(ElementObjectProperty, ElementIndex);
+							OutObjects(FInstancedSubObjRef(ObjectValue, PropertyPath));
+							PropertyPath.Pop();
+						}
+					}
+				}
+			}
+		}
+	}
+	else if (TargetProp->HasAllPropertyFlags(CPF_PersistentInstance))
+	{
+		ensure(TargetProp->HasAllPropertyFlags(CPF_InstancedReference));
+		if (const UObjectProperty* ObjectProperty = Cast<const UObjectProperty>(TargetProp))
+		{
+			for (int32 ArrayIdx = 0; ArrayIdx < ObjectProperty->ArrayDim; ++ArrayIdx)
+			{
+				UObject* ObjectValue = ObjectProperty->GetObjectPropertyValue_InContainer(ContainerAddress, ArrayIdx);
+				if (ObjectValue)
+				{
+					// don't need to push to PropertyPath, since this property is already at its head
+					OutObjects(FInstancedSubObjRef(ObjectValue, PropertyPath));	
+				}
+			}
+		}
+
+		return;
+	}
+	else if (const UStructProperty* StructProperty = Cast<const UStructProperty>(TargetProp))
+	{
+		if (StructProperty->Struct)
+		{
+			for (int32 ArrayIdx = 0; ArrayIdx < StructProperty->ArrayDim; ++ArrayIdx)
+			{
+				const uint8* ValueAddress = StructProperty->ContainerPtrToValuePtr<uint8>(ContainerAddress, ArrayIdx);
+				for (UProperty* StructProp = StructProperty->Struct->RefLink; StructProp; StructProp = StructProp->NextRef)
+				{
+					PropertyPath.Push(StructProp, ArrayIdx);
+					GetInstancedSubObjects_Inner(PropertyPath, ValueAddress, OutObjects);
+					PropertyPath.Pop();
+				}
+			}
+		}
+	}
+}
+
+void FFindInstancedReferenceSubobjectHelper::Duplicate(UObject* OldObject, UObject* NewObject, TMap<UObject*, UObject*>& ReferenceReplacementMap, TArray<UObject*>& DuplicatedObjects)
+{
+	if (OldObject->GetClass()->HasAnyClassFlags(CLASS_HasInstancedReference) &&
+		NewObject->GetClass()->HasAnyClassFlags(CLASS_HasInstancedReference))
+	{
+		TArray<FInstancedSubObjRef> OldInstancedSubObjects;
+		GetInstancedSubObjects(OldObject, OldInstancedSubObjects);
+		if (OldInstancedSubObjects.Num() > 0)
+		{
+			TArray<FInstancedSubObjRef> NewInstancedSubObjects;
+			GetInstancedSubObjects(NewObject, NewInstancedSubObjects);
+			for (const FInstancedSubObjRef& Obj : NewInstancedSubObjects)
+			{
+				const bool bNewObjectHasOldOuter = (Obj->GetOuter() == OldObject);
+				if (bNewObjectHasOldOuter)
+				{
+					const bool bKeptByOld = OldInstancedSubObjects.Contains(Obj);
+					const bool bNotHandledYet = !ReferenceReplacementMap.Contains(Obj);
+					if (bKeptByOld && bNotHandledYet)
+					{
+						UObject* NewEditInlineSubobject = StaticDuplicateObject(Obj, NewObject);
+						ReferenceReplacementMap.Add(Obj, NewEditInlineSubobject);
+
+						// NOTE: we cannot patch OldObject's linker table here, since we don't 
+						//       know the relation between the  two objects (one could be of a 
+						//       super class, and the other a child)
+
+						// We also need to make sure to fixup any properties here
+						DuplicatedObjects.Add(NewEditInlineSubobject);
+					}
+				}
+			}
+		}
+	}
+}
diff --git a/Engine/Source/Runtime/Engine/Private/Level.cpp b/Engine/Source/Runtime/Engine/Private/Level.cpp
--- a/Engine/Source/Runtime/Engine/Private/Level.cpp
+++ b/Engine/Source/Runtime/Engine/Private/Level.cpp
@@ -53,8 +53,19 @@ Level.cpp: Level-related functions
 #include "Engine/LevelStreaming.h"
 #include "LevelUtils.h"
 #include "Components/ModelComponent.h"
+#include "Engine/LevelActorContainer.h"
+#include "Engine/StaticMeshActor.h"
+
 DEFINE_LOG_CATEGORY(LogLevel);
 
+int32 GActorClusteringEnabled = 1;
+static FAutoConsoleVariableRef CVarUseBackgroundLevelStreaming(
+	TEXT("gc.ActorClusteringEnabled"),
+	GActorClusteringEnabled,
+	TEXT("Whether to allow levels to create actor clusters for GC."),
+	ECVF_Default
+);
+
 /*-----------------------------------------------------------------------------
 ULevel implementation.
 -----------------------------------------------------------------------------*/
@@ -308,12 +319,14 @@ void ULevel::AddReferencedObjects(UObject* InThis, FReferenceCollector& Collecto
 	ULevel* This = CastChecked<ULevel>(InThis);
 
 	// Let GC know that we're referencing some AActor objects
-	for (auto& Actor : This->Actors)
+	if (FPlatformProperties::RequiresCookedData() && GActorClusteringEnabled)
 	{
-		Collector.AddReferencedObject(Actor, This);
+		Collector.AddReferencedObjects(This->ActorsForGC, This);
+	}
+	else
+	{
+		Collector.AddReferencedObjects(This->Actors, This);
 	}
-	UObject* ActorsOwner = This;
-	Collector.AddReferencedObject(ActorsOwner, This);
 
 	Super::AddReferencedObjects( This, Collector );
 }
@@ -627,6 +640,47 @@ void ULevel::PostLoad()
 #endif
 }
 
+bool ULevel::CanBeClusterRoot() const
+{
+	return !!GActorClusteringEnabled;
+}
+
+void ULevel::CreateCluster()
+{
+	check(GActorClusteringEnabled);
+
+	// ULevels are not cluster roots themselves, instead they create a special actor container
+	// that holds a reference to all actors that are to be clustered. This is because only
+	// specific actor types can be clustered so the remaining actors that are not clustered
+	// need to be referenced through the level.
+	// Also, we don't want the level to reference the actors that are clusters because that would
+	// make things work even slower (references to clustered objects are expensive). That's why
+	// we keep a separate array for referencing unclustered actors (ActorsForGC).
+	if (ActorCluster == nullptr)
+	{
+		TArray<AActor*> ClusterActors;
+
+		for (int32 ActorIndex = Actors.Num() - 1; ActorIndex >= 0; --ActorIndex)
+		{
+			AActor* Actor = Actors[ActorIndex];
+			if (Actor && Actor->CanBeInCluster())
+			{
+				ClusterActors.Add(Actor);
+			}
+			else
+			{
+				ActorsForGC.Add(Actor);
+			}
+		}
+		if (ClusterActors.Num())
+		{
+			ActorCluster = NewObject<ULevelActorContainer>(this, TEXT("ActorCluster"), RF_Transient);
+			ActorCluster->Actors = MoveTemp(ClusterActors);
+			ActorCluster->CreateCluster();
+		}
+	}
+}
+
 void ULevel::PostDuplicate(bool bDuplicateForPIE)
 {
 	Super::PostDuplicate(bDuplicateForPIE);
diff --git a/Engine/Source/Runtime/Engine/Private/LevelActor.cpp b/Engine/Source/Runtime/Engine/Private/LevelActor.cpp
--- a/Engine/Source/Runtime/Engine/Private/LevelActor.cpp
+++ b/Engine/Source/Runtime/Engine/Private/LevelActor.cpp
@@ -438,6 +438,7 @@ AActor* UWorld::SpawnActor( UClass* Class, FTransform const* UserTransformPtr, c
 		ModifyLevel( LevelToSpawnIn );
 	}
 	LevelToSpawnIn->Actors.Add( Actor );
+	LevelToSpawnIn->ActorsForGC.Add(Actor);
 
 	// Add this newly spawned actor to the network actor list
 	AddNetworkActor( Actor );
diff --git a/Engine/Source/Runtime/Engine/Private/LevelActorContainer.cpp b/Engine/Source/Runtime/Engine/Private/LevelActorContainer.cpp
new file mode 100644
--- /dev/null
+++ b/Engine/Source/Runtime/Engine/Private/LevelActorContainer.cpp
@@ -0,0 +1,372 @@
+// Copyright 1998-2016 Epic Games, Inc. All Rights Reserved.
+
+/*=============================================================================
+Level.cpp: Level-related functions
+=============================================================================*/
+
+#include "Engine/LevelActorContainer.h"
+#include "Engine/Level.h"
+#include "Engine/World.h"
+#include "UObject/FastReferenceCollector.h"
+#include "UObject/UObjectArray.h"
+#include "UObject/Package.h"
+
+DEFINE_LOG_CATEGORY_STATIC(LogLevelActorContainer, Log, All);
+
+/**
+* Pool for reducing allocations when constructing clusters
+*/
+class FActorClusterArrayPool
+{
+public:
+
+	/**
+	* Gets the singleton instance of the FObjectArrayPool
+	* @return Pool singleton.
+	*/
+	FORCEINLINE static FActorClusterArrayPool& Get()
+	{
+		static FActorClusterArrayPool Singleton;
+		return Singleton;
+	}
+
+	/**
+	* Gets an event from the pool or creates one if necessary.
+	*
+	* @return The array.
+	* @see ReturnToPool
+	*/
+	FORCEINLINE TArray<UObject*>* GetArrayFromPool()
+	{
+		TArray<UObject*>* Result = Pool.Pop();
+		if (!Result)
+		{
+			Result = new TArray<UObject*>();
+		}
+		check(Result);
+#if UE_BUILD_DEBUG
+		NumberOfUsedArrays.Increment();
+#endif // UE_BUILD_DEBUG
+		return Result;
+	}
+
+	/**
+	* Returns an array to the pool.
+	*
+	* @param Array The array to return.
+	* @see GetArrayFromPool
+	*/
+	FORCEINLINE void ReturnToPool(TArray<UObject*>* Array)
+	{
+#if UE_BUILD_DEBUG
+		const int32 CheckUsedArrays = NumberOfUsedArrays.Decrement();
+		checkSlow(CheckUsedArrays >= 0);
+#endif // UE_BUILD_DEBUG
+		check(Array);
+		Array->Reset();
+		Pool.Push(Array);
+	}
+
+	/** Performs memory cleanup */
+	void Cleanup()
+	{
+#if UE_BUILD_DEBUG
+		const int32 CheckUsedArrays = NumberOfUsedArrays.GetValue();
+		checkSlow(CheckUsedArrays == 0);
+#endif // UE_BUILD_DEBUG
+
+		uint32 FreedMemory = 0;
+		TArray< TArray<UObject*>* > AllArrays;
+		Pool.PopAll(AllArrays);
+		for (TArray<UObject*>* Array : AllArrays)
+		{
+			FreedMemory += Array->GetAllocatedSize();
+			delete Array;
+		}
+		UE_LOG(LogLevelActorContainer, Log, TEXT("Freed %ub from %d cluster array pools."), FreedMemory, AllArrays.Num());
+	}
+
+#if UE_BUILD_DEBUG
+	void CheckLeaks()
+	{
+		// This function is called after cluster has been created so at this point there should be no
+		// arrays used by cluster creation code and all should be returned to the pool
+		const int32 LeakedPoolArrays = NumberOfUsedArrays.GetValue();
+		checkSlow(LeakedPoolArrays == 0);
+	}
+#endif
+
+private:
+
+	/** Holds the collection of recycled arrays. */
+	TLockFreePointerListLIFO< TArray<UObject*> > Pool;
+
+#if UE_BUILD_DEBUG
+	/** Number of arrays currently acquired from the pool by clusters */
+	FThreadSafeCounter NumberOfUsedArrays;
+#endif // UE_BUILD_DEBUG
+};
+
+/**
+* Handles UObject references found by TFastReferenceCollector
+*/
+class FActorClusterReferenceProcessor
+{
+	int32 ClusterRootIndex;
+	FUObjectCluster& Cluster;
+	ULevel* ParentLevel;
+	UPackage* ParentLevelPackage;
+	volatile bool bIsRunningMultithreaded;
+public:
+
+	FActorClusterReferenceProcessor(int32 InClusterRootIndex, FUObjectCluster& InCluster, ULevel* InParentLevel)
+		: ClusterRootIndex(InClusterRootIndex)
+		, Cluster(InCluster)
+		, ParentLevel(InParentLevel)
+		, bIsRunningMultithreaded(false)
+	{
+		ParentLevelPackage = ParentLevel->GetOutermost();
+	}
+
+	FORCEINLINE int32 GetMinDesiredObjectsPerSubTask() const
+	{
+		// We're not running the processor in parallel when creating clusters
+		return 0;
+	}
+
+	FORCEINLINE volatile bool IsRunningMultithreaded() const
+	{
+		// This should always be false
+		return bIsRunningMultithreaded;
+	}
+
+	FORCEINLINE void SetIsRunningMultithreaded(bool bIsParallel)
+	{
+		check(!bIsParallel);
+		bIsRunningMultithreaded = bIsParallel;
+	}
+
+	void UpdateDetailedStats(UObject* CurrentObject, uint32 DeltaCycles)
+	{
+	}
+
+	void LogDetailedStatsSummary()
+	{
+	}
+
+	FORCENOINLINE bool CanAddToCluster(UObject* Object)
+	{
+		UPackage* Package = Object->GetOutermost();
+		if (Package->ContainsMap() && Package != ParentLevelPackage)
+		{
+			// Can't reference anything from other levels
+			return false;
+		}
+		if (Object->IsIn(ParentLevelPackage) && !Object->IsIn(ParentLevel))
+		{
+			// If the object is in the same package but is not in the level we don't want it.
+			return false;
+		}
+		if (Object->IsA(ULevel::StaticClass()) || Object->IsA(UWorld::StaticClass()))
+		{
+			// And generally, no levels or worlds
+			return false;
+		}
+		return Object->CanBeInCluster();
+	}
+
+	/**
+	* Adds an object to cluster (if possible)
+	*
+	* @param ObjectIndex UObject index in GUObjectArray
+	* @param ObjectItem UObject's entry in GUObjectArray
+	* @param Obj The object to add to cluster
+	* @param ObjectsToSerialize An array of remaining objects to serialize (Obj must be added to it if Obj can be added to cluster)
+	* @param bOuterAndClass If true, the Obj's Outer and Class will also be added to the cluster
+	*/
+	void AddObjectToCluster(int32 ObjectIndex, FUObjectItem* ObjectItem, UObject* Obj, TArray<UObject*>& ObjectsToSerialize, bool bOuterAndClass)
+	{
+		// If we haven't finished loading, we can't be sure we know all the references
+		check(!Obj->HasAnyFlags(RF_NeedLoad | RF_NeedPostLoad));
+		if (ObjectIndex != ClusterRootIndex && ObjectItem->GetOwnerIndex() == 0 && !ObjectItem->IsRootSet() && !GUObjectArray.IsDisregardForGC(Obj) && Obj->CanBeInCluster())
+		{
+			ObjectsToSerialize.Add(Obj);
+			check(!ObjectItem->HasAnyFlags(EInternalObjectFlags::ClusterRoot));
+			ObjectItem->SetOwnerIndex(ClusterRootIndex);
+			Cluster.Objects.Add(ObjectIndex);
+
+			if (bOuterAndClass)
+			{
+				UObject* ObjOuter = Obj->GetOuter();
+				if (CanAddToCluster(ObjOuter))
+				{
+					int32 OuterIndex = GUObjectArray.ObjectToIndex(ObjOuter);
+					AddObjectToCluster(OuterIndex, GUObjectArray.IndexToObjectUnsafeForGC(OuterIndex), ObjOuter, ObjectsToSerialize, false);
+				}
+				else
+				{
+					Cluster.MutableObjects.AddUnique(GUObjectArray.ObjectToIndex(ObjOuter));
+				}
+				if (!Obj->GetClass()->HasAllClassFlags(CLASS_Native))
+				{
+					UObject* ObjectClass = Obj->GetClass();
+					HandleTokenStreamObjectReference(ObjectsToSerialize, Obj, ObjectClass, INDEX_NONE, true);
+					UObject* ObjectClassOuter = Obj->GetClass()->GetOuter();
+					HandleTokenStreamObjectReference(ObjectsToSerialize, Obj, ObjectClassOuter, INDEX_NONE, true);
+				}
+			}
+		}
+	}
+
+	/**
+	* Handles UObject reference from the token stream. Performance is critical here so we're FORCEINLINING this function.
+	*
+	* @param ObjectsToSerialize An array of remaining objects to serialize (Obj must be added to it if Obj can be added to cluster)
+	* @param ReferencingObject Object referencing the object to process.
+	* @param TokenIndex Index to the token stream where the reference was found.
+	* @param bAllowReferenceElimination True if reference elimination is allowed (ignored when constructing clusters).
+	*/
+	FORCEINLINE void HandleTokenStreamObjectReference(TArray<UObject*>& ObjectsToSerialize, UObject* ReferencingObject, UObject*& Object, const int32 TokenIndex, bool bAllowReferenceElimination)
+	{
+		if (Object)
+		{
+			// If we haven't finished loading, we can't be sure we know all the references
+			check(!Object->HasAnyFlags(RF_NeedLoad | RF_NeedPostLoad));
+
+			FUObjectItem* ObjectItem = GUObjectArray.ObjectToObjectItem(Object);
+
+			// Add encountered object reference to list of to be serialized objects if it hasn't already been added.
+			if (ObjectItem->GetOwnerIndex() != ClusterRootIndex)
+			{
+				if (ObjectItem->HasAnyFlags(EInternalObjectFlags::ClusterRoot) || ObjectItem->GetOwnerIndex() != 0)
+				{
+					// Simply reference this cluster and all clusters it's referencing
+					const int32 OtherClusterRootIndex = ObjectItem->HasAnyFlags(EInternalObjectFlags::ClusterRoot) ? GUObjectArray.ObjectToIndex(Object) : ObjectItem->GetOwnerIndex();
+					FUObjectItem* OtherClusterRootItem = GUObjectArray.IndexToObject(OtherClusterRootIndex);
+					FUObjectCluster& OtherCluster = GUObjectClusters[OtherClusterRootItem->GetClusterIndex()];
+					Cluster.ReferencedClusters.AddUnique(OtherClusterRootIndex);
+
+					OtherCluster.ReferencedByClusters.AddUnique(ClusterRootIndex);
+
+					for (int32 OtherClusterReferencedCluster : OtherCluster.ReferencedClusters)
+					{
+						check(OtherClusterReferencedCluster != ClusterRootIndex);
+						Cluster.ReferencedClusters.AddUnique(OtherClusterReferencedCluster);
+					}
+					for (int32 OtherClusterReferencedMutableObjectIndex : OtherCluster.MutableObjects)
+					{
+						Cluster.MutableObjects.AddUnique(OtherClusterReferencedMutableObjectIndex);
+					}
+				}
+				else if (ObjectItem->GetOwnerIndex() == 0 && !ObjectItem->IsRootSet() && !GUObjectArray.IsDisregardForGC(Object) &&
+					!(Object->CanBeClusterRoot() && Object->HasAnyFlags(RF_NeedLoad | RF_NeedPostLoad))) // Objects that can create clusters themselves and haven't been postloaded yet should be excluded
+				{
+					// New object, add it to the cluster.
+					if (CanAddToCluster(Object))
+					{
+						AddObjectToCluster(GUObjectArray.ObjectToIndex(Object), ObjectItem, Object, ObjectsToSerialize, true);
+					}
+					else
+					{
+						Cluster.MutableObjects.AddUnique(GUObjectArray.ObjectToIndex(Object));
+					}
+				}
+			}
+		}
+	}
+};
+
+/**
+* Specialized FReferenceCollector that uses FActorClusterReferenceProcessor to construct the cluster
+*/
+template <class TProcessor>
+class TActorClusterCollector : public FReferenceCollector
+{
+	TProcessor& Processor;
+	TArray<UObject*>& ObjectArray;
+
+public:
+	TActorClusterCollector(TProcessor& InProcessor, TArray<UObject*>& InObjectArray)
+		: Processor(InProcessor)
+		, ObjectArray(InObjectArray)
+	{
+	}
+	virtual void HandleObjectReference(UObject*& Object, const UObject* ReferencingObject, const UProperty* ReferencingProperty) override
+	{
+		Processor.HandleTokenStreamObjectReference(ObjectArray, const_cast<UObject*>(ReferencingObject), Object, INDEX_NONE, false);
+	}
+	virtual void HandleObjectReferences(UObject** InObjects, const int32 ObjectNum, const UObject* ReferencingObject, const UProperty* InReferencingProperty) override
+	{
+		for (int32 ObjectIndex = 0; ObjectIndex < ObjectNum; ++ObjectIndex)
+		{
+			UObject*& Object = InObjects[ObjectIndex];
+			Processor.HandleTokenStreamObjectReference(ObjectArray, const_cast<UObject*>(ReferencingObject), Object, INDEX_NONE, false);
+		}
+	}
+	virtual bool IsIgnoringArchetypeRef() const override
+	{
+		return false;
+	}
+	virtual bool IsIgnoringTransient() const override
+	{
+		return false;
+	}
+};
+
+void ULevelActorContainer::CreateCluster()
+{
+
+	int32 ContainerInternalIndex = GUObjectArray.ObjectToIndex(this);
+	FUObjectItem* RootItem = GUObjectArray.IndexToObject(ContainerInternalIndex);
+	if (RootItem->GetOwnerIndex() != 0 || RootItem->HasAnyFlags(EInternalObjectFlags::ClusterRoot))
+	{
+		return;
+	}
+
+	// If we haven't finished loading, we can't be sure we know all the references
+	check(!HasAnyFlags(RF_NeedLoad | RF_NeedPostLoad));
+
+	// Create a new cluster, reserve an arbitrary amount of memory for it.
+	const int32 ClusterIndex = GUObjectClusters.AllocateCluster(GUObjectArray.ObjectToIndex(this));
+	FUObjectCluster& Cluster = GUObjectClusters[ClusterIndex];
+	Cluster.Objects.Reserve(64);
+
+	// Collect all objects referenced by cluster root and by all objects it's referencing
+	FActorClusterReferenceProcessor Processor(ContainerInternalIndex, Cluster, CastChecked<ULevel>(GetOuter()));
+	TFastReferenceCollector<false, FActorClusterReferenceProcessor, TActorClusterCollector<FActorClusterReferenceProcessor>, FActorClusterArrayPool, true> ReferenceCollector(Processor, FActorClusterArrayPool::Get());
+	TArray<UObject*> ObjectsToProcess;
+	ObjectsToProcess.Add(static_cast<UObject*>(this));
+	ReferenceCollector.CollectReferences(ObjectsToProcess);
+#if UE_BUILD_DEBUG
+	FActorClusterArrayPool::Get().CheckLeaks();
+#endif
+
+	if (Cluster.Objects.Num())
+	{
+		// Sort all objects and set up the cluster root
+		Cluster.Objects.Sort();
+		Cluster.ReferencedClusters.Sort();
+		Cluster.MutableObjects.Sort();
+		check(RootItem->GetOwnerIndex() == 0);
+		RootItem->SetClusterIndex(ClusterIndex);
+		RootItem->SetFlags(EInternalObjectFlags::ClusterRoot);
+
+		UE_LOG(LogLevelActorContainer, Log, TEXT("Created LevelActorCluster (%d) for %s with %d objects, %d referenced clusters and %d mutable objects."),
+			ClusterIndex, *GetOuter()->GetPathName(), Cluster.Objects.Num(), Cluster.ReferencedClusters.Num(), Cluster.MutableObjects.Num());
+	}
+	else
+	{
+		check(RootItem->GetOwnerIndex() == 0);
+		RootItem->SetClusterIndex(ClusterIndex);
+		GUObjectClusters.FreeCluster(ClusterIndex);
+	} 
+}
+
+void ULevelActorContainer::OnClusterMarkedAsPendingKill()
+{
+	ULevel* Level = CastChecked<ULevel>(GetOuter());
+	Level->ActorsForGC.Append(Actors);
+	Actors.Reset();
+
+	Super::OnClusterMarkedAsPendingKill();
+}
diff --git a/Engine/Source/Runtime/Engine/Private/Materials/Material.cpp b/Engine/Source/Runtime/Engine/Private/Materials/Material.cpp
--- a/Engine/Source/Runtime/Engine/Private/Materials/Material.cpp
+++ b/Engine/Source/Runtime/Engine/Private/Materials/Material.cpp
@@ -1219,24 +1219,23 @@ bool UMaterial::CheckMaterialUsage_Concurrent(EMaterialUsage Usage, const bool b
 				UMaterial* Material;
 				EMaterialUsage Usage;
 				bool bSkipPrim;
-				bool& bUsageSetSuccessfully;
 
-				FCallSMU(UMaterial* InMaterial, EMaterialUsage InUsage, bool bInSkipPrim, bool& bInUsageSetSuccessfully)
+				FCallSMU(UMaterial* InMaterial, EMaterialUsage InUsage, bool bInSkipPrim)
 					: Material(InMaterial)
 					, Usage(InUsage)
 					, bSkipPrim(bInSkipPrim)
-					, bUsageSetSuccessfully(bInUsageSetSuccessfully)
 				{
 				}
 
 				void Task()
 				{
-					bUsageSetSuccessfully = Material->CheckMaterialUsage(Usage, bSkipPrim);
+					Material->CheckMaterialUsage(Usage, bSkipPrim);
 				}
 			};
-			UE_LOG(LogMaterial, Warning, TEXT("Has to pass SMU back to game thread. This stalls the tasks graph, but since it is editor only or only happens once, is not such a big deal."));
+			UE_LOG(LogMaterial, Log, TEXT("Had to pass SMU back to game thread. Please ensure correct material usage flags."));
 
-			TSharedRef<FCallSMU, ESPMode::ThreadSafe> CallSMU = MakeShareable(new FCallSMU(const_cast<UMaterial*>(this), Usage, bSkipPrim, bUsageSetSuccessfully));
+			TSharedRef<FCallSMU, ESPMode::ThreadSafe> CallSMU = MakeShareable(new FCallSMU(const_cast<UMaterial*>(this), Usage, bSkipPrim));
+			bUsageSetSuccessfully = false;
 
 			DECLARE_CYCLE_STAT(TEXT("FSimpleDelegateGraphTask.CheckMaterialUsage"),
 				STAT_FSimpleDelegateGraphTask_CheckMaterialUsage,
diff --git a/Engine/Source/Runtime/Engine/Private/Materials/MaterialInstance.cpp b/Engine/Source/Runtime/Engine/Private/Materials/MaterialInstance.cpp
--- a/Engine/Source/Runtime/Engine/Private/Materials/MaterialInstance.cpp
+++ b/Engine/Source/Runtime/Engine/Private/Materials/MaterialInstance.cpp
@@ -1100,24 +1100,23 @@ bool UMaterialInstance::CheckMaterialUsage_Concurrent(const EMaterialUsage Usage
 					UMaterialInstance* Material;
 					EMaterialUsage Usage;
 					bool bSkipPrim;
-					bool& bUsageSetSuccessfully;
 
-					FCallSMU(UMaterialInstance* InMaterial, EMaterialUsage InUsage, bool bInSkipPrim, bool& bInUsageSetSuccessfully)
+					FCallSMU(UMaterialInstance* InMaterial, EMaterialUsage InUsage, bool bInSkipPrim)
 						: Material(InMaterial)
 						, Usage(InUsage)
 						, bSkipPrim(bInSkipPrim)
-						, bUsageSetSuccessfully(bInUsageSetSuccessfully)
 					{
 					}
 
 					void Task()
 					{
-						bUsageSetSuccessfully = Material->CheckMaterialUsage(Usage, bSkipPrim);
+						Material->CheckMaterialUsage(Usage, bSkipPrim);
 					}
 				};
-				UE_LOG(LogMaterial, Warning, TEXT("Has to pass SMU back to game thread. This stalls the tasks graph, but since it is editor only, is not such a big deal."));
+				UE_LOG(LogMaterial, Log, TEXT("Had to pass SMU back to game thread. Please ensure correct material usage flags."));
 
-				TSharedRef<FCallSMU, ESPMode::ThreadSafe> CallSMU = MakeShareable(new FCallSMU(const_cast<UMaterialInstance*>(this), Usage, bSkipPrim, bUsageSetSuccessfully));
+				TSharedRef<FCallSMU, ESPMode::ThreadSafe> CallSMU = MakeShareable(new FCallSMU(const_cast<UMaterialInstance*>(this), Usage, bSkipPrim));
+				bUsageSetSuccessfully = false;
 
 				DECLARE_CYCLE_STAT(TEXT("FSimpleDelegateGraphTask.CheckMaterialUsage"),
 					STAT_FSimpleDelegateGraphTask_CheckMaterialUsage,
