AsyncLoading - Adding USoundBase to the set of CDOs that have a particular fixed boot order, otherwise strange load orders can be triggered where the USoundWave CDO can be initialized while the USoundBase CDO is mid inititialization.  Originally discovered in Dev-Editor, fixed in 3370466.

diff --git a/Engine/Source/Runtime/CoreUObject/Private/Serialization/AsyncLoading.cpp b/Engine/Source/Runtime/CoreUObject/Private/Serialization/AsyncLoading.cpp
--- a/Engine/Source/Runtime/CoreUObject/Private/Serialization/AsyncLoading.cpp
+++ b/Engine/Source/Runtime/CoreUObject/Private/Serialization/AsyncLoading.cpp
@@ -798,6 +798,7 @@ struct FEDLBootNotificationManager
 			FFixedBootOrder()
 			{
 				// add these in reverse order of when we want them processed
+				Array.Add("/Script/Engine/Default__SoundBase");
 				Array.Add("/Script/Engine/Default__MaterialInterface");
 				Array.Add("/Script/Engine/Default__DeviceProfileManager");
 			}
