Merging CL 3411681 from Release-4.16

diff --git a/Engine/Source/Runtime/CoreUObject/Public/UObject/NoExportTypes.h b/Engine/Source/Runtime/CoreUObject/Public/UObject/NoExportTypes.h
--- a/Engine/Source/Runtime/CoreUObject/Public/UObject/NoExportTypes.h
+++ b/Engine/Source/Runtime/CoreUObject/Public/UObject/NoExportTypes.h
@@ -924,7 +924,6 @@ public:
 USTRUCT(immutable, noexport, BlueprintType, meta=(HasNativeMake="Engine.KismetMathLibrary.MakeDateTime", HasNativeBreak="Engine.KismetMathLibrary.BreakDateTime"))
 struct FDateTime
 {
-	UPROPERTY()
 	int64 Ticks;
 };
 
@@ -934,7 +933,6 @@ struct FDateTime
 USTRUCT(immutable, noexport, BlueprintType, meta=(HasNativeMake="Engine.KismetMathLibrary.MakeTimespan", HasNativeBreak="Engine.KismetMathLibrary.BreakTimespan"))
 struct FTimespan
 {
-	UPROPERTY()
 	int64 Ticks;
 };
 
