Copying //UE4/Dev-Framework to //UE4/Dev-Main (Source: //UE4/Dev-Framework @ 3431384)

diff --git a/Engine/Source/Runtime/CoreUObject/Public/Serialization/DuplicatedDataWriter.h b/Engine/Source/Runtime/CoreUObject/Public/Serialization/DuplicatedDataWriter.h
--- a/Engine/Source/Runtime/CoreUObject/Public/Serialization/DuplicatedDataWriter.h
+++ b/Engine/Source/Runtime/CoreUObject/Public/Serialization/DuplicatedDataWriter.h
@@ -7,7 +7,6 @@
 #include "UObject/UObjectAnnotation.h"
 #include "Serialization/DuplicatedObject.h"
 
-class FAssetPtr;
 class FLazyObjectPtr;
 struct FObjectInstancingGraph;
 
@@ -40,7 +39,6 @@ private:
 	virtual FArchive& operator<<(FName& N);
 	virtual FArchive& operator<<(UObject*& Object);
 	virtual FArchive& operator<<(FLazyObjectPtr& LazyObjectPtr);
-	virtual FArchive& operator<<(FAssetPtr& AssetPtr);
 
 	virtual void Serialize(void* Data,int64 Num)
 	{
diff --git a/Engine/Source/Runtime/CoreUObject/Public/UObject/Class.h b/Engine/Source/Runtime/CoreUObject/Public/UObject/Class.h
--- a/Engine/Source/Runtime/CoreUObject/Public/UObject/Class.h
+++ b/Engine/Source/Runtime/CoreUObject/Public/UObject/Class.h
@@ -2011,7 +2011,7 @@ public:
 	int32 ClassUnique;
 
 	// Class flags; See EClassFlags for more information
-	uint32 ClassFlags;
+	EClassFlags ClassFlags;
 
 	// Cast flags used to accelerate dynamic_cast<T*> on objects of this type for common T
 	EClassCastFlags ClassCastFlags;
@@ -2027,7 +2027,8 @@ public:
 	 * Conditionally recompiles the class after loading, in case any dependencies were also newly loaded
 	 * @param ObjLoaded	If set this is the list of objects that are currently loading, usualy GObjLoaded
 	 */
-	virtual void ConditionalRecompileClass(TArray<UObject*>* ObjLoaded) {};
+	virtual void ConditionalRecompileClass(TArray<UObject*>* ObjLoaded) {}
+	virtual void FlushCompilationQueueForLevel() {}
 #endif //WITH_EDITOR
 
 	//
@@ -2124,7 +2125,7 @@ public:
 	// Constructors
 	UClass(const FObjectInitializer& ObjectInitializer = FObjectInitializer::Get());
 	explicit UClass(const FObjectInitializer& ObjectInitializer, UClass* InSuperClass);
-	UClass( EStaticConstructor, FName InName, uint32 InSize, uint32 InClassFlags, EClassCastFlags InClassCastFlags,
+	UClass( EStaticConstructor, FName InName, uint32 InSize, EClassFlags InClassFlags, EClassCastFlags InClassCastFlags,
 		const TCHAR* InClassConfigName, EObjectFlags InFlags, ClassConstructorType InClassConstructor,
 		ClassVTableHelperCtorCallerType InClassVTableHelperCtorCaller,
 		ClassAddReferencedObjectsType InClassAddReferencedObjects);
@@ -2142,7 +2143,7 @@ public:
 	 **/
 	bool HotReloadPrivateStaticClass(
 		uint32			InSize,
-		uint32			InClassFlags,
+		EClassFlags		InClassFlags,
 		EClassCastFlags	InClassCastFlags,
 		const TCHAR*    InConfigName,
 		ClassConstructorType InClassConstructor,
@@ -2358,14 +2359,14 @@ public:
 	/**
 	 * Used to safely check whether the passed in flag is set.
 	 *
-	 * @param	FlagToCheck		Class flag to check for
+	 * @param	FlagsToCheck		Class flag(s) to check for
 	 *
 	 * @return	true if the passed in flag is set, false otherwise
 	 *			(including no flag passed in, unless the FlagsToCheck is CLASS_AllFlags)
 	 */
-	FORCEINLINE bool HasAnyClassFlags( uint32 FlagsToCheck ) const
+	FORCEINLINE bool HasAnyClassFlags( EClassFlags FlagsToCheck ) const
 	{
-		return (ClassFlags & FlagsToCheck) != 0;
+		return EnumHasAnyFlags(ClassFlags, FlagsToCheck) != 0;
 	}
 
 	/**
@@ -2374,9 +2375,9 @@ public:
 	 * @param FlagsToCheck	Class flags to check for
 	 * @return true if all of the passed in flags are set (including no flags passed in), false otherwise
 	 */
-	FORCEINLINE bool HasAllClassFlags( uint32 FlagsToCheck ) const
+	FORCEINLINE bool HasAllClassFlags( EClassFlags FlagsToCheck ) const
 	{
-		return ((ClassFlags & FlagsToCheck) == FlagsToCheck);
+		return EnumHasAllFlags(ClassFlags, FlagsToCheck);
 	}
 
 	/**
@@ -2384,7 +2385,7 @@ public:
 	 *
 	 * @return	The class flags.
 	 */
-	FORCEINLINE uint32 GetClassFlags() const
+	FORCEINLINE EClassFlags GetClassFlags() const
 	{
 		return ClassFlags;
 	}
@@ -2612,7 +2613,7 @@ public:
 
 	UDynamicClass(const FObjectInitializer& ObjectInitializer = FObjectInitializer::Get());
 	explicit UDynamicClass(const FObjectInitializer& ObjectInitializer, UClass* InSuperClass);
-	UDynamicClass(EStaticConstructor, FName InName, uint32 InSize, uint32 InClassFlags, EClassCastFlags InClassCastFlags,
+	UDynamicClass(EStaticConstructor, FName InName, uint32 InSize, EClassFlags InClassFlags, EClassCastFlags InClassCastFlags,
 		const TCHAR* InClassConfigName, EObjectFlags InFlags, ClassConstructorType InClassConstructor,
 		ClassVTableHelperCtorCallerType InClassVTableHelperCtorCaller,
 		ClassAddReferencedObjectsType InClassAddReferencedObjects);
@@ -2696,7 +2697,7 @@ COREUOBJECT_API void GetPrivateStaticClassBody(
 	UClass*& ReturnClass,
 	void(*RegisterNativeFunc)(),
 	uint32 InSize,
-	uint32 InClassFlags,
+	EClassFlags InClassFlags,
 	EClassCastFlags InClassCastFlags,
 	const TCHAR* InConfigName,
 	UClass::ClassConstructorType InClassConstructor,
diff --git a/Engine/Source/Runtime/CoreUObject/Public/UObject/NoExportTypes.h b/Engine/Source/Runtime/CoreUObject/Public/UObject/NoExportTypes.h
--- a/Engine/Source/Runtime/CoreUObject/Public/UObject/NoExportTypes.h
+++ b/Engine/Source/Runtime/CoreUObject/Public/UObject/NoExportTypes.h
@@ -330,7 +330,7 @@ struct FVector
 * A 4-D homogeneous vector.
 * The full C++ class is located here: Engine\Source\Runtime\Core\Public\Math\Vector4.h
 */
-USTRUCT(immutable, noexport)
+USTRUCT(immutable, noexport, BlueprintType)
 struct FVector4
 {
 	UPROPERTY(EditAnywhere, BlueprintReadWrite, Category=Vector4, SaveGame)
@@ -367,10 +367,10 @@ struct FVector2D
 USTRUCT(immutable, noexport)
 struct FTwoVectors
 {
-	UPROPERTY(EditAnywhere, BlueprintReadWrite, Category=TwoVectors, SaveGame)
+	UPROPERTY(EditAnywhere, Category=TwoVectors, SaveGame)
 	FVector v1;
 
-	UPROPERTY(EditAnywhere, BlueprintReadWrite, Category=TwoVectors, SaveGame)
+	UPROPERTY(EditAnywhere, Category=TwoVectors, SaveGame)
 	FVector v2;
 
 };
@@ -380,7 +380,7 @@ struct FTwoVectors
  * A plane definition in 3D space.
  * The full C++ class is located here: Engine\Source\Runtime\Core\Public\Math\Plane.h
  */
-USTRUCT(immutable, noexport)
+USTRUCT(immutable, noexport, BlueprintType)
 struct FPlane : public FVector
 {
 	UPROPERTY(EditAnywhere, BlueprintReadWrite, Category=Plane, SaveGame)
@@ -415,7 +415,7 @@ struct FRotator
  * Quaternion.
  * The full C++ class is located here: Engine\Source\Runtime\Core\Public\Math\Quat.h
  */
-USTRUCT(immutable, noexport)
+USTRUCT(immutable, noexport, BlueprintType)
 struct FQuat
 {
 	UPROPERTY(EditAnywhere, BlueprintReadWrite, Category=Quat, SaveGame)
@@ -440,16 +440,16 @@ struct FQuat
 USTRUCT(immutable, noexport)
 struct FPackedNormal
 {
-	UPROPERTY(EditAnywhere, BlueprintReadWrite, Category=PackedNormal, SaveGame)
+	UPROPERTY(EditAnywhere, Category=PackedNormal, SaveGame)
 	uint8 X;
 
-	UPROPERTY(EditAnywhere, BlueprintReadWrite, Category=PackedNormal, SaveGame)
+	UPROPERTY(EditAnywhere, Category=PackedNormal, SaveGame)
 	uint8 Y;
 
-	UPROPERTY(EditAnywhere, BlueprintReadWrite, Category=PackedNormal, SaveGame)
+	UPROPERTY(EditAnywhere, Category=PackedNormal, SaveGame)
 	uint8 Z;
 
-	UPROPERTY(EditAnywhere, BlueprintReadWrite, Category=PackedNormal, SaveGame)
+	UPROPERTY(EditAnywhere, Category=PackedNormal, SaveGame)
 	uint8 W;
 
 };
@@ -461,7 +461,7 @@ struct FPackedNormal
 USTRUCT(immutable, noexport)
 struct FPackedRGB10A2N
 {
-	UPROPERTY(EditAnywhere, BlueprintReadWrite, Category = PackedBasis, SaveGame)
+	UPROPERTY(EditAnywhere, Category = PackedBasis, SaveGame)
 	int32 Packed;
 };
 
@@ -472,10 +472,10 @@ struct FPackedRGB10A2N
 USTRUCT(immutable, noexport)
 struct FPackedRGBA16N
 {
-	UPROPERTY(EditAnywhere, BlueprintReadWrite, Category = PackedNormal, SaveGame)
+	UPROPERTY(EditAnywhere, Category = PackedNormal, SaveGame)
 	int32 XY;
 
-	UPROPERTY(EditAnywhere, BlueprintReadWrite, Category = PackedNormal, SaveGame)
+	UPROPERTY(EditAnywhere, Category = PackedNormal, SaveGame)
 	int32 ZW;
 };
 
@@ -483,7 +483,7 @@ struct FPackedRGBA16N
  * Screen coordinates.
  * The full C++ class is located here: Engine\Source\Runtime\Core\Public\Math\IntPoint.h
  */
-USTRUCT(immutable, noexport)
+USTRUCT(immutable, noexport, BlueprintType)
 struct FIntPoint
 {
 	UPROPERTY(EditAnywhere, BlueprintReadWrite, Category=IntPoint, SaveGame)
@@ -516,7 +516,7 @@ struct FIntVector
  * A Color (BGRA).
  * The full C++ class is located here: Engine\Source\Runtime\Core\Public\Math\Color.h
  */
-USTRUCT(immutable, noexport)
+USTRUCT(immutable, noexport, BlueprintType)
 struct FColor
 {
 	UPROPERTY(EditAnywhere, BlueprintReadWrite, Category=Color, SaveGame, meta=(ClampMin="0", ClampMax="255"))
@@ -597,7 +597,7 @@ struct FBox2D
  * A bounding box and bounding sphere with the same origin.
  * The full C++ class is located here : Engine\Source\Runtime\Core\Public\Math\BoxSphereBounds.h
  */
-USTRUCT(noexport)
+USTRUCT(noexport, BlueprintType)
 struct FBoxSphereBounds
 {
 	UPROPERTY(EditAnywhere, BlueprintReadWrite, Category=BoxSphereBounds, SaveGame)
@@ -618,25 +618,25 @@ struct FBoxSphereBounds
 USTRUCT(immutable, noexport)
 struct FOrientedBox
 {
-	UPROPERTY(EditAnywhere, BlueprintReadWrite, Category=OrientedBox, SaveGame)
+	UPROPERTY(EditAnywhere, Category=OrientedBox, SaveGame)
 	FVector Center;
 
-	UPROPERTY(EditAnywhere, BlueprintReadWrite, Category=OrientedBox, SaveGame)
+	UPROPERTY(EditAnywhere, Category=OrientedBox, SaveGame)
 	FVector AxisX;
 	
-	UPROPERTY(EditAnywhere, BlueprintReadWrite, Category=OrientedBox, SaveGame)
+	UPROPERTY(EditAnywhere, Category=OrientedBox, SaveGame)
 	FVector AxisY;
 	
-	UPROPERTY(EditAnywhere, BlueprintReadWrite, Category=OrientedBox, SaveGame)
+	UPROPERTY(EditAnywhere, Category=OrientedBox, SaveGame)
 	FVector AxisZ;
 
-	UPROPERTY(EditAnywhere, BlueprintReadWrite, Category=OrientedBox, SaveGame)
+	UPROPERTY(EditAnywhere, Category=OrientedBox, SaveGame)
 	float ExtentX;
 	
-	UPROPERTY(EditAnywhere, BlueprintReadWrite, Category=OrientedBox, SaveGame)
+	UPROPERTY(EditAnywhere, Category=OrientedBox, SaveGame)
 	float ExtentY;
 
-	UPROPERTY(EditAnywhere, BlueprintReadWrite, Category=OrientedBox, SaveGame)
+	UPROPERTY(EditAnywhere, Category=OrientedBox, SaveGame)
 	float ExtentZ;
 };
 
@@ -644,7 +644,7 @@ struct FOrientedBox
  * A 4x4 matrix.
  * The full C++ class is located here: Engine\Source\Runtime\Core\Public\Math\Matrix.h
  */
-USTRUCT(immutable, noexport)
+USTRUCT(immutable, noexport, BlueprintType)
 struct FMatrix
 {
 	UPROPERTY(EditAnywhere, BlueprintReadWrite, Category=Matrix, SaveGame)
@@ -666,19 +666,19 @@ struct FMatrix
 USTRUCT(noexport)
 struct FInterpCurvePointFloat
 {
-	UPROPERTY(EditAnywhere, BlueprintReadWrite, Category=InterpCurvePointFloat)
+	UPROPERTY(EditAnywhere, Category=InterpCurvePointFloat)
 	float InVal;
 
-	UPROPERTY(EditAnywhere, BlueprintReadWrite, Category=InterpCurvePointFloat)
+	UPROPERTY(EditAnywhere, Category=InterpCurvePointFloat)
 	float OutVal;
 
-	UPROPERTY(EditAnywhere, BlueprintReadWrite, Category=InterpCurvePointFloat)
+	UPROPERTY(EditAnywhere, Category=InterpCurvePointFloat)
 	float ArriveTangent;
 
-	UPROPERTY(EditAnywhere, BlueprintReadWrite, Category=InterpCurvePointFloat)
+	UPROPERTY(EditAnywhere, Category=InterpCurvePointFloat)
 	float LeaveTangent;
 
-	UPROPERTY(EditAnywhere, BlueprintReadWrite, Category=InterpCurvePointFloat)
+	UPROPERTY(EditAnywhere, Category=InterpCurvePointFloat)
 	TEnumAsByte<enum EInterpCurveMode> InterpMode;
 
 };
@@ -688,13 +688,13 @@ struct FInterpCurvePointFloat
 USTRUCT(noexport)
 struct FInterpCurveFloat
 {
-	UPROPERTY(EditAnywhere, BlueprintReadWrite, Category=InterpCurveFloat)
+	UPROPERTY(EditAnywhere, Category=InterpCurveFloat)
 	TArray<FInterpCurvePointFloat> Points;
 
-	UPROPERTY(EditAnywhere, BlueprintReadWrite, Category=InterpCurveFloat)
+	UPROPERTY(EditAnywhere, Category=InterpCurveFloat)
 	bool bIsLooped;
 
-	UPROPERTY(EditAnywhere, BlueprintReadWrite, Category=InterpCurveFloat)
+	UPROPERTY(EditAnywhere, Category=InterpCurveFloat)
 	float LoopKeyOffset;
 };
 
@@ -703,19 +703,19 @@ struct FInterpCurveFloat
 USTRUCT(noexport)
 struct FInterpCurvePointVector2D
 {
-	UPROPERTY(EditAnywhere, BlueprintReadWrite, Category=InterpCurvePointVector2D)
+	UPROPERTY(EditAnywhere, Category=InterpCurvePointVector2D)
 	float InVal;
 
-	UPROPERTY(EditAnywhere, BlueprintReadWrite, Category=InterpCurvePointVector2D)
+	UPROPERTY(EditAnywhere, Category=InterpCurvePointVector2D)
 	FVector2D OutVal;
 
-	UPROPERTY(EditAnywhere, BlueprintReadWrite, Category=InterpCurvePointVector2D)
+	UPROPERTY(EditAnywhere, Category=InterpCurvePointVector2D)
 	FVector2D ArriveTangent;
 
-	UPROPERTY(EditAnywhere, BlueprintReadWrite, Category=InterpCurvePointVector2D)
+	UPROPERTY(EditAnywhere, Category=InterpCurvePointVector2D)
 	FVector2D LeaveTangent;
 
-	UPROPERTY(EditAnywhere, BlueprintReadWrite, Category=InterpCurvePointVector2D)
+	UPROPERTY(EditAnywhere, Category=InterpCurvePointVector2D)
 	TEnumAsByte<enum EInterpCurveMode> InterpMode;
 
 };
@@ -725,13 +725,13 @@ struct FInterpCurvePointVector2D
 USTRUCT(noexport)
 struct FInterpCurveVector2D
 {
-	UPROPERTY(EditAnywhere, BlueprintReadWrite, Category=InterpCurveVector2D)
+	UPROPERTY(EditAnywhere, Category=InterpCurveVector2D)
 	TArray<FInterpCurvePointVector2D> Points;
 
-	UPROPERTY(EditAnywhere, BlueprintReadWrite, Category=InterpCurveVector2D)
+	UPROPERTY(EditAnywhere, Category=InterpCurveVector2D)
 	bool bIsLooped;
 
-	UPROPERTY(EditAnywhere, BlueprintReadWrite, Category=InterpCurveVector2D)
+	UPROPERTY(EditAnywhere, Category=InterpCurveVector2D)
 	float LoopKeyOffset;
 };
 
@@ -740,19 +740,19 @@ struct FInterpCurveVector2D
 USTRUCT(noexport)
 struct FInterpCurvePointVector
 {
-	UPROPERTY(EditAnywhere, BlueprintReadWrite, Category=InterpCurvePointVector)
+	UPROPERTY(EditAnywhere, Category=InterpCurvePointVector)
 	float InVal;
 
-	UPROPERTY(EditAnywhere, BlueprintReadWrite, Category=InterpCurvePointVector)
+	UPROPERTY(EditAnywhere, Category=InterpCurvePointVector)
 	FVector OutVal;
 
-	UPROPERTY(EditAnywhere, BlueprintReadWrite, Category=InterpCurvePointVector)
+	UPROPERTY(EditAnywhere, Category=InterpCurvePointVector)
 	FVector ArriveTangent;
 
-	UPROPERTY(EditAnywhere, BlueprintReadWrite, Category=InterpCurvePointVector)
+	UPROPERTY(EditAnywhere, Category=InterpCurvePointVector)
 	FVector LeaveTangent;
 
-	UPROPERTY(EditAnywhere, BlueprintReadWrite, Category=InterpCurvePointVector)
+	UPROPERTY(EditAnywhere, Category=InterpCurvePointVector)
 	TEnumAsByte<enum EInterpCurveMode> InterpMode;
 
 };
@@ -762,13 +762,13 @@ struct FInterpCurvePointVector
 USTRUCT(noexport)
 struct FInterpCurveVector
 {
-	UPROPERTY(EditAnywhere, BlueprintReadWrite, Category=InterpCurveVector)
+	UPROPERTY(EditAnywhere, Category=InterpCurveVector)
 	TArray<FInterpCurvePointVector> Points;
 
-	UPROPERTY(EditAnywhere, BlueprintReadWrite, Category=InterpCurveVector)
+	UPROPERTY(EditAnywhere, Category=InterpCurveVector)
 	bool bIsLooped;
 
-	UPROPERTY(EditAnywhere, BlueprintReadWrite, Category=InterpCurveVector)
+	UPROPERTY(EditAnywhere, Category=InterpCurveVector)
 	float LoopKeyOffset;
 };
 
@@ -777,19 +777,19 @@ struct FInterpCurveVector
 USTRUCT(noexport)
 struct FInterpCurvePointQuat
 {
-	UPROPERTY(EditAnywhere, BlueprintReadWrite, Category=InterpCurvePointQuat)
+	UPROPERTY(EditAnywhere, Category=InterpCurvePointQuat)
 	float InVal;
 
-	UPROPERTY(EditAnywhere, BlueprintReadWrite, Category=InterpCurvePointQuat)
+	UPROPERTY(EditAnywhere, Category=InterpCurvePointQuat)
 	FQuat OutVal;
 
-	UPROPERTY(EditAnywhere, BlueprintReadWrite, Category=InterpCurvePointQuat)
+	UPROPERTY(EditAnywhere, Category=InterpCurvePointQuat)
 	FQuat ArriveTangent;
 
-	UPROPERTY(EditAnywhere, BlueprintReadWrite, Category=InterpCurvePointQuat)
+	UPROPERTY(EditAnywhere, Category=InterpCurvePointQuat)
 	FQuat LeaveTangent;
 
-	UPROPERTY(EditAnywhere, BlueprintReadWrite, Category=InterpCurvePointQuat)
+	UPROPERTY(EditAnywhere, Category=InterpCurvePointQuat)
 	TEnumAsByte<enum EInterpCurveMode> InterpMode;
 
 };
@@ -799,13 +799,13 @@ struct FInterpCurvePointQuat
 USTRUCT(noexport)
 struct FInterpCurveQuat
 {
-	UPROPERTY(EditAnywhere, BlueprintReadWrite, Category=InterpCurveQuat)
+	UPROPERTY(EditAnywhere, Category=InterpCurveQuat)
 	TArray<FInterpCurvePointQuat> Points;
 
-	UPROPERTY(EditAnywhere, BlueprintReadWrite, Category=InterpCurveQuat)
+	UPROPERTY(EditAnywhere, Category=InterpCurveQuat)
 	bool bIsLooped;
 
-	UPROPERTY(EditAnywhere, BlueprintReadWrite, Category=InterpCurveQuat)
+	UPROPERTY(EditAnywhere, Category=InterpCurveQuat)
 	float LoopKeyOffset;
 };
 
@@ -814,19 +814,19 @@ struct FInterpCurveQuat
 USTRUCT(noexport)
 struct FInterpCurvePointTwoVectors
 {
-	UPROPERTY(EditAnywhere, BlueprintReadWrite, Category=InterpCurvePointTwoVectors)
+	UPROPERTY(EditAnywhere, Category=InterpCurvePointTwoVectors)
 	float InVal;
 
-	UPROPERTY(EditAnywhere, BlueprintReadWrite, Category=InterpCurvePointTwoVectors)
+	UPROPERTY(EditAnywhere, Category=InterpCurvePointTwoVectors)
 	FTwoVectors OutVal;
 
-	UPROPERTY(EditAnywhere, BlueprintReadWrite, Category=InterpCurvePointTwoVectors)
+	UPROPERTY(EditAnywhere, Category=InterpCurvePointTwoVectors)
 	FTwoVectors ArriveTangent;
 
-	UPROPERTY(EditAnywhere, BlueprintReadWrite, Category=InterpCurvePointTwoVectors)
+	UPROPERTY(EditAnywhere, Category=InterpCurvePointTwoVectors)
 	FTwoVectors LeaveTangent;
 
-	UPROPERTY(EditAnywhere, BlueprintReadWrite, Category=InterpCurvePointTwoVectors)
+	UPROPERTY(EditAnywhere, Category=InterpCurvePointTwoVectors)
 	TEnumAsByte<enum EInterpCurveMode> InterpMode;
 
 };
@@ -836,13 +836,13 @@ struct FInterpCurvePointTwoVectors
 USTRUCT(noexport)
 struct FInterpCurveTwoVectors
 {
-	UPROPERTY(EditAnywhere, BlueprintReadWrite, Category=InterpCurveTwoVectors)
+	UPROPERTY(EditAnywhere, Category=InterpCurveTwoVectors)
 	TArray<FInterpCurvePointTwoVectors> Points;
 
-	UPROPERTY(EditAnywhere, BlueprintReadWrite, Category=InterpCurveTwoVectors)
+	UPROPERTY(EditAnywhere, Category=InterpCurveTwoVectors)
 	bool bIsLooped;
 
-	UPROPERTY(EditAnywhere, BlueprintReadWrite, Category=InterpCurveTwoVectors)
+	UPROPERTY(EditAnywhere, Category=InterpCurveTwoVectors)
 	float LoopKeyOffset;
 };
 
@@ -851,19 +851,19 @@ struct FInterpCurveTwoVectors
 USTRUCT(noexport)
 struct FInterpCurvePointLinearColor
 {
-	UPROPERTY(EditAnywhere, BlueprintReadWrite, Category=InterpCurvePointLinearColor)
+	UPROPERTY(EditAnywhere, Category=InterpCurvePointLinearColor)
 	float InVal;
 
-	UPROPERTY(EditAnywhere, BlueprintReadWrite, Category=InterpCurvePointLinearColor)
+	UPROPERTY(EditAnywhere, Category=InterpCurvePointLinearColor)
 	FLinearColor OutVal;
 
-	UPROPERTY(EditAnywhere, BlueprintReadWrite, Category=InterpCurvePointLinearColor)
+	UPROPERTY(EditAnywhere, Category=InterpCurvePointLinearColor)
 	FLinearColor ArriveTangent;
 
-	UPROPERTY(EditAnywhere, BlueprintReadWrite, Category=InterpCurvePointLinearColor)
+	UPROPERTY(EditAnywhere, Category=InterpCurvePointLinearColor)
 	FLinearColor LeaveTangent;
 
-	UPROPERTY(EditAnywhere, BlueprintReadWrite, Category=InterpCurvePointLinearColor)
+	UPROPERTY(EditAnywhere, Category=InterpCurvePointLinearColor)
 	TEnumAsByte<enum EInterpCurveMode> InterpMode;
 
 };
@@ -873,13 +873,13 @@ struct FInterpCurvePointLinearColor
 USTRUCT(noexport)
 struct FInterpCurveLinearColor
 {
-	UPROPERTY(EditAnywhere, BlueprintReadWrite, Category=InterpCurveLinearColor)
+	UPROPERTY(EditAnywhere, Category=InterpCurveLinearColor)
 	TArray<FInterpCurvePointLinearColor> Points;
 
-	UPROPERTY(EditAnywhere, BlueprintReadWrite, Category=InterpCurveLinearColor)
+	UPROPERTY(EditAnywhere, Category=InterpCurveLinearColor)
 	bool bIsLooped;
 
-	UPROPERTY(EditAnywhere, BlueprintReadWrite, Category=InterpCurveLinearColor)
+	UPROPERTY(EditAnywhere, Category=InterpCurveLinearColor)
 	float LoopKeyOffset;
 };
 
@@ -939,7 +939,7 @@ struct FTimespan
 
 // A string asset reference
 
-USTRUCT(noexport, meta=(HasNativeMake="Engine.BlueprintFunctionLibrary.MakeStringAssetReference"))
+USTRUCT(noexport, BlueprintType, meta=(HasNativeMake="Engine.BlueprintFunctionLibrary.MakeStringAssetReference"))
 struct FStringAssetReference
 {
 	UPROPERTY(EditAnywhere, BlueprintReadWrite, Category=StringAssetReference)
@@ -983,7 +983,7 @@ struct FFallbackStruct
 {
 };
 
-UENUM()
+UENUM(BlueprintType)
 namespace ERangeBoundTypes
 {
 	/**
@@ -1010,49 +1010,49 @@ namespace ERangeBoundTypes
 
 // A float range bound
 
-USTRUCT(noexport)
+USTRUCT(noexport, BlueprintType)
 struct FFloatRangeBound
 {
-	UPROPERTY(EditAnywhere, Category=Range)
+	UPROPERTY(EditAnywhere, BlueprintReadWrite, Category=Range)
 	TEnumAsByte<ERangeBoundTypes::Type> Type;
 
-	UPROPERTY(EditAnywhere, Category=Range)
+	UPROPERTY(EditAnywhere, BlueprintReadWrite, Category=Range)
 	float Value;
 };
 
 // A float range
 
-USTRUCT(noexport)
+USTRUCT(noexport, BlueprintType)
 struct FFloatRange
 {
-	UPROPERTY(EditAnywhere, Category=Range)
+	UPROPERTY(EditAnywhere, BlueprintReadWrite, Category=Range)
 	FFloatRangeBound LowerBound;
 
-	UPROPERTY(EditAnywhere, Category=Range)
+	UPROPERTY(EditAnywhere, BlueprintReadWrite, Category=Range)
 	FFloatRangeBound UpperBound;
 };
 
 // An int32 range bound
 
-USTRUCT(noexport)
+USTRUCT(noexport, BlueprintType)
 struct FInt32RangeBound
 {
-	UPROPERTY(EditAnywhere, Category = Range)
+	UPROPERTY(EditAnywhere, BlueprintReadWrite, Category = Range)
 	TEnumAsByte<ERangeBoundTypes::Type> Type;
 
-	UPROPERTY(EditAnywhere, Category = Range)
+	UPROPERTY(EditAnywhere, BlueprintReadWrite, Category = Range)
 	int32 Value;
 };
 
 // An int32 range
 
-USTRUCT(noexport)
+USTRUCT(noexport, BlueprintType)
 struct FInt32Range
 {
-	UPROPERTY(EditAnywhere, Category = Range)
+	UPROPERTY(EditAnywhere, BlueprintReadWrite, Category = Range)
 	FInt32RangeBound LowerBound;
 
-	UPROPERTY(EditAnywhere, Category = Range)
+	UPROPERTY(EditAnywhere, BlueprintReadWrite, Category = Range)
 	FInt32RangeBound UpperBound;
 };
 
diff --git a/Engine/Source/Runtime/CoreUObject/Public/UObject/ObjectMacros.h b/Engine/Source/Runtime/CoreUObject/Public/UObject/ObjectMacros.h
--- a/Engine/Source/Runtime/CoreUObject/Public/UObject/ObjectMacros.h
+++ b/Engine/Source/Runtime/CoreUObject/Public/UObject/ObjectMacros.h
@@ -80,6 +80,7 @@ enum ELoadFlags
 	LOAD_PackageForPIE				= 0x00080000,   // This package is being loaded for PIE, it must be flagged as such immediately
 	LOAD_DeferDependencyLoads       = 0x00100000,   // Do not load external (blueprint) dependencies (instead, track them for deferred loading)
 	LOAD_ForFileDiff				= 0x00200000,	// Load the package (not for diffing in the editor), instead verify at the two packages serialized output are the same, if they are not then debug break so that you can get the callstack and object information
+	LOAD_DisableCompileOnLoad		= 0x00400000,	// Prevent this load call from running compile on load for the loaded blueprint (intentionally not recursive, dependencies will still compile on load)
 };
 
 //
@@ -170,79 +171,79 @@ enum EClassFlags
 {
 	/** @name Base flags */
 	//@{
-	CLASS_None				  = 0x00000000,
+	CLASS_None				  = 0x00000000u,
 	/** Class is abstract and can't be instantiated directly. */
-	CLASS_Abstract            = 0x00000001,
+	CLASS_Abstract            = 0x00000001u,
 	/** Save object configuration only to Default INIs, never to local INIs. Must be combined with CLASS_Config */
-	CLASS_DefaultConfig		  = 0x00000002,
+	CLASS_DefaultConfig		  = 0x00000002u,
 	/** Load object configuration at construction time. */
-	CLASS_Config			  = 0x00000004,
+	CLASS_Config			  = 0x00000004u,
 	/** This object type can't be saved; null it out at save time. */
-	CLASS_Transient			  = 0x00000008,
+	CLASS_Transient			  = 0x00000008u,
 	/** Successfully parsed. */
-	CLASS_Parsed              = 0x00000010,
+	CLASS_Parsed              = 0x00000010u,
 	/** */
-	//CLASS_                  = 0x00000020,
+	//CLASS_                  = 0x00000020u,
 	/** All the properties on the class are shown in the advanced section (which is hidden by default) unless SimpleDisplay is specified on the property */
-	CLASS_AdvancedDisplay	  = 0x00000040,
+	CLASS_AdvancedDisplay	  = 0x00000040u,
 	/** Class is a native class - native interfaces will have CLASS_Native set, but not RF_MarkAsNative */
-	CLASS_Native			  = 0x00000080,
+	CLASS_Native			  = 0x00000080u,
 	/** Don't export to C++ header. */
-	CLASS_NoExport            = 0x00000100,
+	CLASS_NoExport            = 0x00000100u,
 	/** Do not allow users to create in the editor. */
-	CLASS_NotPlaceable        = 0x00000200,
+	CLASS_NotPlaceable        = 0x00000200u,
 	/** Handle object configuration on a per-object basis, rather than per-class. */
-	CLASS_PerObjectConfig     = 0x00000400,
+	CLASS_PerObjectConfig     = 0x00000400u,
 	
 	/** */
-	//CLASS_ = 0x00000800,
+	//CLASS_ = 0x00000800u,
 	
 	/** Class can be constructed from editinline New button. */
-	CLASS_EditInlineNew		  = 0x00001000,
+	CLASS_EditInlineNew		  = 0x00001000u,
 	/** Display properties in the editor without using categories. */
-	CLASS_CollapseCategories  = 0x00002000,
+	CLASS_CollapseCategories  = 0x00002000u,
 	/** Class is an interface **/
-	CLASS_Interface           = 0x00004000,
+	CLASS_Interface           = 0x00004000u,
 	/**  Do not export a constructor for this class, assuming it is in the cpptext **/
-	CLASS_CustomConstructor   = 0x00008000,
+	CLASS_CustomConstructor   = 0x00008000u,
 	/** all properties and functions in this class are const and should be exported as const */
-	CLASS_Const			      = 0x00010000,
+	CLASS_Const			      = 0x00010000u,
 
 	/** */
-	//CLASS_ = 0x00020000,
+	//CLASS_ = 0x00020000u,
 	
 	/** Indicates that the class was created from blueprint source material */
-	CLASS_CompiledFromBlueprint  = 0x00040000,
+	CLASS_CompiledFromBlueprint  = 0x00040000u,
 
 	/** Indicates that only the bare minimum bits of this class should be DLL exported/imported */
-	CLASS_MinimalAPI	      = 0x00080000,
+	CLASS_MinimalAPI	      = 0x00080000u,
 	
 	/** Indicates this class must be DLL exported/imported (along with all of it's members) */
-	CLASS_RequiredAPI	      = 0x00100000,
+	CLASS_RequiredAPI	      = 0x00100000u,
 
 	/** Indicates that references to this class default to instanced. Used to be subclasses of UComponent, but now can be any UObject */
-	CLASS_DefaultToInstanced  = 0x00200000,
+	CLASS_DefaultToInstanced  = 0x00200000u,
 
 	/** Indicates that the parent token stream has been merged with ours. */
-	CLASS_TokenStreamAssembled  = 0x00400000,
+	CLASS_TokenStreamAssembled  = 0x00400000u,
 	/** Class has component properties. */
-	CLASS_HasInstancedReference= 0x00800000,
+	CLASS_HasInstancedReference= 0x00800000u,
 	/** Don't show this class in the editor class browser or edit inline new menus. */
-	CLASS_Hidden			  = 0x01000000,
+	CLASS_Hidden			  = 0x01000000u,
 	/** Don't save objects of this class when serializing */
-	CLASS_Deprecated		  = 0x02000000,
+	CLASS_Deprecated		  = 0x02000000u,
 	/** Class not shown in editor drop down for class selection */
-	CLASS_HideDropDown		  = 0x04000000,
+	CLASS_HideDropDown		  = 0x04000000u,
 	/** Class settings are saved to <AppData>/..../Blah.ini (as opposed to CLASS_DefaultConfig) */
-	CLASS_GlobalUserConfig	  = 0x08000000,
+	CLASS_GlobalUserConfig	  = 0x08000000u,
 	/** Class was declared directly in C++ and has no boilerplate generated by UnrealHeaderTool */
-	CLASS_Intrinsic			  = 0x10000000,
+	CLASS_Intrinsic			  = 0x10000000u,
 	/** Class has already been constructed (maybe in a previous DLL version before hot-reload). */
-	CLASS_Constructed		  = 0x20000000,
+	CLASS_Constructed		  = 0x20000000u,
 	/** Indicates that object configuration will not check against ini base/defaults when serialized */
-	CLASS_ConfigDoNotCheckDefaults = 0x40000000,
+	CLASS_ConfigDoNotCheckDefaults = 0x40000000u,
 	/** Class has been consigned to oblivion as part of a blueprint recompile, and a newer version currently exists. */
-	CLASS_NewerVersionExists  = 0x80000000,
+	CLASS_NewerVersionExists  = 0x80000000u,
 
 	//@}
 
@@ -287,10 +288,11 @@ enum EClassFlags
 		CLASS_MinimalAPI |
 		CLASS_RequiredAPI,
 
-	CLASS_AllFlags			= 0xFFFFFFFF,
+	CLASS_AllFlags			= 0xFFFFFFFFu,
 };
 
-
+// Declare bitwise operators to allow EClassFlags to be combined but still retain type safety
+ENUM_CLASS_FLAGS(EClassFlags);
 
 
 
@@ -814,6 +816,9 @@ namespace UF
 		/// This function is cosmetic and will not run on dedicated servers
 		BlueprintCosmetic,
 
+		/// Indicates that a Blueprint exposed function should not be exposed to the end user
+		BlueprintInternalUseOnly,
+	
 		/// This function can be called in the editor on selected instances via a button in the details panel.
 		CallInEditor,
 
@@ -968,6 +973,9 @@ namespace US
 
 		/// Exposes this struct as a type that can be used for variables in blueprints
 		BlueprintType,
+
+		/// Indicates that a BlueprintType struct should not be exposed to the end user
+		BlueprintInternalUseOnly
 	};
 }
 
@@ -1351,7 +1359,7 @@ public: \
 	{ \
 		return TPackage; \
 	} \
-	/** Returns the StaticClassFlags for this class */ \
+	/** Returns the static cast flags for this class */ \
 	inline static EClassCastFlags StaticClassCastFlags() \
 	{ \
 		return TStaticCastFlags; \
@@ -1496,7 +1504,7 @@ public: \
 				PrivateStaticClass, \
 				StaticRegisterNatives##TClass, \
 				sizeof(TClass), \
-				TClass::StaticClassFlags, \
+				(EClassFlags)TClass::StaticClassFlags, \
 				TClass::StaticClassCastFlags(), \
 				TClass::StaticConfigName(), \
 				(UClass::ClassConstructorType)InternalConstructor<TClass>, \
@@ -1556,7 +1564,7 @@ public: \
 			PrivateStaticClass, \
 			StaticRegisterNatives##TClass, \
 			sizeof(TClass), \
-			TClass::StaticClassFlags, \
+			(EClassFlags)TClass::StaticClassFlags, \
 			TClass::StaticClassCastFlags(), \
 			TClass::StaticConfigName(), \
 			(UClass::ClassConstructorType)InternalConstructor<TClass>, \
diff --git a/Engine/Source/Runtime/CoreUObject/Public/UObject/ObjectRedirector.h b/Engine/Source/Runtime/CoreUObject/Public/UObject/ObjectRedirector.h
--- a/Engine/Source/Runtime/CoreUObject/Public/UObject/ObjectRedirector.h
+++ b/Engine/Source/Runtime/CoreUObject/Public/UObject/ObjectRedirector.h
@@ -40,48 +40,3 @@ class UObjectRedirector : public UObject
 	 */
 	virtual bool GetNativePropertyValues( TMap<FString,FString>& out_PropertyValues, uint32 ExportFlags=0 ) const override;
 };
-
-/**
- * Callback structure to respond to redirect-followed callbacks to determine
- * if a redirector was followed from a single object name. Will auto
- * register and unregister the callback on construction/destruction
- */
-class COREUOBJECT_API FScopedRedirectorCatcher
-{
-public:
-	/**
-	 * Constructor for the callback device, will register itself with the RedirectorFollowed delegate
-	 *
-	 * @param InObjectPathNameToMatch Full pathname of the object refrence that is being compiled by script
-	 */
-	FScopedRedirectorCatcher(const FString& InObjectPathNameToMatch);
-
-	/**
-	 * Destructor. Will unregister the callback
-	 */
-	virtual ~FScopedRedirectorCatcher();
-
-	/**
-	 * Returns whether or not a redirector was followed for the ObjectPathName
-	 */
-	inline bool WasRedirectorFollowed()
-	{
-		return bWasRedirectorFollowed;
-	}
-
-	/**
-	 * Responds to FCoreDelegates::RedirectorFollowed. Records all followed redirections
-	 * so they can be cleaned later.
-	 *
-	 * @param InString Name of the package that pointed to the redirect
-	 * @param InObject The UObjectRedirector that was followed
-	 */
-	virtual void OnRedirectorFollowed( const FString& InString, UObject* InObject);
-
-private:
-	/** The full path name of the object that we want to match */
-	FString ObjectPathNameToMatch;
-
-	/** Was a redirector followed, ONLY for the ObjectPathNameToMatch? */
-	bool bWasRedirectorFollowed;
-};
diff --git a/Engine/Source/Runtime/CoreUObject/Public/UObject/PersistentObjectPtr.h b/Engine/Source/Runtime/CoreUObject/Public/UObject/PersistentObjectPtr.h
--- a/Engine/Source/Runtime/CoreUObject/Public/UObject/PersistentObjectPtr.h
+++ b/Engine/Source/Runtime/CoreUObject/Public/UObject/PersistentObjectPtr.h
@@ -31,6 +31,13 @@ public:
 		TagAtLastTest = 0;
 	}
 
+	/** Resets the weak ptr only, call this when ObjectId may change */
+	FORCEINLINE void ResetWeakPtr()
+	{
+		WeakPtr.Reset();
+		TagAtLastTest = 0;
+	}
+
 	/** Construct from another pointer of the same type */
 	FORCEINLINE TPersistentObjectPtr(const TPersistentObjectPtr& Other)
 	{
diff --git a/Engine/Source/Runtime/CoreUObject/Public/UObject/UObjectArray.h b/Engine/Source/Runtime/CoreUObject/Public/UObject/UObjectArray.h
--- a/Engine/Source/Runtime/CoreUObject/Public/UObject/UObjectArray.h
+++ b/Engine/Source/Runtime/CoreUObject/Public/UObject/UObjectArray.h
@@ -16,7 +16,9 @@
 * Controls whether the number of available elements is being tracked in the ObjObjects array.
 * By default it is only tracked in WITH_EDITOR builds as it adds a small amount of tracking overhead
 */
+#if !defined(UE_GC_TRACK_OBJ_AVAILABLE)
 #define UE_GC_TRACK_OBJ_AVAILABLE (WITH_EDITOR)
+#endif
 
 /**
 * Single item in the UObject array.
diff --git a/Engine/Source/Runtime/CoreUObject/Public/UObject/UObjectGlobals.h b/Engine/Source/Runtime/CoreUObject/Public/UObject/UObjectGlobals.h
--- a/Engine/Source/Runtime/CoreUObject/Public/UObject/UObjectGlobals.h
+++ b/Engine/Source/Runtime/CoreUObject/Public/UObject/UObjectGlobals.h
@@ -1866,11 +1866,7 @@ struct COREUOBJECT_API FCoreUObjectDelegates
 
 #endif	//WITH_EDITOR
 
-	// Delegate type for redirector followed events ( Params: const FString& PackageName, UObject* Redirector )
-	DECLARE_MULTICAST_DELEGATE_TwoParams(FOnRedirectorFollowed, const FString&, UObject*);
 
-	// Sent when a UObjectRedirector was followed to find the destination object
-	static FOnRedirectorFollowed RedirectorFollowed;
 
 	/** Delegate used by SavePackage() to create the package backup */
 	static FAutoPackageBackupDelegate AutoPackageBackupDelegate;
@@ -1925,10 +1921,6 @@ struct COREUOBJECT_API FCoreUObjectDelegates
 	/** Queries whether an object should be loaded on top ( replace ) an already existing one */
 	static FOnLoadObjectsOnTop ShouldLoadOnTop;
 
-	/** called when loading a string asset reference */
-	DECLARE_DELEGATE_OneParam(FStringAssetReferenceLoaded, const FString&);
-	static FStringAssetReferenceLoaded StringAssetReferenceLoaded;
-
 	/** called when loading a string asset reference */
 	DECLARE_MULTICAST_DELEGATE_OneParam(FPackageLoadedFromStringAssetReference, const FName&);
 	static FPackageLoadedFromStringAssetReference PackageLoadedFromStringAssetReference;
@@ -1937,13 +1929,21 @@ struct COREUOBJECT_API FCoreUObjectDelegates
 	DECLARE_MULTICAST_DELEGATE_OneParam(FPackageCreatedForLoad, class UPackage*);
 	static FPackageCreatedForLoad PackageCreatedForLoad;
 
-	/** called when saving a string asset reference, can replace the value with something else */
-	DECLARE_DELEGATE_RetVal_OneParam(FString, FStringAssetReferenceSaving, FString const& /*SavingAssetLongPathname*/);
-	static FStringAssetReferenceSaving StringAssetReferenceSaving;
-
 	/** Called when trying to figure out if a UObject is a primary asset, if it doesn't know */
 	DECLARE_DELEGATE_RetVal_OneParam(FPrimaryAssetId, FGetPrimaryAssetIdForObject, const UObject*);
 	static FGetPrimaryAssetIdForObject GetPrimaryAssetIdForObject;
+
+	DECLARE_DELEGATE_OneParam(FStringAssetReferenceLoaded, const FString&);
+	DEPRECATED(4.17, "StringAssetReferenceLoaded is deprecated, call FStringAssetReference::PostLoadPath instead")
+	static FStringAssetReferenceLoaded StringAssetReferenceLoaded;
+
+	DECLARE_DELEGATE_RetVal_OneParam(FString, FStringAssetReferenceSaving, FString const& /*SavingAssetLongPathname*/);
+	DEPRECATED(4.17, "StringAssetReferenceSaving is deprecated, call FStringAssetReference::PreSavePath instead")
+	static FStringAssetReferenceSaving StringAssetReferenceSaving;
+
+	DECLARE_MULTICAST_DELEGATE_TwoParams(FOnRedirectorFollowed, const FString&, UObject*);
+	DEPRECATED(4.17, "RedirectorFollowed is deprecated, FixeupRedirects was replaced with ResavePackages -FixupRedirect")
+	static FOnRedirectorFollowed RedirectorFollowed;
 };
 
 /** Allows release builds to override not verifying GC assumptions. Useful for profiling as it's hitchy. */
diff --git a/Engine/Source/Runtime/Engine/Classes/AI/Navigation/AvoidanceManager.h b/Engine/Source/Runtime/Engine/Classes/AI/Navigation/AvoidanceManager.h
--- a/Engine/Source/Runtime/Engine/Classes/AI/Navigation/AvoidanceManager.h
+++ b/Engine/Source/Runtime/Engine/Classes/AI/Navigation/AvoidanceManager.h
@@ -139,11 +139,9 @@ class ENGINE_API UAvoidanceManager : public UObject, public FSelfRegisteringExec
 	FVector GetAvoidanceVelocityForComponent(UMovementComponent* MovementComp);
 
 	/** Only use if you want manual velocity planning. Provide your AvoidanceUID in order to avoid colliding with yourself. */
-	UFUNCTION(BlueprintCallable, Category="AI")
 	FVector GetAvoidanceVelocityIgnoringUID(const FNavAvoidanceData& AvoidanceData, float DeltaTime, int32 IgnoreThisUID);
 
 	/** Only use if you want manual velocity planning. Will not ignore your own volume if you are registered. */
-	UFUNCTION(BlueprintCallable, Category="AI")
 	FVector GetAvoidanceVelocity(const FNavAvoidanceData& AvoidanceData, float DeltaTime);
 
 	/** Update the RVO avoidance data for the participating UMovementComponent */
diff --git a/Engine/Source/Runtime/Engine/Classes/AI/Navigation/NavLinkDefinition.h b/Engine/Source/Runtime/Engine/Classes/AI/Navigation/NavLinkDefinition.h
--- a/Engine/Source/Runtime/Engine/Classes/AI/Navigation/NavLinkDefinition.h
+++ b/Engine/Source/Runtime/Engine/Classes/AI/Navigation/NavLinkDefinition.h
@@ -175,7 +175,7 @@ struct TStructOpsTypeTraits< FNavigationLinkBase > : public TStructOpsTypeTraits
 	};
 };
 
-USTRUCT()
+USTRUCT(BlueprintType)
 struct ENGINE_API FNavigationLink : public FNavigationLinkBase
 {
 	GENERATED_USTRUCT_BODY()
@@ -183,7 +183,7 @@ struct ENGINE_API FNavigationLink : public FNavigationLinkBase
 	UPROPERTY(EditAnywhere, Category=Default, BlueprintReadWrite, meta=(MakeEditWidget=""))
 	FVector Left;
 
-	UPROPERTY(EditAnywhere, Category=Default, BlueprintReadWrite, meta=(MakeEditWidget=""))
+	UPROPERTY(EditAnywhere, Category=Default, meta=(MakeEditWidget=""))
 	FVector Right;
 
 	FNavigationLink()
@@ -242,16 +242,16 @@ struct ENGINE_API FNavigationSegmentLink : public FNavigationLinkBase
 {
 	GENERATED_USTRUCT_BODY()
 
-	UPROPERTY(EditAnywhere, Category=Default, BlueprintReadWrite, meta=(MakeEditWidget=""))
+	UPROPERTY(EditAnywhere, Category=Default, meta=(MakeEditWidget=""))
 	FVector LeftStart;
 
-	UPROPERTY(EditAnywhere, Category=Default, BlueprintReadWrite, meta=(MakeEditWidget=""))
+	UPROPERTY(EditAnywhere, Category=Default, meta=(MakeEditWidget=""))
 	FVector LeftEnd;
 
-	UPROPERTY(EditAnywhere, Category=Default, BlueprintReadWrite, meta=(MakeEditWidget=""))
+	UPROPERTY(EditAnywhere, Category=Default, meta=(MakeEditWidget=""))
 	FVector RightStart;
 
-	UPROPERTY(EditAnywhere, Category=Default, BlueprintReadWrite, meta=(MakeEditWidget=""))
+	UPROPERTY(EditAnywhere, Category=Default, meta=(MakeEditWidget=""))
 	FVector RightEnd;
 
 	FNavigationSegmentLink() 
diff --git a/Engine/Source/Runtime/Engine/Classes/AI/Navigation/NavigationSystem.h b/Engine/Source/Runtime/Engine/Classes/AI/Navigation/NavigationSystem.h
--- a/Engine/Source/Runtime/Engine/Classes/AI/Navigation/NavigationSystem.h
+++ b/Engine/Source/Runtime/Engine/Classes/AI/Navigation/NavigationSystem.h
@@ -262,36 +262,36 @@ public:
 	// Blueprint functions
 	//----------------------------------------------------------------------//
 
-	UFUNCTION(BlueprintPure, Category = "AI|Navigation", meta = (WorldContext = "WorldContext"))
-	static UNavigationSystem* GetNavigationSystem(UObject* WorldContext);
+	UFUNCTION(BlueprintPure, Category = "AI|Navigation", meta = (WorldContext = "WorldContextObject"))
+	static UNavigationSystem* GetNavigationSystem(UObject* WorldContextObject);
 	
 	/** Project a point onto the NavigationData */
-	UFUNCTION(BlueprintPure, Category = "AI|Navigation", meta = (WorldContext = "WorldContext", DisplayName = "ProjectPointToNavigation"))
-	static bool K2_ProjectPointToNavigation(UObject* WorldContext, const FVector& Point, FVector& ProjectedLocation, ANavigationData* NavData, TSubclassOf<UNavigationQueryFilter> FilterClass, const FVector QueryExtent = FVector::ZeroVector);
+	UFUNCTION(BlueprintPure, Category = "AI|Navigation", meta = (WorldContext = "WorldContextObject", DisplayName = "ProjectPointToNavigation"))
+	static bool K2_ProjectPointToNavigation(UObject* WorldContextObject, const FVector& Point, FVector& ProjectedLocation, ANavigationData* NavData, TSubclassOf<UNavigationQueryFilter> FilterClass, const FVector QueryExtent = FVector::ZeroVector);
 
 	/** Generates a random location reachable from given Origin location.
 	 *	@return Return Value represents if the call was successful */
-	UFUNCTION(BlueprintPure, Category = "AI|Navigation", meta = (WorldContext = "WorldContext", DisplayName = "GetRandomReachablePointInRadius"))
-	static bool K2_GetRandomReachablePointInRadius(UObject* WorldContext, const FVector& Origin, FVector& RandomLocation, float Radius, ANavigationData* NavData = NULL, TSubclassOf<UNavigationQueryFilter> FilterClass = NULL);
+	UFUNCTION(BlueprintPure, Category = "AI|Navigation", meta = (WorldContext = "WorldContextObject", DisplayName = "GetRandomReachablePointInRadius"))
+	static bool K2_GetRandomReachablePointInRadius(UObject* WorldContextObject, const FVector& Origin, FVector& RandomLocation, float Radius, ANavigationData* NavData = NULL, TSubclassOf<UNavigationQueryFilter> FilterClass = NULL);
 
 	/** Generates a random location in navigable space within given radius of Origin.
 	 *	@return Return Value represents if the call was successful */
-	UFUNCTION(BlueprintPure, Category = "AI|Navigation", meta = (WorldContext = "WorldContext", DisplayName = "GetRandomPointInNavigableRadius"))
-	static bool K2_GetRandomPointInNavigableRadius(UObject* WorldContext, const FVector& Origin, FVector& RandomLocation, float Radius, ANavigationData* NavData = NULL, TSubclassOf<UNavigationQueryFilter> FilterClass = NULL);
+	UFUNCTION(BlueprintPure, Category = "AI|Navigation", meta = (WorldContext = "WorldContextObject", DisplayName = "GetRandomPointInNavigableRadius"))
+	static bool K2_GetRandomPointInNavigableRadius(UObject* WorldContextObject, const FVector& Origin, FVector& RandomLocation, float Radius, ANavigationData* NavData = NULL, TSubclassOf<UNavigationQueryFilter> FilterClass = NULL);
 	
 	/** Potentially expensive. Use with caution. Consider using UPathFollowingComponent::GetRemainingPathCost instead */
-	UFUNCTION(BlueprintPure, Category="AI|Navigation", meta=(WorldContext="WorldContext" ) )
-	static ENavigationQueryResult::Type GetPathCost(UObject* WorldContext, const FVector& PathStart, const FVector& PathEnd, float& PathCost, ANavigationData* NavData = NULL, TSubclassOf<UNavigationQueryFilter> FilterClass = NULL);
+	UFUNCTION(BlueprintPure, Category="AI|Navigation", meta=(WorldContext="WorldContextObject" ) )
+	static ENavigationQueryResult::Type GetPathCost(UObject* WorldContextObject, const FVector& PathStart, const FVector& PathEnd, float& PathCost, ANavigationData* NavData = NULL, TSubclassOf<UNavigationQueryFilter> FilterClass = NULL);
 
 	/** Potentially expensive. Use with caution */
-	UFUNCTION(BlueprintPure, Category="AI|Navigation", meta=(WorldContext="WorldContext" ) )
-	static ENavigationQueryResult::Type GetPathLength(UObject* WorldContext, const FVector& PathStart, const FVector& PathEnd, float& PathLength, ANavigationData* NavData = NULL, TSubclassOf<UNavigationQueryFilter> FilterClass = NULL);
+	UFUNCTION(BlueprintPure, Category="AI|Navigation", meta=(WorldContext="WorldContextObject" ) )
+	static ENavigationQueryResult::Type GetPathLength(UObject* WorldContextObject, const FVector& PathStart, const FVector& PathEnd, float& PathLength, ANavigationData* NavData = NULL, TSubclassOf<UNavigationQueryFilter> FilterClass = NULL);
 
-	UFUNCTION(BlueprintPure, Category="AI|Navigation", meta=(WorldContext="WorldContext" ) )
-	static bool IsNavigationBeingBuilt(UObject* WorldContext);
+	UFUNCTION(BlueprintPure, Category="AI|Navigation", meta=(WorldContext="WorldContextObject" ) )
+	static bool IsNavigationBeingBuilt(UObject* WorldContextObject);
 
-	UFUNCTION(BlueprintPure, Category = "AI|Navigation", meta = (WorldContext = "WorldContext"))
-	static bool IsNavigationBeingBuiltOrLocked(UObject* WorldContext);
+	UFUNCTION(BlueprintPure, Category = "AI|Navigation", meta = (WorldContext = "WorldContextObject"))
+	static bool IsNavigationBeingBuiltOrLocked(UObject* WorldContextObject);
 
 	UFUNCTION(BlueprintCallable, Category="AI|Navigation")
 	static void SimpleMoveToActor(AController* Controller, const AActor* Goal);
@@ -301,21 +301,21 @@ public:
 
 	/** Finds path instantly, in a FindPath Synchronously. 
 	 *	@param PathfindingContext could be one of following: NavigationData (like Navmesh actor), Pawn or Controller. This parameter determines parameters of specific pathfinding query */
-	UFUNCTION(BlueprintCallable, Category = "AI|Navigation", meta = (WorldContext="WorldContext"))
-	static UNavigationPath* FindPathToLocationSynchronously(UObject* WorldContext, const FVector& PathStart, const FVector& PathEnd, AActor* PathfindingContext = NULL, TSubclassOf<UNavigationQueryFilter> FilterClass = NULL);
+	UFUNCTION(BlueprintCallable, Category = "AI|Navigation", meta = (WorldContext="WorldContextObject"))
+	static UNavigationPath* FindPathToLocationSynchronously(UObject* WorldContextObject, const FVector& PathStart, const FVector& PathEnd, AActor* PathfindingContext = NULL, TSubclassOf<UNavigationQueryFilter> FilterClass = NULL);
 
 	/** Finds path instantly, in a FindPath Synchronously. Main advantage over FindPathToLocationSynchronously is that 
 	 *	the resulting path will automatically get updated if goal actor moves more than TetherDistance away from last path node
 	 *	@param PathfindingContext could be one of following: NavigationData (like Navmesh actor), Pawn or Controller. This parameter determines parameters of specific pathfinding query */
-	UFUNCTION(BlueprintCallable, Category = "AI|Navigation", meta = (WorldContext="WorldContext"))
-	static UNavigationPath* FindPathToActorSynchronously(UObject* WorldContext, const FVector& PathStart, AActor* GoalActor, float TetherDistance = 50.f, AActor* PathfindingContext = NULL, TSubclassOf<UNavigationQueryFilter> FilterClass = NULL);
+	UFUNCTION(BlueprintCallable, Category = "AI|Navigation", meta = (WorldContext="WorldContextObject"))
+	static UNavigationPath* FindPathToActorSynchronously(UObject* WorldContextObject, const FVector& PathStart, AActor* GoalActor, float TetherDistance = 50.f, AActor* PathfindingContext = NULL, TSubclassOf<UNavigationQueryFilter> FilterClass = NULL);
 
 	/** Performs navigation raycast on NavigationData appropriate for given Querier.
 	 *	@param Querier if not passed default navigation data will be used
 	 *	@param HitLocation if line was obstructed this will be set to hit location. Otherwise it contains SegmentEnd
 	 *	@return true if line from RayStart to RayEnd was obstructed. Also, true when no navigation data present */
-	UFUNCTION(BlueprintCallable, Category="AI|Navigation", meta=(WorldContext="WorldContext" ))
-	static bool NavigationRaycast(UObject* WorldContext, const FVector& RayStart, const FVector& RayEnd, FVector& HitLocation, TSubclassOf<UNavigationQueryFilter> FilterClass = NULL, AController* Querier = NULL);
+	UFUNCTION(BlueprintCallable, Category="AI|Navigation", meta=(WorldContext="WorldContextObject" ))
+	static bool NavigationRaycast(UObject* WorldContextObject, const FVector& RayStart, const FVector& RayEnd, FVector& HitLocation, TSubclassOf<UNavigationQueryFilter> FilterClass = NULL, AController* Querier = NULL);
 
 	/** will limit the number of simultaneously running navmesh tile generation jobs to specified number.
 	 *	@param MaxNumberOfJobs gets trimmed to be at least 1. You cannot use this function to pause navmesh generation */
@@ -1000,13 +1000,13 @@ public:
 	DEPRECATED(4.11, "UpdateNavOctreeAll is deprecated. Use UpdateActorAndComponentsInNavOctree")
 	static void UpdateNavOctreeAll(AActor* Actor);
 	DEPRECATED(4.16, "This version of ProjectPointToNavigation is deprecated. Please use the new version")
-	UFUNCTION(BlueprintPure, Category = "AI|Navigation", meta = (WorldContext = "WorldContext", DisplayName = "ProjectPointToNavigation_DEPRECATED", DeprecatedFunction, DeprecationMessage = "This version of ProjectPointToNavigation is deprecated. Please use the new version"))
-	static FVector ProjectPointToNavigation(UObject* WorldContext, const FVector& Point, ANavigationData* NavData = NULL, TSubclassOf<UNavigationQueryFilter> FilterClass = NULL, const FVector QueryExtent = FVector::ZeroVector);
+	UFUNCTION(BlueprintPure, Category = "AI|Navigation", meta = (WorldContext = "WorldContextObject", DisplayName = "ProjectPointToNavigation_DEPRECATED", DeprecatedFunction, DeprecationMessage = "This version of ProjectPointToNavigation is deprecated. Please use the new version"))
+	static FVector ProjectPointToNavigation(UObject* WorldContextObject, const FVector& Point, ANavigationData* NavData = NULL, TSubclassOf<UNavigationQueryFilter> FilterClass = NULL, const FVector QueryExtent = FVector::ZeroVector);
 	DEPRECATED(4.16, "This version of GetRandomReachablePointInRadius is deprecated. Please use the new version")
-	UFUNCTION(BlueprintPure, Category = "AI|Navigation", meta = (WorldContext = "WorldContext", DisplayName = "GetRandomReachablePointInRadius_DEPRECATED", DeprecatedFunction, DeprecationMessage = "This version of GetRandomReachablePointInRadius is deprecated. Please use the new version"))
-	static FVector GetRandomReachablePointInRadius(UObject* WorldContext, const FVector& Origin, float Radius, ANavigationData* NavData = NULL, TSubclassOf<UNavigationQueryFilter> FilterClass = NULL);
+	UFUNCTION(BlueprintPure, Category = "AI|Navigation", meta = (WorldContext = "WorldContextObject", DisplayName = "GetRandomReachablePointInRadius_DEPRECATED", DeprecatedFunction, DeprecationMessage = "This version of GetRandomReachablePointInRadius is deprecated. Please use the new version"))
+	static FVector GetRandomReachablePointInRadius(UObject* WorldContextObject, const FVector& Origin, float Radius, ANavigationData* NavData = NULL, TSubclassOf<UNavigationQueryFilter> FilterClass = NULL);
 	DEPRECATED(4.16, "This version of GetRandomPointInNavigableRadius is deprecated. Please use the new version")
-	UFUNCTION(BlueprintPure, Category = "AI|Navigation", meta = (WorldContext = "WorldContext", DisplayName = "GetRandomPointInNavigableRadius_DEPRECATED", DeprecatedFunction, DeprecationMessage = "This version of GetRandomPointInNavigableRadius is deprecated. Please use the new version"))
-	static FVector GetRandomPointInNavigableRadius(UObject* WorldContext, const FVector& Origin, float Radius, ANavigationData* NavData = NULL, TSubclassOf<UNavigationQueryFilter> FilterClass = NULL);
+	UFUNCTION(BlueprintPure, Category = "AI|Navigation", meta = (WorldContext = "WorldContextObject", DisplayName = "GetRandomPointInNavigableRadius_DEPRECATED", DeprecatedFunction, DeprecationMessage = "This version of GetRandomPointInNavigableRadius is deprecated. Please use the new version"))
+	static FVector GetRandomPointInNavigableRadius(UObject* WorldContextObject, const FVector& Origin, float Radius, ANavigationData* NavData = NULL, TSubclassOf<UNavigationQueryFilter> FilterClass = NULL);
 };
 
diff --git a/Engine/Source/Runtime/Engine/Classes/Animation/AnimCompositeBase.h b/Engine/Source/Runtime/Engine/Classes/Animation/AnimCompositeBase.h
--- a/Engine/Source/Runtime/Engine/Classes/Animation/AnimCompositeBase.h
+++ b/Engine/Source/Runtime/Engine/Classes/Animation/AnimCompositeBase.h
@@ -60,26 +60,26 @@ struct FAnimSegment
 	GENERATED_USTRUCT_BODY()
 
 	/** Anim Reference to play - only allow AnimSequence or AnimComposite **/
-	UPROPERTY(EditAnywhere, BlueprintReadWrite, Category=AnimSegment)
+	UPROPERTY(EditAnywhere, Category=AnimSegment)
 	UAnimSequenceBase* AnimReference;
 
 	/** Start Pos within this AnimCompositeBase */
-	UPROPERTY(VisibleAnywhere, BlueprintReadWrite, Category=AnimSegment)
+	UPROPERTY(VisibleAnywhere, Category=AnimSegment)
 	float	StartPos;
 
 	/** Time to start playing AnimSequence at. */
-	UPROPERTY(EditAnywhere, BlueprintReadWrite, Category=AnimSegment)
+	UPROPERTY(EditAnywhere, Category=AnimSegment)
 	float	AnimStartTime;
 
 	/** Time to end playing the AnimSequence at. */
-	UPROPERTY(EditAnywhere, BlueprintReadWrite, Category=AnimSegment)
+	UPROPERTY(EditAnywhere, Category=AnimSegment)
 	float	AnimEndTime;
 
 	/** Playback speed of this animation. If you'd like to reverse, set -1*/
-	UPROPERTY(EditAnywhere, BlueprintReadWrite, Category=AnimSegment)
+	UPROPERTY(EditAnywhere, Category=AnimSegment)
 	float	AnimPlayRate;
 
-	UPROPERTY(EditAnywhere, BlueprintReadWrite, Category=AnimSegment)
+	UPROPERTY(EditAnywhere, Category=AnimSegment)
 	int32		LoopingCount;
 
 	FAnimSegment()
@@ -187,7 +187,7 @@ struct FAnimTrack
 {
 	GENERATED_USTRUCT_BODY()
 
-	UPROPERTY(EditAnywhere, BlueprintReadWrite, Category=AnimTrack, EditFixedSize)
+	UPROPERTY(EditAnywhere, Category=AnimTrack, EditFixedSize)
 	TArray<FAnimSegment>	AnimSegments;
 
 	FAnimTrack() {}
diff --git a/Engine/Source/Runtime/Engine/Classes/Animation/AnimNodeBase.h b/Engine/Source/Runtime/Engine/Classes/Animation/AnimNodeBase.h
--- a/Engine/Source/Runtime/Engine/Classes/Animation/AnimNodeBase.h
+++ b/Engine/Source/Runtime/Engine/Classes/Animation/AnimNodeBase.h
@@ -332,7 +332,7 @@ namespace EPinHidingMode
 #define ENABLE_ANIMGRAPH_TRAVERSAL_DEBUG 0
 
 /** A pose link to another node */
-USTRUCT()
+USTRUCT(BlueprintInternalUseOnly)
 struct ENGINE_API FPoseLinkBase
 {
 	GENERATED_USTRUCT_BODY()
@@ -390,7 +390,7 @@ public:
 #define ENABLE_ANIMNODE_POSE_DEBUG 0
 
 /** A local-space pose link to another node */
-USTRUCT()
+USTRUCT(BlueprintInternalUseOnly)
 struct ENGINE_API FPoseLink : public FPoseLinkBase
 {
 	GENERATED_USTRUCT_BODY()
@@ -407,7 +407,7 @@ private:
 };
 
 /** A component-space pose link to another node */
-USTRUCT()
+USTRUCT(BlueprintInternalUseOnly)
 struct ENGINE_API FComponentSpacePoseLink : public FPoseLinkBase
 {
 	GENERATED_USTRUCT_BODY()
diff --git a/Engine/Source/Runtime/Engine/Classes/Animation/AnimNodeSpaceConversions.h b/Engine/Source/Runtime/Engine/Classes/Animation/AnimNodeSpaceConversions.h
--- a/Engine/Source/Runtime/Engine/Classes/Animation/AnimNodeSpaceConversions.h
+++ b/Engine/Source/Runtime/Engine/Classes/Animation/AnimNodeSpaceConversions.h
@@ -7,7 +7,7 @@
 #include "Animation/AnimNodeBase.h"
 #include "AnimNodeSpaceConversions.generated.h"
 
-USTRUCT()
+USTRUCT(BlueprintInternalUseOnly)
 struct ENGINE_API FAnimNode_ConvertComponentToLocalSpace : public FAnimNode_Base
 {
 	GENERATED_USTRUCT_BODY()
@@ -29,7 +29,7 @@ public:
 
 };
 
-USTRUCT()
+USTRUCT(BlueprintInternalUseOnly)
 struct ENGINE_API FAnimNode_ConvertLocalToComponentSpace : public FAnimNode_Base
 {
 	GENERATED_USTRUCT_BODY()
diff --git a/Engine/Source/Runtime/Engine/Classes/Components/BoxComponent.h b/Engine/Source/Runtime/Engine/Classes/Components/BoxComponent.h
--- a/Engine/Source/Runtime/Engine/Classes/Components/BoxComponent.h
+++ b/Engine/Source/Runtime/Engine/Classes/Components/BoxComponent.h
@@ -63,7 +63,7 @@ public:
 
 FORCEINLINE FVector UBoxComponent::GetScaledBoxExtent() const
 {
-	return BoxExtent * ComponentToWorld.GetScale3D();
+	return BoxExtent * GetComponentTransform().GetScale3D();
 }
 
 FORCEINLINE FVector UBoxComponent::GetUnscaledBoxExtent() const
diff --git a/Engine/Source/Runtime/Engine/Classes/Components/CapsuleComponent.h b/Engine/Source/Runtime/Engine/Classes/Components/CapsuleComponent.h
--- a/Engine/Source/Runtime/Engine/Classes/Components/CapsuleComponent.h
+++ b/Engine/Source/Runtime/Engine/Classes/Components/CapsuleComponent.h
@@ -255,6 +255,6 @@ FORCEINLINE void UCapsuleComponent::GetUnscaledCapsuleSize_WithoutHemisphere(flo
 
 FORCEINLINE float UCapsuleComponent::GetShapeScale() const
 {
-	return ComponentToWorld.GetMinimumAxisScale();
+	return GetComponentTransform().GetMinimumAxisScale();
 }
 
diff --git a/Engine/Source/Runtime/Engine/Classes/Components/ChildActorComponent.h b/Engine/Source/Runtime/Engine/Classes/Components/ChildActorComponent.h
--- a/Engine/Source/Runtime/Engine/Classes/Components/ChildActorComponent.h
+++ b/Engine/Source/Runtime/Engine/Classes/Components/ChildActorComponent.h
@@ -80,6 +80,7 @@ public:
 #if WITH_EDITOR
 	virtual void PostEditChangeProperty(FPropertyChangedEvent& PropertyChangedEvent) override;
 	virtual void PostEditChangeChainProperty(FPropertyChangedChainEvent& PropertyChangedEvent) override;
+	virtual void PostEditImport() override;
 	virtual void PostEditUndo() override;
 	virtual void PostLoad() override;
 #endif
diff --git a/Engine/Source/Runtime/Engine/Classes/Components/DirectionalLightComponent.h b/Engine/Source/Runtime/Engine/Classes/Components/DirectionalLightComponent.h
--- a/Engine/Source/Runtime/Engine/Classes/Components/DirectionalLightComponent.h
+++ b/Engine/Source/Runtime/Engine/Classes/Components/DirectionalLightComponent.h
@@ -127,7 +127,7 @@ class ENGINE_API UDirectionalLightComponent : public ULightComponent
 	float TraceDistance;
 
 	/** The Lightmass settings for this object. */
-	UPROPERTY(EditAnywhere, BlueprintReadOnly, Category=Light, meta=(ShowOnlyInnerProperties))
+	UPROPERTY(EditAnywhere, Category=Light, meta=(ShowOnlyInnerProperties))
 	struct FLightmassDirectionalLightSettings LightmassSettings;
 
 	/**
diff --git a/Engine/Source/Runtime/Engine/Classes/Components/PointLightComponent.h b/Engine/Source/Runtime/Engine/Classes/Components/PointLightComponent.h
--- a/Engine/Source/Runtime/Engine/Classes/Components/PointLightComponent.h
+++ b/Engine/Source/Runtime/Engine/Classes/Components/PointLightComponent.h
@@ -61,7 +61,7 @@ class ENGINE_API UPointLightComponent : public ULightComponent
 	float SourceLength;
 
 	/** The Lightmass settings for this object. */
-	UPROPERTY(EditAnywhere, BlueprintReadOnly, Category=Light, meta=(ShowOnlyInnerProperties))
+	UPROPERTY(EditAnywhere, Category=Light, meta=(ShowOnlyInnerProperties))
 	struct FLightmassPointLightSettings LightmassSettings;
 
 	UFUNCTION(BlueprintCallable, Category="Rendering|Lighting")
diff --git a/Engine/Source/Runtime/Engine/Classes/Components/PrimitiveComponent.h b/Engine/Source/Runtime/Engine/Classes/Components/PrimitiveComponent.h
--- a/Engine/Source/Runtime/Engine/Classes/Components/PrimitiveComponent.h
+++ b/Engine/Source/Runtime/Engine/Classes/Components/PrimitiveComponent.h
@@ -736,7 +736,6 @@ public:
 	void GetOverlappingComponents(TArray<UPrimitiveComponent*>& InOverlappingComponents) const;
 
 	/** Returns list of components this component is overlapping. */
-	UFUNCTION(BlueprintCallable, Category="Collision", meta=(UnsafeDuringActorConstruction="true"))
 	const TArray<FOverlapInfo>& GetOverlapInfos() const;
 
 	/** 
@@ -1591,7 +1590,7 @@ protected:
 	 * accelerate simulation
 	 *
 	 * @param : UseBI - root body instsance
-	 * @return : New ComponentToWorld to use
+	 * @return : New GetComponentTransform() to use
 	 */
 	virtual FTransform GetComponentTransformFromBodyInstance(FBodyInstance* UseBI);
 public:
diff --git a/Engine/Source/Runtime/Engine/Classes/Components/SceneCaptureComponent.h b/Engine/Source/Runtime/Engine/Classes/Components/SceneCaptureComponent.h
--- a/Engine/Source/Runtime/Engine/Classes/Components/SceneCaptureComponent.h
+++ b/Engine/Source/Runtime/Engine/Classes/Components/SceneCaptureComponent.h
@@ -25,7 +25,7 @@ struct FSceneCaptureViewInfo
 	EStereoscopicPass StereoPass;
 };
 
-USTRUCT()
+USTRUCT(BlueprintType)
 struct FEngineShowFlagsSetting
 {
 	GENERATED_USTRUCT_BODY()
@@ -80,7 +80,7 @@ class ENGINE_API USceneCaptureComponent : public USceneComponent
 	int32 CaptureSortPriority;
 
 	/** ShowFlags for the SceneCapture's ViewFamily, to control rendering settings for this view. Hidden but accessible through details customization */
-	UPROPERTY(EditAnywhere, BlueprintReadWrite, interp, Category=SceneComponent)
+	UPROPERTY(EditAnywhere, interp, Category=SceneComponent)
 	TArray<struct FEngineShowFlagsSetting> ShowFlagSettings;
 
 	// TODO: Make this a UStruct to set directly?
diff --git a/Engine/Source/Runtime/Engine/Classes/Components/SceneComponent.h b/Engine/Source/Runtime/Engine/Classes/Components/SceneComponent.h
--- a/Engine/Source/Runtime/Engine/Classes/Components/SceneComponent.h
+++ b/Engine/Source/Runtime/Engine/Classes/Components/SceneComponent.h
@@ -18,11 +18,8 @@ class USceneComponent;
 struct FLevelCollection;
 
 /** Overlap info consisting of the primitive and the body that is overlapping */
-USTRUCT()
 struct ENGINE_API FOverlapInfo
 {
-	GENERATED_USTRUCT_BODY()
-
 	FOverlapInfo()
 	{}
 
@@ -135,18 +132,11 @@ private:
 	UPROPERTY(ReplicatedUsing = OnRep_AttachSocketName)
 	FName AttachSocketName;
 
-public:
-	DEPRECATED(4.11, "SceneComponent no longer supports custom locations.")
-	uint32 bRequiresCustomLocation:1;
-
-protected:
-
 	/** True if we have ever updated ComponentToWorld based on RelativeLocation/Rotation/Scale. Used at startup to make sure it is initialized. */
 	UPROPERTY(Transient)
 	uint32 bWorldToComponentUpdated : 1;
 
 public:
-
 	/** If RelativeLocation should be considered relative to the world, rather than the parent */
 	UPROPERTY(EditAnywhere, AdvancedDisplay, BlueprintReadWrite, ReplicatedUsing=OnRep_Transform, Category=Transform)
 	uint32 bAbsoluteLocation:1;
@@ -210,10 +200,11 @@ public:
 	FBoxSphereBounds Bounds;
 
 	/** Current transform of the component, relative to the world */
+	DEPRECATED_FORGAME(4.17, "ComponentToWorld will be made private, use GetComponentTransform() instead.")
 	FTransform ComponentToWorld;
 
 private:
-	/** Cache that avoids Quat<->Rotator conversions if possible. Only to be used with ComponentToWorld.GetRotation(). */
+	/** Cache that avoids Quat<->Rotator conversions if possible. Only to be used with GetComponentTransform().GetRotation(). */
 	FRotationConversionCache WorldRotationCache;
 
 public:
@@ -622,6 +613,7 @@ public:
 	 * @param  bWeldSimulatedBodies Whether to weld together simulated physics bodies.
 	 * @return True if attachment is successful (or already attached to requested parent/socket), false if attachment is rejected and there is no change in AttachParent.
 	*/
+	DEPRECATED(4.17, "This function is deprecated, please use AttachToComponent instead.")
 	UFUNCTION(BlueprintCallable, Category = "Utilities|Transformation", meta = (DisplayName = "AttachTo (Deprecated)", AttachType = "KeepRelativeOffset"))
 	bool K2_AttachTo(USceneComponent* InParent, FName InSocketName = NAME_None, EAttachLocation::Type AttachType = EAttachLocation::KeepRelativeOffset, bool bWeldSimulatedBodies = true);
 
@@ -649,7 +641,8 @@ public:
 	bool K2_AttachToComponent(USceneComponent* Parent, FName SocketName, EAttachmentRule LocationRule, EAttachmentRule RotationRule, EAttachmentRule ScaleRule, bool bWeldSimulatedBodies);
 
 	/** Zeroes out the relative transform of the component, and calls AttachTo(). Useful for attaching directly to a scene component or socket location  */
-	UFUNCTION(BlueprintCallable, meta=(DeprecatedFunction, DeprecationMessage = "Use AttachToComponent instead"), Category="Utilities|Transformation")
+	DEPRECATED(4.17, "Use AttachToComponent instead.")
+	UFUNCTION(BlueprintCallable, meta=(DeprecatedFunction, DeprecationMessage = "Use AttachToComponent instead."), Category="Utilities|Transformation")
 	bool SnapTo(USceneComponent* InParent, FName InSocketName = NAME_None);
 
 	/** 
@@ -948,44 +941,53 @@ public:
 	/** Return location of the component, in world space */
 	FORCEINLINE FVector GetComponentLocation() const
 	{
-		return ComponentToWorld.GetLocation();
-	}
-
-	DEPRECATED(4.11, "Custom locations on SceneComponent are now deprecated.")
-	virtual FVector GetCustomLocation() const
-	{
-		check(0);
-		return ComponentToWorld.GetLocation();
+		return GetComponentTransform().GetLocation();
 	}
 
 	/** Return rotation of the component, in world space */
 	FORCEINLINE FRotator GetComponentRotation() const
 	{
-		return WorldRotationCache.NormalizedQuatToRotator(ComponentToWorld.GetRotation());
+		return WorldRotationCache.NormalizedQuatToRotator(GetComponentTransform().GetRotation());
 	}
 
 	/** Return rotation quaternion of the component, in world space */
 	FORCEINLINE FQuat GetComponentQuat() const
 	{
-		return ComponentToWorld.GetRotation();
+		return GetComponentTransform().GetRotation();
 	}
 
 	/** Return scale of the component, in world space */
 	FORCEINLINE FVector GetComponentScale() const
 	{
-		return ComponentToWorld.GetScale3D();
+		return GetComponentTransform().GetScale3D();
 	}
 
-	/** Get the current component-to-world transform for this component */
-	FORCEINLINE FTransform GetComponentToWorld() const 
+	/** Sets the cached component to world directly. This should be used very rarely. */
+	FORCEINLINE void SetComponentToWorld(const FTransform& NewComponentToWorld)
+	{
+		bWorldToComponentUpdated = true;
+		PRAGMA_DISABLE_DEPRECATION_WARNINGS
+		ComponentToWorld = NewComponentToWorld;
+		PRAGMA_ENABLE_DEPRECATION_WARNINGS
+	}
+
+	/** 
+	 * Get the current component-to-world transform for this component 
+	 * TODO: probably deprecate this in favor of GetComponentTransform
+	 */
+	FORCEINLINE const FTransform& GetComponentToWorld() const 
 	{ 
-		return ComponentToWorld; // TODO: probably deprecate this in favor of GetComponentTransform
+		PRAGMA_DISABLE_DEPRECATION_WARNINGS
+		return ComponentToWorld;
+		PRAGMA_ENABLE_DEPRECATION_WARNINGS
 	}
 
 	/** Get the current component-to-world transform for this component */
-	FORCEINLINE FTransform GetComponentTransform() const
+	FORCEINLINE const FTransform& GetComponentTransform() const
 	{
+		PRAGMA_DISABLE_DEPRECATION_WARNINGS
 		return ComponentToWorld;
+		PRAGMA_ENABLE_DEPRECATION_WARNINGS
 	}
 
 	/** Update transforms of any components attached to this one. */
diff --git a/Engine/Source/Runtime/Engine/Classes/Components/SphereComponent.h b/Engine/Source/Runtime/Engine/Classes/Components/SphereComponent.h
--- a/Engine/Source/Runtime/Engine/Classes/Components/SphereComponent.h
+++ b/Engine/Source/Runtime/Engine/Classes/Components/SphereComponent.h
@@ -80,5 +80,5 @@ FORCEINLINE float USphereComponent::GetUnscaledSphereRadius() const
 
 FORCEINLINE float USphereComponent::GetShapeScale() const
 {
-	return ComponentToWorld.GetMinimumAxisScale();
+	return GetComponentTransform().GetMinimumAxisScale();
 }
diff --git a/Engine/Source/Runtime/Engine/Classes/Components/StaticMeshComponent.h b/Engine/Source/Runtime/Engine/Classes/Components/StaticMeshComponent.h
--- a/Engine/Source/Runtime/Engine/Classes/Components/StaticMeshComponent.h
+++ b/Engine/Source/Runtime/Engine/Classes/Components/StaticMeshComponent.h
@@ -313,7 +313,7 @@ class ENGINE_API UStaticMeshComponent : public UMeshComponent
 #endif
 
 	/** The Lightmass settings for this object. */
-	UPROPERTY(EditAnywhere, BlueprintReadWrite, Category=Lighting)
+	UPROPERTY(EditAnywhere, Category=Lighting)
 	struct FLightmassPrimitiveSettings LightmassSettings;
 
 	virtual ~UStaticMeshComponent();
diff --git a/Engine/Source/Runtime/Engine/Classes/EdGraph/EdGraphNode.h b/Engine/Source/Runtime/Engine/Classes/EdGraph/EdGraphNode.h
--- a/Engine/Source/Runtime/Engine/Classes/EdGraph/EdGraphNode.h
+++ b/Engine/Source/Runtime/Engine/Classes/EdGraph/EdGraphNode.h
@@ -75,6 +75,16 @@ enum EEdGraphPinDirection
 	EGPD_MAX,
 };
 
+/** Enum used to define what container type a pin represents. */
+UENUM()
+enum class EPinContainerType : uint8
+{
+	None,
+	Array,
+	Set,
+	Map
+};
+
 /** Enum to indicate what sort of title we want. */
 UENUM()
 namespace ENodeTitleType
@@ -176,20 +186,55 @@ class ENGINE_API UEdGraphNode : public UObject
 	UPROPERTY()
 	int32 NodeHeight;
 
+	/** Enum to indicate if a node has advanced-display-pins, and if they are shown */
+	UPROPERTY()
+	TEnumAsByte<ENodeAdvancedPins::Type> AdvancedPinDisplay;
+
+	/** Indicates in what state the node is enabled, which may eliminate it from being compiled */
+	UPROPERTY()
+	ENodeEnabledState EnabledState;
+
+	/** Indicates whether or not the user explicitly set the enabled state */
+	UPROPERTY()
+	uint8 bUserSetEnabledState:1;
+
+protected:
+	/** (DEPRECATED) Value used for AllowSplitPins(). Do not override. */
+	uint8 bAllowSplitPins_DEPRECATED:1;
+
+private:
+	/** (DEPRECATED) FALSE if the node is a disabled, which eliminates it from being compiled */
+	UPROPERTY()
+	uint8 bIsNodeEnabled_DEPRECATED:1;
+
+public:
+
 #if WITH_EDITORONLY_DATA
 	/** If true, this node can be resized and should be drawn with a resize handle */
 	UPROPERTY()
-	uint32 bCanResizeNode:1;
+	uint8 bCanResizeNode:1;
 #endif // WITH_EDITORONLY_DATA
 
 	/** Flag to check for compile error/warning */
 	UPROPERTY()
-	uint32 bHasCompilerMessage:1;
+	uint8 bHasCompilerMessage:1;
+
+	/** Comment bubble pinned state */
+	UPROPERTY()
+	uint8 bCommentBubblePinned:1;
+
+	/** Comment bubble visibility */
+	UPROPERTY()
+	uint8 bCommentBubbleVisible:1;
+
+	/** Make comment bubble visible */
+	UPROPERTY(Transient)
+	uint8 bCommentBubbleMakeVisible:1;
 
 #if WITH_EDITORONLY_DATA
 	/** If true, this node can be renamed in the editor */
 	UPROPERTY()
-	uint32 bCanRenameNode:1;
+	uint8 bCanRenameNode:1;
 
 	/** Note for a node that lingers until saved */
 	UPROPERTY(Transient)
@@ -200,18 +245,6 @@ class ENGINE_API UEdGraphNode : public UObject
 	UPROPERTY()
 	FString NodeComment;
 
-	/** Comment bubble pinned state */
-	UPROPERTY()
-	bool bCommentBubblePinned;
-
-	/** Comment bubble visibility */
-	UPROPERTY()
-	bool bCommentBubbleVisible;
-
-	/** Make comment bubble visible */
-	UPROPERTY(Transient)
-	bool bCommentBubbleMakeVisible;
-
 	/** Flag to store node specific compile error/warning*/
 	UPROPERTY()
 	int32 ErrorType;
@@ -224,23 +257,6 @@ class ENGINE_API UEdGraphNode : public UObject
 	UPROPERTY()
 	FGuid NodeGuid;
 
-	/** Enum to indicate if a node has advanced-display-pins, and if they are shown */
-	UPROPERTY()
-	TEnumAsByte<ENodeAdvancedPins::Type> AdvancedPinDisplay;
-
-	/** Indicates in what state the node is enabled, which may eliminate it from being compiled */
-	UPROPERTY()
-	ENodeEnabledState EnabledState;
-
-	/** Indicates whether or not the user explicitly set the enabled state */
-	UPROPERTY()
-	bool bUserSetEnabledState;
-
-private:
-	/** (DEPRECATED) FALSE if the node is a disabled, which eliminates it from being compiled */
-	UPROPERTY()
-	bool bIsNodeEnabled_DEPRECATED;
-
 public:
 	/** Enables this node. */
 	FORCEINLINE void EnableNode()
@@ -290,6 +306,7 @@ public:
 	TArray<UEdGraphPin*> GetAllPins() { return Pins; }
 
 	/** Create a new pin on this node using the supplied info, and return the new pin */
+	DEPRECATED(4.17, "Use version that takes PinContainerType instead of separate booleans for array, set, and map")
 	UEdGraphPin* CreatePin(
 		EEdGraphPinDirection Dir, 
 		const FString& PinCategory, 
@@ -304,6 +321,19 @@ public:
 		bool bIsMap = false,
 		const FEdGraphTerminalType& ValueTerminalType = FEdGraphTerminalType());
 
+	/** Create a new pin on this node using the supplied info, and return the new pin */
+	UEdGraphPin* CreatePin(
+		EEdGraphPinDirection Dir, 
+		const FString& PinCategory, 
+		const FString& PinSubCategory, 
+		UObject* PinSubCategoryObject, 
+		const FString& PinName, 
+		EPinContainerType PinContainerType = EPinContainerType::None,
+		bool bIsReference = false, 
+		bool bIsConst = false, 
+		int32 Index = INDEX_NONE, 
+		const FEdGraphTerminalType& ValueTerminalType = FEdGraphTerminalType());
+
 	/** Create a new pin on this node using the supplied pin type, and return the new pin */
 	UEdGraphPin* CreatePin(EEdGraphPinDirection Dir, const FEdGraphPinType& InPinType, const FString& PinName, int32 Index = INDEX_NONE);
 
@@ -475,7 +505,7 @@ public:
 	/**
 	 * Returns the link used for external documentation for the graph node
 	 */
-	virtual FString GetDocumentationLink() const { return TEXT(""); }
+	virtual FString GetDocumentationLink() const { return FString(); }
 
 	/**
 	 * Returns the name of the excerpt to display from the specified external documentation link for the graph node
@@ -644,11 +674,6 @@ protected:
 	virtual FString GetPropertyNameAndValueForDiff(const UProperty* Prop, const uint8* PropertyAddr) const;
 
 #endif // WITH_EDITOR
-
-protected:
-
-	/** (DEPRECATED) Value used for AllowSplitPins(). Do not override. */
-	bool bAllowSplitPins_DEPRECATED;
 };
 
 
diff --git a/Engine/Source/Runtime/Engine/Classes/EdGraph/EdGraphPin.h b/Engine/Source/Runtime/Engine/Classes/EdGraph/EdGraphPin.h
--- a/Engine/Source/Runtime/Engine/Classes/EdGraph/EdGraphPin.h
+++ b/Engine/Source/Runtime/Engine/Classes/EdGraph/EdGraphPin.h
@@ -104,55 +104,82 @@ struct FEdGraphPinType
 	UPROPERTY()
 	FEdGraphTerminalType PinValueType;
 
-	/** Whether or not this pin represents a map of keys to values */
 	UPROPERTY()
-	bool bIsMap;
+	EPinContainerType ContainerType;
 
-	/** Whether or not this pin represents a set of (unique) values */
+private:
+	/** DEPRECATED(4.17) Whether or not this pin represents a map of keys to values */
+	UPROPERTY()
+	uint8 bIsMap_DEPRECATED:1;
+
+	/** DEPRECATED(4.17) Whether or not this pin represents a set of (unique) values */
 	UPROPERTY()
-	bool bIsSet;
+	uint8 bIsSet_DEPRECATED:1;
 
-	/** Whether or not this pin represents an array of values */
+	/** DEPRECATED(4.17) Whether or not this pin represents an array of values */
 	UPROPERTY()
-	bool bIsArray;
+	uint8 bIsArray_DEPRECATED:1;
 
+public:
 	/** Whether or not this pin is a value passed by reference or not */
 	UPROPERTY()
-	bool bIsReference;
+	uint8 bIsReference:1;
 
 	/** Whether or not this pin is a immutable const value */
 	UPROPERTY()
-	bool bIsConst;
+	uint8 bIsConst:1;
 
 	/** Whether or not this is a weak reference */
 	UPROPERTY()
-	bool bIsWeakPointer;
+	uint8 bIsWeakPointer:1;
 
-	FORCEINLINE bool IsContainer() const { return bIsMap || bIsSet || bIsArray; }
+	FORCEINLINE bool IsContainer() const { return (ContainerType != EPinContainerType::None); }
+	FORCEINLINE bool IsArray() const { return (ContainerType == EPinContainerType::Array); }
+	FORCEINLINE bool IsSet() const { return (ContainerType == EPinContainerType::Set); }
+	FORCEINLINE bool IsMap() const { return (ContainerType == EPinContainerType::Map); }
 
 public:
 	FEdGraphPinType() 
+		: PinSubCategoryObject(nullptr)
+		, ContainerType(EPinContainerType::None)
+		, bIsMap_DEPRECATED(false)
+		, bIsSet_DEPRECATED(false)
+		, bIsArray_DEPRECATED(false)
+		, bIsReference(false)
+		, bIsConst(false)
+		, bIsWeakPointer(false)	
 	{
-		PinSubCategoryObject = nullptr;
-		bIsMap = false;
-		bIsSet = false;
-		bIsArray = false;
-		bIsReference = false;
-		bIsConst = false;
-		bIsWeakPointer = false;
 	}
-	FEdGraphPinType(const FString& InPinCategory, const FString& InPinSubCategory, UObject* InPinSubCategoryObject, bool bInIsArray, bool bInIsReference, bool bInIsSet, bool bInIsMap, const FEdGraphTerminalType& InValueTerminalType )
+
+	DEPRECATED(4.17, "Use version that takes PinContainerType instead of separate booleans for array, set, and map")
+	FEdGraphPinType(FString InPinCategory, FString InPinSubCategory, UObject* InPinSubCategoryObject, bool bInIsArray, bool bInIsReference, bool bInIsSet, bool bInIsMap, const FEdGraphTerminalType& InValueTerminalType )
+		: PinCategory(MoveTemp(InPinCategory))
+		, PinSubCategory(MoveTemp(InPinSubCategory))
+		, PinSubCategoryObject(InPinSubCategoryObject)
+		, PinValueType(InValueTerminalType)
+		, ContainerType(ToPinContainerType(bInIsArray, bInIsSet, bInIsMap))
+		, bIsMap_DEPRECATED(false)
+		, bIsSet_DEPRECATED(false)
+		, bIsArray_DEPRECATED(false)
+		, bIsReference(bInIsReference)
+		, bIsConst(false)
+		, bIsWeakPointer(false)
+	{
+	}
+
+	FEdGraphPinType(FString InPinCategory, FString InPinSubCategory, UObject* InPinSubCategoryObject, EPinContainerType InPinContainerType, bool bInIsReference, const FEdGraphTerminalType& InValueTerminalType )
+		: PinCategory(MoveTemp(InPinCategory))
+		, PinSubCategory(MoveTemp(InPinSubCategory))
+		, PinSubCategoryObject(InPinSubCategoryObject)
+		, PinValueType(InValueTerminalType)
+		, ContainerType(InPinContainerType)
+		, bIsMap_DEPRECATED(false)
+		, bIsSet_DEPRECATED(false)
+		, bIsArray_DEPRECATED(false)
+		, bIsReference(bInIsReference)
+		, bIsConst(false)
+		, bIsWeakPointer(false)
 	{
-		PinCategory = InPinCategory;
-		PinSubCategory = InPinSubCategory;
-		PinSubCategoryObject = InPinSubCategoryObject;
-		PinValueType = InValueTerminalType;
-		bIsMap = bInIsMap;
-		bIsSet = bInIsSet;
-		bIsArray = bInIsArray;
-		bIsReference = bInIsReference;
-		bIsConst = false;
-		bIsWeakPointer = false;
 	}
 	bool operator == ( const FEdGraphPinType& Other ) const
 	{
@@ -160,9 +187,7 @@ public:
 			&& (PinSubCategory == Other.PinSubCategory)
 			&& (PinSubCategoryObject == Other.PinSubCategoryObject)
 			&& (PinValueType == Other.PinValueType)
-			&& (bIsMap == Other.bIsMap)
-			&& (bIsSet == Other.bIsSet)
-			&& (bIsArray == Other.bIsArray)
+			&& (ContainerType == Other.ContainerType)
 			&& (bIsReference == Other.bIsReference)
 			&& (bIsWeakPointer == Other.bIsWeakPointer)
 			&& (PinSubCategoryMemberReference == Other.PinSubCategoryMemberReference)
@@ -170,28 +195,17 @@ public:
 	}
 	bool operator != ( const FEdGraphPinType& Other ) const
 	{
-		return (PinCategory != Other.PinCategory) 
-			|| (PinSubCategory != Other.PinSubCategory) 
-			|| (PinSubCategoryObject != Other.PinSubCategoryObject) 
-			|| (PinValueType != Other.PinValueType)
-			|| (bIsMap != Other.bIsMap)
-			|| (bIsSet != Other.bIsSet)
-			|| (bIsArray != Other.bIsArray) 
-			|| (bIsReference != Other.bIsReference)
-			|| (bIsWeakPointer != Other.bIsWeakPointer)
-			|| (bIsConst != Other.bIsConst);
+		return !(*this == Other);
 	}
 
 	void ResetToDefaults()
 	{
-		PinCategory.Empty();
-		PinSubCategory.Empty();
+		PinCategory.Reset();
+		PinSubCategory.Reset();
 		PinSubCategoryObject = nullptr;
 		PinValueType = FEdGraphTerminalType();
 		PinSubCategoryMemberReference.Reset();
-		bIsMap = false;
-		bIsSet = false;
-		bIsArray = false;
+		ContainerType = EPinContainerType::None;
 		bIsReference = false;
 		bIsWeakPointer = false;
 		bIsConst = false;
@@ -201,6 +215,29 @@ public:
 
 	static ENGINE_API FEdGraphPinType GetPinTypeForTerminalType( const FEdGraphTerminalType& TerminalType );
 	static ENGINE_API FEdGraphPinType GetTerminalTypeForContainer( const FEdGraphPinType& ContainerType );
+
+	static EPinContainerType ToPinContainerType(const bool bInIsArray, const bool bInIsSet, const bool bInIsMap)
+	{
+		EPinContainerType ContainerType = EPinContainerType::None;
+
+		if (bInIsArray)
+		{
+			check(!bInIsSet && !bInIsMap);
+			ContainerType = EPinContainerType::Array;
+		}
+		else if (bInIsSet)
+		{
+			check(!bInIsMap);
+			ContainerType = EPinContainerType::Set;
+		}
+		else if (bInIsMap)
+		{
+			ContainerType = EPinContainerType::Map;
+		}
+
+		return ContainerType;
+	}
+
 };
 
 template<>
@@ -397,6 +434,9 @@ public:
 	/** Get the current DefaultObject path name, or DefaultValue if its null */
 	ENGINE_API FString GetDefaultAsString() const;
 
+	/** Returns true if the current default value matches the autogenerated default value */
+	ENGINE_API bool DoesDefaultValueMatchAutogenerated() const;
+
 #if WITH_EDITORONLY_DATA
 	/** Returns how the name of the pin should be displayed in the UI */
 	ENGINE_API FText GetDisplayName() const;
@@ -411,6 +451,7 @@ public:
 	 */
 	const FString GetLinkInfoString( const FString& InFunctionName, const FString& InInfoData, const UEdGraphPin* InToPin ) const;
 
+	/** Reset default values to empty. This should not be called when AutogeneratedDefaultValue needs to be respected! */
 	void ResetDefaultValue()
 	{
 		DefaultValue.Empty();
@@ -418,6 +459,7 @@ public:
 		DefaultTextValue = FText::GetEmpty();
 	}
 
+	/** Resets node to default constructor state */
 	void ResetToDefaults()
 	{
 		check(LinkedTo.Num() == 0);
diff --git a/Engine/Source/Runtime/Engine/Classes/Engine/AssetManager.h b/Engine/Source/Runtime/Engine/Classes/Engine/AssetManager.h
--- a/Engine/Source/Runtime/Engine/Classes/Engine/AssetManager.h
+++ b/Engine/Source/Runtime/Engine/Classes/Engine/AssetManager.h
@@ -145,7 +145,7 @@ public:
 	 * @param LoadBundles		List of bundles to load for those assets
 	 * @param DelegateToCall	Delegate that will be called on completion, may be called before function returns if assets are already loaded
 	 * @param Priority			Async loading priority for this request
-	 * @return					Streamable Handle that can be used to poll or wait
+	 * @return					Streamable Handle that can be used to poll or wait. You do not need to keep this handle to stop the assets from being unloaded
 	 */
 	virtual TSharedPtr<FStreamableHandle> LoadPrimaryAssets(const TArray<FPrimaryAssetId>& AssetsToLoad, const TArray<FName>& LoadBundles = TArray<FName>(), FStreamableDelegate DelegateToCall = FStreamableDelegate(), TAsyncLoadPriority Priority = FStreamableManager::DefaultAsyncLoadPriority);
 
@@ -156,7 +156,8 @@ public:
 	virtual TSharedPtr<FStreamableHandle> LoadPrimaryAssetsWithType(FPrimaryAssetType PrimaryAssetType, const TArray<FName>& LoadBundles = TArray<FName>(), FStreamableDelegate DelegateToCall = FStreamableDelegate(), TAsyncLoadPriority Priority = FStreamableManager::DefaultAsyncLoadPriority);
 
 	/** 
-	 * Unloads a list of Primary Assets. This will drop hard references to these assets, but they may be in memory due to other references or GC delay
+	 * Unloads a list of Primary Assets that were previously Loaded.
+	 * If the only thing keeping these assets in memory was a prior Load call, they will be freed.
 	 *
 	 * @param AssetsToUnload	List of primary assets to load
 	 * @return					Number of assets unloaded
@@ -180,7 +181,7 @@ public:
 	 * @param RemoveAllBundles	If true, remove all existing bundles even if not in remove list
 	 * @param DelegateToCall	Delegate that will be called on completion, may be called before function returns if assets are already loaded
 	 * @param Priority			Async loading priority for this request
-	 * @return					Streamable Handle that can be used to poll or wait
+	 * @return					Streamable Handle that can be used to poll or wait. You do not need to keep this handle to stop the assets from being unloaded
 	 */
 	virtual TSharedPtr<FStreamableHandle> ChangeBundleStateForPrimaryAssets(const TArray<FPrimaryAssetId>& AssetsToChange, const TArray<FName>& AddBundles, const TArray<FName>& RemoveBundles, bool bRemoveAllBundles = false, FStreamableDelegate DelegateToCall = FStreamableDelegate(), TAsyncLoadPriority Priority = FStreamableManager::DefaultAsyncLoadPriority);
 
@@ -192,7 +193,7 @@ public:
 	 * @param Bundles			If not null, will fill in with a list of the requested bundle state
 	 * @return					Streamable Handle that can be used to poll or wait
 	 */
-	TSharedPtr<FStreamableHandle> GetPrimaryAssetHandle(const FPrimaryAssetId& PrimaryAssetId, bool bForceCurrent = false, TArray<FName>* Bundles = nullptr);
+	TSharedPtr<FStreamableHandle> GetPrimaryAssetHandle(const FPrimaryAssetId& PrimaryAssetId, bool bForceCurrent = false, TArray<FName>* Bundles = nullptr) const ;
 
 	/** 
 	 * Returns a list of primary assets that are in the given bundle state. Only assets that are loaded or being loaded are valid
@@ -204,10 +205,23 @@ public:
 	 * @param bForceCurrent		If true, only use the current state. If false, use the current or pending
 	 * @return					True if any found
 	 */
-	bool GetPrimaryAssetsWithBundleState(TArray<FPrimaryAssetId>& PrimaryAssetList, const TArray<FName>& ValidTypes, const TArray<FName>& RequiredBundles, const TArray<FName>& ExcludedBundles = TArray<FName>(), bool bForceCurrent = false);
+	bool GetPrimaryAssetsWithBundleState(TArray<FPrimaryAssetId>& PrimaryAssetList, const TArray<FName>& ValidTypes, const TArray<FName>& RequiredBundles, const TArray<FName>& ExcludedBundles = TArray<FName>(), bool bForceCurrent = false) const;
 
 	/** Fills in a TMap with the pending/active loading state of every asset */
-	void GetPrimaryAssetBundleStateMap(TMap<FPrimaryAssetId, TArray<FName>>& BundleStateMap, bool bForceCurrent = false);
+	void GetPrimaryAssetBundleStateMap(TMap<FPrimaryAssetId, TArray<FName>>& BundleStateMap, bool bForceCurrent = false) const;
+
+	/**
+	 * Preloads data for a set of assets in a specific bundle state, and returns a handle you must keep active.
+	 * These assets are not officially Loaded, so Unload/ChangeBundleState will not affect them and if you release the handle without otherwise loading the assets they will be freed.
+	 *
+	 * @param AssetsToLoad		List of primary assets to load
+	 * @param LoadBundles		List of bundles to load for those assets
+	 * @param bLoadRecursive	If true, this will call RecursivelyExpandBundleData and recurse into sub bundles of other primary assets loaded by a bundle reference
+	 * @param DelegateToCall	Delegate that will be called on completion, may be called before function returns if assets are already loaded
+	 * @param Priority			Async loading priority for this request
+	 * @return					Streamable Handle that must be stored to keep the preloaded assets from being freed
+	 */
+	virtual TSharedPtr<FStreamableHandle> PreloadPrimaryAssets(const TArray<FPrimaryAssetId>& AssetsToLoad, const TArray<FName>& LoadBundles, bool bLoadRecursive, FStreamableDelegate DelegateToCall = FStreamableDelegate(), TAsyncLoadPriority Priority = FStreamableManager::DefaultAsyncLoadPriority);
 
 	/** Quick wrapper to async load some non primary assets with the primary streamable manager. This will not auto release the handle, release it if needed */
 	virtual TSharedPtr<FStreamableHandle> LoadAssetList(const TArray<FStringAssetReference>& AssetList);
@@ -234,7 +248,7 @@ public:
 	virtual bool GetManagedPackageList(FPrimaryAssetId PrimaryAssetId, TArray<FName>& AssetPackageList) const;
 
 	/** Returns list of PrimaryAssetIds that manage a package. Will optionally recurse up the management chain */
-	virtual bool GetPackageManagerList(FName PackageName, bool bRecurseToParents, TArray<FPrimaryAssetId>& ManagerList) const;
+	virtual bool GetPackageManagers(FName PackageName, bool bRecurseToParents, TSet<FPrimaryAssetId>& ManagerSet) const;
 
 
 	// GENERAL ASSET UTILITY FUNCTIONS
@@ -264,6 +278,9 @@ public:
 	/** Dumps out list of loaded asset bundles to log */
 	static void DumpLoadedAssetState();
 
+	/** Dumps information about the Asset Registry to log */
+	static void DumpAssetRegistryInfo();
+
 	/** Dumps out list of primary asset -> managed assets to log */
 	static void DumpReferencersForPackage(const TArray< FString >& PackageNames);
 
@@ -285,8 +302,8 @@ public:
 #if WITH_EDITOR
 	// EDITOR ONLY FUNCTIONALITY
 
-	/** Gets package names to add to the cook. Do this instead of loading assets so RAM can be properly managed on build machines */
-	virtual void ModifyCook(TArray<FName>& PackageNames);
+	/** Gets package names to add to the cook, and packages to never cook even if in startup set memory or referenced */
+	virtual void ModifyCook(TArray<FName>& PackagesToCook, TArray<FName>& PackagesToNeverCook);
 
 	/** Returns cook rule for a package name using Management rules, games should override this to take into account their individual workflows */
 	virtual EPrimaryAssetCookRule GetPackageCookRule(FName PackageName) const;
@@ -400,9 +417,6 @@ protected:
 	TMap<FPrimaryAssetId, TArray<FName>> PrimaryAssetStateBeforePIE;
 #endif // WITH_EDITOR
 
-	/** Per-type asset information */
-	TMap<FName, TSharedRef<FPrimaryAssetTypeData>> AssetTypeMap;
-
 	/** Map from object path to Primary Asset Id */
 	TMap<FName, FPrimaryAssetId> AssetPathMap;
 
@@ -429,10 +443,6 @@ protected:
 	UPROPERTY()
 	bool bIsGlobalAsyncScanEnvironment;
 
-	/** True if we should keep hard references to loaded assets, to stop them from garbage collecting */
-	UPROPERTY()
-	bool bShouldKeepHardRefs;
-
 	/** True if PrimaryAssetType/Name will be implied for loading assets that don't have it saved on disk. Won't work for all projects */
 	UPROPERTY()
 	bool bShouldGuessTypeAndName;
@@ -459,6 +469,9 @@ protected:
 	TMap<FName, FName> AssetPathRedirects;
 
 private:
+	/** Per-type asset information, cannot be accessed by children as it is defined in CPP file */
+	TMap<FName, TSharedRef<FPrimaryAssetTypeData>> AssetTypeMap;
+
 	mutable class IAssetRegistry* CachedAssetRegistry;
 	mutable const class UAssetManagerSettings* CachedSettings;
 };
diff --git a/Engine/Source/Runtime/Engine/Classes/Engine/AssetManagerSettings.h b/Engine/Source/Runtime/Engine/Classes/Engine/AssetManagerSettings.h
--- a/Engine/Source/Runtime/Engine/Classes/Engine/AssetManagerSettings.h
+++ b/Engine/Source/Runtime/Engine/Classes/Engine/AssetManagerSettings.h
@@ -46,7 +46,7 @@ class ENGINE_API UAssetManagerSettings : public UDeveloperSettings
 	GENERATED_BODY()
 
 public:
-	UAssetManagerSettings() : bOnlyCookProductionAssets(false) {}
+	UAssetManagerSettings() : bOnlyCookProductionAssets(false), bShouldGuessTypeAndNameInEditor(true) {}
 
 	/** List of asset types to scan at startup */
 	UPROPERTY(config, EditAnywhere, Category = "Asset Manager")
@@ -64,6 +64,10 @@ public:
 	UPROPERTY(config, EditAnywhere, Category = "Asset Manager")
 	bool bOnlyCookProductionAssets;
 
+	/** If true, PrimaryAsset Type/Name will be implied for assets in the editor (cooked builds always must be explicit). This allows guessing for content that hasn't been resaved yet */
+	UPROPERTY(config, EditAnywhere, Category = "Asset Manager")
+	bool bShouldGuessTypeAndNameInEditor;
+
 	/** Redirect from Type:Name to Type:NameNew */
 	UPROPERTY(config, EditAnywhere, Category = "Redirects")
 	TArray<FAssetManagerRedirect> PrimaryAssetIdRedirects;
diff --git a/Engine/Source/Runtime/Engine/Classes/Engine/Blueprint.h b/Engine/Source/Runtime/Engine/Classes/Engine/Blueprint.h
--- a/Engine/Source/Runtime/Engine/Classes/Engine/Blueprint.h
+++ b/Engine/Source/Runtime/Engine/Classes/Engine/Blueprint.h
@@ -391,6 +391,9 @@ class ENGINE_API UBlueprint : public UBlueprintCore
 	UPROPERTY(transient)
 	uint32 bForceFullEditor : 1;
 
+	UPROPERTY(transient)
+	uint32 bQueuedForCompilation : 1 ;
+
 	/**whether or not you want to continuously rerun the construction script for an actor as you drag it in the editor, or only when the drag operation is complete*/
 	UPROPERTY(EditAnywhere, Category=BlueprintOptions)
 	uint32 bRunConstructionScriptOnDrag : 1;
diff --git a/Engine/Source/Runtime/Engine/Classes/Engine/BlueprintGeneratedClass.h b/Engine/Source/Runtime/Engine/Classes/Engine/BlueprintGeneratedClass.h
--- a/Engine/Source/Runtime/Engine/Classes/Engine/BlueprintGeneratedClass.h
+++ b/Engine/Source/Runtime/Engine/Classes/Engine/BlueprintGeneratedClass.h
@@ -698,6 +698,7 @@ public:
 	virtual FPrimaryAssetId GetPrimaryAssetId() const override;
 	virtual bool NeedsLoadForServer() const override;
 	virtual bool NeedsLoadForClient() const override;
+	virtual bool NeedsLoadForEditorGame() const override;
 	virtual bool CanBeClusterRoot() const override;
 	// End UObject interface
 	
@@ -705,6 +706,7 @@ public:
 #if WITH_EDITOR
 	virtual UClass* GetAuthoritativeClass() override;
 	virtual void ConditionalRecompileClass(TArray<UObject*>* ObjLoaded) override;
+	virtual void FlushCompilationQueueForLevel() override;
 	virtual UObject* GetArchetypeForCDO() const override;
 #endif //WITH_EDITOR
 	virtual void SerializeDefaultObject(UObject* Object, FArchive& Ar) override;
diff --git a/Engine/Source/Runtime/Engine/Classes/Engine/DataTable.h b/Engine/Source/Runtime/Engine/Classes/Engine/DataTable.h
--- a/Engine/Source/Runtime/Engine/Classes/Engine/DataTable.h
+++ b/Engine/Source/Runtime/Engine/Classes/Engine/DataTable.h
@@ -28,7 +28,7 @@ class UDataTable;
 /**
  * Base class for all table row structs to inherit from.
  */
-USTRUCT()
+USTRUCT(BlueprintInternalUseOnly)
 struct FTableRowBase
 {
 	GENERATED_USTRUCT_BODY()
@@ -333,7 +333,7 @@ struct TStructOpsTypeTraits< FDataTableRowHandle > : public TStructOpsTypeTraits
 };
 
 /** Handle to a particular row in a table*/
-USTRUCT()
+USTRUCT(BlueprintType)
 struct ENGINE_API FDataTableCategoryHandle
 {
 	GENERATED_USTRUCT_BODY()
diff --git a/Engine/Source/Runtime/Engine/Classes/Engine/DeveloperSettings.h b/Engine/Source/Runtime/Engine/Classes/Engine/DeveloperSettings.h
--- a/Engine/Source/Runtime/Engine/Classes/Engine/DeveloperSettings.h
+++ b/Engine/Source/Runtime/Engine/Classes/Engine/DeveloperSettings.h
@@ -6,7 +6,8 @@
 #include "UObject/ObjectMacros.h"
 #include "UObject/UObjectGlobals.h"
 #include "UObject/Object.h"
-#include "Widgets/SWidget.h"
+class SWidget;
+
 #include "DeveloperSettings.generated.h"
 
 /**
@@ -38,7 +39,7 @@ public:
 #endif
 
 	/** Gets a custom widget for the settings.  This is only for very custom situations. */
-	virtual TSharedPtr<class SWidget> GetCustomSettingsWidget() const;
+	virtual TSharedPtr<SWidget> GetCustomSettingsWidget() const;
 
 protected:
 	/**
diff --git a/Engine/Source/Runtime/Engine/Classes/Engine/Engine.h b/Engine/Source/Runtime/Engine/Classes/Engine/Engine.h
--- a/Engine/Source/Runtime/Engine/Classes/Engine/Engine.h
+++ b/Engine/Source/Runtime/Engine/Classes/Engine/Engine.h
@@ -573,7 +573,7 @@ class ENGINE_API UEngine
 {
 	GENERATED_UCLASS_BODY()
 
-	// Called after GEngine->Init has been called
+	DEPRECATED(4.17, "UEngine::OnPostEngineInit is deprecated, bind to FCoreDelegates::OnPostEngineInit instead, which will also be called for commandlets")
 	static FSimpleMulticastDelegate OnPostEngineInit;
 
 private:
