Core: Partially rolling back CL# 3372764, so duplicating on top of an object that was not part of the annotation set will recreate it (reopening UE-43328)
#reimplementing CL# 3436378 from 4.16

diff --git a/Engine/Source/Runtime/CoreUObject/Private/Serialization/DuplicateDataWriter.cpp b/Engine/Source/Runtime/CoreUObject/Private/Serialization/DuplicateDataWriter.cpp
--- a/Engine/Source/Runtime/CoreUObject/Private/Serialization/DuplicateDataWriter.cpp
+++ b/Engine/Source/Runtime/CoreUObject/Private/Serialization/DuplicateDataWriter.cpp
@@ -5,7 +5,6 @@
 #include "Misc/StringAssetReference.h"
 #include "UObject/AssetPtr.h"
 #include "UObject/PropertyPortFlags.h"
-#include "UObject/UObjectHash.h"
 #include "Serialization/DuplicatedObject.h"
 #include "Serialization/DuplicatedDataWriter.h"
 
@@ -137,23 +136,11 @@ UObject* FDuplicateDataWriter::GetDuplicatedObject(UObject* Object, bool bCreate
 			UObject* DupOuter = GetDuplicatedObject(Object->GetOuter());
 			if(DupOuter != nullptr)
 			{
-				// First check if the duplicated outer already has an allocated duplicate of this object
-				Result = static_cast<UObject*>(FindObjectWithOuter(DupOuter, Object->GetClass(), Object->GetFName()));
-
-				if (Result == nullptr)
-				{
-					// The object's outer is being duplicated, create a duplicate of this object.
-					Result = StaticConstructObject_Internal(Object->GetClass(), DupOuter, Object->GetFName(),
-						ApplyFlags | Object->GetMaskedFlags(FlagMask),
-						ApplyInternalFlags | (Object->GetInternalFlags() & InternalFlagMask),
-						Object->GetArchetype(), true, InstanceGraph);
-				}
-				else
-				{
-					// If we found the object then set the flags up such that it matches up with what would have been allocated
-					Result->SetFlags(ApplyFlags | Object->GetMaskedFlags(FlagMask));
-					Result->SetInternalFlags(ApplyInternalFlags | (Object->GetInternalFlags() & InternalFlagMask));
-				}
+				// The object's outer is being duplicated, create a duplicate of this object.
+				Result = StaticConstructObject_Internal(Object->GetClass(), DupOuter, Object->GetFName(),
+					ApplyFlags | Object->GetMaskedFlags(FlagMask),
+					ApplyInternalFlags | (Object->GetInternalFlags() & InternalFlagMask),
+					Object->GetArchetype(), true, InstanceGraph);
 
 				AddDuplicate(Object, Result);
 			}
