Copying //UE4/Dev-Framework to //UE4/Dev-Main (Source: //UE4/Dev-Framework @ 3459469)

diff --git a/Engine/Source/Runtime/CoreUObject/Public/Misc/PackageName.h b/Engine/Source/Runtime/CoreUObject/Public/Misc/PackageName.h
--- a/Engine/Source/Runtime/CoreUObject/Public/Misc/PackageName.h
+++ b/Engine/Source/Runtime/CoreUObject/Public/Misc/PackageName.h
@@ -36,7 +36,8 @@ public:
 	static FName* FindScriptPackageName(FName InShortName);
 
 	/** 
-	 * Tries to convert the supplied filename to long package name. Will attempt to find the package on disk (very slow).
+	 * Tries to convert the supplied relative or absolute filename to a long package name/path starting with a root like /game
+	 * This works on both package names and directories, and it does not validate that it actually exists on disk.
 	 * 
 	 * @param InFilename Filename to convert.
 	 * @param OutPackageName The resulting long package name if the conversion was successful.
@@ -44,37 +45,52 @@ public:
 	 * @return Returns true if the supplied filename properly maps to one of the long package roots.
 	 */
 	static bool TryConvertFilenameToLongPackageName(const FString& InFilename, FString& OutPackageName, FString* OutFailureReason = nullptr);
+
 	/** 
 	 * Converts the supplied filename to long package name.
-	 *  Throws a fatal error if the conversion is not successfull.
+	 * Throws a fatal error if the conversion is not successful.
 	 * 
 	 * @param InFilename Filename to convert.
 	 * @return Long package name.
 	 */
 	static FString FilenameToLongPackageName(const FString& InFilename);
+
 	/** 
 	 * Tries to convert a long package name to a file name with the supplied extension.
+	 * This can be called on package paths as well, provide no extension in that case
 	 *
 	 * @param InLongPackageName Long Package Name
 	 * @param InExtension Package extension.
 	 * @return Package filename.
 	 */
 	static bool TryConvertLongPackageNameToFilename(const FString& InLongPackageName, FString& OutFilename, const FString& InExtension = TEXT(""));
+
 	/** 
 	 * Converts a long package name to a file name with the supplied extension.
+	 * Throws a fatal error if the conversion is not successful.
 	 *
 	 * @param InLongPackageName Long Package Name
 	 * @param InExtension Package extension.
 	 * @return Package filename.
 	 */
 	static FString LongPackageNameToFilename(const FString& InLongPackageName, const FString& InExtension = TEXT(""));
+
 	/** 
 	 * Returns the path to the specified package, excluding the short package name
 	 *
-	 * @param InLongPackageName Package Name.
-	 * @return The path to the specified package.
+	 * @param InLongPackageName Long Package Name.
+	 * @return The path containing the specified package.
 	 */
 	static FString GetLongPackagePath(const FString& InLongPackageName);
+
+	/** 
+	 * Returns the clean asset name for the specified package, same as GetShortName
+	 *
+	 * @param InLongPackageName Long Package Name
+	 * @return Clean asset name.
+	 */
+	static FString GetLongPackageAssetName(const FString& InLongPackageName);
+
 	/** 
 	 * Convert a long package name into root, path, and name components
 	 *
@@ -86,21 +102,7 @@ public:
 	 * @return True if the conversion was possible, false otherwise
 	 */
 	static bool SplitLongPackageName(const FString& InLongPackageName, FString& OutPackageRoot, FString& OutPackagePath, FString& OutPackageName, const bool bStripRootLeadingSlash = false);
-	/** 
-	 * Returns the clean asset name for the specified package 
-	 *
-	 * @param InLongPackageName Long Package Name
-	 * @return Clean asset name.
-	 */
-	static FString GetLongPackageAssetName(const FString& InLongPackageName);
-	/**
-	 * Convert a root path to the content path associated with it
-	 * Similar to TryConvertLongPackageNameToFilename except doesn't require a file just returns the pat
-	 * See also RegisterMountPoint and UnRegisterMountPoint
-	 * @param RootPath The package root path, eg "/Game/"
-	 * @param OutContentPath The path from the mount point to the package, eg "Maps/TestMaps/
-	 */
-	static bool ConvertRootPathToContentPath(const FString& RootPath, FString& OutContentPath);
+
 	/** 
 	 * Returns true if the path starts with a valid root (i.e. /Game/, /Engine/, etc) and contains no illegal characters.
 	 *
@@ -110,6 +112,7 @@ public:
 	 * @return							true if a valid long package name
 	 */
 	static bool IsValidLongPackageName(const FString& InLongPackageName, bool bIncludeReadOnlyRoots = false, FText* OutReason = nullptr);
+
 	/** 
 	 * Returns true if the path starts with a valid root (i.e. /Game/, /Engine/, etc) and contains no illegal characters.
 	 * This validates that the packagename is valid, and also makes sure the object after package name is also correct.
@@ -120,6 +123,7 @@ public:
 	 * @return							true if a valid object path
 	 */
 	static bool IsValidObjectPath(const FString& InObjectPath, FText* OutReason = nullptr);
+
 	/**
 	 * Checks if the given string is a long package name or not.
 	 *
@@ -127,20 +131,16 @@ public:
 	 * @return true if the given name is a long package name, false otherwise.
 	 */
 	static bool IsShortPackageName(const FString& PossiblyLongName);
-	/**
-	 * Checks if the given name is a long package name or not.
-	 *
-	 * @param PossiblyLongName Package name.
-	 * @return true if the given name is a long package name, false otherwise.
-	 */
 	static bool IsShortPackageName(const FName PossiblyLongName);
+
 	/**
 	 * Converts package name to short name.
 	 *
-	 * @param Package Package which name to convert.
+	 * @param Package Package with name to convert.
 	 * @return Short package name.
 	 */
 	static FString GetShortName(const UPackage* Package);
+
 	/**
 	 * Converts package name to short name.
 	 *
@@ -148,20 +148,9 @@ public:
 	 * @return Short package name.
 	 */
 	static FString GetShortName(const FString& LongName);
-	/**
-	 * Converts package name to short name.
-	 *
-	 * @param LongName Package name to convert.
-	 * @return Short package name.
-	 */
 	static FString GetShortName(const FName& LongName);
-	/**
-	 * Converts package name to short name.
-	 *
-	 * @param LongName Package name to convert.
-	 * @return Short package name.
-	 */
 	static FString GetShortName(const TCHAR* LongName);
+
 	/**
 	 * Converts package name to short name.
 	 *
@@ -169,33 +158,22 @@ public:
 	 * @return Short package name.
 	 */
 	static FName GetShortFName(const FString& LongName);
-	/**
-	 * Converts package name to short name.
-	 *
-	 * @param LongName Package name to convert.
-	 * @return Short package name.
-	 */
 	static FName GetShortFName(const FName& LongName);
-	/**
-	 * Converts package name to short name.
-	 *
-	 * @param LongName Package name to convert.
-	 * @return Short package name.
-	 */
 	static FName GetShortFName(const TCHAR* LongName);
+
 	/**
 	 * This will insert a mount point at the head of the search chain (so it can overlap an existing mount point and win).
 	 *
-	 * @param RootPath Root Path.
-	 * @param ContentPath Content Path.
+	 * @param RootPath Logical Root Path.
+	 * @param ContentPath Content Path on disk.
 	 */
 	static void RegisterMountPoint(const FString& RootPath, const FString& ContentPath);
 
 	/**
 	 * This will remove a previously inserted mount point.
 	 *
-	 * @param RootPath Root Path.
-	 * @param ContentPath Content Path.
+	 * @param RootPath Logical Root Path.
+	 * @param ContentPath Content Path on disk.
 	 */
 	static void UnRegisterMountPoint(const FString& RootPath, const FString& ContentPath);
 
@@ -283,19 +261,10 @@ public:
 	 */
 	static FString GetLocalizedPackagePath(const FString& InSourcePackagePath, const FString& InCultureName);
 
-	/**
-	 * Strips all path and extension information from a relative or fully qualified file name.
-	 *
-	 * @param	InPathName	a relative or fully qualified file name
-	 *
-	 * @return	the passed in string, stripped of path and extensions
-	 */
-	static FString PackageFromPath(const TCHAR* InPathName);
-
 	/** 
 	 * Returns the file extension for packages containing assets.
 	 *
-	 * @return	file extension for asset pacakges ( dot included )
+	 * @return	file extension for asset packages ( dot included )
 	 */
 	static FORCEINLINE const FString& GetAssetPackageExtension()
 	{
@@ -304,7 +273,7 @@ public:
 	/** 
 	 * Returns the file extension for packages containing assets.
 	 *
-	 * @return	file extension for asset pacakges ( dot included )
+	 * @return	file extension for asset packages ( dot included )
 	 */
 	static FORCEINLINE const FString& GetMapPackageExtension()
 	{
@@ -437,6 +406,12 @@ public:
 	*/
 	static bool FindPackageFileWithoutExtension(const FString& InPackageFilename, FString& OutFilename);
 
+	DEPRECATED(4.17, "Deprecated. Call TryConvertLongPackageNameToFilename instead, which also works on nested paths")
+	static bool ConvertRootPathToContentPath(const FString& RootPath, FString& OutContentPath);
+
+	DEPRECATED(4.17, "Deprecated. Call TryConvertFilenameToLongPackageName instead")
+	static FString PackageFromPath(const TCHAR* InPathName);
+
 private:
 
 	static FString AssetPackageExtension;
diff --git a/Engine/Source/Runtime/CoreUObject/Public/Serialization/SerializedPropertyScope.h b/Engine/Source/Runtime/CoreUObject/Public/Serialization/SerializedPropertyScope.h
--- a/Engine/Source/Runtime/CoreUObject/Public/Serialization/SerializedPropertyScope.h
+++ b/Engine/Source/Runtime/CoreUObject/Public/Serialization/SerializedPropertyScope.h
@@ -19,12 +19,16 @@ class FSerializedPropertyScope
 	void PopEditorOnlyProperty();
 #endif
 public:
-	FSerializedPropertyScope(FArchive& InAr, UProperty* InProperty)
+	FSerializedPropertyScope(FArchive& InAr, UProperty* InProperty, const UProperty* OnlyIfOldProperty = nullptr)
 		: Ar(InAr)
 		, Property(InProperty)
 	{
 		OldProperty = Ar.GetSerializedProperty();
-		Ar.SetSerializedProperty(Property);
+		if (!OnlyIfOldProperty || OldProperty == OnlyIfOldProperty)
+		{
+			Ar.SetSerializedProperty(Property);
+		}
+		
 #if WITH_EDITORONLY_DATA
 		PushEditorOnlyProperty();
 #endif
diff --git a/Engine/Source/Runtime/CoreUObject/Public/UObject/CoreRedirects.h b/Engine/Source/Runtime/CoreUObject/Public/UObject/CoreRedirects.h
--- a/Engine/Source/Runtime/CoreUObject/Public/UObject/CoreRedirects.h
+++ b/Engine/Source/Runtime/CoreUObject/Public/UObject/CoreRedirects.h
@@ -24,7 +24,6 @@ enum class ECoreRedirectFlags : int32
 	Type_Function =			0x00000010, // UFunction
 	Type_Property =			0x00000020, // UProperty
 	Type_Package =			0x00000040, // UPackage
-	Type_Delegate =			0x00000080, // Dynamic delegate, special type of UFunction
 
 	// Option flags, specify rules for this redirect
 	Option_InstanceOnly =	0x00010000, // Only redirect instances of this type, not the type itself
diff --git a/Engine/Source/Runtime/CoreUObject/Public/UObject/FastReferenceCollector.h b/Engine/Source/Runtime/CoreUObject/Public/UObject/FastReferenceCollector.h
--- a/Engine/Source/Runtime/CoreUObject/Public/UObject/FastReferenceCollector.h
+++ b/Engine/Source/Runtime/CoreUObject/Public/UObject/FastReferenceCollector.h
@@ -17,6 +17,129 @@ struct FStackEntry;
 	FastReferenceCollector.h: Unreal realtime garbage collection helpers
 =============================================================================*/
 
+/** Struct to hold the objects to serialize array and the list of weak references. This is allocated by ArrayPool */
+struct FGCArrayStruct
+{
+	TArray<UObject*> ObjectsToSerialize;
+	TArray<UObject**> WeakReferences;
+};
+
+/**
+ * Pool for reducing GC allocations
+ */
+class FGCArrayPool
+{
+public:
+
+	/**
+	* Gets the singleton instance of the FObjectArrayPool
+	* @return Pool singleton.
+	*/
+	FORCEINLINE static FGCArrayPool& Get()
+	{
+		static FGCArrayPool Singleton;
+		return Singleton;
+	}
+
+	/**
+	* Gets an event from the pool or creates one if necessary.
+	*
+	* @return The array.
+	* @see ReturnToPool
+	*/
+	FORCEINLINE FGCArrayStruct* GetArrayStructFromPool()
+	{
+		FGCArrayStruct* Result = Pool.Pop();
+		if (!Result)
+		{
+			Result = new FGCArrayStruct();
+		}
+		check(Result);
+#if UE_BUILD_DEBUG
+		NumberOfUsedArrays.Increment();
+#endif // UE_BUILD_DEBUG
+		return Result;
+	}
+
+	/**
+	* Returns an array to the pool.
+	*
+	* @param Array The array to return.
+	* @see GetArrayFromPool
+	*/
+	FORCEINLINE void ReturnToPool(FGCArrayStruct* ArrayStruct)
+	{
+#if UE_BUILD_DEBUG
+		const int32 CheckUsedArrays = NumberOfUsedArrays.Decrement();
+		checkSlow(CheckUsedArrays >= 0);
+#endif // UE_BUILD_DEBUG
+		check(ArrayStruct);
+		ArrayStruct->ObjectsToSerialize.Reset();
+		Pool.Push(ArrayStruct);
+	}
+
+	/** Performs memory cleanup */
+	void Cleanup()
+	{
+#if UE_BUILD_DEBUG
+		const int32 CheckUsedArrays = NumberOfUsedArrays.GetValue();
+		checkSlow(CheckUsedArrays == 0);
+#endif // UE_BUILD_DEBUG
+
+		uint32 FreedMemory = 0;
+		TArray<FGCArrayStruct*> AllArrays;
+		Pool.PopAll(AllArrays);
+		for (FGCArrayStruct* ArrayStruct : AllArrays)
+		{
+			FreedMemory += ArrayStruct->ObjectsToSerialize.GetAllocatedSize();
+			FreedMemory += ArrayStruct->WeakReferences.GetAllocatedSize();
+			delete ArrayStruct;
+		}
+		UE_LOG(LogGarbage, Log, TEXT("Freed %ub from %d GC array pools."), FreedMemory, AllArrays.Num());
+	}
+
+	/** Clears weak references for everything in the pool */
+	void ClearWeakReferences()
+	{
+		TArray<FGCArrayStruct*> AllArrays;
+		Pool.PopAll(AllArrays);
+		for (FGCArrayStruct* ArrayStruct : AllArrays)
+		{
+			for (UObject** WeakReference : ArrayStruct->WeakReferences)
+			{
+				UObject*& ReferencedObject = *WeakReference;
+				if (ReferencedObject && ReferencedObject->IsUnreachable())
+				{
+					ReferencedObject = nullptr;
+				}
+			}
+			ArrayStruct->WeakReferences.Reset();
+
+			Pool.Push(ArrayStruct);
+		}
+	}
+
+#if UE_BUILD_DEBUG
+	void CheckLeaks()
+	{
+		// This function is called after GC has finished so at this point there should be no
+		// arrays used by GC and all should be returned to the pool
+		const int32 LeakedGCPoolArrays = NumberOfUsedArrays.GetValue();
+		checkSlow(LeakedGCPoolArrays == 0);
+	}
+#endif
+
+private:
+
+	/** Holds the collection of recycled arrays. */
+	TLockFreePointerListLIFO< FGCArrayStruct > Pool;
+
+#if UE_BUILD_DEBUG
+	/** Number of arrays currently acquired from the pool by GC */
+	FThreadSafeCounter NumberOfUsedArrays;
+#endif // UE_BUILD_DEBUG
+};
+
 /**
  * Helper class that looks for UObject references by traversing UClass token stream and calls AddReferencedObjects.
  * Provides a generic way of processing references that is used by Unreal Engine garbage collection.
@@ -44,8 +167,8 @@ struct FStackEntry;
 	 class FSampleArrayPool
 	 {
 	   static FSampleArrayPool& Get();
-		 TArray<UObject*>* GetArrayFromPool();
-		 void ReturnToPool(TArray<UObject*>* Array);
+		 FGCArrayStruct* GetArrayStryctFromPool();
+		 void ReturnToPool(FGCArrayStruct* ArrayStruct);
 	 };
  */
 template <bool bParallel, typename ReferenceProcessorType, typename CollectorType, typename ArrayPoolType, bool bAutoGenerateTokenStream = false>
@@ -57,7 +180,7 @@ private:
 	{
 		TFastReferenceCollector*	Owner;
 		ArrayPoolType& ArrayPool;
-		TLockFreePointerListUnordered<TArray<UObject*>, PLATFORM_CACHE_LINE_SIZE> Tasks;
+		TLockFreePointerListUnordered<FGCArrayStruct, PLATFORM_CACHE_LINE_SIZE> Tasks;
 
 		FCriticalSection WaitingThreadsLock;
 		TArray<FEvent*> WaitingThreads;
@@ -84,10 +207,10 @@ private:
 
 		FORCENOINLINE void AddTask(const TArray<UObject*>* InObjectsToSerialize, int32 StartIndex, int32 NumObjects)
 		{
-			TArray<UObject*>* ObjectsToSerialize(ArrayPool.GetArrayFromPool());
-			ObjectsToSerialize->AddUninitialized(NumObjects);
-			FMemory::Memcpy(ObjectsToSerialize->GetData(), InObjectsToSerialize->GetData() + StartIndex, NumObjects * sizeof(UObject*));
-			Tasks.Push(ObjectsToSerialize);
+			FGCArrayStruct* ArrayStruct = ArrayPool.GetArrayStructFromPool();
+			ArrayStruct->ObjectsToSerialize.AddUninitialized(NumObjects);
+			FMemory::Memcpy(ArrayStruct->ObjectsToSerialize.GetData(), InObjectsToSerialize->GetData() + StartIndex, NumObjects * sizeof(UObject*));
+			Tasks.Push(ArrayStruct);
 
 			FEvent* WaitingThread = nullptr;
 			{
@@ -116,8 +239,8 @@ private:
 			}
 			while (true)
 			{
-				TArray<UObject*>* ObjectsToSerialize = Tasks.Pop();
-				while (!ObjectsToSerialize)
+				FGCArrayStruct* ArrayStruct = Tasks.Pop();
+				while (!ArrayStruct)
 				{
 					if (bDone)
 					{
@@ -130,8 +253,8 @@ private:
 						{
 							return;
 						}
-						ObjectsToSerialize = Tasks.Pop();
-						if (!ObjectsToSerialize)
+						ArrayStruct = Tasks.Pop();
+						if (!ArrayStruct)
 						{
 							if (WaitingThreads.Num() + 1 == NumThreadsStarted)
 							{
@@ -151,7 +274,7 @@ private:
 							}
 						}
 					}
-					if (ObjectsToSerialize)
+					if (ArrayStruct)
 					{
 						check(!WaitEvent);
 					}
@@ -160,12 +283,12 @@ private:
 						check(WaitEvent);
 						WaitEvent->Wait();
 						FPlatformProcess::ReturnSynchEventToPool(WaitEvent);
-						ObjectsToSerialize = Tasks.Pop();
-						check(!ObjectsToSerialize || !bDone);
+						ArrayStruct = Tasks.Pop();
+						check(!ArrayStruct || !bDone);
 					}
 				}
-				Owner->ProcessObjectArray(*ObjectsToSerialize, FGraphEventRef());
-				ArrayPool.ReturnToPool(ObjectsToSerialize);
+				Owner->ProcessObjectArray(*ArrayStruct, FGraphEventRef());
+				ArrayPool.ReturnToPool(ArrayStruct);
 			}
 		}
 	};
@@ -203,21 +326,21 @@ private:
 	class FCollectorTask
 	{
 		TFastReferenceCollector*	Owner;
-		TArray<UObject*>*	ObjectsToSerialize;
+		FGCArrayStruct*	ArrayStruct;
 		ArrayPoolType& ArrayPool;
 
 	public:
 		FCollectorTask(TFastReferenceCollector* InOwner, const TArray<UObject*>* InObjectsToSerialize, int32 StartIndex, int32 NumObjects, ArrayPoolType& InArrayPool)
 			: Owner(InOwner)
-			, ObjectsToSerialize(InArrayPool.GetArrayFromPool())
+			, ArrayStruct(InArrayPool.GetArrayStructFromPool())
 			, ArrayPool(InArrayPool)
 		{
-			ObjectsToSerialize->AddUninitialized(NumObjects);
-			FMemory::Memcpy(ObjectsToSerialize->GetData(), InObjectsToSerialize->GetData() + StartIndex, NumObjects * sizeof(UObject*));
+			ArrayStruct->ObjectsToSerialize.AddUninitialized(NumObjects);
+			FMemory::Memcpy(ArrayStruct->ObjectsToSerialize.GetData(), InObjectsToSerialize->GetData() + StartIndex, NumObjects * sizeof(UObject*));
 		}
 		~FCollectorTask()
 		{
-			ArrayPool.ReturnToPool(ObjectsToSerialize);
+			ArrayPool.ReturnToPool(ArrayStruct);
 		}
 		FORCEINLINE TStatId GetStatId() const
 		{
@@ -255,7 +378,7 @@ private:
 		}
 		void DoTask(ENamedThreads::Type CurrentThread, FGraphEventRef& MyCompletionGraphEvent)
 		{
-			Owner->ProcessObjectArray(*ObjectsToSerialize, MyCompletionGraphEvent);
+			Owner->ProcessObjectArray(*ArrayStruct, MyCompletionGraphEvent);
 		}
 	};
 
@@ -295,14 +418,15 @@ public:
 	* @param ObjectsToCollectReferencesFor List of objects which references should be collected
 	* @param bForceSingleThreaded Collect references on a single thread
 	*/
-	void CollectReferences(TArray<UObject*>& ObjectsToCollectReferencesFor)
+	void CollectReferences(FGCArrayStruct& ArrayStruct)
 	{
+		TArray<UObject*>& ObjectsToCollectReferencesFor = ArrayStruct.ObjectsToSerialize;
 		if (ObjectsToCollectReferencesFor.Num())
 		{
 			if (!bParallel)
 			{
 				FGraphEventRef InvalidRef;
-				ProcessObjectArray(ObjectsToCollectReferencesFor, InvalidRef);
+				ProcessObjectArray(ArrayStruct, InvalidRef);
 			}
 			else
 			{
@@ -356,20 +480,20 @@ private:
 	 * @param InObjectsToSerializeArray Objects to process
 	 * @param MyCompletionGraphEvent Task graph event
 	 */
-	void ProcessObjectArray(TArray<UObject*>& InObjectsToSerializeArray, const FGraphEventRef& MyCompletionGraphEvent)
+	void ProcessObjectArray(FGCArrayStruct& InObjectsToSerializeStruct, const FGraphEventRef& MyCompletionGraphEvent)
 	{
 		DECLARE_SCOPE_CYCLE_COUNTER(TEXT("TFastReferenceCollector::ProcessObjectArray"), STAT_FFastReferenceCollector_ProcessObjectArray, STATGROUP_GC);
 
-		UObject* CurrentObject = NULL;
+		UObject* CurrentObject = nullptr;
 
 		const int32 MinDesiredObjectsPerSubTask = ReferenceProcessor.GetMinDesiredObjectsPerSubTask(); // sometimes there will be less, a lot less
 
 		/** Growing array of objects that require serialization */
-		TArray<UObject*>&	NewObjectsToSerializeArray = *ArrayPool.GetArrayFromPool();
+		FGCArrayStruct&	NewObjectsToSerializeStruct = *ArrayPool.GetArrayStructFromPool();
 
 		// Ping-pong between these two arrays if there's not enough objects to spawn a new task
-		TArray<UObject*>& ObjectsToSerialize = InObjectsToSerializeArray;
-		TArray<UObject*>& NewObjectsToSerialize = NewObjectsToSerializeArray;
+		TArray<UObject*>& ObjectsToSerialize = InObjectsToSerializeStruct.ObjectsToSerialize;
+		TArray<UObject*>& NewObjectsToSerialize = NewObjectsToSerializeStruct.ObjectsToSerialize;
 
 		// Presized "recursion" stack for handling arrays and structs.
 		TArray<FStackEntry> Stack;
@@ -383,7 +507,7 @@ private:
 		int32 CurrentIndex = 0;
 		do
 		{
-			CollectorType ReferenceCollector(ReferenceProcessor, NewObjectsToSerialize);
+			CollectorType ReferenceCollector(ReferenceProcessor, NewObjectsToSerializeStruct);
 			while (CurrentIndex < ObjectsToSerialize.Num())
 			{
 #if PERF_DETAILED_PER_CLASS_GC_STATS
@@ -662,6 +786,6 @@ EndLoop:
 		ReferenceProcessor.LogDetailedStatsSummary();
 #endif
 
-		ArrayPool.ReturnToPool(&NewObjectsToSerializeArray);
+		ArrayPool.ReturnToPool(&NewObjectsToSerializeStruct);
 	}
 };
diff --git a/Engine/Source/Runtime/CoreUObject/Public/UObject/NoExportTypes.h b/Engine/Source/Runtime/CoreUObject/Public/UObject/NoExportTypes.h
--- a/Engine/Source/Runtime/CoreUObject/Public/UObject/NoExportTypes.h
+++ b/Engine/Source/Runtime/CoreUObject/Public/UObject/NoExportTypes.h
@@ -957,9 +957,8 @@ struct FStringClassReference : public FStringAssetReference
 USTRUCT(noexport, BlueprintType)
 struct FPrimaryAssetType
 {
-private:
 	/** The Type of this object, by default it's base class's name */
-	UPROPERTY()
+	UPROPERTY(EditAnywhere, BlueprintReadWrite, Category = PrimaryAssetType)
 	FName Name;
 };
 
diff --git a/Engine/Source/Runtime/CoreUObject/Public/UObject/ObjectMacros.h b/Engine/Source/Runtime/CoreUObject/Public/UObject/ObjectMacros.h
--- a/Engine/Source/Runtime/CoreUObject/Public/UObject/ObjectMacros.h
+++ b/Engine/Source/Runtime/CoreUObject/Public/UObject/ObjectMacros.h
@@ -535,11 +535,10 @@ enum class EInternalObjectFlags : int32
 	Unreachable = 1 << 28, ///< Object is not reachable on the object graph.
 	PendingKill = 1 << 29, ///< Objects that are pending destruction (invalid for gameplay but valid objects)
 	RootSet = 1 << 30, ///< Object will not be garbage collected, even if unreferenced.
-	NoStrongReference = 1 << 31, ///< The object is not referenced by any strong reference. The flag is used by GC.
 
 	GarbageCollectionKeepFlags = Native | Async | AsyncLoading,
 	// Make sure this is up to date!
-	AllFlags = ReachableInCluster | ClusterRoot | Native | Async | AsyncLoading | Unreachable | PendingKill | RootSet | NoStrongReference
+	AllFlags = ReachableInCluster | ClusterRoot | Native | Async | AsyncLoading | Unreachable | PendingKill | RootSet
 };
 ENUM_CLASS_FLAGS(EInternalObjectFlags);
 
@@ -807,6 +806,12 @@ namespace UF
 		/// This function can be called from blueprint code and should be exposed to the user of blueprint editing tools.
 		BlueprintCallable,
 
+		/// This function is used as the get accessor for a blueprint exposed property. Implies BlueprintPure and BlueprintCallable.
+		BlueprintGetter,
+
+		/// This function is used as the set accessor for a blueprint exposed property. Implies BlueprintCallable.
+		BlueprintSetter,
+
 		/// This function will not execute from blueprint code if running on something without network authority
 		BlueprintAuthorityOnly,
 
@@ -929,9 +934,15 @@ namespace UP
 		/// This property can be read by blueprints, but not modified.
 		BlueprintReadOnly,
 
+		/// This property has an accessor to return the value. Implies BlueprintReadOnly if BlueprintSetter or BlueprintReadWrite is not specified. (usage: BlueprintGetter=FunctionName).
+		BlueprintGetter,
+
 		/// This property can be read or written from a blueprint.
 		BlueprintReadWrite,
 
+		/// This property has an accessor to set the value. Implies BlueprintReadWrite. (usage: BlueprintSetter=FunctionName).
+		BlueprintSetter,
+
 		/// The AssetRegistrySearchable keyword indicates that this property and it's value will be automatically added
 		/// to the asset registry for any asset class instances containing this as a member variable.  It is not legal
 		/// to use on struct properties or parameters.
diff --git a/Engine/Source/Runtime/CoreUObject/Public/UObject/UObjectArray.h b/Engine/Source/Runtime/CoreUObject/Public/UObject/UObjectArray.h
--- a/Engine/Source/Runtime/CoreUObject/Public/UObject/UObjectArray.h
+++ b/Engine/Source/Runtime/CoreUObject/Public/UObject/UObjectArray.h
@@ -184,18 +184,6 @@ struct FUObjectItem
 		return !!(Flags & int32(EInternalObjectFlags::RootSet));
 	}
 
-	FORCEINLINE void SetNoStrongReference()
-	{
-		Flags |= int32(EInternalObjectFlags::NoStrongReference);
-	}
-	FORCEINLINE void ClearNoStrongReference()
-	{
-		Flags &= ~int32(EInternalObjectFlags::NoStrongReference);
-	}
-	FORCEINLINE bool IsNoStrongReference() const
-	{
-		return !!(Flags & int32(EInternalObjectFlags::NoStrongReference));
-	}
 	FORCEINLINE void ResetSerialNumberAndFlags()
 	{
 		Flags = 0;
diff --git a/Engine/Source/Runtime/CoreUObject/Public/UObject/UObjectGlobals.h b/Engine/Source/Runtime/CoreUObject/Public/UObject/UObjectGlobals.h
--- a/Engine/Source/Runtime/CoreUObject/Public/UObject/UObjectGlobals.h
+++ b/Engine/Source/Runtime/CoreUObject/Public/UObject/UObjectGlobals.h
@@ -1689,15 +1689,17 @@ public:
 	 */
 	virtual void AllowEliminatingReferences(bool bAllow) {}
 	/**
-	* Sets the property that is currently being serialized
-	*/
+	 * Sets the property that is currently being serialized
+	 */
 	virtual void SetSerializedProperty(class UProperty* Inproperty) {}
 	/**
-	* Gets the property that is currently being serialized
-	*/
+	 * Gets the property that is currently being serialized
+	 */
 	virtual class UProperty* GetSerializedProperty() const { return nullptr; }
-
-	virtual void SetShouldHandleAsWeakRef(bool bWeakRef) {}
+	/** 
+	 * Marks a specific object reference as a weak reference. This does not affect GC but will be freed at a later point 
+	 */
+	virtual void MarkWeakObjectReferenceForClearing(UObject** WeakReference) {}
 protected:
 	/**
 	 * Handle object reference. Called by AddReferencedObject.
diff --git a/Engine/Source/Runtime/Engine/Classes/AI/Navigation/NavAreas/NavArea.h b/Engine/Source/Runtime/Engine/Classes/AI/Navigation/NavAreas/NavArea.h
--- a/Engine/Source/Runtime/Engine/Classes/AI/Navigation/NavAreas/NavArea.h
+++ b/Engine/Source/Runtime/Engine/Classes/AI/Navigation/NavAreas/NavArea.h
@@ -80,6 +80,7 @@ public:
 
 	virtual void FinishDestroy() override;
 	virtual void PostLoad() override;
+	virtual void PostInitProperties() override;
 	virtual void Serialize(FArchive& Ar) override;
 
 	FORCEINLINE uint16 GetAreaFlags() const { return AreaFlags; }
@@ -107,4 +108,6 @@ protected:
 
 	/** these flags will be applied to navigation data along with AreaID */
 	uint16 AreaFlags;
+	
+	void RegisterArea();
 };
diff --git a/Engine/Source/Runtime/Engine/Classes/AI/Navigation/NavigationData.h b/Engine/Source/Runtime/Engine/Classes/AI/Navigation/NavigationData.h
--- a/Engine/Source/Runtime/Engine/Classes/AI/Navigation/NavigationData.h
+++ b/Engine/Source/Runtime/Engine/Classes/AI/Navigation/NavigationData.h
@@ -586,6 +586,7 @@ public:
 
 	bool CanBeMainNavData() const { return bCanBeMainNavData; }
 	bool CanSpawnOnRebuild() const { return bCanSpawnOnRebuild; }
+	bool NeedsRebuildOnLoad() const { return bForceRebuildOnLoad; }
 
 protected:
 	virtual void FillConfig(FNavDataConfig& Dest) { Dest = NavDataConfig; }
diff --git a/Engine/Source/Runtime/Engine/Classes/AI/Navigation/NavigationSystem.h b/Engine/Source/Runtime/Engine/Classes/AI/Navigation/NavigationSystem.h
--- a/Engine/Source/Runtime/Engine/Classes/AI/Navigation/NavigationSystem.h
+++ b/Engine/Source/Runtime/Engine/Classes/AI/Navigation/NavigationSystem.h
@@ -955,7 +955,7 @@ private:
 	void AddNavigationBoundsUpdateRequest(const FNavigationBoundsUpdateRequest& UpdateRequest);
 
 	/** Triggers navigation building on all eligible navigation data. */
-	void RebuildAll();
+	void RebuildAll(bool bIsLoadTime = false);
 		 
 	/** Handler for FWorldDelegates::LevelAddedToWorld event */
 	 void OnLevelAddedToWorld(ULevel* InLevel, UWorld* InWorld);
diff --git a/Engine/Source/Runtime/Engine/Classes/Animation/AnimNotifies/AnimNotify_PlayParticleEffect.h b/Engine/Source/Runtime/Engine/Classes/Animation/AnimNotifies/AnimNotify_PlayParticleEffect.h
--- a/Engine/Source/Runtime/Engine/Classes/Animation/AnimNotifies/AnimNotify_PlayParticleEffect.h
+++ b/Engine/Source/Runtime/Engine/Classes/Animation/AnimNotifies/AnimNotify_PlayParticleEffect.h
@@ -34,15 +34,15 @@ public:
 	// End UAnimNotify interface
 
 	// Particle System to Spawn
-	UPROPERTY(EditAnywhere, Category="AnimNotify", meta=(DisplayName="Particle System"))
+	UPROPERTY(EditAnywhere, BlueprintReadWrite, Category="AnimNotify", meta=(DisplayName="Particle System"))
 	UParticleSystem* PSTemplate;
 
 	// Location offset from the socket
-	UPROPERTY(EditAnywhere, Category="AnimNotify")
+	UPROPERTY(EditAnywhere, BlueprintReadWrite, Category="AnimNotify")
 	FVector LocationOffset;
 
 	// Rotation offset from socket
-	UPROPERTY(EditAnywhere, Category="AnimNotify")
+	UPROPERTY(EditAnywhere, BlueprintReadWrite, Category="AnimNotify")
 	FRotator RotationOffset;
 
 private:
@@ -52,11 +52,11 @@ private:
 public:
 
 	// Should attach to the bone/socket
-	UPROPERTY(EditAnywhere, Category="AnimNotify")
+	UPROPERTY(EditAnywhere, BlueprintReadWrite, Category="AnimNotify")
 	uint32 Attached:1; 	//~ Does not follow coding standard due to redirection from BP
 
 	// SocketName to attach to
-	UPROPERTY(EditAnywhere, Category = "AnimNotify")
+	UPROPERTY(EditAnywhere, BlueprintReadWrite, Category = "AnimNotify")
 	FName SocketName;
 };
 
diff --git a/Engine/Source/Runtime/Engine/Classes/Animation/AnimNotifies/AnimNotify_PlaySound.h b/Engine/Source/Runtime/Engine/Classes/Animation/AnimNotifies/AnimNotify_PlaySound.h
--- a/Engine/Source/Runtime/Engine/Classes/Animation/AnimNotifies/AnimNotify_PlaySound.h
+++ b/Engine/Source/Runtime/Engine/Classes/Animation/AnimNotifies/AnimNotify_PlaySound.h
@@ -39,7 +39,7 @@ public:
 	float PitchMultiplier;
 
 	// If this sound should follow its owner
-	UPROPERTY(EditAnywhere, Category = "AnimNotify")
+	UPROPERTY(EditAnywhere, BlueprintReadWrite, Category = "AnimNotify")
 	uint32 bFollow:1;
 
 	// Socket or bone name to attach sound to
diff --git a/Engine/Source/Runtime/Engine/Classes/EdGraph/EdGraph.h b/Engine/Source/Runtime/Engine/Classes/EdGraph/EdGraph.h
--- a/Engine/Source/Runtime/Engine/Classes/EdGraph/EdGraph.h
+++ b/Engine/Source/Runtime/Engine/Classes/EdGraph/EdGraph.h
@@ -136,9 +136,18 @@ public:
 	 * a FGraphNodeCreator builder.
 	 */
 	template <typename NodeClass>
+	DEPRECATED(4.17, "Use CreateIntermediateNode instead.")
 	NodeClass* CreateBlankNode()
 	{
-		return (NodeClass*)CreateNode(NodeClass::StaticClass());
+		return CreateIntermediateNode<NodeClass>();
+	}
+
+	template <typename NodeClass>
+	NodeClass* CreateIntermediateNode()
+	{
+		NodeClass* Node = (NodeClass*)CreateNode(NodeClass::StaticClass());
+		FSetAsIntermediateNode SetAsIntermediate(Node);
+		return Node;
 	}
 
 	/** 
diff --git a/Engine/Source/Runtime/Engine/Classes/EdGraph/EdGraphNode.h b/Engine/Source/Runtime/Engine/Classes/EdGraph/EdGraphNode.h
--- a/Engine/Source/Runtime/Engine/Classes/EdGraph/EdGraphNode.h
+++ b/Engine/Source/Runtime/Engine/Classes/EdGraph/EdGraphNode.h
@@ -213,8 +213,14 @@ public:
 	/** If true, this node can be resized and should be drawn with a resize handle */
 	UPROPERTY()
 	uint8 bCanResizeNode:1;
+
 #endif // WITH_EDITORONLY_DATA
 
+private:
+	/** Whether the node was created as part of an expansion step */
+	uint8 bIsIntermediateNode : 1;
+
+public:
 	/** Flag to check for compile error/warning */
 	UPROPERTY()
 	uint8 bHasCompilerMessage:1;
@@ -355,6 +361,9 @@ public:
 	/** Find a pin on this node with the supplied name and remove it, returns TRUE if successful */
 	bool RemovePin(UEdGraphPin* Pin);
 
+	/** Returns whether the node was created by UEdGraph::CreateIntermediateNode. */
+	bool IsIntermediateNode() const { return bIsIntermediateNode; }
+
 	/** Whether or not this node should be given the chance to override pin names.  If this returns true, then GetPinNameOverride() will be called for each pin, each frame */
 	virtual bool ShouldOverridePinNames() const { return false; }
 
@@ -674,7 +683,19 @@ protected:
 	virtual FString GetPropertyNameAndValueForDiff(const UProperty* Prop, const uint8* PropertyAddr) const;
 
 #endif // WITH_EDITOR
+
+	friend struct FSetAsIntermediateNode;
 };
 
+struct FSetAsIntermediateNode
+{
+	friend UEdGraph;
+
+private:
+	FSetAsIntermediateNode(UEdGraphNode* GraphNode)
+	{
+		GraphNode->bIsIntermediateNode = true;
+	}
+};
 
 
diff --git a/Engine/Source/Runtime/Engine/Classes/Engine/AssetManager.h b/Engine/Source/Runtime/Engine/Classes/Engine/AssetManager.h
--- a/Engine/Source/Runtime/Engine/Classes/Engine/AssetManager.h
+++ b/Engine/Source/Runtime/Engine/Classes/Engine/AssetManager.h
@@ -33,6 +33,9 @@ public:
 	/** Returns the current AssetManager object */
 	static UAssetManager& Get();
 
+	/** Returns the current AssetManager object if it exists, null otherwise */
+	static UAssetManager* GetIfValid();
+
 	/** Accesses the StreamableManager used by this Asset Manager. Static for easy access */
 	static FStreamableManager& GetStreamableManager() { return Get().StreamableManager; }
 
@@ -114,6 +117,9 @@ public:
 	/** Gets the list of all FStringAssetReferences for a given type, returns true if any found */
 	virtual bool GetPrimaryAssetPathList(FPrimaryAssetType PrimaryAssetType, TArray<FStringAssetReference>& AssetPathList) const;
 
+	/** Sees if the passed in object is a registered primary asset, if so return it. Returns invalid Identifier if not found */
+	virtual FPrimaryAssetId GetPrimaryAssetIdForObject(UObject* Object) const;
+
 	/** Sees if the passed in object path is a registered primary asset, if so return it. Returns invalid Identifier if not found */
 	virtual FPrimaryAssetId GetPrimaryAssetIdForPath(const FStringAssetReference& ObjectPath) const;
 	virtual FPrimaryAssetId GetPrimaryAssetIdForPath(FName ObjectPath) const;
@@ -121,8 +127,8 @@ public:
 	/** Sees if the package has a primary asset, useful if only the package name is available */
 	virtual FPrimaryAssetId GetPrimaryAssetIdForPackage(FName PackagePath) const;
 
-	/** Parses AssetData to extract the primary type/name from it. This works even if it isn't in the directory yet */
-	virtual FPrimaryAssetId GetPrimaryAssetIdFromData(const FAssetData& AssetData, FPrimaryAssetType SuggestedType = NAME_None) const;
+	/** Returns the primary asset Id for the given FAssetData, only works if in directory */
+	virtual FPrimaryAssetId GetPrimaryAssetIdForData(const FAssetData& AssetData) const;
 
 	/** Gets list of all FPrimaryAssetId for a primary asset type, returns true if any were found */
 	virtual bool GetPrimaryAssetIdList(FPrimaryAssetType PrimaryAssetType, TArray<FPrimaryAssetId>& PrimaryAssetIdList) const;
@@ -185,6 +191,19 @@ public:
 	 */
 	virtual TSharedPtr<FStreamableHandle> ChangeBundleStateForPrimaryAssets(const TArray<FPrimaryAssetId>& AssetsToChange, const TArray<FName>& AddBundles, const TArray<FName>& RemoveBundles, bool bRemoveAllBundles = false, FStreamableDelegate DelegateToCall = FStreamableDelegate(), TAsyncLoadPriority Priority = FStreamableManager::DefaultAsyncLoadPriority);
 
+	/** 
+	 * Changes the bundle state of all loaded primary assets. Only assets matching OldBundles will be modified
+	 * You can wait on the returned streamable request or poll as needed.
+	 * If there is no work to do, returned handle will be null and delegate will get called before function returns.
+	 *
+	 * @param NewBundles		New bundle state for the assets that are changed
+	 * @param OldBundles		Old bundle state, it will remove these bundles and replace with NewBundles
+	 * @param DelegateToCall	Delegate that will be called on completion, may be called before function returns if assets are already loaded
+	 * @param Priority			Async loading priority for this request
+	 * @return					Streamable Handle that can be used to poll or wait. You do not need to keep this handle to stop the assets from being unloaded
+	 */
+	virtual TSharedPtr<FStreamableHandle> ChangeBundleStateForMatchingPrimaryAssets(const TArray<FName>& NewBundles, const TArray<FName>& OldBundles, FStreamableDelegate DelegateToCall = FStreamableDelegate(), TAsyncLoadPriority Priority = FStreamableManager::DefaultAsyncLoadPriority);
+
 	/** 
 	 * Returns the loading handle associated with the primary asset, it can then be checked for progress or waited on
 	 *
@@ -205,7 +224,7 @@ public:
 	 * @param bForceCurrent		If true, only use the current state. If false, use the current or pending
 	 * @return					True if any found
 	 */
-	bool GetPrimaryAssetsWithBundleState(TArray<FPrimaryAssetId>& PrimaryAssetList, const TArray<FName>& ValidTypes, const TArray<FName>& RequiredBundles, const TArray<FName>& ExcludedBundles = TArray<FName>(), bool bForceCurrent = false) const;
+	bool GetPrimaryAssetsWithBundleState(TArray<FPrimaryAssetId>& PrimaryAssetList, const TArray<FPrimaryAssetType>& ValidTypes, const TArray<FName>& RequiredBundles, const TArray<FName>& ExcludedBundles = TArray<FName>(), bool bForceCurrent = false) const;
 
 	/** Fills in a TMap with the pending/active loading state of every asset */
 	void GetPrimaryAssetBundleStateMap(TMap<FPrimaryAssetId, TArray<FName>>& BundleStateMap, bool bForceCurrent = false) const;
@@ -253,6 +272,9 @@ public:
 
 	// GENERAL ASSET UTILITY FUNCTIONS
 
+	/** Parses AssetData to extract the primary type/name from it. This works even if it isn't in the directory */
+	virtual FPrimaryAssetId ExtractPrimaryAssetIdFromData(const FAssetData& AssetData, FPrimaryAssetType SuggestedType = NAME_None) const;
+
 	/** Gets the FAssetData at a specific path, handles redirectors and blueprint classes correctly. Returns true if it found a valid data */
 	virtual bool GetAssetDataForPath(const FStringAssetReference& ObjectPath, FAssetData& AssetData) const;
 
@@ -320,6 +342,9 @@ public:
 	/** Refresh the entire set of asset data, can call from editor when things have changed dramatically */
 	virtual void RefreshPrimaryAssetDirectory();
 
+	/** Resets all asset manager data, called in the editor to reinitialize the config */
+	virtual void ReinitializeFromConfig();
+
 	/** Updates the asset management database if needed */
 	virtual void UpdateManagementDatabase(bool bForceRefresh = false);
 
@@ -463,6 +488,10 @@ protected:
 	UPROPERTY()
 	bool bUpdateManagementDatabaseAfterScan;
 
+	/** Number of notifications seen in this update */
+	UPROPERTY()
+	int32 NumberOfSpawnedNotifications;
+
 	/** Redirector maps loaded out of AssetMigrations.ini */
 	TMap<FName, FName> PrimaryAssetTypeRedirects;
 	TMap<FString, FString> PrimaryAssetIdRedirects;
diff --git a/Engine/Source/Runtime/Engine/Classes/Engine/AssetManagerSettings.h b/Engine/Source/Runtime/Engine/Classes/Engine/AssetManagerSettings.h
--- a/Engine/Source/Runtime/Engine/Classes/Engine/AssetManagerSettings.h
+++ b/Engine/Source/Runtime/Engine/Classes/Engine/AssetManagerSettings.h
@@ -53,7 +53,7 @@ public:
 	TArray<FPrimaryAssetTypeInfo> PrimaryAssetTypesToScan;
 
 	/** List of directories to exclude from scanning for Primary Assets, useful to exclude test assets */
-	UPROPERTY(config, EditAnywhere, Category = "Asset Manager", meta = (RelativeToGameContentDir))
+	UPROPERTY(config, EditAnywhere, Category = "Asset Manager", meta = (RelativeToGameContentDir, LongPackageName))
 	TArray<FDirectoryPath> DirectoriesToExclude;
 
 	/** List of specific asset rule overrides */
@@ -79,4 +79,8 @@ public:
 	/** Redirect from /game/assetpath to /game/assetpathnew */
 	UPROPERTY(config, EditAnywhere, Category = "Redirects")
 	TArray<FAssetManagerRedirect> AssetPathRedirects;
+
+#if WITH_EDITOR
+	virtual void PostEditChangeProperty(FPropertyChangedEvent& PropertyChangedEvent) override;
+#endif
 };
\ No newline at end of file
diff --git a/Engine/Source/Runtime/Engine/Classes/Engine/AssetManagerTypes.h b/Engine/Source/Runtime/Engine/Classes/Engine/AssetManagerTypes.h
--- a/Engine/Source/Runtime/Engine/Classes/Engine/AssetManagerTypes.h
+++ b/Engine/Source/Runtime/Engine/Classes/Engine/AssetManagerTypes.h
@@ -103,7 +103,7 @@ public:
 
 private:
 	/** Directories to search for this asset type */
-	UPROPERTY(EditAnywhere, Category = AssetType, meta = (RelativeToGameContentDir))
+	UPROPERTY(EditAnywhere, Category = AssetType, meta = (RelativeToGameContentDir, LongPackageName))
 	TArray<FDirectoryPath> Directories;
 
 	/** Individual assets to scan */
diff --git a/Engine/Source/Runtime/Engine/Classes/Engine/PrimaryAssetLabel.h b/Engine/Source/Runtime/Engine/Classes/Engine/PrimaryAssetLabel.h
--- a/Engine/Source/Runtime/Engine/Classes/Engine/PrimaryAssetLabel.h
+++ b/Engine/Source/Runtime/Engine/Classes/Engine/PrimaryAssetLabel.h
@@ -27,6 +27,10 @@ public:
 	UPROPERTY(EditAnywhere, Category = PrimaryAssetLabel)
 	uint32 bLabelAssetsInMyDirectory : 1;
 
+	/** Set to true if the label asset itself should be cooked and available at runtime. This does not affect the assets that are labeled, they are set with cook rule */
+	UPROPERTY(EditAnywhere, Category = PrimaryAssetLabel)
+	uint32 bIsRuntimeLabel : 1;
+
 	/** List of manually specified assets to label */
 	UPROPERTY(EditAnywhere, Category = PrimaryAssetLabel, meta = (AssetBundles = "Explicit"))
 	TArray<TAssetPtr<UObject>> ExplicitAssets;
@@ -35,6 +39,12 @@ public:
 	UPROPERTY(EditAnywhere, Category = PrimaryAssetLabel, meta = (AssetBundles = "Explicit", BlueprintBaseOnly))
 	TArray<TAssetSubclassOf<UObject>> ExplicitBlueprints;
 
+	/** Set to editor only if this is not available in a cooked build */
+	virtual bool IsEditorOnly() const
+	{
+		return !bIsRuntimeLabel;
+	}
+
 	// TODO add collections
 
 #if WITH_EDITORONLY_DATA
diff --git a/Engine/Source/Runtime/Engine/Classes/Engine/StreamableManager.h b/Engine/Source/Runtime/Engine/Classes/Engine/StreamableManager.h
--- a/Engine/Source/Runtime/Engine/Classes/Engine/StreamableManager.h
+++ b/Engine/Source/Runtime/Engine/Classes/Engine/StreamableManager.h
@@ -297,7 +297,7 @@ private:
 	UObject* GetStreamed(const FStringAssetReference& Target) const;
 	void CheckCompletedRequests(const FStringAssetReference& Target, struct FStreamable* Existing);
 
-	void OnPostGarbageCollect();
+	void OnPreGarbageCollect();
 	void AsyncLoadCallback(FStringAssetReference Request);
 
 	/** Map of paths to streamable objects, this will be the post-redirector name */
diff --git a/Engine/Source/Runtime/Engine/Classes/GameFramework/Actor.h b/Engine/Source/Runtime/Engine/Classes/GameFramework/Actor.h
--- a/Engine/Source/Runtime/Engine/Classes/GameFramework/Actor.h
+++ b/Engine/Source/Runtime/Engine/Classes/GameFramework/Actor.h
@@ -469,7 +469,7 @@ protected:
 	/**
 	 * Collision primitive that defines the transform (location, rotation, scale) of this Actor.
 	 */
-	UPROPERTY()
+	UPROPERTY(BlueprintGetter=K2_GetRootComponent, Category="Utilities|Transformation")
 	USceneComponent* RootComponent;
 
 #if WITH_EDITORONLY_DATA
@@ -805,7 +805,7 @@ public:
 	void GetActorBounds(bool bOnlyCollidingComponents, FVector& Origin, FVector& BoxExtent) const;
 
 	/** Returns the RootComponent of this Actor */
-	UFUNCTION(BlueprintCallable, meta=(DisplayName = "GetRootComponent"), Category="Utilities|Transformation")
+	UFUNCTION(BlueprintGetter)
 	USceneComponent* K2_GetRootComponent() const;
 
 	/** Returns velocity (in cm/s (Unreal Units/second) of the rootcomponent if it is either using physics or has an associated MovementComponent */
diff --git a/Engine/Source/Runtime/Engine/Classes/Kismet/BlueprintAsyncActionBase.h b/Engine/Source/Runtime/Engine/Classes/Kismet/BlueprintAsyncActionBase.h
--- a/Engine/Source/Runtime/Engine/Classes/Kismet/BlueprintAsyncActionBase.h
+++ b/Engine/Source/Runtime/Engine/Classes/Kismet/BlueprintAsyncActionBase.h
@@ -12,7 +12,10 @@ class ENGINE_API UBlueprintAsyncActionBase : public UObject
 {
 	GENERATED_UCLASS_BODY()
 
-	// Called to trigger the action once the delegates have been bound
+	/** Called to trigger the action once the delegates have been bound */
 	UFUNCTION(BlueprintCallable, meta=(BlueprintInternalUseOnly="true"))
 	virtual void Activate();
+
+	/** Call when the action is completely done, this makes the action free to delete */
+	virtual void SetReadyToDestroy();
 };
diff --git a/Engine/Source/Runtime/Engine/Classes/Kismet/KismetSystemLibrary.h b/Engine/Source/Runtime/Engine/Classes/Kismet/KismetSystemLibrary.h
--- a/Engine/Source/Runtime/Engine/Classes/Kismet/KismetSystemLibrary.h
+++ b/Engine/Source/Runtime/Engine/Classes/Kismet/KismetSystemLibrary.h
@@ -169,14 +169,45 @@ class ENGINE_API UKismetSystemLibrary : public UBlueprintFunctionLibrary
 	UFUNCTION(BlueprintPure, Category="Utilities|Platform")
 	static FString GetDeviceId();
 
+	/** Converts an interfance into an object */
 	UFUNCTION(BlueprintPure, meta=(DisplayName = "ToObject (interface)", CompactNodeTitle = "->"), Category="Utilities")
 	static UObject* Conv_InterfaceToObject(const FScriptInterface& Interface); 
 
+	/** Returns true if the Soft Object Reference is not null */
+	UFUNCTION(BlueprintPure, Category = "Utilities")
+	static bool IsValidSoftObjectReference(const TAssetPtr<UObject>& SoftObjectReference);
+
+	/** Returns true if the values are equal (A == B) */
+	UFUNCTION(BlueprintPure, meta = (DisplayName = "Equal (SoftObjectReference)", CompactNodeTitle = "=="), Category = "Utilities")
+	static bool EqualEqual_SoftObjectReference(const TAssetPtr<UObject>& A, const TAssetPtr<UObject>& B);
+
+	/** Returns true if the values are not equal (A != B) */
+	UFUNCTION(BlueprintPure, meta = (DisplayName = "NotEqual (SoftObjectReference)", CompactNodeTitle = "!="), Category = "Utilities")
+	static bool NotEqual_SoftObjectReference(const TAssetPtr<UObject>& A, const TAssetPtr<UObject>& B);
+
+	/** Returns true if the Soft Class Reference is not null */
+	UFUNCTION(BlueprintPure, Category = "Utilities")
+	static bool IsValidSoftClassReference(const TAssetSubclassOf<UObject>& SoftClassReference);
+
+	/** Returns true if the values are equal (A == B) */
+	UFUNCTION(BlueprintPure, meta = (DisplayName = "Equal (SoftClassReference)", CompactNodeTitle = "=="), Category = "Utilities")
+	static bool EqualEqual_SoftClassReference(const TAssetSubclassOf<UObject>& A, const TAssetSubclassOf<UObject>& B);
+
+	/** Returns true if the values are not equal (A != B) */
+	UFUNCTION(BlueprintPure, meta = (DisplayName = "NotEqual (SoftClassReference)", CompactNodeTitle = "!="), Category = "Utilities")
+	static bool NotEqual_SoftClassReference(const TAssetSubclassOf<UObject>& A, const TAssetSubclassOf<UObject>& B);
+
+	UFUNCTION(BlueprintPure, meta = (BlueprintInternalUseOnly = "true"), Category = "Utilities")
+	static UObject* Conv_SoftObjectReferenceToObject(const TAssetPtr<UObject>& Asset);
+
+	UFUNCTION(BlueprintPure, meta = (BlueprintInternalUseOnly = "true"), Category = "Utilities")
+	static TSubclassOf<UObject> Conv_SoftClassReferenceToClass(const TAssetSubclassOf<UObject>& AssetClass);
+
 	UFUNCTION(BlueprintPure, meta = (BlueprintInternalUseOnly = "true"), Category = "Utilities")
-	static UObject* Conv_AssetToObject(const TAssetPtr<UObject>& Asset);
+	static TAssetPtr<UObject> Conv_ObjectToSoftObjectReference(UObject* Object);
 
 	UFUNCTION(BlueprintPure, meta = (BlueprintInternalUseOnly = "true"), Category = "Utilities")
-	static TSubclassOf<UObject> Conv_AssetClassToClass(const TAssetSubclassOf<UObject>& AssetClass);
+	static TAssetSubclassOf<UObject> Conv_ClassToSoftClassReference(const TSubclassOf<UObject>& Class);
 
 	DECLARE_DYNAMIC_DELEGATE_OneParam(FOnAssetLoaded, class UObject*, Loaded);
 
@@ -1556,6 +1587,92 @@ class ENGINE_API UKismetSystemLibrary : public UBlueprintFunctionLibrary
 	 */
 	UFUNCTION(BlueprintCallable, Category="Utilities")
 	static FString GetCommandLine();
+
+	// --- Asset Manager ------------------------------
+
+	/** Returns the Object associated with a Primary Asset Id, this will only return a valid object if it is in memory, it will not load it */
+	UFUNCTION(BlueprintPure, Category = "AssetManager")
+	static UObject* GetObjectFromPrimaryAssetId(FPrimaryAssetId PrimaryAssetId);
+
+	/** Returns the Blueprint Class associated with a Primary Asset Id, this will only return a valid object if it is in memory, it will not load it */
+	UFUNCTION(BlueprintPure, Category = "AssetManager")
+	static TSubclassOf<UObject> GetClassFromPrimaryAssetId(FPrimaryAssetId PrimaryAssetId);
+
+	/** Returns the Object Id associated with a Primary Asset Id, this works even if the asset is not loaded */
+	UFUNCTION(BlueprintPure, Category = "AssetManager")
+	static TAssetPtr<UObject> GetSoftObjectReferenceFromPrimaryAssetId(FPrimaryAssetId PrimaryAssetId);
+
+	/** Returns the Blueprint Class Id associated with a Primary Asset Id, this works even if the asset is not loaded */
+	UFUNCTION(BlueprintPure, Category = "AssetManager")
+	static TAssetSubclassOf<UObject> GetSoftClassReferenceFromPrimaryAssetId(FPrimaryAssetId PrimaryAssetId);
+
+	/** Returns the Primary Asset Id for an Object, this can return an invalid one if not registered */
+	UFUNCTION(BlueprintPure, Category = "AssetManager")
+	static FPrimaryAssetId GetPrimaryAssetIdFromObject(UObject* Object);
+
+	/** Returns the Primary Asset Id for a Class, this can return an invalid one if not registered */
+	UFUNCTION(BlueprintPure, Category = "AssetManager")
+	static FPrimaryAssetId GetPrimaryAssetIdFromClass(TSubclassOf<UObject> Class);
+
+	/** Returns the Primary Asset Id for a Soft Object Reference, this can return an invalid one if not registered */
+	UFUNCTION(BlueprintPure, Category = "AssetManager")
+	static FPrimaryAssetId GetPrimaryAssetIdFromSoftObjectReference(TAssetPtr<UObject> SoftObjectReference);
+
+	/** Returns the Primary Asset Id for a Soft Class Reference, this can return an invalid one if not registered */
+	UFUNCTION(BlueprintPure, Category = "AssetManager")
+	static FPrimaryAssetId GetPrimaryAssetIdFromSoftClassReference(TAssetSubclassOf<UObject> SoftClassReference);
+
+	/** Returns list of PrimaryAssetIds for a PrimaryAssetType */
+	UFUNCTION(BlueprintCallable, Category = "AssetManager")
+	static void GetPrimaryAssetIdList(FPrimaryAssetType PrimaryAssetType, TArray<FPrimaryAssetId>& OutPrimaryAssetIdList);
+
+	/** Returns true if the Primary Asset Id is valid */
+	UFUNCTION(BlueprintPure, Category = "AssetManager")
+	static bool IsValidPrimaryAssetId(FPrimaryAssetId PrimaryAssetId);
+
+	/** Returns true if the values are equal (A == B) */
+	UFUNCTION(BlueprintPure, meta = (DisplayName = "Equal (PrimaryAssetId)", CompactNodeTitle = "=="), Category = "AssetManager")
+	static bool EqualEqual_PrimaryAssetId(FPrimaryAssetId A, FPrimaryAssetId B);
+
+	/** Returns true if the values are not equal (A != B) */
+	UFUNCTION(BlueprintPure, meta = (DisplayName = "NotEqual (PrimaryAssetId)", CompactNodeTitle = "!="), Category = "AssetManager")
+	static bool NotEqual_PrimaryAssetId(FPrimaryAssetId A, FPrimaryAssetId B);
+
+	/** Returns list of PrimaryAssetIds for a PrimaryAssetType */
+	UFUNCTION(BlueprintPure, Category = "AssetManager")
+	static bool IsValidPrimaryAssetType(FPrimaryAssetType PrimaryAssetType);
+
+	/** Returns true if the values are equal (A == B) */
+	UFUNCTION(BlueprintPure, meta = (DisplayName = "Equal (PrimaryAssetType)", CompactNodeTitle = "=="), Category = "AssetManager")
+	static bool EqualEqual_PrimaryAssetType(FPrimaryAssetType A, FPrimaryAssetType B);
+
+	/** Returns true if the values are not equal (A != B) */
+	UFUNCTION(BlueprintPure, meta = (DisplayName = "NotEqual (PrimaryAssetType)", CompactNodeTitle = "!="), Category = "AssetManager")
+	static bool NotEqual_PrimaryAssetType(FPrimaryAssetType A, FPrimaryAssetType B);
+
+	/** Unloads a primary asset, which allows it to be garbage collected if nothing else is referencing it */
+	UFUNCTION(BlueprintCallable, Category = "AssetManager")
+	static void UnloadPrimaryAsset(FPrimaryAssetId PrimaryAssetId);
+
+	/** Unloads a primary asset, which allows it to be garbage collected if nothing else is referencing it */
+	UFUNCTION(BlueprintCallable, Category = "AssetManager")
+	static void UnloadPrimaryAssetList(const TArray<FPrimaryAssetId>& PrimaryAssetIdList);
+
+	/** 
+	 * Returns the list of loaded bundles for a given Primary Asset. This will return false if the asset is not loaded at all.
+	 * If ForceCurrentState is true it will return the current state even if a load is in process
+	 */
+	UFUNCTION(BlueprintCallable, Category = "AssetManager")
+	static bool GetCurrentBundleState(FPrimaryAssetId PrimaryAssetId, bool bForceCurrentState, TArray<FName>& OutBundles);
+
+	/** 
+	 * Returns the list of assets that are in a given bundle state. Required Bundles must be specified
+	 * If ExcludedBundles is not empty, it will not return any assets in those bundle states
+	 * If ValidTypes is not empty, it will only return assets of those types
+	 * If ForceCurrentState is true it will use the current state even if a load is in process
+	 */
+	UFUNCTION(BlueprintCallable, Category = "AssetManager", meta=(AutoCreateRefTerm = "ExcludedBundles, ValidTypes"))
+	static void GetPrimaryAssetsWithBundleState(const TArray<FName>& RequiredBundles, const TArray<FName>& ExcludedBundles, const TArray<FPrimaryAssetType>& ValidTypes, bool bForceCurrentState, TArray<FPrimaryAssetId>& OutPrimaryAssetIdList);
 };
 
 
diff --git a/Engine/Source/Runtime/Engine/Private/AI/Navigation/NavAreas/NavArea.cpp b/Engine/Source/Runtime/Engine/Private/AI/Navigation/NavAreas/NavArea.cpp
--- a/Engine/Source/Runtime/Engine/Private/AI/Navigation/NavAreas/NavArea.cpp
+++ b/Engine/Source/Runtime/Engine/Private/AI/Navigation/NavAreas/NavArea.cpp
@@ -12,7 +12,7 @@ UNavArea::UNavArea(const FObjectInitializer& ObjectInitializer) : Super(ObjectIn
 	DrawColor = FColor::Magenta;
 	SupportedAgentsBits = 0xffffffff;
 	// NOTE! AreaFlags == 0 means UNWALKABLE!
-	AreaFlags = 1;  
+	AreaFlags = 1;
 }
 
 void UNavArea::FinishDestroy()
@@ -32,7 +32,17 @@ void UNavArea::FinishDestroy()
 void UNavArea::PostLoad()
 {
 	Super::PostLoad();
+	RegisterArea();
+}
 
+void UNavArea::PostInitProperties()
+{
+	Super::PostInitProperties();
+	RegisterArea();
+}
+
+void UNavArea::RegisterArea()
+{
 	if (HasAnyFlags(RF_ClassDefaultObject)
 #if WITH_HOT_RELOAD
 		&& !GIsHotReload
diff --git a/Engine/Source/Runtime/Engine/Private/AI/Navigation/NavigationSystem.cpp b/Engine/Source/Runtime/Engine/Private/AI/Navigation/NavigationSystem.cpp
--- a/Engine/Source/Runtime/Engine/Private/AI/Navigation/NavigationSystem.cpp
+++ b/Engine/Source/Runtime/Engine/Private/AI/Navigation/NavigationSystem.cpp
@@ -720,17 +720,23 @@ void UNavigationSystem::OnWorldInitDone(FNavigationSystemRunMode Mode)
 		}
 	}
 
-	if (Mode == FNavigationSystemRunMode::EditorMode && bGenerateNavigationOnlyAroundNavigationInvokers)
+	if (Mode == FNavigationSystemRunMode::EditorMode)
 	{
-		UWorld* MyWorld = GetWorld();
-		// gather enforcers manually to be able to see the results in editor as well
-		for (TObjectIterator<UNavigationInvokerComponent> It; It; ++It)
+		// update navigation invokers
+		if (bGenerateNavigationOnlyAroundNavigationInvokers)
 		{
-			if (MyWorld == It->GetWorld())
+			for (TObjectIterator<UNavigationInvokerComponent> It; It; ++It)
 			{
-				It->RegisterWithNavigationSystem(*this);
+				if (World == It->GetWorld())
+				{
+					It->RegisterWithNavigationSystem(*this);
+				}
 			}
 		}
+
+		// update navdata after loading world
+		const bool bIsLoadTime = true;
+		RebuildAll(bIsLoadTime);
 	}
 
 	bWorldInitDone = true;
@@ -3369,7 +3375,7 @@ void UNavigationSystem::RemoveNavigationBuildLock(uint8 Flags, bool bSkipRebuild
 	}
 }
 
-void UNavigationSystem::RebuildAll()
+void UNavigationSystem::RebuildAll(bool bIsLoadTime)
 {
 	const bool bIsInGame = GetWorld()->IsGameWorld();
 	
@@ -3391,7 +3397,7 @@ void UNavigationSystem::RebuildAll()
 	{
 		ANavigationData* NavData = NavDataSet[NavDataIndex];
 				
-		if (NavData && (!bIsInGame || NavData->SupportsRuntimeGeneration()))
+		if (NavData && (!bIsLoadTime || NavData->NeedsRebuildOnLoad()) && (!bIsInGame || NavData->SupportsRuntimeGeneration()))
 		{
 			NavData->RebuildAll();
 		}
diff --git a/Engine/Source/Runtime/Engine/Private/AssetManager.cpp b/Engine/Source/Runtime/Engine/Private/AssetManager.cpp
--- a/Engine/Source/Runtime/Engine/Private/AssetManager.cpp
+++ b/Engine/Source/Runtime/Engine/Private/AssetManager.cpp
@@ -104,6 +104,7 @@ UAssetManager::UAssetManager()
 	bIsBulkScanning = false;
 	bIsManagementDatabaseCurrent = false;
 	bUpdateManagementDatabaseAfterScan = false;
+	NumberOfSpawnedNotifications = 0;
 }
 
 void UAssetManager::PostInitProperties()
@@ -167,6 +168,16 @@ UAssetManager& UAssetManager::Get()
 	}
 }
 
+UAssetManager* UAssetManager::GetIfValid()
+{
+	if (GEngine && GEngine->AssetManager)
+	{
+		return GEngine->AssetManager;
+	}
+
+	return nullptr;
+}
+
 IAssetRegistry& UAssetManager::GetAssetRegistry() const
 {
 	if (!CachedAssetRegistry)
@@ -375,7 +386,7 @@ int32 UAssetManager::ScanPathsForPrimaryAssets(FPrimaryAssetType PrimaryAssetTyp
 			}
 		}
 
-		FPrimaryAssetId PrimaryAssetId = GetPrimaryAssetIdFromData(Data, PrimaryAssetType);
+		FPrimaryAssetId PrimaryAssetId = ExtractPrimaryAssetIdFromData(Data, PrimaryAssetType);
 
 		// Remove invalid or wrong type assets
 		if (!PrimaryAssetId.IsValid() || PrimaryAssetId.PrimaryAssetType != PrimaryAssetType)
@@ -401,6 +412,7 @@ void UAssetManager::StartBulkScanning()
 	if (ensure(!bIsBulkScanning))
 	{
 		bIsBulkScanning = true;
+		NumberOfSpawnedNotifications = 0;
 	}
 }
 
@@ -436,14 +448,18 @@ void UAssetManager::UpdateCachedAssetData(const FPrimaryAssetId& PrimaryAssetId,
 #if WITH_EDITOR
 			if (GIsEditor)
 			{
-				FNotificationInfo Info(FText::Format(LOCTEXT("DuplicateAssetId", "Duplicate Asset ID {0} used by {1} and {2}, you must delete or rename one!"), 
-					FText::FromString(PrimaryAssetId.ToString()), FText::FromString(OldData->AssetPtr.ToStringReference().GetLongPackageName()), FText::FromString(NewStringReference.GetLongPackageName())));
-				Info.ExpireDuration = 30.0f;
-				
-				TSharedPtr<SNotificationItem> Notification = FSlateNotificationManager::Get().AddNotification(Info);
-				if (Notification.IsValid())
+				const int MaxNotificationsPerFrame = 5;
+				if (NumberOfSpawnedNotifications++ < MaxNotificationsPerFrame)
 				{
-					Notification->SetCompletionState(SNotificationItem::CS_Fail);
+					FNotificationInfo Info(FText::Format(LOCTEXT("DuplicateAssetId", "Duplicate Asset ID {0} used by {1} and {2}, you must delete or rename one!"),
+						FText::FromString(PrimaryAssetId.ToString()), FText::FromString(OldData->AssetPtr.ToStringReference().GetLongPackageName()), FText::FromString(NewStringReference.GetLongPackageName())));
+					Info.ExpireDuration = 30.0f;
+
+					TSharedPtr<SNotificationItem> Notification = FSlateNotificationManager::Get().AddNotification(Info);
+					if (Notification.IsValid())
+					{
+						Notification->SetCompletionState(SNotificationItem::CS_Fail);
+					}
 				}
 			}
 #endif
@@ -769,6 +785,17 @@ bool UAssetManager::GetPrimaryAssetPathList(FPrimaryAssetType PrimaryAssetType,
 	return AssetPathList.Num() > 0;
 }
 
+FPrimaryAssetId UAssetManager::GetPrimaryAssetIdForObject(UObject* Object) const
+{
+	// Use path instead of calling on Object, we only want it if it's registered
+	return GetPrimaryAssetIdForPath(FName(*Object->GetPathName()));
+}
+
+FPrimaryAssetId UAssetManager::GetPrimaryAssetIdForData(const FAssetData& AssetData) const
+{
+	return GetPrimaryAssetIdForPath(GetAssetPathForData(AssetData));
+}
+
 FPrimaryAssetId UAssetManager::GetPrimaryAssetIdForPath(const FStringAssetReference& ObjectPath) const
 {
 	FName PossibleAssetPath = FName(*ObjectPath.ToString(), FNAME_Find);
@@ -835,7 +862,7 @@ FPrimaryAssetId UAssetManager::GetPrimaryAssetIdForPackage(FName PackagePath) co
 	return FoundId;
 }
 
-FPrimaryAssetId UAssetManager::GetPrimaryAssetIdFromData(const FAssetData& AssetData, FPrimaryAssetType SuggestedType) const
+FPrimaryAssetId UAssetManager::ExtractPrimaryAssetIdFromData(const FAssetData& AssetData, FPrimaryAssetType SuggestedType) const
 {
 	FPrimaryAssetId FoundId = AssetData.GetPrimaryAssetId();
 
@@ -1074,6 +1101,20 @@ TSharedPtr<FStreamableHandle> UAssetManager::ChangeBundleStateForPrimaryAssets(c
 	return ReturnHandle;
 }
 
+TSharedPtr<FStreamableHandle> UAssetManager::ChangeBundleStateForMatchingPrimaryAssets(const TArray<FName>& NewBundles, const TArray<FName>& OldBundles, FStreamableDelegate DelegateToCall, TAsyncLoadPriority Priority)
+{
+	TArray<FPrimaryAssetId> AssetsToChange;
+
+	if (GetPrimaryAssetsWithBundleState(AssetsToChange, TArray<FPrimaryAssetType>(), OldBundles))
+	{
+		// This will call delegate when done
+		return ChangeBundleStateForPrimaryAssets(AssetsToChange, NewBundles, OldBundles, false, DelegateToCall, Priority);
+	}
+
+	// Nothing to transition, call delegate now
+	DelegateToCall.ExecuteIfBound();
+	return nullptr;
+}
 
 TSharedPtr<FStreamableHandle> UAssetManager::PreloadPrimaryAssets(const TArray<FPrimaryAssetId>& AssetsToLoad, const TArray<FName>& LoadBundles, bool bLoadRecursive, FStreamableDelegate DelegateToCall, TAsyncLoadPriority Priority)
 {
@@ -1210,13 +1251,13 @@ TSharedPtr<FStreamableHandle> UAssetManager::GetPrimaryAssetHandle(const FPrimar
 	return LoadState.Handle;
 }
 
-bool UAssetManager::GetPrimaryAssetsWithBundleState(TArray<FPrimaryAssetId>& PrimaryAssetList, const TArray<FName>& ValidTypes, const TArray<FName>& RequiredBundles, const TArray<FName>& ExcludedBundles, bool bForceCurrent) const
+bool UAssetManager::GetPrimaryAssetsWithBundleState(TArray<FPrimaryAssetId>& PrimaryAssetList, const TArray<FPrimaryAssetType>& ValidTypes, const TArray<FName>& RequiredBundles, const TArray<FName>& ExcludedBundles, bool bForceCurrent) const
 {
 	bool bFoundAny = false;
 
 	for (const TPair<FName, TSharedRef<FPrimaryAssetTypeData>>& TypePair : AssetTypeMap)
 	{
-		if (ValidTypes.Num() > 0 && !ValidTypes.Contains(TypePair.Key))
+		if (ValidTypes.Num() > 0 && !ValidTypes.Contains(FPrimaryAssetType(TypePair.Key)))
 		{
 			// Skip this type
 			continue;
@@ -2547,6 +2588,19 @@ void UAssetManager::RefreshPrimaryAssetDirectory()
 	PostInitialAssetScan();
 }
 
+void UAssetManager::ReinitializeFromConfig()
+{
+	AssetPathMap.Reset();
+	AssetRuleOverrides.Reset();
+	ManagementParentMap.Reset();
+	CachedAssetBundles.Reset();
+	AlreadyScannedDirectories.Reset();
+	AssetTypeMap.Reset();
+
+	LoadRedirectorMaps();
+	ScanPrimaryAssetTypesFromConfig();
+}
+
 void UAssetManager::OnInMemoryAssetCreated(UObject *Object)
 {
 	// Ignore PIE and CDO changes
diff --git a/Engine/Source/Runtime/Engine/Private/AssetManagerTypes.cpp b/Engine/Source/Runtime/Engine/Private/AssetManagerTypes.cpp
--- a/Engine/Source/Runtime/Engine/Private/AssetManagerTypes.cpp
+++ b/Engine/Source/Runtime/Engine/Private/AssetManagerTypes.cpp
@@ -2,6 +2,7 @@
 
 #include "Engine/AssetManagerTypes.h"
 #include "Engine/AssetManager.h"
+#include "Engine/AssetManagerSettings.h"
 
 bool FPrimaryAssetTypeInfo::FillRuntimeData()
 {
@@ -14,12 +15,30 @@ bool FPrimaryAssetTypeInfo::FillRuntimeData()
 
 	for (const FStringAssetReference& AssetRef : SpecificAssets)
 	{
-		AssetScanPaths.AddUnique(AssetRef.ToString());
+		if (!AssetRef.IsNull())
+		{
+			AssetScanPaths.AddUnique(AssetRef.ToString());
+		}
 	}
 
 	for (const FDirectoryPath& PathRef : Directories)
 	{
-		AssetScanPaths.AddUnique(PathRef.Path);
+		if (!PathRef.Path.IsEmpty())
+		{
+			AssetScanPaths.AddUnique(PathRef.Path);
+		}
+	}
+
+	if (AssetScanPaths.Num() == 0)
+	{
+		// No scan locations picked out
+		return false;
+	}
+
+	if (PrimaryAssetType == NAME_None)
+	{
+		// Invalid type
+		return false;
 	}
 
 	return true;
@@ -68,4 +87,15 @@ void FPrimaryAssetRules::PropagateCookRules(const FPrimaryAssetRules& ParentRule
 	{
 		CookRule = ParentRules.CookRule;
 	}
-}
\ No newline at end of file
+}
+
+#if WITH_EDITOR
+void UAssetManagerSettings::PostEditChangeProperty(FPropertyChangedEvent& PropertyChangedEvent)
+{
+	Super::PostEditChangeProperty(PropertyChangedEvent);
+	if (PropertyChangedEvent.Property && UAssetManager::IsValid())
+	{
+		UAssetManager::Get().ReinitializeFromConfig();
+	}
+}
+#endif
\ No newline at end of file
diff --git a/Engine/Source/Runtime/Engine/Private/AsyncActionLoadPrimaryAsset.cpp b/Engine/Source/Runtime/Engine/Private/AsyncActionLoadPrimaryAsset.cpp
new file mode 100644
--- /dev/null
+++ b/Engine/Source/Runtime/Engine/Private/AsyncActionLoadPrimaryAsset.cpp
@@ -0,0 +1,169 @@
+// Copyright 1998-2017 Epic Games, Inc. All Rights Reserved.
+
+#include "AsyncActionLoadPrimaryAsset.h"
+#include "Engine/AssetManager.h"
+
+void UAsyncActionLoadPrimaryAssetBase::Activate()
+{
+	if (UAssetManager* Manager = UAssetManager::GetIfValid())
+	{
+		switch (Operation)
+		{
+		case EAssetManagerOperation::Load:
+			LoadHandle = Manager->LoadPrimaryAssets(AssetsToLoad, LoadBundles);
+			break;
+		case EAssetManagerOperation::ChangeBundleStateMatching:
+			LoadHandle = Manager->ChangeBundleStateForMatchingPrimaryAssets(LoadBundles, OldBundles);
+			break;
+		case EAssetManagerOperation::ChangeBundleStateList:
+			LoadHandle = Manager->ChangeBundleStateForPrimaryAssets(AssetsToLoad, LoadBundles, OldBundles);
+			break;
+		}
+		
+		if (LoadHandle.IsValid())
+		{
+			if (!LoadHandle->HasLoadCompleted())
+			{
+				LoadHandle->BindCompleteDelegate(FStreamableDelegate::CreateUObject(this, &UAsyncActionLoadPrimaryAssetBase::HandleLoadCompleted));
+				return;
+			}
+		}
+	}
+
+	// Either load already succeeded, or it failed
+	HandleLoadCompleted();
+}
+
+void UAsyncActionLoadPrimaryAssetBase::HandleLoadCompleted()
+{
+	LoadHandle.Reset();
+	SetReadyToDestroy();
+}
+
+UAsyncActionLoadPrimaryAsset* UAsyncActionLoadPrimaryAsset::AsyncLoadPrimaryAsset(FPrimaryAssetId PrimaryAsset, const TArray<FName>& LoadBundles)
+{
+	UAsyncActionLoadPrimaryAsset* Action = NewObject<UAsyncActionLoadPrimaryAsset>();
+	Action->AssetsToLoad.Add(PrimaryAsset);
+	Action->LoadBundles = LoadBundles;
+	Action->Operation = EAssetManagerOperation::Load;
+
+	return Action;
+}
+
+void UAsyncActionLoadPrimaryAsset::HandleLoadCompleted()
+{
+	UObject* AssetLoaded = nullptr;
+	if (LoadHandle.IsValid())
+	{
+		AssetLoaded = LoadHandle->GetLoadedAsset();
+	}
+
+	Super::HandleLoadCompleted();
+	Completed.Broadcast(AssetLoaded);
+}
+
+UAsyncActionLoadPrimaryAssetClass* UAsyncActionLoadPrimaryAssetClass::AsyncLoadPrimaryAssetClass(FPrimaryAssetId PrimaryAsset, const TArray<FName>& LoadBundles)
+{
+	UAsyncActionLoadPrimaryAssetClass* Action = NewObject<UAsyncActionLoadPrimaryAssetClass>();
+	Action->AssetsToLoad.Add(PrimaryAsset);
+	Action->LoadBundles = LoadBundles;
+	Action->Operation = EAssetManagerOperation::Load;
+
+	return Action;
+}
+
+void UAsyncActionLoadPrimaryAssetClass::HandleLoadCompleted()
+{
+	TSubclassOf<UObject> AssetLoaded = nullptr;
+	if (LoadHandle.IsValid())
+	{
+		AssetLoaded = Cast<UClass>(LoadHandle->GetLoadedAsset());
+	}
+
+	Super::HandleLoadCompleted();
+	Completed.Broadcast(AssetLoaded);
+}
+
+UAsyncActionLoadPrimaryAssetList* UAsyncActionLoadPrimaryAssetList::AsyncLoadPrimaryAssetList(const TArray<FPrimaryAssetId>& PrimaryAssetList, const TArray<FName>& LoadBundles)
+{
+	UAsyncActionLoadPrimaryAssetList* Action = NewObject<UAsyncActionLoadPrimaryAssetList>();
+	Action->AssetsToLoad = PrimaryAssetList;
+	Action->LoadBundles = LoadBundles;
+	Action->Operation = EAssetManagerOperation::Load;
+
+	return Action;
+}
+
+void UAsyncActionLoadPrimaryAssetList::HandleLoadCompleted()
+{
+	TArray<UObject*> AssetList;
+
+	if (LoadHandle.IsValid())
+	{
+		LoadHandle->GetLoadedAssets(AssetList);
+	}
+
+	Super::HandleLoadCompleted();
+	Completed.Broadcast(AssetList);
+}
+
+UAsyncActionLoadPrimaryAssetClassList* UAsyncActionLoadPrimaryAssetClassList::AsyncLoadPrimaryAssetClassList(const TArray<FPrimaryAssetId>& PrimaryAssetList, const TArray<FName>& LoadBundles)
+{
+	UAsyncActionLoadPrimaryAssetClassList* Action = NewObject<UAsyncActionLoadPrimaryAssetClassList>();
+	Action->AssetsToLoad = PrimaryAssetList;
+	Action->LoadBundles = LoadBundles;
+	Action->Operation = EAssetManagerOperation::Load;
+
+	return Action;
+}
+
+void UAsyncActionLoadPrimaryAssetClassList::HandleLoadCompleted()
+{
+	TArray<TSubclassOf<UObject>> AssetClassList;
+	TArray<UObject*> AssetList;
+
+	if (LoadHandle.IsValid())
+	{
+		LoadHandle->GetLoadedAssets(AssetList);
+
+		for (UObject* LoadedAsset : AssetList)
+		{
+			UClass* LoadedClass = Cast<UClass>(LoadedAsset);
+
+			if (LoadedClass)
+			{
+				AssetClassList.Add(LoadedClass);
+			}
+		}
+	}
+
+	Super::HandleLoadCompleted();
+	Completed.Broadcast(AssetClassList);
+}
+
+UAsyncActionChangePrimaryAssetBundles* UAsyncActionChangePrimaryAssetBundles::AsyncChangeBundleStateForMatchingPrimaryAssets(const TArray<FName>& NewBundles, const TArray<FName>& OldBundles)
+{
+	UAsyncActionChangePrimaryAssetBundles* Action = NewObject<UAsyncActionChangePrimaryAssetBundles>();
+	Action->LoadBundles = NewBundles;
+	Action->OldBundles = OldBundles;
+	Action->Operation = EAssetManagerOperation::ChangeBundleStateMatching;
+	
+	return Action;
+}
+
+UAsyncActionChangePrimaryAssetBundles* UAsyncActionChangePrimaryAssetBundles::AsyncChangeBundleStateForPrimaryAssetList(const TArray<FPrimaryAssetId>& PrimaryAssetList, const TArray<FName>& AddBundles, const TArray<FName>& RemoveBundles)
+{
+	UAsyncActionChangePrimaryAssetBundles* Action = NewObject<UAsyncActionChangePrimaryAssetBundles>();
+	Action->LoadBundles = AddBundles;
+	Action->OldBundles = RemoveBundles;
+	Action->AssetsToLoad = PrimaryAssetList;
+	Action->Operation = EAssetManagerOperation::ChangeBundleStateList;
+
+	return Action;
+}
+
+void UAsyncActionChangePrimaryAssetBundles::HandleLoadCompleted()
+{
+	Super::HandleLoadCompleted();
+	Completed.Broadcast();
+}
\ No newline at end of file
diff --git a/Engine/Source/Runtime/Engine/Private/AsyncActionLoadPrimaryAsset.h b/Engine/Source/Runtime/Engine/Private/AsyncActionLoadPrimaryAsset.h
new file mode 100644
--- /dev/null
+++ b/Engine/Source/Runtime/Engine/Private/AsyncActionLoadPrimaryAsset.h
@@ -0,0 +1,174 @@
+// Copyright 1998-2017 Epic Games, Inc. All Rights Reserved.
+
+#pragma once
+
+#include "CoreMinimal.h"
+#include "UObject/ObjectMacros.h"
+#include "Engine/StreamableManager.h"
+#include "Kismet/BlueprintAsyncActionBase.h"
+#include "Templates/SubclassOf.h"
+
+#include "AsyncActionLoadPrimaryAsset.generated.h"
+
+/** Base class of all asset manager load calls */
+UCLASS(Abstract)
+class UAsyncActionLoadPrimaryAssetBase : public UBlueprintAsyncActionBase
+{
+	GENERATED_BODY()
+
+public:
+	/** Execute the actual load */
+	virtual void Activate() override;
+
+protected:
+
+	/** Called from asset manager */
+	virtual void HandleLoadCompleted();
+
+	/** Specific assets requested */
+	TArray<FPrimaryAssetId> AssetsToLoad;
+
+	/** Bundle states */
+	TArray<FName> LoadBundles;
+
+	/** Bundle states */
+	TArray<FName> OldBundles;
+
+	/** Handle of load request */
+	TSharedPtr<FStreamableHandle> LoadHandle;
+
+	enum class EAssetManagerOperation : uint8
+	{
+		Load,
+		ChangeBundleStateMatching,
+		ChangeBundleStateList
+	};
+
+	/** Which operation is being run */
+	EAssetManagerOperation Operation;
+};
+
+DECLARE_DYNAMIC_MULTICAST_DELEGATE_OneParam(FOnPrimaryAssetLoaded, UObject*, Loaded);
+
+UCLASS()
+class UAsyncActionLoadPrimaryAsset : public UAsyncActionLoadPrimaryAssetBase
+{
+	GENERATED_BODY()
+
+public:
+	/** 
+	 * Load a primary asset into memory. The completed delegate will go off when the load succeeds or fails, you should cast the Loaded object to verify it is the correct type.
+	 * If LoadBundles is specified, those bundles are loaded along with the asset
+	 */
+	UFUNCTION(BlueprintCallable, meta=(BlueprintInternalUseOnly="true", Category = "AssetManager", AutoCreateRefTerm = "LoadBundles"))
+	static UAsyncActionLoadPrimaryAsset* AsyncLoadPrimaryAsset(FPrimaryAssetId PrimaryAsset, const TArray<FName>& LoadBundles);
+
+	UPROPERTY(BlueprintAssignable)
+	FOnPrimaryAssetLoaded Completed;
+
+protected:
+
+	/** Called from asset manager */
+	virtual void HandleLoadCompleted() override;
+};
+
+DECLARE_DYNAMIC_MULTICAST_DELEGATE_OneParam(FOnPrimaryAssetClassLoaded, TSubclassOf<UObject>, Loaded);
+
+UCLASS()
+class UAsyncActionLoadPrimaryAssetClass : public UAsyncActionLoadPrimaryAssetBase
+{
+	GENERATED_BODY()
+
+public:
+	/** 
+	 * Load a primary asset class into memory. The completed delegate will go off when the load succeeds or fails, you should cast the Loaded class to verify it is the correct type 
+	 * If LoadBundles is specified, those bundles are loaded along with the asset
+	 */
+	UFUNCTION(BlueprintCallable, meta = (BlueprintInternalUseOnly = "true", Category = "AssetManager", AutoCreateRefTerm = "LoadBundles"))
+	static UAsyncActionLoadPrimaryAssetClass* AsyncLoadPrimaryAssetClass(FPrimaryAssetId PrimaryAsset, const TArray<FName>& LoadBundles);
+
+	UPROPERTY(BlueprintAssignable)
+	FOnPrimaryAssetClassLoaded Completed;
+
+protected:
+
+	/** Called from asset manager */
+	virtual void HandleLoadCompleted() override;
+};
+
+DECLARE_DYNAMIC_MULTICAST_DELEGATE_OneParam(FOnPrimaryAssetListLoaded, const TArray<UObject*>&, Loaded);
+
+UCLASS()
+class UAsyncActionLoadPrimaryAssetList : public UAsyncActionLoadPrimaryAssetBase
+{
+	GENERATED_BODY()
+
+public:
+	/** 
+	 * Load a list primary assets into memory. The completed delegate will go off when the load succeeds or fails, you should cast the Loaded object list to verify it is the correct type 
+	 * If LoadBundles is specified, those bundles are loaded along with the asset list
+	 */
+	UFUNCTION(BlueprintCallable, meta=(BlueprintInternalUseOnly="true", Category = "AssetManager", AutoCreateRefTerm = "LoadBundles"))
+	static UAsyncActionLoadPrimaryAssetList* AsyncLoadPrimaryAssetList(const TArray<FPrimaryAssetId>& PrimaryAssetList, const TArray<FName>& LoadBundles);
+
+	UPROPERTY(BlueprintAssignable)
+	FOnPrimaryAssetListLoaded Completed;
+
+protected:
+
+	/** Called from asset manager */
+	virtual void HandleLoadCompleted() override;
+};
+
+DECLARE_DYNAMIC_MULTICAST_DELEGATE_OneParam(FOnPrimaryAssetClassListLoaded, const TArray<TSubclassOf<UObject>>&, Loaded);
+
+UCLASS()
+class UAsyncActionLoadPrimaryAssetClassList : public UAsyncActionLoadPrimaryAssetBase
+{
+	GENERATED_BODY()
+
+public:
+	/** 
+	 * Load a list primary asset classes into memory. The completed delegate will go off when the load succeeds or fails, you should cast the Loaded object list to verify it is the correct type 
+	 * If LoadBundles is specified, those bundles are loaded along with the asset list
+	 */
+	UFUNCTION(BlueprintCallable, meta = (BlueprintInternalUseOnly = "true", Category = "AssetManager", AutoCreateRefTerm = "LoadBundles"))
+	static UAsyncActionLoadPrimaryAssetClassList* AsyncLoadPrimaryAssetClassList(const TArray<FPrimaryAssetId>& PrimaryAssetList, const TArray<FName>& LoadBundles);
+
+	UPROPERTY(BlueprintAssignable)
+	FOnPrimaryAssetClassListLoaded Completed;
+
+protected:
+
+	/** Called from asset manager */
+	virtual void HandleLoadCompleted() override;
+};
+
+DECLARE_DYNAMIC_MULTICAST_DELEGATE(FOnPrimaryAssetBundlesChanged);
+
+UCLASS()
+class UAsyncActionChangePrimaryAssetBundles : public UAsyncActionLoadPrimaryAssetBase
+{
+	GENERATED_BODY()
+
+public:
+	/**
+	 * Change the bundle state of all assets that match OldBundles to instead contain NewBundles. 
+	 */
+	UFUNCTION(BlueprintCallable, meta = (BlueprintInternalUseOnly = "true", Category = "AssetManager"))
+	static UAsyncActionChangePrimaryAssetBundles* AsyncChangeBundleStateForMatchingPrimaryAssets(const TArray<FName>& NewBundles, const TArray<FName>& OldBundles);
+
+	/**
+	 * Change the bundle state of assets in PrimaryAssetList. AddBundles are added and RemoveBundles are removed, both must be filled in but an empty array is allowed
+	 */
+	UFUNCTION(BlueprintCallable, meta = (BlueprintInternalUseOnly = "true", Category = "AssetManager"))
+	static UAsyncActionChangePrimaryAssetBundles* AsyncChangeBundleStateForPrimaryAssetList(const TArray<FPrimaryAssetId>& PrimaryAssetList, const TArray<FName>& AddBundles, const TArray<FName>& RemoveBundles);
+
+	UPROPERTY(BlueprintAssignable)
+	FOnPrimaryAssetBundlesChanged Completed;
+
+protected:
+
+	/** Called from asset manager */
+	virtual void HandleLoadCompleted() override;
+};
diff --git a/Engine/Source/Runtime/Engine/Private/BlueprintGeneratedClass.cpp b/Engine/Source/Runtime/Engine/Private/BlueprintGeneratedClass.cpp
--- a/Engine/Source/Runtime/Engine/Private/BlueprintGeneratedClass.cpp
+++ b/Engine/Source/Runtime/Engine/Private/BlueprintGeneratedClass.cpp
@@ -1450,10 +1450,34 @@ protected:
 			Object = nullptr;
 		}
 #endif
+		if (Object)
+		{
+			bool bWeakRef = false;
+
+			// If the property that serialized us is not an object property we are in some native serializer, we have to treat these as strong
+			if (!Object->HasAnyFlags(RF_StrongRefOnFrame))
+			{
+				UObjectProperty* ObjectProperty = Cast<UObjectProperty>(GetSerializedProperty());
+
+				if (ObjectProperty)
+				{
+					bWeakRef = true;
+				}
+			}
 
-		const bool bWeakRef = Object ? !Object->HasAnyFlags(RF_StrongRefOnFrame) : false;
-		Collector.SetShouldHandleAsWeakRef(bWeakRef); 
-		return FSimpleObjectReferenceCollectorArchive::operator<<(Object);
+			if (bWeakRef)
+			{
+				// This was a raw UObject * serialized by UObjectProperty, so just save the address
+				Collector.MarkWeakObjectReferenceForClearing(&Object);				
+			}
+			else
+			{
+				// This is a hard reference or we don't know what's serializing it, so serialize it normally
+				return FSimpleObjectReferenceCollectorArchive::operator<<(Object);
+			}
+		}
+
+		return *this;
 	}
 };
 
@@ -1473,9 +1497,6 @@ void UBlueprintGeneratedClass::AddReferencedObjectsInUbergraphFrame(UObject* InT
 					checkSlow(BPGC->UberGraphFunction);
 					FPersistentFrameCollectorArchive ObjectReferenceCollector(InThis, Collector);
 					BPGC->UberGraphFunction->SerializeBin(ObjectReferenceCollector, PointerToUberGraphFrame->RawPointer);
-
-					// Reset the ShouldHandleAsWeakRef state, before the collector is used by a different archive.
-					Collector.SetShouldHandleAsWeakRef(false);
 				}
 			}
 		}
diff --git a/Engine/Source/Runtime/Engine/Private/GameEngine.cpp b/Engine/Source/Runtime/Engine/Private/GameEngine.cpp
--- a/Engine/Source/Runtime/Engine/Private/GameEngine.cpp
+++ b/Engine/Source/Runtime/Engine/Private/GameEngine.cpp
@@ -743,7 +743,8 @@ bool UGameEngine::NetworkRemapPath(UNetDriver* Driver, FString& Str, bool bReadi
 
 	// If the prefixed path matches the world package name or the name of a streaming level,
 	// return the prefixed name.
-	const FString PackageNameOnly = FPackageName::PackageFromPath(*Str);
+	FString PackageNameOnly = Str;
+	FPackageName::TryConvertFilenameToLongPackageName(PackageNameOnly, PackageNameOnly);
 
 	const FString PrefixedFullName = UWorld::ConvertToPIEPackageName(Str, Context.PIEInstance);
 	const FString PrefixedPackageName = UWorld::ConvertToPIEPackageName(PackageNameOnly, Context.PIEInstance);
diff --git a/Engine/Source/Runtime/Engine/Private/Kismet/BlueprintAsyncActionBase.cpp b/Engine/Source/Runtime/Engine/Private/Kismet/BlueprintAsyncActionBase.cpp
--- a/Engine/Source/Runtime/Engine/Private/Kismet/BlueprintAsyncActionBase.cpp
+++ b/Engine/Source/Runtime/Engine/Private/Kismet/BlueprintAsyncActionBase.cpp
@@ -8,9 +8,17 @@
 UBlueprintAsyncActionBase::UBlueprintAsyncActionBase(const FObjectInitializer& ObjectInitializer)
 	: Super(ObjectInitializer)
 {
-	SetFlags(RF_StrongRefOnFrame);
+	if (!HasAnyFlags(RF_ClassDefaultObject))
+	{
+		SetFlags(RF_StrongRefOnFrame);
+	}
 }
 
 void UBlueprintAsyncActionBase::Activate()
 {
 }
+
+void UBlueprintAsyncActionBase::SetReadyToDestroy()
+{
+	ClearFlags(RF_StrongRefOnFrame);
+}
diff --git a/Engine/Source/Runtime/Engine/Private/KismetSystemLibrary.cpp b/Engine/Source/Runtime/Engine/Private/KismetSystemLibrary.cpp
--- a/Engine/Source/Runtime/Engine/Private/KismetSystemLibrary.cpp
+++ b/Engine/Source/Runtime/Engine/Private/KismetSystemLibrary.cpp
@@ -31,6 +31,7 @@
 #include "Net/OnlineEngineInterface.h"
 #include "UserActivityTracking.h"
 #include "KismetTraceUtils.h"
+#include "Engine/AssetManager.h"
 
 //////////////////////////////////////////////////////////////////////////
 // UKismetSystemLibrary
@@ -829,16 +830,56 @@ void UKismetSystemLibrary::SetAssetClassPropertyByName(UObject* Object, FName Pr
 	}
 }
 
-UObject* UKismetSystemLibrary::Conv_AssetToObject(const TAssetPtr<UObject>& Asset)
+bool UKismetSystemLibrary::IsValidSoftObjectReference(const TAssetPtr<UObject>& SoftObjectReference)
+{
+	return !SoftObjectReference.IsNull();
+}
+
+bool UKismetSystemLibrary::EqualEqual_SoftObjectReference(const TAssetPtr<UObject>& A, const TAssetPtr<UObject>& B)
+{
+	return A == B;
+}
+
+bool UKismetSystemLibrary::NotEqual_SoftObjectReference(const TAssetPtr<UObject>& A, const TAssetPtr<UObject>& B)
+{
+	return A != B;
+}
+
+bool UKismetSystemLibrary::IsValidSoftClassReference(const TAssetSubclassOf<UObject>& SoftClassReference)
+{
+	return !SoftClassReference.IsNull();
+}
+
+bool UKismetSystemLibrary::EqualEqual_SoftClassReference(const TAssetSubclassOf<UObject>& A, const TAssetSubclassOf<UObject>& B)
+{
+	return A == B;
+}
+
+bool UKismetSystemLibrary::NotEqual_SoftClassReference(const TAssetSubclassOf<UObject>& A, const TAssetSubclassOf<UObject>& B)
+{
+	return A != B;
+}
+
+UObject* UKismetSystemLibrary::Conv_SoftObjectReferenceToObject(const TAssetPtr<UObject>& Asset)
 {
 	return ((const FAssetPtr*)&Asset)->Get();
 }
 
-TSubclassOf<UObject> UKismetSystemLibrary::Conv_AssetClassToClass(const TAssetSubclassOf<UObject>& AssetClass)
+TSubclassOf<UObject> UKismetSystemLibrary::Conv_SoftClassReferenceToClass(const TAssetSubclassOf<UObject>& AssetClass)
 {
 	return Cast<UClass>(((const FAssetPtr*)&AssetClass)->Get());
 }
 
+TAssetPtr<UObject> UKismetSystemLibrary::Conv_ObjectToSoftObjectReference(UObject *Object)
+{
+	return TAssetPtr<UObject>(Object);
+}
+
+TAssetSubclassOf<UObject> UKismetSystemLibrary::Conv_ClassToSoftClassReference(const TSubclassOf<UObject>& Class)
+{
+	return TAssetSubclassOf<UObject>(*Class);
+}
+
 void UKismetSystemLibrary::SetTextPropertyByName(UObject* Object, FName PropertyName, const FText& Value)
 {
 	if(Object != NULL)
@@ -2402,3 +2443,171 @@ FString UKismetSystemLibrary::GetCommandLine()
 {
 	return FString(FCommandLine::Get());
 }
+
+UObject* UKismetSystemLibrary::GetObjectFromPrimaryAssetId(FPrimaryAssetId PrimaryAssetId)
+{
+	if (UAssetManager* Manager = UAssetManager::GetIfValid())
+	{
+		FPrimaryAssetTypeInfo Info;
+		if (Manager->GetPrimaryAssetTypeInfo(PrimaryAssetId.PrimaryAssetType, Info) && !Info.bHasBlueprintClasses)
+		{
+			return Manager->GetPrimaryAssetObject(PrimaryAssetId);
+		}
+	}
+	return nullptr;
+}
+
+TSubclassOf<UObject> UKismetSystemLibrary::GetClassFromPrimaryAssetId(FPrimaryAssetId PrimaryAssetId)
+{
+	if (UAssetManager* Manager = UAssetManager::GetIfValid())
+	{
+		FPrimaryAssetTypeInfo Info;
+		if (Manager->GetPrimaryAssetTypeInfo(PrimaryAssetId.PrimaryAssetType, Info) && Info.bHasBlueprintClasses)
+		{
+			return Manager->GetPrimaryAssetObjectClass<UObject>(PrimaryAssetId);
+		}
+	}
+	return nullptr;
+}
+
+TAssetPtr<UObject> UKismetSystemLibrary::GetSoftObjectReferenceFromPrimaryAssetId(FPrimaryAssetId PrimaryAssetId)
+{
+	if (UAssetManager* Manager = UAssetManager::GetIfValid())
+	{
+		FPrimaryAssetTypeInfo Info;
+		if (Manager->GetPrimaryAssetTypeInfo(PrimaryAssetId.PrimaryAssetType, Info) && !Info.bHasBlueprintClasses)
+		{
+			return TAssetPtr<UObject>(Manager->GetPrimaryAssetPath(PrimaryAssetId));
+		}
+	}
+	return nullptr;
+}
+
+TAssetSubclassOf<UObject> UKismetSystemLibrary::GetSoftClassReferenceFromPrimaryAssetId(FPrimaryAssetId PrimaryAssetId)
+{
+	if (UAssetManager* Manager = UAssetManager::GetIfValid())
+	{
+		FPrimaryAssetTypeInfo Info;
+		if (Manager->GetPrimaryAssetTypeInfo(PrimaryAssetId.PrimaryAssetType, Info) && Info.bHasBlueprintClasses)
+		{
+			return TAssetSubclassOf<UObject>(Manager->GetPrimaryAssetPath(PrimaryAssetId));
+		}
+	}
+	return nullptr;
+}
+
+FPrimaryAssetId UKismetSystemLibrary::GetPrimaryAssetIdFromObject(UObject* Object)
+{
+	if (UAssetManager* Manager = UAssetManager::GetIfValid())
+	{
+		if (Object)
+		{
+			return Manager->GetPrimaryAssetIdForObject(Object);
+		}
+	}
+	return FPrimaryAssetId();
+}
+
+FPrimaryAssetId UKismetSystemLibrary::GetPrimaryAssetIdFromClass(TSubclassOf<UObject> Class)
+{
+	if (UAssetManager* Manager = UAssetManager::GetIfValid())
+	{
+		if (Class)
+		{
+			return Manager->GetPrimaryAssetIdForObject(*Class);
+		}
+	}
+	return FPrimaryAssetId();
+}
+
+FPrimaryAssetId UKismetSystemLibrary::GetPrimaryAssetIdFromSoftObjectReference(TAssetPtr<UObject> SoftObjectReference)
+{
+	if (UAssetManager* Manager = UAssetManager::GetIfValid())
+	{
+		return Manager->GetPrimaryAssetIdForPath(SoftObjectReference.ToStringReference());
+	}
+	return FPrimaryAssetId();
+}
+
+FPrimaryAssetId UKismetSystemLibrary::GetPrimaryAssetIdFromSoftClassReference(TAssetSubclassOf<UObject> SoftClassReference)
+{
+	if (UAssetManager* Manager = UAssetManager::GetIfValid())
+	{
+		return Manager->GetPrimaryAssetIdForPath(SoftClassReference.ToStringReference());
+	}
+	return FPrimaryAssetId();
+}
+
+void UKismetSystemLibrary::GetPrimaryAssetIdList(FPrimaryAssetType PrimaryAssetType, TArray<FPrimaryAssetId>& OutPrimaryAssetIdList)
+{
+	if (UAssetManager* Manager = UAssetManager::GetIfValid())
+	{
+		Manager->GetPrimaryAssetIdList(PrimaryAssetType, OutPrimaryAssetIdList);
+	}
+}
+
+bool UKismetSystemLibrary::IsValidPrimaryAssetId(FPrimaryAssetId PrimaryAssetId)
+{
+	return PrimaryAssetId.IsValid();
+}
+
+bool UKismetSystemLibrary::EqualEqual_PrimaryAssetId(FPrimaryAssetId A, FPrimaryAssetId B)
+{
+	return A == B;
+}
+
+bool UKismetSystemLibrary::NotEqual_PrimaryAssetId(FPrimaryAssetId A, FPrimaryAssetId B)
+{
+	return A != B;
+}
+
+bool UKismetSystemLibrary::IsValidPrimaryAssetType(FPrimaryAssetType PrimaryAssetType)
+{
+	return PrimaryAssetType.IsValid();
+}
+
+bool UKismetSystemLibrary::EqualEqual_PrimaryAssetType(FPrimaryAssetType A, FPrimaryAssetType B)
+{
+	return A == B;
+}
+
+bool UKismetSystemLibrary::NotEqual_PrimaryAssetType(FPrimaryAssetType A, FPrimaryAssetType B)
+{
+	return A != B;
+}
+
+void UKismetSystemLibrary::UnloadPrimaryAsset(FPrimaryAssetId PrimaryAssetId)
+{
+	if (UAssetManager* Manager = UAssetManager::GetIfValid())
+	{
+		Manager->UnloadPrimaryAsset(PrimaryAssetId);
+	}
+}
+
+void UKismetSystemLibrary::UnloadPrimaryAssetList(const TArray<FPrimaryAssetId>& PrimaryAssetIdList)
+{
+	if (UAssetManager* Manager = UAssetManager::GetIfValid())
+	{
+		Manager->UnloadPrimaryAssets(PrimaryAssetIdList);
+	}
+}
+
+bool UKismetSystemLibrary::GetCurrentBundleState(FPrimaryAssetId PrimaryAssetId, bool bForceCurrentState, TArray<FName>& OutBundles)
+{
+	if (UAssetManager* Manager = UAssetManager::GetIfValid())
+	{
+		if (Manager->GetPrimaryAssetHandle(PrimaryAssetId, bForceCurrentState, &OutBundles).IsValid())
+		{
+			return true;
+		}
+	}
+	return false;
+}
+
+void UKismetSystemLibrary::GetPrimaryAssetsWithBundleState(const TArray<FName>& RequiredBundles, const TArray<FName>& ExcludedBundles, const TArray<FPrimaryAssetType>& ValidTypes, bool bForceCurrentState, TArray<FPrimaryAssetId>& OutPrimaryAssetIdList)
+{
+	if (UAssetManager* Manager = UAssetManager::GetIfValid())
+	{
+		Manager->GetPrimaryAssetsWithBundleState(OutPrimaryAssetIdList, ValidTypes, RequiredBundles, ExcludedBundles, bForceCurrentState);
+	}
+}
\ No newline at end of file
diff --git a/Engine/Source/Runtime/Engine/Private/LevelActorContainer.cpp b/Engine/Source/Runtime/Engine/Private/LevelActorContainer.cpp
--- a/Engine/Source/Runtime/Engine/Private/LevelActorContainer.cpp
+++ b/Engine/Source/Runtime/Engine/Private/LevelActorContainer.cpp
@@ -14,100 +14,6 @@ Level.cpp: Level-related functions
 
 DEFINE_LOG_CATEGORY_STATIC(LogLevelActorContainer, Log, All);
 
-/**
-* Pool for reducing allocations when constructing clusters
-*/
-class FActorClusterArrayPool
-{
-public:
-
-	/**
-	* Gets the singleton instance of the FObjectArrayPool
-	* @return Pool singleton.
-	*/
-	FORCEINLINE static FActorClusterArrayPool& Get()
-	{
-		static FActorClusterArrayPool Singleton;
-		return Singleton;
-	}
-
-	/**
-	* Gets an event from the pool or creates one if necessary.
-	*
-	* @return The array.
-	* @see ReturnToPool
-	*/
-	FORCEINLINE TArray<UObject*>* GetArrayFromPool()
-	{
-		TArray<UObject*>* Result = Pool.Pop();
-		if (!Result)
-		{
-			Result = new TArray<UObject*>();
-		}
-		check(Result);
-#if UE_BUILD_DEBUG
-		NumberOfUsedArrays.Increment();
-#endif // UE_BUILD_DEBUG
-		return Result;
-	}
-
-	/**
-	* Returns an array to the pool.
-	*
-	* @param Array The array to return.
-	* @see GetArrayFromPool
-	*/
-	FORCEINLINE void ReturnToPool(TArray<UObject*>* Array)
-	{
-#if UE_BUILD_DEBUG
-		const int32 CheckUsedArrays = NumberOfUsedArrays.Decrement();
-		checkSlow(CheckUsedArrays >= 0);
-#endif // UE_BUILD_DEBUG
-		check(Array);
-		Array->Reset();
-		Pool.Push(Array);
-	}
-
-	/** Performs memory cleanup */
-	void Cleanup()
-	{
-#if UE_BUILD_DEBUG
-		const int32 CheckUsedArrays = NumberOfUsedArrays.GetValue();
-		checkSlow(CheckUsedArrays == 0);
-#endif // UE_BUILD_DEBUG
-
-		uint32 FreedMemory = 0;
-		TArray< TArray<UObject*>* > AllArrays;
-		Pool.PopAll(AllArrays);
-		for (TArray<UObject*>* Array : AllArrays)
-		{
-			FreedMemory += Array->GetAllocatedSize();
-			delete Array;
-		}
-		UE_LOG(LogLevelActorContainer, Log, TEXT("Freed %ub from %d cluster array pools."), FreedMemory, AllArrays.Num());
-	}
-
-#if UE_BUILD_DEBUG
-	void CheckLeaks()
-	{
-		// This function is called after cluster has been created so at this point there should be no
-		// arrays used by cluster creation code and all should be returned to the pool
-		const int32 LeakedPoolArrays = NumberOfUsedArrays.GetValue();
-		checkSlow(LeakedPoolArrays == 0);
-	}
-#endif
-
-private:
-
-	/** Holds the collection of recycled arrays. */
-	TLockFreePointerListLIFO< TArray<UObject*> > Pool;
-
-#if UE_BUILD_DEBUG
-	/** Number of arrays currently acquired from the pool by clusters */
-	FThreadSafeCounter NumberOfUsedArrays;
-#endif // UE_BUILD_DEBUG
-};
-
 /**
 * Handles UObject references found by TFastReferenceCollector
 */
@@ -288,24 +194,24 @@ template <class TProcessor>
 class TActorClusterCollector : public FReferenceCollector
 {
 	TProcessor& Processor;
-	TArray<UObject*>& ObjectArray;
+	FGCArrayStruct& ObjectArrayStruct;
 
 public:
-	TActorClusterCollector(TProcessor& InProcessor, TArray<UObject*>& InObjectArray)
+	TActorClusterCollector(TProcessor& InProcessor, FGCArrayStruct& InObjectArrayStruct)
 		: Processor(InProcessor)
-		, ObjectArray(InObjectArray)
+		, ObjectArrayStruct(InObjectArrayStruct)
 	{
 	}
 	virtual void HandleObjectReference(UObject*& Object, const UObject* ReferencingObject, const UProperty* ReferencingProperty) override
 	{
-		Processor.HandleTokenStreamObjectReference(ObjectArray, const_cast<UObject*>(ReferencingObject), Object, INDEX_NONE, false);
+		Processor.HandleTokenStreamObjectReference(ObjectArrayStruct.ObjectsToSerialize, const_cast<UObject*>(ReferencingObject), Object, INDEX_NONE, false);
 	}
 	virtual void HandleObjectReferences(UObject** InObjects, const int32 ObjectNum, const UObject* ReferencingObject, const UProperty* InReferencingProperty) override
 	{
 		for (int32 ObjectIndex = 0; ObjectIndex < ObjectNum; ++ObjectIndex)
 		{
 			UObject*& Object = InObjects[ObjectIndex];
-			Processor.HandleTokenStreamObjectReference(ObjectArray, const_cast<UObject*>(ReferencingObject), Object, INDEX_NONE, false);
+			Processor.HandleTokenStreamObjectReference(ObjectArrayStruct.ObjectsToSerialize, const_cast<UObject*>(ReferencingObject), Object, INDEX_NONE, false);
 		}
 	}
 	virtual bool IsIgnoringArchetypeRef() const override
@@ -338,12 +244,13 @@ void ULevelActorContainer::CreateCluster()
 
 	// Collect all objects referenced by cluster root and by all objects it's referencing
 	FActorClusterReferenceProcessor Processor(ContainerInternalIndex, Cluster, CastChecked<ULevel>(GetOuter()));
-	TFastReferenceCollector<false, FActorClusterReferenceProcessor, TActorClusterCollector<FActorClusterReferenceProcessor>, FActorClusterArrayPool, true> ReferenceCollector(Processor, FActorClusterArrayPool::Get());
-	TArray<UObject*> ObjectsToProcess;
+	TFastReferenceCollector<false, FActorClusterReferenceProcessor, TActorClusterCollector<FActorClusterReferenceProcessor>, FGCArrayPool, true> ReferenceCollector(Processor, FGCArrayPool::Get());
+	FGCArrayStruct ArrayStruct;
+	TArray<UObject*>& ObjectsToProcess = ArrayStruct.ObjectsToSerialize;
 	ObjectsToProcess.Add(static_cast<UObject*>(this));
-	ReferenceCollector.CollectReferences(ObjectsToProcess);
+	ReferenceCollector.CollectReferences(ArrayStruct);
 #if UE_BUILD_DEBUG
-	FActorClusterArrayPool::Get().CheckLeaks();
+	FGCArrayPool::Get().CheckLeaks();
 #endif
 
 	if (Cluster.Objects.Num())
diff --git a/Engine/Source/Runtime/Engine/Private/PrimaryAssetLabel.cpp b/Engine/Source/Runtime/Engine/Private/PrimaryAssetLabel.cpp
--- a/Engine/Source/Runtime/Engine/Private/PrimaryAssetLabel.cpp
+++ b/Engine/Source/Runtime/Engine/Private/PrimaryAssetLabel.cpp
@@ -12,6 +12,7 @@ const FName UPrimaryAssetLabel::DirectoryBundle = FName("Directory");
 UPrimaryAssetLabel::UPrimaryAssetLabel()
 {
 	bLabelAssetsInMyDirectory = false;
+	bIsRuntimeLabel = false;
 
 	// By default have low priority and don't recurse
 	Rules.bApplyRecursively = false;
diff --git a/Engine/Source/Runtime/Engine/Private/StreamableManager.cpp b/Engine/Source/Runtime/Engine/Private/StreamableManager.cpp
--- a/Engine/Source/Runtime/Engine/Private/StreamableManager.cpp
+++ b/Engine/Source/Runtime/Engine/Private/StreamableManager.cpp
@@ -422,13 +422,13 @@ void FStreamableHandle::UpdateCombinedHandle()
 	}
 
 	// Check all our children, complete if done
-	bool bAllInactive = true;
+	bool bAllCompleted = true;
 	bool bAllCanceled = true;
 	for (TSharedPtr<FStreamableHandle> ChildHandle : ChildHandles)
 	{
-		if (ChildHandle->IsActive())
+		if (ChildHandle->IsLoadingInProgress())
 		{
-			bAllInactive = false;
+			bAllCompleted = false;
 		}
 		if (!ChildHandle->WasCanceled())
 		{
@@ -436,12 +436,12 @@ void FStreamableHandle::UpdateCombinedHandle()
 		}
 	}
 
-	// If all our sub handles were canceled, cancel us. Otherwise complete us if at least one was completed and there are no active
+	// If all our sub handles were canceled, cancel us. Otherwise complete us if at least one was completed and there are none in progress
 	if (bAllCanceled)
 	{
 		CancelHandle();
 	}
-	else if (bAllInactive)
+	else if (bAllCompleted)
 	{
 		CompleteLoad();
 
@@ -569,13 +569,13 @@ struct FStreamable
 
 FStreamableManager::FStreamableManager()
 {
-	FCoreUObjectDelegates::PostGarbageCollect.AddRaw(this, &FStreamableManager::OnPostGarbageCollect);
+	FCoreUObjectDelegates::PreGarbageCollect.AddRaw(this, &FStreamableManager::OnPreGarbageCollect);
 	bForceSynchronousLoads = false;
 }
 
 FStreamableManager::~FStreamableManager()
 {
-	FCoreUObjectDelegates::PostGarbageCollect.RemoveAll(this);
+	FCoreUObjectDelegates::PreGarbageCollect.RemoveAll(this);
 
 	for (TStreamableMap::TIterator It(StreamableItems); It; ++It)
 	{
@@ -584,7 +584,7 @@ FStreamableManager::~FStreamableManager()
 	}
 }
 
-void FStreamableManager::OnPostGarbageCollect()
+void FStreamableManager::OnPreGarbageCollect()
 {
 	TSet<FStringAssetReference> RedirectsToRemove;
 
diff --git a/Engine/Source/Runtime/Engine/Private/UnrealEngine.cpp b/Engine/Source/Runtime/Engine/Private/UnrealEngine.cpp
--- a/Engine/Source/Runtime/Engine/Private/UnrealEngine.cpp
+++ b/Engine/Source/Runtime/Engine/Private/UnrealEngine.cpp
@@ -11634,7 +11634,7 @@ void UEngine::CopyPropertiesForUnrelatedObjects(UObject* OldObject, UObject* New
 			pRecord->OldInstance = OldInstance;
 			OldInstanceMap.Add(OldInstance->GetPathName(OldObject), SavedInstances.Num() - 1);
 			const uint32 AdditionalPortFlags = Params.bCopyDeprecatedProperties ? PPF_UseDeprecatedProperties : PPF_None;
-			FObjectWriter SubObjWriter(OldInstance, pRecord->SavedProperties, true, true, true, AdditionalPortFlags);
+			FObjectWriter SubObjWriter(OldInstance, pRecord->SavedProperties, true, true, Params.bDoDelta, AdditionalPortFlags);
 		}
 	}
 
diff --git a/Engine/Source/Runtime/Slate/Private/Framework/Application/MenuStack.cpp b/Engine/Source/Runtime/Slate/Private/Framework/Application/MenuStack.cpp
--- a/Engine/Source/Runtime/Slate/Private/Framework/Application/MenuStack.cpp
+++ b/Engine/Source/Runtime/Slate/Private/Framework/Application/MenuStack.cpp
@@ -166,7 +166,7 @@ TSharedRef<SWindow> FMenuStack::PushMenu( const TSharedRef<SWindow>& ParentWindo
 
 	FSlateRect Anchor(SummonLocation, SummonLocation + SummonLocationSize);
 	FVector2D ExpectedSize(300, 200);
-	FVector2D AdjustedSummonLocation = FSlateApplication::Get().CalculatePopupWindowPosition(Anchor, ExpectedSize, Orient_Vertical);
+	FVector2D AdjustedSummonLocation = FSlateApplication::Get().CalculatePopupWindowPosition(Anchor, ExpectedSize, FVector2D::ZeroVector, Orient_Vertical);
 
 	TSharedRef<SWindow> NewMenuWindow =
 		SNew(SWindow)
@@ -357,7 +357,7 @@ FMenuStack::FPrePushResults FMenuStack::PrePush(const FPrePushArgs& InArgs)
 	if (ActiveMethod.GetPopupMethod() == EPopupMethod::CreateNewWindow)
 	{
 		// Places the menu's window in the work area
-		OutResults.AnimStartLocation = OutResults.AnimFinalLocation = FSlateApplication::Get().CalculatePopupWindowPosition(InArgs.Anchor, OutResults.ExpectedSize, Orientation);
+		OutResults.AnimStartLocation = OutResults.AnimFinalLocation = FSlateApplication::Get().CalculatePopupWindowPosition(InArgs.Anchor, OutResults.ExpectedSize, FVector2D::ZeroVector, Orientation);
 	}
 	else
 	{
