Dupe of CL #3461534, fixes CIS failure of automation tests

diff --git a/Engine/Source/Runtime/CoreUObject/Private/Serialization/AsyncLoading.cpp b/Engine/Source/Runtime/CoreUObject/Private/Serialization/AsyncLoading.cpp
--- a/Engine/Source/Runtime/CoreUObject/Private/Serialization/AsyncLoading.cpp
+++ b/Engine/Source/Runtime/CoreUObject/Private/Serialization/AsyncLoading.cpp
@@ -37,6 +37,7 @@
 #include "UniquePtr.h"
 #include "Serialization/BufferReader.h"
 #include "TaskGraphInterfaces.h"
+#include "Blueprint/BlueprintSupport.h"
 
 #define FIND_MEMORY_STOMPS (1 && (PLATFORM_WINDOWS || PLATFORM_LINUX) && !WITH_EDITORONLY_DATA)
 
@@ -4560,6 +4561,11 @@ EAsyncPackageState::Type FAsyncLoadingThread::ProcessLoadedPackages(bool bUseTim
 
 	if (Result == EAsyncPackageState::Complete)
 	{
+#if WITH_EDITORONLY_DATA
+		// This needs to happen after loading new blueprints in the editor, and this is handled in EndLoad for synchronous loads
+		FBlueprintSupport::FlushReinstancingQueue();
+#endif
+
 		// We're not done until all packages have been deleted
 		Result = PackagesToDelete.Num() ? EAsyncPackageState::PendingImports : EAsyncPackageState::Complete;
 	}
