Mirror of 3476972, a fix for a copy paste error introduced in 3437368

diff --git a/Engine/Source/Runtime/CoreUObject/Private/Serialization/AsyncLoading.cpp b/Engine/Source/Runtime/CoreUObject/Private/Serialization/AsyncLoading.cpp
--- a/Engine/Source/Runtime/CoreUObject/Private/Serialization/AsyncLoading.cpp
+++ b/Engine/Source/Runtime/CoreUObject/Private/Serialization/AsyncLoading.cpp
@@ -3122,6 +3122,7 @@ void FAsyncPackage::EventDrivenCreateExport(int32 LocalExportIndex)
 						{
 							Export.Object->SetFlags(RF_NeedLoad | RF_NeedPostLoad | RF_NeedPostLoadSubobjects | RF_WasLoaded);
 						}
+						Export.Object->ClearFlags(RF_WillBeLoaded);
 					}
 				}
 			}
