Fix case of CommandLine.h include, and move it to the top of the file to prevent an IncludeTool warning.

diff --git a/Engine/Source/Runtime/CoreUObject/Private/UObject/UObjectBase.cpp b/Engine/Source/Runtime/CoreUObject/Private/UObject/UObjectBase.cpp
--- a/Engine/Source/Runtime/CoreUObject/Private/UObject/UObjectBase.cpp
+++ b/Engine/Source/Runtime/CoreUObject/Private/UObject/UObjectBase.cpp
@@ -18,6 +18,7 @@
 #include "Templates/Casts.h"
 #include "UObject/GCObject.h"
 #include "LinkerLoad.h"
+#include "CommandLine.h"
 
 DEFINE_LOG_CATEGORY_STATIC(LogUObjectBase, Log, All);
 DEFINE_STAT(STAT_UObjectsStatGroupTester);
@@ -935,7 +936,6 @@ static FAutoConsoleVariableRef CMaxObjectsInGame(
 	);
 
 
-#include "Commandline.h"
 /**
  * Final phase of UObject initialization. all auto register objects are added to the main data structures.
  */
