Copying //UE4/Dev-Framework to //UE4/Dev-Main (Source: //UE4/Dev-Framework @ 3510040)

diff --git a/Engine/Source/Runtime/CoreUObject/Public/UObject/Class.h b/Engine/Source/Runtime/CoreUObject/Public/UObject/Class.h
--- a/Engine/Source/Runtime/CoreUObject/Public/UObject/Class.h
+++ b/Engine/Source/Runtime/CoreUObject/Public/UObject/Class.h
@@ -25,18 +25,9 @@ struct FNetDeltaSerializeInfo;
 struct FObjectInstancingGraph;
 struct FPropertyTag;
 
-/*-----------------------------------------------------------------------------
-	Mirrors of mirror structures in Object.h. These are used by generated code 
-	to facilitate correct offsets and alignments for structures containing these '
-	odd types.
------------------------------------------------------------------------------*/
-
 COREUOBJECT_API DECLARE_LOG_CATEGORY_EXTERN(LogClass, Log, All);
 COREUOBJECT_API DECLARE_LOG_CATEGORY_EXTERN(LogScriptSerialization, Log, All);
 
-struct FPropertyTag;
-struct FNetDeltaSerializeInfo;
-
 /*-----------------------------------------------------------------------------
 	FRepRecord.
 -----------------------------------------------------------------------------*/
@@ -64,7 +55,7 @@ class COREUOBJECT_API UField : public UObject
 {
 	DECLARE_CASTED_CLASS_INTRINSIC(UField, UObject, CLASS_Abstract, TEXT("/Script/CoreUObject"), CASTCLASS_UField)
 
-	// Variables.
+	/** Next Field in the linked list */
 	UField*			Next;
 
 	// Constructors.
@@ -77,7 +68,11 @@ class COREUOBJECT_API UField : public UObject
 	// UField interface.
 	virtual void AddCppProperty( UProperty* Property );
 	virtual void Bind();
+
+	/** Goes up the outer chain to look for a UClass */
 	UClass* GetOwnerClass() const;
+
+	/** Goes up the outer chain to look for a UStruct */
 	UStruct* GetOwnerStruct() const;
 
 #if WITH_EDITOR || HACK_HEADER_GENERATOR
@@ -104,13 +99,6 @@ class COREUOBJECT_API UField : public UObject
 	 * @return true if there is a (possibly blank) value associated with this key
 	 */
 	bool HasMetaData(const TCHAR* Key) const;
-
-	/**
-	 * Determines if the property has any metadata associated with the key
-	 * 
-	 * @param Key The key to lookup in the metadata
-	 * @return true if there is a (possibly blank) value associated with this key
-	 */
 	bool HasMetaData(const FName& Key) const;
 
 	/**
@@ -120,8 +108,6 @@ class COREUOBJECT_API UField : public UObject
 	 * @return The value associated with the key
 	 */
 	const FString& GetMetaData(const TCHAR* Key) const;
-
-	/** Find the metadata value associated with the key */
 	const FString& GetMetaData(const FName& Key) const;
 
 	/**
@@ -133,15 +119,6 @@ class COREUOBJECT_API UField : public UObject
 	 * @return							Localized metadata if available, defaults to whatever is provided via GetMetaData
 	 */
 	const FText GetMetaDataText(const TCHAR* MetaDataKey, const FString LocalizationNamespace = FString(), const FString LocalizationKey = FString()) const;
-
-	/**
-	 * Find the metadata value associated with the key and localization namespace and key
-	 *
-	 * @param Key						The key to lookup in the metadata
-	 * @param LocalizationNamespace		Namespace to lookup in the localization manager
-	 * @param LocalizationKey			Key to lookup in the localization manager
-	 * @return							Localized metadata if available, defaults to whatever is provided via GetMetaData
-	 */
 	const FText GetMetaDataText(const FName& MetaDataKey, const FString LocalizationNamespace = FString(), const FString LocalizationKey = FString()) const;
 
 	/**
@@ -151,29 +128,20 @@ class COREUOBJECT_API UField : public UObject
 	 * @return The value associated with the key
 	 */
 	void SetMetaData(const TCHAR* Key, const TCHAR* InValue);
-
-	/** Sets the metadata value associated with the key */
 	void SetMetaData(const FName& Key, const TCHAR* InValue);
 
 	/**
-	* Find the metadata value associated with the key
-	* and return bool 
-	* @param Key The key to lookup in the metadata
-	* @return return true if the value was true (case insensitive)
-	*/
+	 * Find the metadata value associated with the key
+	 * and return bool 
+	 * @param Key The key to lookup in the metadata
+	 * @return return true if the value was true (case insensitive)
+	 */
 	bool GetBoolMetaData(const TCHAR* Key) const
 	{		
 		const FString& BoolString = GetMetaData(Key);
 		// FString == operator does case insensitive comparison
 		return (BoolString == "true");
 	}
-	
-	/**
-	 * Find the metadata value associated with the key
-	 * and return bool  
-	 * @param Key The key to lookup in the metadata
-	 * @return return true if the value was true (case insensitive)
-	 */
 	bool GetBoolMetaData(const FName& Key) const
 	{		
 		const FString& BoolString = GetMetaData(Key);
@@ -182,24 +150,17 @@ class COREUOBJECT_API UField : public UObject
 	}
 
 	/**
-	* Find the metadata value associated with the key
-	* and return int32 
-	* @param Key The key to lookup in the metadata
-	* @return the int value stored in the metadata.
-	*/
+	 * Find the metadata value associated with the key
+	 * and return int32 
+	 * @param Key The key to lookup in the metadata
+	 * @return the int value stored in the metadata.
+	 */
 	int32 GetINTMetaData(const TCHAR* Key) const
 	{
 		const FString& INTString = GetMetaData(Key);
 		int32 Value = FCString::Atoi(*INTString);
 		return Value;
 	}
-
-	/**
-	* Find the metadata value associated with the key
-	* and return int32 
-	* @param Key The key to lookup in the metadata
-	* @return the int value stored in the metadata.
-	*/
 	int32 GetINTMetaData(const FName& Key) const
 	{
 		const FString& INTString = GetMetaData(Key);
@@ -208,11 +169,11 @@ class COREUOBJECT_API UField : public UObject
 	}
 
 	/**
-	* Find the metadata value associated with the key
-	* and return float
-	* @param Key The key to lookup in the metadata
-	* @return the float value stored in the metadata.
-	*/
+	 * Find the metadata value associated with the key
+	 * and return float
+	 * @param Key The key to lookup in the metadata
+	 * @return the float value stored in the metadata.
+	 */
 	float GetFLOATMetaData(const TCHAR* Key) const
 	{
 		const FString& FLOATString = GetMetaData(Key);
@@ -220,13 +181,6 @@ class COREUOBJECT_API UField : public UObject
 		float Value = FCString::Atof(*FLOATString);
 		return Value;
 	}
-
-	/**
-	* Find the metadata value associated with the key
-	* and return float
-	* @param Key The key to lookup in the metadata
-	* @return the float value stored in the metadata.
-	*/
 	float GetFLOATMetaData(const FName& Key) const
 	{
 		const FString& FLOATString = GetMetaData(Key);
@@ -234,14 +188,18 @@ class COREUOBJECT_API UField : public UObject
 		float Value = FCString::Atof(*FLOATString);
 		return Value;
 	}
-
+	
+	/**
+	 * Find the metadata value associated with the key
+	 * and return Class
+	 * @param Key The key to lookup in the metadata
+	 * @return the class value stored in the metadata.
+	 */
 	UClass* GetClassMetaData(const TCHAR* Key) const;
 	UClass* GetClassMetaData(const FName& Key) const;
 
 	/** Clear any metadata associated with the key */
 	void RemoveMetaData(const TCHAR* Key);
-
-	/** Clear any metadata associated with the key */
 	void RemoveMetaData(const FName& Key);
 #endif
 };
@@ -261,22 +219,27 @@ class COREUOBJECT_API UStruct : public UField
 protected:
 	friend COREUOBJECT_API UClass* Z_Construct_UClass_UStruct();
 private:
+	/** Struct this inherits from, may be null */
 	UStruct* SuperStruct;
 public:
+	/** Pointer to start of linked list of child fields */
 	UField* Children;
+	
+	/** Total size of all UProperties, the allocated structure may be larger due to alignment */
 	int32 PropertiesSize;
-
+	/** Alignment of structure in memory, structure will be at least this large */
 	int32 MinAlignment;
 	
+	/** Script bytecode associated with this object */
 	TArray<uint8> Script;
 
-	/** In memory only: Linked list of properties from most-derived to base **/
+	/** In memory only: Linked list of properties from most-derived to base */
 	UProperty* PropertyLink;
-	/** In memory only: Linked list of object reference properties from most-derived to base **/
+	/** In memory only: Linked list of object reference properties from most-derived to base */
 	UProperty* RefLink;
-	/** In memory only: Linked list of properties requiring destruction. Note this does not include things that will be destroyed byt he native destructor **/
+	/** In memory only: Linked list of properties requiring destruction. Note this does not include things that will be destroyed byt he native destructor */
 	UProperty* DestructorLink;
-	/** In memory only: Linked list of properties requiring post constructor initialization.**/
+	/** In memory only: Linked list of properties requiring post constructor initialization */
 	UProperty* PostConstructLink;
 
 	/** Array of object references embedded in script code. Mirrored for easy access by realtime garbage collection code */
@@ -294,10 +257,12 @@ public:
 	virtual void RegisterDependencies() override;
 	static void AddReferencedObjects(UObject* InThis, FReferenceCollector& Collector);
 	virtual void GetPreloadDependencies(TArray<UObject*>& OutDeps) override;
+	virtual void TagSubobjects(EObjectFlags NewFlags) override;
 
 	// UField interface.
 	virtual void AddCppProperty(UProperty* Property) override;
 
+	/** Searches property link chain for a property with the specified name */
 	UProperty* FindPropertyByName(FName InName) const;
 
 	/**
@@ -311,15 +276,16 @@ public:
 	 */
 	void InstanceSubobjectTemplates( void* Data, void const* DefaultData, UStruct* DefaultStruct, UObject* Owner, FObjectInstancingGraph* InstanceGraph );
 
-	//
+	/** Returns the structure used for inheritance, may be changed by child types */
 	virtual UStruct* GetInheritanceSuper() const {return GetSuperStruct();}
 
-	//
+	/** Static wrapper for Link, using a dummy archive */
 	void StaticLink(bool bRelinkExistingProperties = false);
 
-	//
+	/** Creates the field/property links and gets structure ready for use at runtime */
 	virtual void Link(FArchive& Ar, bool bRelinkExistingProperties);
 
+	/** Serializes struct properties, does not handle defaults */
 	virtual void SerializeBin( FArchive& Ar, void* Data ) const;
 
 	/**
@@ -332,7 +298,8 @@ public:
 	 */
 	void SerializeBinEx( FArchive& Ar, void* Data, void const* DefaultData, UStruct* DefaultStruct ) const;
 
-	virtual void SerializeTaggedProperties( FArchive& Ar, uint8* Data, UStruct* DefaultsStruct, uint8* Defaults, const UObject* BreakRecursionIfFullyLoad=NULL) const;
+	/** Serializes list of properties, using property tags to handle mismatches */
+	virtual void SerializeTaggedProperties( FArchive& Ar, uint8* Data, UStruct* DefaultsStruct, uint8* Defaults, const UObject* BreakRecursionIfFullyLoad=nullptr) const;
 
 	/**
 	 * Initialize a struct over uninitialized memory. This may be done by calling the native constructor or individually initializing properties
@@ -342,6 +309,7 @@ public:
 	 * @param	Stride		Stride of the array, If this default (0), then we will pull the size from the struct
 	 */
 	virtual void InitializeStruct(void* Dest, int32 ArrayDim = 1) const;
+	
 	/**
 	 * Destroy a struct in memory. This may be done by calling the native destructor and then the constructor or individually reinitializing properties
 	 *
@@ -353,11 +321,13 @@ public:
 
 #if WITH_EDITOR
 public:
-	virtual UProperty* CustomFindProperty(const FName InName) const { return NULL; };
+	/** Used to search for properties in user defined structs */
+	virtual UProperty* CustomFindProperty(const FName InName) const { return nullptr; };
 #endif // WITH_EDITOR
 public:
+
+	/** Serialize an expression to an archive. Returns expression token */
 	virtual EExprToken SerializeExpr(int32& iCode, FArchive& Ar);
-	virtual void TagSubobjects(EObjectFlags NewFlags) override;
 
 	/**
 	 * Returns the struct/ class prefix used for the C++ declaration of this struct/ class.
@@ -366,32 +336,38 @@ public:
 	 */
 	virtual const TCHAR* GetPrefixCPP() const { return TEXT("F"); }
 
+	/** Total size of all UProperties, the allocated structure may be larger due to alignment */
 	FORCEINLINE int32 GetPropertiesSize() const
 	{
 		return PropertiesSize;
 	}
 
+	/** Alignment of structure in memory, structure will be at least this large */
 	FORCEINLINE int32 GetMinAlignment() const
 	{
 		return MinAlignment;
 	}
 
+	/** Returns actual allocated size of structure in memory */
 	FORCEINLINE int32 GetStructureSize() const
 	{
 		return Align(PropertiesSize,MinAlignment);
 	}
 
+	/** Modifies the property size after it's been recomputed */
 	void SetPropertiesSize( int32 NewSize )
 	{
 		PropertiesSize = NewSize;
 	}
 
+	/** Returns true if this struct either is class T, or is a child of class T. This will not crash on null structs */
 	template<class T>
 	bool IsChildOf() const
 	{
 		return IsChildOf(T::StaticClass());
 	}
 
+	/** Returns true if this struct either is SomeBase, or is a child of SomeBase. This will not crash on null structs */
 	bool IsChildOf( const UStruct* SomeBase ) const
 	{
 		for (const UStruct* Struct = this; Struct; Struct = Struct->GetSuperStruct())
@@ -403,6 +379,7 @@ public:
 		return false;
 	}
 
+	/** Struct this inherits from, may be null */
 	UStruct* GetSuperStruct() const
 	{
 		return SuperStruct;
@@ -414,17 +391,17 @@ public:
 	 */
 	virtual void SetSuperStruct(UStruct* NewSuperStruct);
 
-	/**
-	 * Serializes the SuperStruct pointer.
-	 */
+	/** Serializes the SuperStruct pointer */
 	virtual void SerializeSuperStruct(FArchive& Ar);
 
+	/** Called to add a Field to the front of the linked list for the parent struct */
 	void LinkChild(UField* Child)
 	{
 		Child->Next = Children;
 		Children = Child;
 	}
 
+	/** Returns the a human readable string for a property, overridden for user defined structs */
 	virtual FString PropertyNameToDisplayName(FName InName) const 
 	{ 
 		return InName.ToString();
@@ -438,11 +415,11 @@ public:
 	bool GetStringMetaDataHierarchical(const FName& Key, FString* OutValue = nullptr) const;
 
 	/**
-	* Determines if the struct or any of its super structs has any metadata associated with the provided key
-	*
-	* @param Key The key to lookup in the metadata
-	* @return pointer to the UStruct that has associated metadata, nullptr if Key is not associated with any UStruct in the hierarchy
-	*/
+	 * Determines if the struct or any of its super structs has any metadata associated with the provided key
+	 *
+	 * @param Key The key to lookup in the metadata
+	 * @return pointer to the UStruct that has associated metadata, nullptr if Key is not associated with any UStruct in the hierarchy
+	 */
 	const UStruct* HasMetaDataHierarchical(const FName& Key) const;
 #endif
 
@@ -792,7 +769,7 @@ FORCEINLINE typename TEnableIf<THasGetTypeHash<CPPSTRUCT>::Value, uint32>::Type
 
 
 /**
- * Reflection data for a structure.
+ * Reflection data for a standalone structure declared in a header or as a user defined struct
  */
 class UScriptStruct : public UStruct
 {
@@ -803,7 +780,7 @@ public:
 		/**
 		 * Constructor
 		 * @param InSize: sizeof() of the structure
-		**/
+		 */
 		ICppStructOps(int32 InSize, int32 InAlignment)
 			: Size(InSize)
 			, Alignment(InAlignment)
@@ -836,14 +813,12 @@ public:
 		/** 
 		 * Serialize this structure 
 		 * @return true if the package is new enough to support this, if false, it will fall back to ordinary script struct serialization
-		**/
+		 */
 		virtual bool Serialize(FArchive& Ar, void *Data) = 0;
 
 		/** return true if this class implements a post serialize call **/
 		virtual bool HasPostSerialize() = 0;
-		/** 
-		 * Call PostLoad on this structure
-		**/
+		/** Call PostLoad on this structure */
 		virtual void PostSerialize(const FArchive& Ar, void *Data) = 0;
 
 		/** return true if this struct can net serialize **/
@@ -851,7 +826,7 @@ public:
 		/** 
 		 * Net serialize this structure 
 		 * @return true if the struct was serialized, otherwise it will fall back to ordinary script struct net serialization
-		**/
+		 */
 		virtual bool NetSerialize(FArchive& Ar, class UPackageMap* Map, bool& bOutSuccess, void *Data) = 0;
 
 		/** return true if this struct can net delta serialize delta (serialize a network delta from a base state) **/
@@ -859,7 +834,7 @@ public:
 		/** 
 		 * Net serialize delta this structure. Serialize a network delta from a base state
 		 * @return true if the struct was serialized, otherwise it will fall back to ordinary script struct net delta serialization
-		**/
+		 */
 		virtual bool NetDeltaSerialize(FNetDeltaSerializeInfo & DeltaParms, void *Data) = 0;
 
 		/** return true if this struct should be memcopied **/
@@ -870,7 +845,7 @@ public:
 		/** 
 		 * Copy this structure 
 		 * @return true if the copy was handled, otherwise it will fall back to CopySingleValue
-		**/
+		 */
 		virtual bool Copy(void* Dest, void const* Src, int32 ArrayDim) = 0;
 
 		/** return true if this struct can compare **/
@@ -878,7 +853,7 @@ public:
 		/** 
 		 * Compare this structure 
 		 * @return true if the copy was handled, otherwise it will fall back to UStructProperty::Identical
-		**/
+		 */
 		virtual bool Identical(const void* A, const void* B, uint32 PortFlags, bool& bOutResult) = 0;
 
 		/** return true if this struct can export **/
@@ -886,7 +861,7 @@ public:
 		/** 
 		 * export this structure 
 		 * @return true if the copy was exported, otherwise it will fall back to UStructProperty::ExportTextItem
-		**/
+		 */
 		virtual bool ExportTextItem(FString& ValueStr, const void* PropertyValue, const void* DefaultValue, class UObject* Parent, int32 PortFlags, class UObject* ExportRootScope) = 0;
 
 		/** return true if this struct can import **/
@@ -894,7 +869,7 @@ public:
 		/** 
 		 * import this structure 
 		 * @return true if the copy was imported, otherwise it will fall back to UStructProperty::ImportText
-		**/
+		 */
 		virtual bool ImportTextItem(const TCHAR*& Buffer, void* Data, int32 PortFlags, class UObject* OwnerObject, FOutputDevice* ErrorText) = 0;
 
 		/** return true if this struct has custom GC code **/
@@ -902,7 +877,7 @@ public:
 		/** 
 		 * return a pointer to a function that can add referenced objects
 		 * @return true if the copy was imported, otherwise it will fall back to UStructProperty::ImportText
-		**/
+		 */
 		typedef void (*TPointerToAddStructReferencedObjects)(const void* A, class FReferenceCollector& Collector);
 		virtual TPointerToAddStructReferencedObjects AddStructReferencedObjects() = 0;
 
@@ -911,7 +886,7 @@ public:
 		/** 
 		 * Serialize this structure, from some other tag
 		 * @return true if this succeeded, false will trigger a warning and not serialize at all
-		**/
+		 */
 		virtual bool SerializeFromMismatchedTag(struct FPropertyTag const& Tag, FArchive& Ar, void *Data) = 0;
 
 		/** return true if this struct has a GetTypeHash */
@@ -933,7 +908,6 @@ public:
 	};
 
 
-
 	/** Template to manage dynamic access to C++ struct construction and destruction **/
 	template<class CPPSTRUCT>
 	struct TCppStructOps : public ICppStructOps
@@ -1098,7 +1072,7 @@ public:
 	DECLARE_CASTED_CLASS_INTRINSIC_NO_CTOR(UScriptStruct, UStruct, 0, TEXT("/Script/CoreUObject"), CASTCLASS_UScriptStruct, COREUOBJECT_API)
 
 	COREUOBJECT_API UScriptStruct( EStaticConstructor, int32 InSize, EObjectFlags InFlags );
-	COREUOBJECT_API explicit UScriptStruct(const FObjectInitializer& ObjectInitializer, UScriptStruct* InSuperStruct, ICppStructOps* InCppStructOps = NULL, EStructFlags InStructFlags = STRUCT_NoFlags, SIZE_T ExplicitSize = 0, SIZE_T ExplicitAlignment = 0);
+	COREUOBJECT_API explicit UScriptStruct(const FObjectInitializer& ObjectInitializer, UScriptStruct* InSuperStruct, ICppStructOps* InCppStructOps = nullptr, EStructFlags InStructFlags = STRUCT_NoFlags, SIZE_T ExplicitSize = 0, SIZE_T ExplicitAlignment = 0);
 	COREUOBJECT_API explicit UScriptStruct(const FObjectInitializer& ObjectInitializer = FObjectInitializer::Get());
 
 public:
@@ -1129,31 +1103,35 @@ public:
 	virtual COREUOBJECT_API void DestroyStruct(void* Dest, int32 ArrayDim = 1) const override;
 	// End of UStruct interface.
 
-	/** Stash a CppStructOps for future use 
+	/** 
+	 * Stash a CppStructOps for future use 
 	 * @param Target Name of the struct 
 	 * @param InCppStructOps Cpp ops for this struct
-	**/
+	 */
 	static COREUOBJECT_API void DeferCppStructOps(FName Target, ICppStructOps* InCppStructOps);
 
 	/** Look for the CppStructOps and hook it up **/
 	COREUOBJECT_API void PrepareCppStructOps();
 
+	/** Returns the CppStructOps that can be used to do custom operations */
 	FORCEINLINE ICppStructOps* GetCppStructOps() const
 	{
 		checkf(bPrepareCppStructOpsCompleted, TEXT("GetCppStructOps: PrepareCppStructOps() has not been called for class %s"), *GetName());
 		return CppStructOps;
 	}
 
+	/** Resets currently assigned CppStructOps, called when loading a struct */
 	void ClearCppStructOps()
 	{
 		StructFlags = EStructFlags(StructFlags & ~STRUCT_ComputedFlags);
 		bPrepareCppStructOpsCompleted = false;
-		CppStructOps = NULL;
+		CppStructOps = nullptr;
 	}
+
 	/** 
 	 * If it is native, it is assumed to have defaults because it has a constructor
 	 * @return true if this struct has defaults
-	**/
+	 */
 	FORCEINLINE bool HasDefaults() const
 	{
 		return !!GetCppStructOps();
@@ -1161,7 +1139,6 @@ public:
 
 	/**
 	 * Returns whether this struct should be serialized atomically.
-	 *
 	 * @param	Ar	Archive the struct is going to be serialized with later on
 	 */
 	bool ShouldSerializeAtomically(FArchive& Ar) const
@@ -1265,8 +1242,10 @@ public:
 	 */
 	virtual COREUOBJECT_API uint32 GetStructTypeHash(const void* Src) const;
 
+	/** Used by User Defined Structs to preload this struct and any child objects */
 	virtual COREUOBJECT_API void RecursivelyPreload();
 
+	/** Returns the custom Guid assigned to this struct for User Defined Structs, or an invalid Guid */
 	virtual COREUOBJECT_API FGuid GetCustomGuid() const;
 	
 	/** Returns the (native, c++) name of the struct */
@@ -1274,9 +1253,9 @@ public:
 
 #if WITH_EDITOR
 	/**
-	* Initializes this structure to its default values
-	* @param InStructData		The memory location to initialize
-	*/
+	 * Initializes this structure to its default values
+	 * @param InStructData		The memory location to initialize
+	 */
 	virtual COREUOBJECT_API void InitializeDefaultValue(uint8* InStructData) const;
 #endif // WITH_EDITOR
 };
@@ -1294,12 +1273,20 @@ class COREUOBJECT_API UFunction : public UStruct
 	DECLARE_WITHIN(UClass)
 public:
 	// Persistent variables.
+
+	/** EFunctionFlags set defined for this function */
 	EFunctionFlags FunctionFlags;
+
+	/** Memory Offset of replicated data, only valid if Net flag set */
 	uint16 RepOffset;
 
 	// Variables in memory only.
+	
+	/** Number of parameters total */
 	uint8 NumParms;
+	/** Total size of parameters in memory */
 	uint16 ParmsSize;
+	/** Memory offset of return value property */
 	uint16 ReturnValueOffset;
 	/** Id of this RPC function call (must be FUNC_Net & (FUNC_NetService|FUNC_NetResponse)) */
 	uint16 RPCId;
@@ -1310,14 +1297,15 @@ public:
 	UProperty* FirstPropertyToInit;
 
 #if UE_BLUEPRINT_EVENTGRAPH_FASTCALLS
-	// The event graph this function calls in to (persistent)
+	/** The event graph this function calls in to (persistent) */
 	UFunction* EventGraphFunction;
 
-	// The state offset inside of the event graph (persistent)
+	/** The state offset inside of the event graph (persistent) */
 	int32 EventGraphCallOffset;
 #endif
 
 private:
+	/** C++ function this is bound to */
 	Native Func;
 
 public:
@@ -1354,6 +1342,7 @@ public:
 	explicit UFunction(const FObjectInitializer& ObjectInitializer, UFunction* InSuperFunction, EFunctionFlags InFunctionFlags = FUNC_None, uint16 InRepOffset = 0, SIZE_T ParamsSize = 0 );
 	explicit UFunction(UFunction* InSuperFunction, EFunctionFlags InFunctionFlags = FUNC_None, uint16 InRepOffset = 0, SIZE_T ParamsSize = 0);
 
+	/** Initializes transient members like return value offset */
 	void InitializeDerivedMembers();
 
 	// UObject interface.
@@ -1363,12 +1352,13 @@ public:
 	virtual void Bind() override;
 
 	// UStruct interface.
-	virtual UStruct* GetInheritanceSuper() const override { return NULL;}
+	virtual UStruct* GetInheritanceSuper() const override { return nullptr;}
 	virtual void Link(FArchive& Ar, bool bRelinkExistingProperties) override;
 
-	// UFunction interface.
+	/** Returns parent function if there is one, or null */
 	UFunction* GetSuperFunction() const;
 
+	/** Returns the return value property if there is one, or null */
 	UProperty* GetReturnProperty() const;
 
 	/**
@@ -1403,7 +1393,7 @@ public:
 		//@TODO: UCREMOVAL: CPF_ConstParm added as a hack to get blueprints compiling with a const DamageType parameter.
 		const uint64 IgnoreFlags = CPF_PersistentInstance | CPF_ExportObject | CPF_InstancedReference 
 			| CPF_ContainsInstancedReference | CPF_ComputedFlags | CPF_ConstParm | CPF_UObjectWrapper
-			| CPF_NativeAccessSpecifiers | CPF_AdvancedDisplay;
+			| CPF_NativeAccessSpecifiers | CPF_AdvancedDisplay | CPF_BlueprintVisible | CPF_BlueprintReadOnly;
 		return IgnoreFlags;
 	}
 
@@ -1431,6 +1421,9 @@ public:
 	bool IsSignatureCompatibleWith(const UFunction* OtherFunction, uint64 IgnoreFlags) const;
 };
 
+//
+// Function definition used by dynamic delegate declarations
+//
 class COREUOBJECT_API UDelegateFunction : public UFunction
 {
 	DECLARE_CASTED_CLASS_INTRINSIC(UDelegateFunction, UFunction, 0, TEXT("/Script/CoreUObject"), CASTCLASS_UDelegateFunction)
@@ -1594,13 +1587,13 @@ public:
 	/** searches the list of all enum value names for the specified name
 	 * @return the value the specified name represents if found, otherwise INDEX_NONE
 	 */
-	static int64 LookupEnumNameSlow(const TCHAR* InTestShortName, UEnum** FoundEnum = NULL)
+	static int64 LookupEnumNameSlow(const TCHAR* InTestShortName, UEnum** FoundEnum = nullptr)
 	{
 		int64 Result = LookupEnumName(InTestShortName, FoundEnum);
 		if (Result == INDEX_NONE)
 		{
 			FString TestShortName = FString(TEXT("::")) + InTestShortName;
-			UEnum* TheEnum = NULL;
+			UEnum* TheEnum = nullptr;
 			for (TMap<FName, UEnum*>::TIterator It(AllEnumNames); It; ++It)
 			{
 				if (It.Key().ToString().Contains(TestShortName) )
@@ -1608,11 +1601,11 @@ public:
 					TheEnum = It.Value();
 				}
 			}
-			if (FoundEnum != NULL)
+			if (FoundEnum != nullptr)
 			{
 				*FoundEnum = TheEnum;
 			}
-			Result = (TheEnum != NULL) ? TheEnum->GetValueByName(InTestShortName) : INDEX_NONE;
+			Result = (TheEnum != nullptr) ? TheEnum->GetValueByName(InTestShortName) : INDEX_NONE;
 		}
 		return Result;
 	}
@@ -1886,7 +1879,7 @@ struct COREUOBJECT_API FImplementedInterface
 	bool bImplementedByK2;
 
 	FImplementedInterface()
-		: Class(NULL)
+		: Class(nullptr)
 		, PointerOffset(0)
 		, bImplementedByK2(false)
 	{}
@@ -2007,19 +2000,19 @@ public:
 	/** Pointer to a static AddReferencedObjects method. */
 	ClassAddReferencedObjectsType ClassAddReferencedObjects;
 
-	// Class pseudo-unique counter; used to accelerate unique instance name generation
+	/** Class pseudo-unique counter; used to accelerate unique instance name generation */
 	int32 ClassUnique;
 
-	// Class flags; See EClassFlags for more information
+	/** Class flags; See EClassFlags for more information */
 	EClassFlags ClassFlags;
 
-	// Cast flags used to accelerate dynamic_cast<T*> on objects of this type for common T
+	/** Cast flags used to accelerate dynamic_cast<T*> on objects of this type for common T */
 	EClassCastFlags ClassCastFlags;
 
-	// The required type for the outer of instances of this class.
+	/** The required type for the outer of instances of this class */
 	UClass* ClassWithin;
 
-	// This is the blueprint that caused the generation of this class, or NULL if it is a native compiled-in class
+	/** This is the blueprint that caused the generation of this class, or null if it is a native compiled-in class */
 	UObject* ClassGeneratedBy;
 
 #if WITH_EDITOR
@@ -2031,20 +2024,18 @@ public:
 	virtual void FlushCompilationQueueForLevel() {}
 #endif //WITH_EDITOR
 
-	//
+	/** Which Name.ini file to load Config variables out of */
 	FName ClassConfigName;
 
-	// Used to check if the class was cooked or not.
+	/** Used to check if the class was cooked or not */
 	bool bCooked;
 
-	// List of replication records
+	/** List of replication records */
 	TArray<FRepRecord> ClassReps;
 
-	// List of network relevant fields (properties and functions)
+	/** List of network relevant fields (properties and functions) */
 	TArray<UField*> NetFields;
 
-	virtual bool IsNameStableForNetworking() const override { return true; }		// For now, assume all classes have stable net names
-
 #if WITH_EDITOR || HACK_HEADER_GENERATOR 
 	// Editor only properties
 	void GetHideFunctions(TArray<FString>& OutHideFunctions) const;
@@ -2067,16 +2058,14 @@ public:
 		// The object must of this class type.
 		check(This->IsA(this)); 
 		// This is should always be set to something, at the very least to UObject::ARO
-		check(ClassAddReferencedObjects != NULL); 
+		check(ClassAddReferencedObjects != nullptr);
 		ClassAddReferencedObjects(This, Collector);
 	}
 
-	// The class default object; used for delta serialization and object initialization
+	/** The class default object; used for delta serialization and object initialization */
 	UObject* ClassDefaultObject;
 
-	/**
-	* Assemble reference token streams for all classes if they haven't had it assembled already
-	*/
+	/** Assemble reference token streams for all classes if they haven't had it assembled already */
 	static void AssembleReferenceTokenStreams();
 
 private:
@@ -2095,8 +2084,8 @@ private:
 public:
 	/**
 	 * The list of interfaces which this class implements, along with the pointer property that is located at the offset of the interface's vtable.
-	 * If the interface class isn't native, the property will be NULL.
-	 **/
+	 * If the interface class isn't native, the property will be null.
+	 */
 	TArray<FImplementedInterface> Interfaces;
 
 	/**
@@ -2112,12 +2101,10 @@ public:
 	FCriticalSection ReferenceTokenStreamCritical;
 
 #if !(UE_BUILD_TEST || UE_BUILD_SHIPPING)
-	/* TokenIndex map to look-up token stream index origin. */
+	/** TokenIndex map to look-up token stream index origin. */
 	FGCDebugReferenceTokenMap DebugTokenMap;
 #endif
 
-	
-
 	/** This class's native functions. */
 	TArray<FNativeFunctionLookup> NativeFunctionLookupTable;
 
@@ -2140,7 +2127,7 @@ public:
 	 * @param	InClassConstructor				Pointer to InternalConstructor<TClass>
 	 * @param	TClass_Super_StaticClass		Static class of the super class
 	 * @param	TClass_WithinClass_StaticClass	Static class of the WithinClass
-	 **/
+	 */
 	bool HotReloadPrivateStaticClass(
 		uint32			InSize,
 		EClassFlags		InClassFlags,
@@ -2181,7 +2168,7 @@ public:
 	 * Add a native function to the internal native function table
 	 * @param	InName							name of the function
 	 * @param	InPointer						pointer to the function
-	 **/
+	 */
 	void AddNativeFunction(const ANSICHAR* InName, Native InPointer);
 
 	/**
@@ -2189,35 +2176,39 @@ public:
 	 * which can have unicode identifiers for functions and properties.
 	 * @param	InName							name of the function
 	 * @param	InPointer						pointer to the function
-	 **/
+	 */
 	void AddNativeFunction(const WIDECHAR* InName, Native InPointer);
 
-	// Add a function to the function map
+	/** Add a function to the function map */
 	void AddFunctionToFunctionMap(UFunction* NewFunction)
 	{
 		FuncMap.Add(NewFunction->GetFName(), NewFunction);
 	}
 
-	// This is used by the code generator, which instantiates UFunctions with a name that is later overridden. Overridden names
-	// are needed to support generated versions of blueprint classes, properties of which do not have the same naming 
-	// restrictions as native C++ properties.
+	/** 
+	 * This is used by the code generator, which instantiates UFunctions with a name that is later overridden. Overridden names
+	 * are needed to support generated versions of blueprint classes, properties of which do not have the same naming 
+	 * restrictions as native C++ properties
+	 */
 	void AddFunctionToFunctionMapWithOverriddenName(UFunction* NewFunction, FName OverriddenName)
 	{
 		FuncMap.Add(OverriddenName, NewFunction);
 	}
 
-	// Remove a function from the function map
+	/** Remove a function from the function map */
 	void RemoveFunctionFromFunctionMap(UFunction* Function)
 	{
 		FuncMap.Remove(Function->GetFName());
 	}
 
+	/** Clears the function name caches, in case things have changed */
 	void ClearFunctionMapsCaches()
 	{
 		ParentFuncMap.Empty();
 		InterfaceFuncMap.Empty();
 	}
 
+	/** Looks for a given function name */
 	UFunction* FindFunctionByName(FName InName, EIncludeSuperFlag::Type IncludeSuper = EIncludeSuperFlag::IncludeSuper) const;
 
 	// UObject interface.
@@ -2225,14 +2216,15 @@ public:
 	virtual void PostLoad() override;
 	virtual void FinishDestroy() override;
 	virtual void DeferredRegister(UClass *UClassStaticClass,const TCHAR* PackageName,const TCHAR* InName) override;
-	virtual bool Rename(const TCHAR* NewName = NULL, UObject* NewOuter = NULL, ERenameFlags Flags = REN_None) override;
+	virtual bool Rename(const TCHAR* NewName = nullptr, UObject* NewOuter = nullptr, ERenameFlags Flags = REN_None) override;
 	virtual void TagSubobjects(EObjectFlags NewFlags) override;
 	virtual void PostInitProperties() override;
 	static void AddReferencedObjects(UObject* InThis, FReferenceCollector& Collector);
 	virtual FRestoreForUObjectOverwrite* GetRestoreForUObjectOverwrite() override;
 	virtual FString GetDesc() override;
 	virtual void GetAssetRegistryTags(TArray<FAssetRegistryTag>& OutTags) const override;
-	virtual bool IsAsset() const override { return false; }
+	virtual bool IsAsset() const override { return false; }	
+	virtual bool IsNameStableForNetworking() const override { return true; } // For now, assume all classes have stable net names
 	// End of UObject interface.
 
 	// UField interface.
@@ -2269,6 +2261,7 @@ public:
 	 */
 	const FString GetConfigName() const;
 
+	/** Returns parent class, the parent of a Class is always another class */
 	UClass* GetSuperClass() const
 	{
 		return (UClass*)GetSuperStruct();
@@ -2277,19 +2270,20 @@ public:
 	/** Feedback context for default property import **/
 	static class FFeedbackContext& GetDefaultPropertiesFeedbackContext();
 
+	/** Returns amount of memory used by default object */
 	int32 GetDefaultsCount()
 	{
-		return ClassDefaultObject != NULL ? GetPropertiesSize() : 0;
+		return ClassDefaultObject != nullptr ? GetPropertiesSize() : 0;
 	}
 
 	/**
-	  * Get the default object from the class
-	  * @param	bCreateIfNeeded if true (default) then the CDO is created if it is NULL.
-	  * @return		the CDO for this class
-	**/
+	 * Get the default object from the class
+	 * @param	bCreateIfNeeded if true (default) then the CDO is created if it is null
+	 * @return		the CDO for this class
+	 */
 	UObject* GetDefaultObject(bool bCreateIfNeeded = true)
 	{
-		if (ClassDefaultObject == NULL && bCreateIfNeeded)
+		if (ClassDefaultObject == nullptr && bCreateIfNeeded)
 		{
 			CreateDefaultObject();
 		}
@@ -2298,37 +2292,40 @@ public:
 	}
 
 	/**
-	* Helper method to assist with initializing object properties from an explicit list.
-	*
-	* @param	InStruct			the current scope for which the given property list applies
-	* @param	DataPtr				destination address (where to start copying values to)
-	* @param	DefaultDataPtr		source address (where to start copying the defaults data from)
-	*/
+	 * Helper method to assist with initializing object properties from an explicit list.
+	 *
+	 * @param	InStruct			the current scope for which the given property list applies
+	 * @param	DataPtr				destination address (where to start copying values to)
+	 * @param	DefaultDataPtr		source address (where to start copying the defaults data from)
+	 */
 	virtual void InitPropertiesFromCustomList(uint8* DataPtr, const uint8* DefaultDataPtr) {}
 
 	/**
-	* Get the name of the CDO for the this class
-	* @return The name of the CDO
-	*/
+	 * Get the name of the CDO for the this class
+	 * @return The name of the CDO
+	 */
 	FName GetDefaultObjectName();
 
+	/** Returns memory used to store temporary data on an instance, used by blueprints */
 	virtual uint8* GetPersistentUberGraphFrame(UObject* Obj, UFunction* FuncToCheck) const
 	{
-		return NULL;
+		return nullptr;
 	}
 
+	/** Creates memory to store temporary data */
 	virtual void CreatePersistentUberGraphFrame(UObject* Obj, bool bCreateOnlyIfEmpty = false, bool bSkipSuperClass = false, UClass* OldClass = nullptr) const
 	{
 	}
-
+	
+	/** Clears memory to store temporary data */
 	virtual void DestroyPersistentUberGraphFrame(UObject* Obj, bool bSkipSuperClass = false) const
 	{
 	}
 
 	/**
-	  * Get the default object from the class and cast to a particular type
-	  * @return		the CDO for this class
-	**/
+	 * Get the default object from the class and cast to a particular type
+	 * @return		the CDO for this class
+	 */
 	template<class T>
 	T* GetDefaultObject()
 	{
@@ -2477,7 +2474,7 @@ public:
 	 * This will return whether or not this class implements the passed in class / interface 
 	 *
 	 * @param SomeClass - the interface to check and see if this class implements it
-	 **/
+	 */
 	bool ImplementsInterface(const class UClass* SomeInterface) const;
 
 	/** serializes the passed in object as this class's default object using the given archive
@@ -2529,8 +2526,12 @@ public:
 	 */
 	virtual bool HasProperty(UProperty* InProperty) const;
 
+	/** Finds the object that is used as the parent object when serializing properties, overridden for blueprints */
 	virtual UObject* FindArchetype(UClass* ArchetypeClass, const FName ArchetypeName) const { return nullptr; }
 
+	/** Returns archetype object for CDO */
+	virtual UObject* GetArchetypeForCDO() const;
+
 	/**
 	 * On save, we order a package's exports in class dependency order (so that
 	 * on load, we create the class dependencies before we create the class).  
@@ -2547,8 +2548,6 @@ public:
 	 */
 	virtual void GetRequiredPreloadDependencies(TArray<UObject*>& DependenciesOut) {}
 
-	virtual UObject* GetArchetypeForCDO() const;
-
 	/** Returns true if this class implements script instrumentation. */
 	virtual bool HasInstrumentation() const { return false; }
 
@@ -2567,22 +2566,28 @@ private:
 		friend class FBlueprintCompileReinstancer;
 	#endif
 
-	// This signature intentionally hides the method declared in UObjectBaseUtility to make it private.
-	// Call IsChildOf instead; Hidden because calling IsA on a class almost always indicates an error where the caller should use IsChildOf.
+	/** 
+	 * This signature intentionally hides the method declared in UObjectBaseUtility to make it private.
+	 * Call IsChildOf instead; Hidden because calling IsA on a class almost always indicates an error where the caller should use IsChildOf
+	 */
 	bool IsA(const UClass* Parent) const
 	{
 		return UObject::IsA(Parent);
 	}
 
-	// This signature intentionally hides the method declared in UObject to make it private.
-	// Call FindFunctionByName instead; This method will search for a function declared in UClass instead of the class it was called on
+	/** 
+	 * This signature intentionally hides the method declared in UObject to make it private.
+	 * Call FindFunctionByName instead; This method will search for a function declared in UClass instead of the class it was called on
+	 */
 	UFunction* FindFunction(FName InName) const
 	{
 		return UObject::FindFunction(InName);
 	}
 
-	// This signature intentionally hides the method declared in UObject to make it private.
-	// Call FindFunctionByName instead; This method will search for a function declared in UClass instead of the class it was called on
+	/** 
+	 * This signature intentionally hides the method declared in UObject to make it private.
+	 * Call FindFunctionByName instead; This method will search for a function declared in UClass instead of the class it was called on
+	 */
 	UFunction* FindFunctionChecked(FName InName) const
 	{
 		return UObject::FindFunctionChecked(InName);
@@ -2626,6 +2631,7 @@ public:
 	virtual void PurgeClass(bool bRecompilingOnLoad) override;
 	virtual UObject* FindArchetype(UClass* ArchetypeClass, const FName ArchetypeName) const override;
 
+	/** Find a struct property, called from generated code */
 	UStructProperty* FindStructPropertyChecked(const TCHAR* PropertyName) const;
 
 	/** Misc objects owned by the class. */
@@ -2637,12 +2643,12 @@ public:
 	/** Outer assets used by the class */
 	TArray<UObject*> UsedAssets;
 
-	// Specialized sub-object containers
+	/** Specialized sub-object containers */
 	TArray<UObject*> DynamicBindingObjects;
 	TArray<UObject*> ComponentTemplates;
 	TArray<UObject*> Timelines;
 
-	// IAnimClassInterface (UAnimClassData) or null
+	/** IAnimClassInterface (UAnimClassData) or null */
 	UObject* AnimClassImplementation;
 };
 
@@ -2814,7 +2820,7 @@ private:
 	 */
 	bool HasDestinationRoot() const
 	{
-		return DestinationRoot != NULL;
+		return DestinationRoot != nullptr;
 	}
 
 	/**
@@ -2985,13 +2991,19 @@ struct FStructUtils
 {
 	static bool ArePropertiesTheSame(const UProperty* A, const UProperty* B, bool bCheckPropertiesNames);
 
-	// does structures have exactly the same memory layout
+	/** Do structures have exactly the same memory layout */
 	COREUOBJECT_API static bool TheSameLayout(const UStruct* StructA, const UStruct* StructB, bool bCheckPropertiesNames = false);
 
-	/** Locates a named structure in the package with the given name. Not expected to fail. */
+	/** Locates a named structure in the package with the given name. Not expected to fail */
 	COREUOBJECT_API static UStruct* FindStructureInPackageChecked(const TCHAR* StructName, const TCHAR* PackageName);
 };
 
+/*-----------------------------------------------------------------------------
+	Mirrors of mirror structures in Object.h. These are used by generated code 
+	to facilitate correct offsets and alignments for structures containing these
+	odd types.
+-----------------------------------------------------------------------------*/
+
 template< class T > struct TBaseStructure
 {
 };
diff --git a/Engine/Source/Runtime/Engine/Classes/AI/Navigation/AbstractNavData.h b/Engine/Source/Runtime/Engine/Classes/AI/Navigation/AbstractNavData.h
--- a/Engine/Source/Runtime/Engine/Classes/AI/Navigation/AbstractNavData.h
+++ b/Engine/Source/Runtime/Engine/Classes/AI/Navigation/AbstractNavData.h
@@ -41,7 +41,12 @@ public:
 UCLASS()
 class ENGINE_API AAbstractNavData : public ANavigationData
 {
-	GENERATED_UCLASS_BODY()
+	GENERATED_BODY()
+
+public:
+	AAbstractNavData(const FObjectInitializer& ObjectInitializer = FObjectInitializer::Get());
+
+	virtual void PostLoad() override;
 
 	// Begin ANavigationData overrides
 	virtual void BatchRaycast(TArray<FNavigationRaycastWork>& Workload, FSharedConstNavQueryFilter QueryFilter, const UObject* Querier = NULL) const override {};
diff --git a/Engine/Source/Runtime/Engine/Classes/AI/Navigation/CrowdManagerBase.h b/Engine/Source/Runtime/Engine/Classes/AI/Navigation/CrowdManagerBase.h
new file mode 100644
--- /dev/null
+++ b/Engine/Source/Runtime/Engine/Classes/AI/Navigation/CrowdManagerBase.h
@@ -0,0 +1,17 @@
+// Copyright 1998-2017 Epic Games, Inc. All Rights Reserved.
+
+#pragma once
+
+#include "CrowdManagerBase.generated.h"
+
+/** Base class for Crowd Managers. If you want to create a custom crowd manager
+ *	implement a class extending this one and set UNavigationSystem::CrowdManagerClass
+ *	to point at your class */
+UCLASS(Abstract, Transient)
+class ENGINE_API UCrowdManagerBase : public UObject
+{
+	GENERATED_BODY()
+public:
+	virtual void Tick(float DeltaTime) PURE_VIRTUAL(UCrowdManagerBase::Tick, );
+};
+
diff --git a/Engine/Source/Runtime/Engine/Classes/AI/Navigation/NavigationSystem.h b/Engine/Source/Runtime/Engine/Classes/AI/Navigation/NavigationSystem.h
--- a/Engine/Source/Runtime/Engine/Classes/AI/Navigation/NavigationSystem.h
+++ b/Engine/Source/Runtime/Engine/Classes/AI/Navigation/NavigationSystem.h
@@ -26,7 +26,7 @@ class FNavDataGenerator;
 class FNavigationOctree;
 class INavLinkCustomInterface;
 class INavRelevantInterface;
-class UCrowdManager;
+class UCrowdManagerBase;
 class UNavArea;
 class UNavigationPath;
 struct FNavigationRelevantData;
@@ -159,10 +159,13 @@ class ENGINE_API UNavigationSystem : public UBlueprintFunctionLibrary
 	ANavigationData* MainNavData;
 
 	/** special navigation data for managing direct paths, not part of NavDataSet! */
-	UPROPERTY()
+	UPROPERTY(Transient)
 	ANavigationData* AbstractNavData;
 
 protected:
+	UPROPERTY(EditAnywhere, BlueprintReadOnly, Category = Navigation)
+	TSubclassOf<UCrowdManagerBase> CrowdManagerClass;
+
 	/** Should navigation system spawn default Navigation Data when there's none and there are navigation bounds present? */
 	UPROPERTY(config, EditAnywhere, Category=NavigationSystem)
 	uint32 bAutoCreateNavigationData:1;
@@ -244,7 +247,7 @@ public:
 	FOnNavigationInitDone OnNavigationInitDone;
 	
 private:
-	TWeakObjectPtr<UCrowdManager> CrowdManager;
+	TWeakObjectPtr<UCrowdManagerBase> CrowdManager;
 	
 	/** set to true when navigation processing was blocked due to missing nav bounds */
 	uint32 bNavDataRemovedDueToMissingNavBounds : 1;
@@ -386,7 +389,7 @@ public:
 
 	UWorld* GetWorld() const override { return GetOuterUWorld(); }
 
-	UCrowdManager* GetCrowdManager() const { return CrowdManager.Get(); }
+	UCrowdManagerBase* GetCrowdManager() const { return CrowdManager.Get(); }
 
 protected:
 	/** spawn new crowd manager */
@@ -874,6 +877,7 @@ protected:
 	uint8 bInitialLevelsAdded : 1;
 	uint8 bWorldInitDone : 1;
 	uint8 bAsyncBuildPaused : 1;
+	uint8 bCanAccumulateDirtyAreas : 1;
 
 	/** cached navigable world bounding box*/
 	mutable FBox NavigableWorldBounds;
@@ -940,7 +944,7 @@ protected:
 	 *	in this regard **/
 	void AddElementToNavOctree(const FNavigationDirtyElement& DirtyElement);
 
-	void SetCrowdManager(UCrowdManager* NewCrowdManager); 
+	void SetCrowdManager(UCrowdManagerBase* NewCrowdManager);
 
 	/** Add BSP collision data to navigation octree */
 	void AddLevelCollisionToOctree(ULevel* Level);
diff --git a/Engine/Source/Runtime/Engine/Classes/AI/Navigation/RecastNavMesh.h b/Engine/Source/Runtime/Engine/Classes/AI/Navigation/RecastNavMesh.h
--- a/Engine/Source/Runtime/Engine/Classes/AI/Navigation/RecastNavMesh.h
+++ b/Engine/Source/Runtime/Engine/Classes/AI/Navigation/RecastNavMesh.h
@@ -47,6 +47,8 @@ class UNavigationSystem;
 class UPrimitiveComponent;
 class URecastNavMeshDataChunk;
 struct FRecastAreaNavModifierElement;
+class dtNavMesh;
+class dtQueryFilter;
 
 UENUM()
 namespace ERecastPartitioning
@@ -893,6 +895,15 @@ public:
 	//----------------------------------------------------------------------//
 	// Querying                                                                
 	//----------------------------------------------------------------------//
+	
+	/** dtNavMesh getter */
+	dtNavMesh* GetRecastMesh();
+
+	/** dtNavMesh getter */
+	const dtNavMesh* GetRecastMesh() const;
+
+	/** Retrieves LinkUserID associated with indicated PolyID */
+	uint32 GetLinkUserId(NavNodeRef LinkPolyID) const;
 
 	FColor GetAreaIDColor(uint8 AreaID) const;
 
@@ -1030,7 +1041,6 @@ protected:
 private:
 	friend class FRecastNavMeshGenerator;
 	friend class FPImplRecastNavMesh;
-	friend class UCrowdManager;
 	// destroys FPImplRecastNavMesh instance if it has been created 
 	void DestroyRecastPImpl();
 	// @todo docuement
diff --git a/Engine/Source/Runtime/Engine/Classes/Components/InterpToMovementComponent.h b/Engine/Source/Runtime/Engine/Classes/Components/InterpToMovementComponent.h
--- a/Engine/Source/Runtime/Engine/Classes/Components/InterpToMovementComponent.h
+++ b/Engine/Source/Runtime/Engine/Classes/Components/InterpToMovementComponent.h
@@ -203,7 +203,7 @@ protected:
 	virtual void HandleImpact(const FHitResult& Hit, float TimeSlice=0.f, const FVector& MoveDelta = FVector::ZeroVector) override;
 
 	/* Reverse direction we are moving */
-	float ReverseDirection(const FHitResult& Hit, float Time, bool InBroadcastEvent);
+	void ReverseDirection(const FHitResult& Hit, float Time, bool InBroadcastEvent);
 
 	/** Compute the distance for the given time. */
 	virtual FVector ComputeMoveDelta(float Time) const;
@@ -233,7 +233,7 @@ protected:
 	virtual void UpdateControlPoints(bool InForceUpdate);	
 
 	/* Calculate the new current time */
-	float CalculateNewTime( float TimeNow, float Delta, FHitResult& HitResult, bool InBroadcastEvent, bool& OutStopped );
+	float CalculateNewTime( float TimeNow, float Delta, FHitResult& HitResult, bool InBroadcastEvent, bool& OutStopped, float& OutTimeRemainder );
 
 private:
 	float TotalDistance;
diff --git a/Engine/Source/Runtime/Engine/Classes/EdGraph/EdGraphNode.h b/Engine/Source/Runtime/Engine/Classes/EdGraph/EdGraphNode.h
--- a/Engine/Source/Runtime/Engine/Classes/EdGraph/EdGraphNode.h
+++ b/Engine/Source/Runtime/Engine/Classes/EdGraph/EdGraphNode.h
@@ -131,6 +131,14 @@ enum class ENodeEnabledState : uint8
 	DevelopmentOnly
 };
 
+/** Enum that defines what kind of orphaned pins should be retained. */
+enum class ESaveOrphanPinMode : uint8
+{
+	SaveNone,
+	SaveAll,
+	SaveAllButExec
+};
+
 /** Holds metadata keys, so as to discourage text duplication throughout the engine. */
 struct ENGINE_API FNodeMetadata
 {
@@ -194,6 +202,9 @@ class ENGINE_API UEdGraphNode : public UObject
 	UPROPERTY()
 	ENodeEnabledState EnabledState;
 
+	/** When reconstructing a node should the orphaned pins be retained and transfered to the new pin list. */
+	ESaveOrphanPinMode OrphanedPinSaveMode;
+
 	/** Indicates whether or not the user explicitly set the enabled state */
 	UPROPERTY()
 	uint8 bUserSetEnabledState:1;
diff --git a/Engine/Source/Runtime/Engine/Classes/EdGraph/EdGraphPin.h b/Engine/Source/Runtime/Engine/Classes/EdGraph/EdGraphPin.h
--- a/Engine/Source/Runtime/Engine/Classes/EdGraph/EdGraphPin.h
+++ b/Engine/Source/Runtime/Engine/Classes/EdGraph/EdGraphPin.h
@@ -378,6 +378,16 @@ public:
 
 	/** If true, the pin is displayed as ref. This is transient. */
 	uint32 bDisplayAsMutableRef:1;
+
+	/** 
+	 * If true, this pin existed on an older version of the owning node, but when the node was reconstructed a matching pin was not found.
+	 * This pin must be linked to other pins or have a non-default value and will be removed if disconnected, reset to default, or the node is refreshed.
+	 **/
+	uint32 bOrphanedPin:1;
+
+	/** If true, this pin will be retained when reconstructing a node if there is no matching pin on the new version of the pin. This is transient. */
+	uint32 bSavePinIfOrphaned:1;
+
 #endif // WITH_EDITORONLY_DATA
 
 	/** True when InvalidateAndTrash was called. This pin is intended to be discarded and destroyed. */
@@ -395,7 +405,7 @@ public:
 	ENGINE_API void BreakLinkTo(UEdGraphPin* ToPin);
 
 	/** Break all links from this pin */
-	ENGINE_API void BreakAllPinLinks();
+	ENGINE_API void BreakAllPinLinks(bool bNotifyNodes = false);
 
 	/**
 	 * Copies the persistent data (across a node refresh) from the SourcePin.
@@ -478,6 +488,8 @@ public:
 		bNotConnectable = false;
 		bDefaultValueIsReadOnly = false;
 		bDefaultValueIsIgnored = false;
+		bOrphanedPin = false;
+		bSavePinIfOrphaned = true;
 #endif // WITH_EDITORONLY_DATA
 	}
 
@@ -522,6 +534,9 @@ public:
 
 	/** This needs to be called if you want to use pin data within PostEditUndo */
 	ENGINE_API static void ResolveAllPinReferences();
+
+	ENGINE_API static bool EnableOrphanPins();
+
 private:
 	/** Private Constructor. Create pins using CreatePin since all pin instances are managed by TSharedPtr. */
 	UEdGraphPin(UEdGraphNode* InOwningNode, const FGuid& PinGuid);
diff --git a/Engine/Source/Runtime/Engine/Classes/EdGraph/EdGraphSchema.h b/Engine/Source/Runtime/Engine/Classes/EdGraph/EdGraphSchema.h
--- a/Engine/Source/Runtime/Engine/Classes/EdGraph/EdGraphSchema.h
+++ b/Engine/Source/Runtime/Engine/Classes/EdGraph/EdGraphSchema.h
@@ -709,6 +709,12 @@ class ENGINE_API UEdGraphSchema : public UObject
 	/** Sets the text to the specified pin */
 	virtual void TrySetDefaultText(UEdGraphPin& InPin, const FText& InNewDefaultText) const;
 
+	/** Returns if the pin's value matches what the true (autogenerated) default value for that pin would be */
+	virtual bool DoesDefaultValueMatchAutogenerated(const UEdGraphPin& InPin) const;
+
+	/** Resets a pin back to it's autogenerated default value, optionally calling the default value change callbacks */
+	virtual void ResetPinToAutogeneratedDefaultValue(UEdGraphPin* Pin, bool bCallModifyCallbacks = true) const { }
+
 	/** If we should disallow viewing and editing of the supplied pin */
 	virtual bool ShouldHidePinDefaultValue(UEdGraphPin* Pin) const { return false; }
 
diff --git a/Engine/Source/Runtime/Engine/Classes/Engine/AssetManager.h b/Engine/Source/Runtime/Engine/Classes/Engine/AssetManager.h
--- a/Engine/Source/Runtime/Engine/Classes/Engine/AssetManager.h
+++ b/Engine/Source/Runtime/Engine/Classes/Engine/AssetManager.h
@@ -336,9 +336,6 @@ public:
 	/** For a given package and platform, return what Chunks it should be assigned to, games can override this as needed. Returns false if no information found */
 	virtual bool GetPackageChunkIds(FName PackageName, const class ITargetPlatform* TargetPlatform, const TArray<int32>& ExistingChunkList, TArray<int32>& OutChunkList) const;
 
-	/** Runs game-specific code to update the source hash on a package at cook time. This is added to the base package guid check */
-	virtual void UpdatePackageSourceHash(FName PackageName, FMD5& PackageSourceHash) const;
-
 	/** Refresh the entire set of asset data, can call from editor when things have changed dramatically */
 	virtual void RefreshPrimaryAssetDirectory();
 
@@ -401,9 +398,6 @@ protected:
 	/** Helper function to write out asset reports */
 	virtual bool WriteCustomReport(FString FileName, TArray<FString>& FileLines) const;
 
-	/** Function used during creating Management references to decide when to recurse and set references */
-	virtual EAssetSetManagerResult::Type ShouldSetManager(const FAssetIdentifier& Manager, const FAssetIdentifier& Source, const FAssetIdentifier& Target, EAssetRegistryDependencyType::Type DependencyType, EAssetSetManagerFlags::Type Flags) const;
-	
 	/** Returns true if the specified TypeInfo should be scanned. Can be implemented by the game. */
 	virtual bool ShouldScanPrimaryAssetType(FPrimaryAssetTypeInfo& TypeInfo) const;
 
@@ -417,6 +411,9 @@ protected:
 	virtual bool IsPathExcludedFromScan(const FString& Path) const;
 
 #if WITH_EDITOR
+	/** Function used during creating Management references to decide when to recurse and set references */
+	virtual EAssetSetManagerResult::Type ShouldSetManager(const FAssetIdentifier& Manager, const FAssetIdentifier& Source, const FAssetIdentifier& Target, EAssetRegistryDependencyType::Type DependencyType, EAssetSetManagerFlags::Type Flags) const;
+
 	/** Helper function which requests the asset register scan a list of directories/assets */
 	virtual void ScanPathsSynchronous(const TArray<FString>& PathsToScan) const;
 
diff --git a/Engine/Source/Runtime/Engine/Classes/Engine/AssetManagerTypes.h b/Engine/Source/Runtime/Engine/Classes/Engine/AssetManagerTypes.h
--- a/Engine/Source/Runtime/Engine/Classes/Engine/AssetManagerTypes.h
+++ b/Engine/Source/Runtime/Engine/Classes/Engine/AssetManagerTypes.h
@@ -85,7 +85,7 @@ struct FPrimaryAssetTypeInfo
 
 private:
 	/** Base Class of all assets of this type */
-	UPROPERTY(EditAnywhere, Category = AssetType)
+	UPROPERTY(EditAnywhere, Category = AssetType, meta = (AllowAbstract))
 	TAssetSubclassOf<UObject> AssetBaseClass;
 public:
 
diff --git a/Engine/Source/Runtime/Engine/Classes/Engine/Engine.h b/Engine/Source/Runtime/Engine/Classes/Engine/Engine.h
--- a/Engine/Source/Runtime/Engine/Classes/Engine/Engine.h
+++ b/Engine/Source/Runtime/Engine/Classes/Engine/Engine.h
@@ -45,6 +45,19 @@ class UNetDriver;
 class FFineGrainedPerformanceTracker;
 #endif
 
+// The kind of failure handling that GetWorldFromContextObject uses 
+enum class EGetWorldErrorMode
+{
+	// Silently returns nullptr, the calling code is expected to handle this gracefully
+	ReturnNull,
+
+	// Raises a runtime error but still returns nullptr, the calling code is expected to handle this gracefully
+	LogAndReturnNull,
+
+	// Asserts, the calling code is not expecting to handle a failure gracefully
+	Assert
+};
+
 /**
  * Enumerates types of fully loaded packages.
  */
@@ -2175,12 +2188,31 @@ public:
 	 * Obtain a world object pointer from an object with has a world context.
 	 *
 	 * @param Object		Object whose owning world we require.
-	 * @param bChecked      Allows calling function to specify not to do ensure check and that a nullptr return value is acceptable
-	 *						This flag is only used when called by main game thread. 
-	 * returns				The world to which the object belongs.
+	 * @param ErrorMode		Controls what happens if the Object cannot be found
+	 * @return				The world to which the object belongs or nullptr if it cannot be found.
 	 */
-	UWorld* GetWorldFromContextObject(const UObject* Object, bool bChecked = true) const;
+	UWorld* GetWorldFromContextObject(const UObject* Object, EGetWorldErrorMode ErrorMode) const;
 
+	/** 
+	 * Obtain a world object pointer from an object with has a world context.
+	 *
+	 * @param Object		Object whose owning world we require.
+	 * @return				The world to which the object belongs; asserts if the world cannot be found!
+	 */
+	UWorld* GetWorldFromContextObjectChecked(const UObject* Object) const
+	{
+		return GetWorldFromContextObject(Object, EGetWorldErrorMode::Assert);
+	}
+
+	/** 
+	 * This function is deprecated
+	 */
+	DEPRECATED(4.17, "GetWorldFromContextObject(Object) and GetWorldFromContextObject(Object, boolean) are replaced by GetWorldFromContextObject(Object, Enum) or GetWorldFromContextObjectChecked(Object)")
+	UWorld* GetWorldFromContextObject(const UObject* Object, bool bChecked = true) const
+	{
+		// Note: The behavior in 4.16 and before was similar to Assert if bChecked was true, but almost no callers actually wanted to pass in bChecked=true
+		return GetWorldFromContextObject(Object, bChecked ? EGetWorldErrorMode::LogAndReturnNull : EGetWorldErrorMode::ReturnNull);
+	}
 
 	/** 
 	 * mostly done to check if PIE is being set up, go GWorld is going to change, and it's not really _the_G_World_
diff --git a/Engine/Source/Runtime/Engine/Classes/Engine/LevelStreaming.h b/Engine/Source/Runtime/Engine/Classes/Engine/LevelStreaming.h
--- a/Engine/Source/Runtime/Engine/Classes/Engine/LevelStreaming.h
+++ b/Engine/Source/Runtime/Engine/Classes/Engine/LevelStreaming.h
@@ -213,6 +213,7 @@ public:
 	ENGINE_API FString GetWorldAssetPackageName() const;
 
 	/** Gets the package name for the world asset referred to by this level streaming as an FName */
+	UFUNCTION(BlueprintCallable, Category = "Game")
 	ENGINE_API FName GetWorldAssetPackageFName() const;
 
 	/** Sets the world asset based on the package name assuming it contains a world of the same name. */
diff --git a/Engine/Source/Runtime/Engine/Classes/Engine/LevelStreamingVolume.h b/Engine/Source/Runtime/Engine/Classes/Engine/LevelStreamingVolume.h
--- a/Engine/Source/Runtime/Engine/Classes/Engine/LevelStreamingVolume.h
+++ b/Engine/Source/Runtime/Engine/Classes/Engine/LevelStreamingVolume.h
@@ -51,12 +51,12 @@ class ALevelStreamingVolume : public AVolume
 	TEnumAsByte<enum EStreamingVolumeUsage> StreamingUsage;
 
 	//~ Begin UObject Interface.
-	void PostLoad() override;
+	virtual ENGINE_API void PostLoad() override;
 	//~ End UObject Interface.
 	
 #if WITH_EDITOR
 	//~ Begin AActor Interface.
-	virtual void CheckForErrors() override;
+	virtual ENGINE_API void CheckForErrors() override;
 	//~ End AActor Interface.
 	
 	/** Updates list of streaming levels that are referenced by this streaming volume */
diff --git a/Engine/Source/Runtime/Engine/Classes/GameFramework/Actor.h b/Engine/Source/Runtime/Engine/Classes/GameFramework/Actor.h
--- a/Engine/Source/Runtime/Engine/Classes/GameFramework/Actor.h
+++ b/Engine/Source/Runtime/Engine/Classes/GameFramework/Actor.h
@@ -2722,6 +2722,9 @@ public:
 		static_assert(TPointerIsConvertibleFromTo<T, const UActorComponent>::Value, "'T' template parameter to GetComponents must be derived from UActorComponent");
 		SCOPE_CYCLE_COUNTER(STAT_GetComponentsTime);
 
+		// Empty input array, but don't affect allocated size.
+		OutComponents.Reset(0);
+
 		TArray<UChildActorComponent*> ChildActorComponents;
 
 		for (UActorComponent* OwnedComponent : OwnedComponents)
diff --git a/Engine/Source/Runtime/Engine/Classes/GameFramework/SpringArmComponent.h b/Engine/Source/Runtime/Engine/Classes/GameFramework/SpringArmComponent.h
--- a/Engine/Source/Runtime/Engine/Classes/GameFramework/SpringArmComponent.h
+++ b/Engine/Source/Runtime/Engine/Classes/GameFramework/SpringArmComponent.h
@@ -118,6 +118,7 @@ class ENGINE_API USpringArmComponent : public USceneComponent
 	virtual void OnRegister() override;
 	virtual void TickComponent(float DeltaTime, enum ELevelTick TickType, FActorComponentTickFunction* ThisTickFunction) override;
 	virtual void PostLoad() override;
+	virtual void ApplyWorldOffset(const FVector& InOffset, bool bWorldShift) override;
 	// End of UActorComponent interface
 
 	// USceneComponent interface
diff --git a/Engine/Source/Runtime/Engine/Classes/Kismet/KismetMathLibrary.h b/Engine/Source/Runtime/Engine/Classes/Kismet/KismetMathLibrary.h
--- a/Engine/Source/Runtime/Engine/Classes/Kismet/KismetMathLibrary.h
+++ b/Engine/Source/Runtime/Engine/Classes/Kismet/KismetMathLibrary.h
@@ -1805,6 +1805,10 @@ class ENGINE_API UKismetMathLibrary : public UBlueprintFunctionLibrary
 	UFUNCTION(BlueprintPure, meta=(DisplayName = "vector2d / float", CompactNodeTitle = "/", Keywords = "/ divide division"), Category="Math|Vector2D")
 	static FVector2D Divide_Vector2DFloat(FVector2D A, float B = 1.f);
 
+	/* Element-wise Vector divide (Result = {A.x/B.x, A.y/B.y}) */
+	UFUNCTION(BlueprintPure, meta = (DisplayName = "vector2d / vector2d", CompactNodeTitle = "/", Keywords = "/ divide division"), Category = "Math|Vector2D")
+	static FVector2D Divide_Vector2DVector2D(FVector2D A, FVector2D B);
+
 	/** Returns Vector A added by B */
 	UFUNCTION(BlueprintPure, meta=(DisplayName = "vector2d + float", CompactNodeTitle = "+", Keywords = "+ add plus"), Category="Math|Vector2D")
 	static FVector2D Add_Vector2DFloat(FVector2D A, float B);
@@ -2092,6 +2096,7 @@ private:
 	static void ReportError_Divide_VectorVector();
 	static void ReportError_ProjectVectorOnToVector();
 	static void ReportError_Divide_Vector2DFloat();
+	static void ReportError_Divide_Vector2DVector2D();
 	static void ReportError_DaysInMonth();
 };
 
diff --git a/Engine/Source/Runtime/Engine/Classes/Kismet/KismetMathLibrary.inl b/Engine/Source/Runtime/Engine/Classes/Kismet/KismetMathLibrary.inl
--- a/Engine/Source/Runtime/Engine/Classes/Kismet/KismetMathLibrary.inl
+++ b/Engine/Source/Runtime/Engine/Classes/Kismet/KismetMathLibrary.inl
@@ -1094,6 +1094,19 @@ FVector2D UKismetMathLibrary::Divide_Vector2DFloat(FVector2D A, float B)
 	return A / B;
 }
 
+KISMET_MATH_FORCEINLINE
+FVector2D UKismetMathLibrary::Divide_Vector2DVector2D(FVector2D A, FVector2D B)
+{
+	if (B.X == 0.f || B.Y == 0.f)
+	{
+		ReportError_Divide_Vector2DVector2D();
+		return FVector2D::ZeroVector;
+	}
+
+	return A / B;
+}
+
+
 KISMET_MATH_FORCEINLINE
 FVector2D UKismetMathLibrary::Add_Vector2DFloat(FVector2D A, float B)
 {
@@ -1478,36 +1491,6 @@ float UKismetMathLibrary::GetTotalSeconds( FTimespan A )
 	return A.GetTotalSeconds();
 }
 
-KISMET_MATH_FORCEINLINE
-FTimespan UKismetMathLibrary::FromDays( float Days )
-{
-	return FTimespan::FromDays(Days);
-}
-
-KISMET_MATH_FORCEINLINE
-FTimespan UKismetMathLibrary::FromHours( float Hours )
-{
-	return FTimespan::FromHours(Hours);
-}
-
-KISMET_MATH_FORCEINLINE
-FTimespan UKismetMathLibrary::FromMilliseconds( float Milliseconds )
-{
-	return FTimespan::FromMilliseconds(Milliseconds);
-}
-
-KISMET_MATH_FORCEINLINE
-FTimespan UKismetMathLibrary::FromMinutes( float Minutes )
-{
-	return FTimespan::FromMinutes(Minutes);
-}
-
-KISMET_MATH_FORCEINLINE
-FTimespan UKismetMathLibrary::FromSeconds( float Seconds )
-{
-	return FTimespan::FromSeconds(Seconds);
-}
-
 KISMET_MATH_FORCEINLINE
 FTimespan UKismetMathLibrary::TimespanMaxValue( )
 {
diff --git a/Engine/Source/Runtime/Engine/Classes/Kismet/KismetSystemLibrary.h b/Engine/Source/Runtime/Engine/Classes/Kismet/KismetSystemLibrary.h
--- a/Engine/Source/Runtime/Engine/Classes/Kismet/KismetSystemLibrary.h
+++ b/Engine/Source/Runtime/Engine/Classes/Kismet/KismetSystemLibrary.h
@@ -212,12 +212,12 @@ class ENGINE_API UKismetSystemLibrary : public UBlueprintFunctionLibrary
 	DECLARE_DYNAMIC_DELEGATE_OneParam(FOnAssetLoaded, class UObject*, Loaded);
 
 	UFUNCTION(BlueprintCallable, meta = (Latent, LatentInfo = "LatentInfo", WorldContext = "WorldContextObject", BlueprintInternalUseOnly = "true"), Category = "Utilities")
-	static void LoadAsset(UObject* WorldContextObject, const TAssetPtr<UObject>& Asset, FOnAssetLoaded OnLoaded, FLatentActionInfo LatentInfo);
+	static void LoadAsset(UObject* WorldContextObject, TAssetPtr<UObject> Asset, FOnAssetLoaded OnLoaded, FLatentActionInfo LatentInfo);
 
 	DECLARE_DYNAMIC_DELEGATE_OneParam(FOnAssetClassLoaded, TSubclassOf<UObject>, Loaded);
 
 	UFUNCTION(BlueprintCallable, meta = (Latent, LatentInfo = "LatentInfo", WorldContext = "WorldContextObject", BlueprintInternalUseOnly = "true"), Category = "Utilities")
-	static void LoadAssetClass(UObject* WorldContextObject, const TAssetSubclassOf<UObject>& AssetClass, FOnAssetClassLoaded OnLoaded, FLatentActionInfo LatentInfo);
+	static void LoadAssetClass(UObject* WorldContextObject, TAssetSubclassOf<UObject> AssetClass, FOnAssetClassLoaded OnLoaded, FLatentActionInfo LatentInfo);
 
 
 	/**
diff --git a/Engine/Source/Runtime/Engine/Private/AI/Navigation/AbstractNavData.cpp b/Engine/Source/Runtime/Engine/Private/AI/Navigation/AbstractNavData.cpp
--- a/Engine/Source/Runtime/Engine/Private/AI/Navigation/AbstractNavData.cpp
+++ b/Engine/Source/Runtime/Engine/Private/AI/Navigation/AbstractNavData.cpp
@@ -14,7 +14,8 @@ INavigationQueryFilterInterface* FAbstractQueryFilter::CreateCopy() const
 	return new FAbstractQueryFilter();
 }
 
-AAbstractNavData::AAbstractNavData(const FObjectInitializer& ObjectInitializer) : Super(ObjectInitializer)
+AAbstractNavData::AAbstractNavData(const FObjectInitializer& ObjectInitializer) 
+	: Super(ObjectInitializer)
 {
 #if WITH_EDITORONLY_DATA
 	bEditable = false;
@@ -38,6 +39,16 @@ AAbstractNavData::AAbstractNavData(const FObjectInitializer& ObjectInitializer)
 	}
 }
 
+void AAbstractNavData::PostLoad()
+{
+	Super::PostLoad();
+	SetFlags(RF_Transient);
+	// marking as pending kill might seem an overkill, but one of the things 
+	// this changes aims to achieve is to get rig of the excess number of 
+	// AAbstractNavData instances. "There should be only one!"
+	MarkPendingKill();
+}
+
 FPathFindingResult AAbstractNavData::FindPathAbstract(const FNavAgentProperties& AgentProperties, const FPathFindingQuery& Query)
 {
 	const ANavigationData* Self = Query.NavData.Get();
diff --git a/Engine/Source/Runtime/Engine/Private/AI/Navigation/NavFilters/RecastFilter_UseDefaultArea.cpp b/Engine/Source/Runtime/Engine/Private/AI/Navigation/NavFilters/RecastFilter_UseDefaultArea.cpp
--- a/Engine/Source/Runtime/Engine/Private/AI/Navigation/NavFilters/RecastFilter_UseDefaultArea.cpp
+++ b/Engine/Source/Runtime/Engine/Private/AI/Navigation/NavFilters/RecastFilter_UseDefaultArea.cpp
@@ -3,7 +3,7 @@
 #include "AI/Navigation/NavFilters/RecastFilter_UseDefaultArea.h"
 #include "Templates/Casts.h"
 #include "AI/Navigation/RecastNavMesh.h"
-#include "AI/Navigation/PImplRecastNavMesh.h"
+#include "AI/Navigation/RecastQueryFilter.h"
 
 URecastFilter_UseDefaultArea::URecastFilter_UseDefaultArea(const FObjectInitializer& ObjectInitializer) : Super(ObjectInitializer)
 {
diff --git a/Engine/Source/Runtime/Engine/Private/AI/Navigation/NavModifierVolume.cpp b/Engine/Source/Runtime/Engine/Private/AI/Navigation/NavModifierVolume.cpp
--- a/Engine/Source/Runtime/Engine/Private/AI/Navigation/NavModifierVolume.cpp
+++ b/Engine/Source/Runtime/Engine/Private/AI/Navigation/NavModifierVolume.cpp
@@ -6,6 +6,7 @@
 #include "AI/Navigation/NavAreas/NavArea_Null.h"
 #include "AI/NavigationOctree.h"
 #include "Components/BrushComponent.h"
+#include "AI/NavigationSystemHelpers.h"
 #include "Engine/CollisionProfile.h"
 
 //----------------------------------------------------------------------//
@@ -62,12 +63,31 @@ void ANavModifierVolume::PostEditUndo()
 
 void ANavModifierVolume::PostEditChangeProperty(FPropertyChangedEvent& PropertyChangedEvent)
 {
+	static const FName NAME_AreaClass = GET_MEMBER_NAME_CHECKED(ANavModifierVolume, AreaClass);
+	static const FName NAME_BrushComponent = TEXT("BrushComponent");
+
 	Super::PostEditChangeProperty(PropertyChangedEvent);
+	
+	const FName PropName = PropertyChangedEvent.Property ? PropertyChangedEvent.Property->GetFName() : NAME_None;
 
-	if (PropertyChangedEvent.Property && PropertyChangedEvent.Property->GetFName() == GET_MEMBER_NAME_CHECKED(ANavModifierVolume, AreaClass))
+	if (PropName == NAME_AreaClass)
 	{
 		UNavigationSystem::UpdateActorInNavOctree(*this);
 	}
+	else if (PropName == NAME_BrushComponent)
+	{
+		if (GetBrushComponent())
+		{
+			if (GetBrushComponent()->GetBodySetup() && NavigationHelper::IsBodyNavigationRelevant(*GetBrushComponent()->GetBodySetup()))
+			{
+				UNavigationSystem::UpdateActorInNavOctree(*this);
+			}
+			else
+			{
+				UNavigationSystem::OnActorUnregistered(this);
+			}
+		}
+	}
 }
 
 #endif
diff --git a/Engine/Source/Runtime/Engine/Private/AI/Navigation/NavigationSystem.cpp b/Engine/Source/Runtime/Engine/Private/AI/Navigation/NavigationSystem.cpp
--- a/Engine/Source/Runtime/Engine/Private/AI/Navigation/NavigationSystem.cpp
+++ b/Engine/Source/Runtime/Engine/Private/AI/Navigation/NavigationSystem.cpp
@@ -292,12 +292,15 @@ UNavigationSystem::UNavigationSystem(const FObjectInitializer& ObjectInitializer
 #if WITH_EDITOR
 	NavUpdateLockFlags = 0;
 #endif
+	
+	CrowdManagerClass = UCrowdManager::StaticClass();
 
 	// active tiles
 	NextInvokersUpdateTime = 0.f;
 	ActiveTilesUpdateInterval = 1.f;
 	bGenerateNavigationOnlyAroundNavigationInvokers = false;
 	DataGatheringMode = ENavDataGatheringModeConfig::Instant;
+	bCanAccumulateDirtyAreas = true;
 
 	if (HasAnyFlags(RF_ClassDefaultObject) == false)
 	{
@@ -384,6 +387,7 @@ void UNavigationSystem::UpdateAbstractNavData()
 		FNavDataConfig DummyConfig;
 		DummyConfig.NavigationDataClass = AAbstractNavData::StaticClass();
 		AbstractNavData = CreateNavigationDataInstance(DummyConfig);
+		AbstractNavData->SetFlags(RF_Transient);
 	}
 }
 
@@ -741,6 +745,11 @@ void UNavigationSystem::OnWorldInitDone(FNavigationSystemRunMode Mode)
 		RebuildAll(bIsLoadTime);
 	}
 
+	if (!bCanAccumulateDirtyAreas)
+	{
+		DirtyAreas.Empty();
+	}
+
 	bWorldInitDone = true;
 	OnNavigationInitDone.Broadcast();
 }
@@ -767,10 +776,13 @@ void UNavigationSystem::RegisterNavigationDataInstances()
 
 void UNavigationSystem::CreateCrowdManager()
 {
-	SetCrowdManager(NewObject<UCrowdManager>(this));
+	if (CrowdManagerClass)
+	{
+		SetCrowdManager(NewObject<UCrowdManagerBase>(this, CrowdManagerClass));
+	}
 }
 
-void UNavigationSystem::SetCrowdManager(UCrowdManager* NewCrowdManager)
+void UNavigationSystem::SetCrowdManager(UCrowdManagerBase* NewCrowdManager)
 {
 	if (NewCrowdManager == CrowdManager.Get())
 	{
@@ -852,12 +864,7 @@ void UNavigationSystem::Tick(float DeltaSeconds)
 	}
 
 	// Tick navigation mesh async builders
-	if (!bAsyncBuildPaused && (bNavigationAutoUpdateEnabled || bIsGame 
-#if WITH_EDITOR
-		// continue ticking if build is in progress
-		|| (GIsEditor && IsNavigationBuildInProgress(/*bCheckDirtyToo=*/false))
-#endif // WITH_EDITOR
-		))
+	if (!bAsyncBuildPaused)
 	{
 		SCOPE_CYCLE_COUNTER(STAT_Navigation_TickAsyncBuild);
 		for (ANavigationData* NavData : NavDataSet)
@@ -885,7 +892,7 @@ void UNavigationSystem::Tick(float DeltaSeconds)
 void UNavigationSystem::AddReferencedObjects(UObject* InThis, FReferenceCollector& Collector)
 {
 	UNavigationSystem* This = CastChecked<UNavigationSystem>(InThis);
-	UCrowdManager* CrowdManager = This->GetCrowdManager();
+	UCrowdManagerBase* CrowdManager = This->GetCrowdManager();
 	Collector.AddReferencedObject(CrowdManager, InThis);
 
 	// don't reference NavAreaClasses in editor (unless PIE is active)
@@ -904,9 +911,12 @@ void UNavigationSystem::SetNavigationAutoUpdateEnabled(bool bNewEnable, UNavigat
 
 		if (InNavigationSystem)
 		{
+			InNavigationSystem->bCanAccumulateDirtyAreas = bNavigationAutoUpdateEnabled || (InNavigationSystem->OperationMode != FNavigationSystemRunMode::EditorMode);
+
 			if (bNavigationAutoUpdateEnabled)
 			{
-				InNavigationSystem->RemoveNavigationBuildLock(ENavigationBuildLock::NoUpdateInEditor);
+				const bool bSkipRebuildsInEditor = false;
+				InNavigationSystem->RemoveNavigationBuildLock(ENavigationBuildLock::NoUpdateInEditor, bSkipRebuildsInEditor);
 			}
 			else
 			{
@@ -1356,11 +1366,11 @@ UNavigationPath* UNavigationSystem::FindPathToLocationSynchronously(UObject* Wor
 
 	if (WorldContextObject != NULL)
 	{
-		World = GEngine->GetWorldFromContextObject(WorldContextObject);
+		World = GEngine->GetWorldFromContextObject(WorldContextObject, EGetWorldErrorMode::LogAndReturnNull);
 	}
 	if (World == NULL && PathfindingContext != NULL)
 	{
-		World = GEngine->GetWorldFromContextObject(PathfindingContext);
+		World = GEngine->GetWorldFromContextObject(PathfindingContext, EGetWorldErrorMode::LogAndReturnNull);
 	}
 
 	UNavigationPath* ResultPath = NULL;
@@ -1414,11 +1424,11 @@ bool UNavigationSystem::NavigationRaycast(UObject* WorldContextObject, const FVe
 
 	if (WorldContextObject != NULL)
 	{
-		World = GEngine->GetWorldFromContextObject(WorldContextObject);
+		World = GEngine->GetWorldFromContextObject(WorldContextObject, EGetWorldErrorMode::LogAndReturnNull);
 	}
 	if (World == NULL && Querier != NULL)
 	{
-		World = GEngine->GetWorldFromContextObject(Querier);
+		World = GEngine->GetWorldFromContextObject(Querier, EGetWorldErrorMode::LogAndReturnNull);
 	}
 
 	// blocked, i.e. not traversable, by default
@@ -1709,11 +1719,41 @@ const TSet<FNavigationBounds>& UNavigationSystem::GetNavigationBounds() const
 
 void UNavigationSystem::ApplyWorldOffset(const FVector& InOffset, bool bWorldShift)
 {
-	for (ANavigationData* NavData : NavDataSet)
+	// Attempt at generation of new nav mesh after the shift
+	// dynamic navmesh, we regenerate completely
+	if (GetRuntimeGenerationType() == ERuntimeGenerationType::Dynamic)
 	{
-		if (NavData)
+		//stop generators from building navmesh
+		for (ANavigationData* NavData : NavDataSet)
+		{
+			if (NavData)
+			{
+				if (NavData->GetGenerator()) NavData->GetGenerator()->CancelBuild();
+			}
+		}
+
+		ConditionalPopulateNavOctree();
+		Build();
+
+		for (ANavigationData* NavData : NavDataSet)
+		{
+			if (NavData)
+			{
+				NavData->ConditionalConstructGenerator();
+				ARecastNavMesh* RecastNavMesh = Cast<ARecastNavMesh>(NavData);
+				if (RecastNavMesh) RecastNavMesh->RequestDrawingUpdate();
+			}
+		}
+	}
+	else // static navmesh
+	{
+		//not sure what happens when we shift farther than the extents of the NavOctree are
+		for (ANavigationData* NavData : NavDataSet)
 		{
-			NavData->ApplyWorldOffset(InOffset, bWorldShift);
+			if (NavData)
+			{
+				NavData->ApplyWorldOffset(InOffset, bWorldShift);
+			}
 		}
 	}
 }
@@ -2279,7 +2319,7 @@ UNavigationSystem* UNavigationSystem::GetCurrent(UObject* WorldContextObject)
 
 	if (WorldContextObject != NULL)
 	{
-		World = GEngine->GetWorldFromContextObject(WorldContextObject);
+		World = GEngine->GetWorldFromContextObject(WorldContextObject, EGetWorldErrorMode::LogAndReturnNull);
 	}
 
 	return World ? World->GetNavigationSystem() : NULL;
@@ -2301,7 +2341,7 @@ ANavigationData* UNavigationSystem::GetNavDataWithID(const uint16 NavDataID) con
 
 void UNavigationSystem::AddDirtyArea(const FBox& NewArea, int32 Flags)
 {
-	if (Flags > 0)
+	if (Flags > 0 && bCanAccumulateDirtyAreas)
 	{
 		DirtyAreas.Add(FNavigationDirtyArea(NewArea, Flags));
 	}
@@ -3727,7 +3767,7 @@ UNavigationSystem* UNavigationSystem::GetNavigationSystem(UObject* WorldContextO
 
 bool UNavigationSystem::K2_ProjectPointToNavigation(UObject* WorldContextObject, const FVector& Point, FVector& ProjectedLocation, ANavigationData* NavData, TSubclassOf<UNavigationQueryFilter> FilterClass, const FVector QueryExtent)
 {
-	UWorld* World = GEngine->GetWorldFromContextObject(WorldContextObject);
+	UWorld* World = GEngine->GetWorldFromContextObject(WorldContextObject, EGetWorldErrorMode::LogAndReturnNull);
 	UNavigationSystem* NavSys = UNavigationSystem::GetCurrent(World);
 
 	ProjectedLocation = Point;
@@ -3753,7 +3793,7 @@ bool UNavigationSystem::K2_GetRandomReachablePointInRadius(UObject* WorldContext
 	FNavLocation RandomPoint(Origin);
 	bool bResult = false;
 
-	UWorld* World = GEngine->GetWorldFromContextObject(WorldContextObject);
+	UWorld* World = GEngine->GetWorldFromContextObject(WorldContextObject, EGetWorldErrorMode::LogAndReturnNull);
 	UNavigationSystem* NavSys = UNavigationSystem::GetCurrent(World);
 	if (NavSys)
 	{
@@ -3773,7 +3813,7 @@ bool UNavigationSystem::K2_GetRandomPointInNavigableRadius(UObject* WorldContext
 	FNavLocation RandomPoint(Origin);
 	bool bResult = false;
 
-	UWorld* World = GEngine->GetWorldFromContextObject(WorldContextObject);
+	UWorld* World = GEngine->GetWorldFromContextObject(WorldContextObject, EGetWorldErrorMode::LogAndReturnNull);
 	UNavigationSystem* NavSys = UNavigationSystem::GetCurrent(World);
 	if (NavSys)
 	{
@@ -3790,7 +3830,7 @@ bool UNavigationSystem::K2_GetRandomPointInNavigableRadius(UObject* WorldContext
 
 ENavigationQueryResult::Type UNavigationSystem::GetPathCost(UObject* WorldContextObject, const FVector& PathStart, const FVector& PathEnd, float& OutPathCost, ANavigationData* NavData, TSubclassOf<UNavigationQueryFilter> FilterClass)
 {
-	UWorld* World = GEngine->GetWorldFromContextObject( WorldContextObject );
+	UWorld* World = GEngine->GetWorldFromContextObject(WorldContextObject, EGetWorldErrorMode::LogAndReturnNull);
 	UNavigationSystem* NavSys = UNavigationSystem::GetCurrent(World);
 	if (NavSys)
 	{
@@ -3808,7 +3848,7 @@ ENavigationQueryResult::Type UNavigationSystem::GetPathLength(UObject* WorldCont
 {
 	float PathLength = 0.f;
 
-	UWorld* World = GEngine->GetWorldFromContextObject( WorldContextObject );
+	UWorld* World = GEngine->GetWorldFromContextObject(WorldContextObject, EGetWorldErrorMode::LogAndReturnNull);
 	UNavigationSystem* NavSys = UNavigationSystem::GetCurrent(World);
 	if (NavSys)
 	{
@@ -3824,7 +3864,7 @@ ENavigationQueryResult::Type UNavigationSystem::GetPathLength(UObject* WorldCont
 
 bool UNavigationSystem::IsNavigationBeingBuilt(UObject* WorldContextObject)
 {
-	UWorld* World = GEngine->GetWorldFromContextObject(WorldContextObject);
+	UWorld* World = GEngine->GetWorldFromContextObject(WorldContextObject, EGetWorldErrorMode::LogAndReturnNull);
 	UNavigationSystem* NavSys = UNavigationSystem::GetCurrent(World);
 	
 	if (NavSys && !NavSys->IsNavigationBuildingPermanentlyLocked())
@@ -3837,7 +3877,7 @@ bool UNavigationSystem::IsNavigationBeingBuilt(UObject* WorldContextObject)
 
 bool UNavigationSystem::IsNavigationBeingBuiltOrLocked(UObject* WorldContextObject)
 {
-	UWorld* World = GEngine->GetWorldFromContextObject(WorldContextObject);
+	UWorld* World = GEngine->GetWorldFromContextObject(WorldContextObject, EGetWorldErrorMode::LogAndReturnNull);
 	UNavigationSystem* NavSys = UNavigationSystem::GetCurrent(World);
 
 	if (NavSys)
@@ -4176,7 +4216,7 @@ FVector UNavigationSystem::ProjectPointToNavigation(UObject* WorldContextObject,
 {
 	FNavLocation ProjectedPoint(Point);
 
-	UWorld* World = GEngine->GetWorldFromContextObject(WorldContextObject);
+	UWorld* World = GEngine->GetWorldFromContextObject(WorldContextObject, EGetWorldErrorMode::LogAndReturnNull);
 	UNavigationSystem* NavSys = UNavigationSystem::GetCurrent(World);
 	if (NavSys)
 	{
@@ -4195,7 +4235,7 @@ FVector UNavigationSystem::GetRandomReachablePointInRadius(UObject* WorldContext
 {
 	FNavLocation RandomPoint;
 
-	UWorld* World = GEngine->GetWorldFromContextObject(WorldContextObject);
+	UWorld* World = GEngine->GetWorldFromContextObject(WorldContextObject, EGetWorldErrorMode::LogAndReturnNull);
 	UNavigationSystem* NavSys = UNavigationSystem::GetCurrent(World);
 	if (NavSys)
 	{
@@ -4213,7 +4253,7 @@ FVector UNavigationSystem::GetRandomPointInNavigableRadius(UObject* WorldContext
 {
 	FNavLocation RandomPoint;
 
-	UWorld* World = GEngine->GetWorldFromContextObject(WorldContextObject);
+	UWorld* World = GEngine->GetWorldFromContextObject(WorldContextObject, EGetWorldErrorMode::LogAndReturnNull);
 	UNavigationSystem* NavSys = UNavigationSystem::GetCurrent(World);
 	if (NavSys)
 	{
diff --git a/Engine/Source/Runtime/Engine/Private/AI/Navigation/PImplRecastNavMesh.cpp b/Engine/Source/Runtime/Engine/Private/AI/Navigation/PImplRecastNavMesh.cpp
--- a/Engine/Source/Runtime/Engine/Private/AI/Navigation/PImplRecastNavMesh.cpp
+++ b/Engine/Source/Runtime/Engine/Private/AI/Navigation/PImplRecastNavMesh.cpp
@@ -13,7 +13,7 @@
 #include "Recast/RecastAlloc.h"
 #include "AI/Navigation/NavAreas/NavArea.h"
 #include "AI/Navigation/RecastNavMeshGenerator.h"
-
+#include "AI/Navigation/RecastQueryFilter.h"
 #include "AI/Navigation/NavLinkCustomInterface.h"
 #include "VisualLogger/VisualLogger.h"
 
diff --git a/Engine/Source/Runtime/Engine/Private/AI/Navigation/RecastNavMesh.cpp b/Engine/Source/Runtime/Engine/Private/AI/Navigation/RecastNavMesh.cpp
--- a/Engine/Source/Runtime/Engine/Private/AI/Navigation/RecastNavMesh.cpp
+++ b/Engine/Source/Runtime/Engine/Private/AI/Navigation/RecastNavMesh.cpp
@@ -16,6 +16,7 @@
 #include "AI/Navigation/NavAreas/NavArea_LowHeight.h"
 #include "AI/Navigation/NavLinkCustomInterface.h"
 #include "AI/Navigation/RecastNavMeshDataChunk.h"
+#include "AI/Navigation/RecastQueryFilter.h"
 #include "VisualLogger/VisualLogger.h"
 
 #if WITH_EDITOR
@@ -1412,9 +1413,10 @@ void ARecastNavMesh::UpdateCustomLink(const INavLinkCustomInterface* CustomLink)
 	if (AreaId >= 0 && RecastNavMeshImpl)
 	{
 		UNavArea* DefArea = (UNavArea*)(AreaClass->GetDefaultObject());
+		const uint16 PolyFlags = DefArea->GetAreaFlags() | ARecastNavMesh::GetNavLinkFlag();
 
-		RecastNavMeshImpl->UpdateNavigationLinkArea(UserId, AreaId, DefArea->GetAreaFlags());
-		RecastNavMeshImpl->UpdateSegmentLinkArea(UserId, AreaId, DefArea->GetAreaFlags());
+		RecastNavMeshImpl->UpdateNavigationLinkArea(UserId, AreaId, PolyFlags);
+		RecastNavMeshImpl->UpdateSegmentLinkArea(UserId, AreaId, PolyFlags);
 	}
 }
 
@@ -1424,8 +1426,9 @@ void ARecastNavMesh::UpdateNavigationLinkArea(int32 UserId, TSubclassOf<UNavArea
 	if (AreaId >= 0 && RecastNavMeshImpl)
 	{
 		UNavArea* DefArea = (UNavArea*)(AreaClass->GetDefaultObject());
+		const uint16 PolyFlags = DefArea->GetAreaFlags() | ARecastNavMesh::GetNavLinkFlag();
 
-		RecastNavMeshImpl->UpdateNavigationLinkArea(UserId, AreaId, DefArea->GetAreaFlags());
+		RecastNavMeshImpl->UpdateNavigationLinkArea(UserId, AreaId, PolyFlags);
 	}
 }
 
@@ -1435,8 +1438,9 @@ void ARecastNavMesh::UpdateSegmentLinkArea(int32 UserId, TSubclassOf<UNavArea> A
 	if (AreaId >= 0 && RecastNavMeshImpl)
 	{
 		UNavArea* DefArea = (UNavArea*)(AreaClass->GetDefaultObject());
+		const uint16 PolyFlags = DefArea->GetAreaFlags() | ARecastNavMesh::GetNavLinkFlag();
 
-		RecastNavMeshImpl->UpdateSegmentLinkArea(UserId, AreaId, DefArea->GetAreaFlags());
+		RecastNavMeshImpl->UpdateSegmentLinkArea(UserId, AreaId, PolyFlags);
 	}
 }
 
@@ -2603,4 +2607,19 @@ void FRecastNavMeshCachedData::OnAreaAdded(const UClass* AreaClass, int32 AreaID
 	}		
 }
 
+uint32 ARecastNavMesh::GetLinkUserId(NavNodeRef LinkPolyID) const
+{
+	return RecastNavMeshImpl ? RecastNavMeshImpl->GetLinkUserId(LinkPolyID) : 0;
+}
+
+dtNavMesh* ARecastNavMesh::GetRecastMesh()
+{
+	return RecastNavMeshImpl ? RecastNavMeshImpl->GetRecastMesh() : nullptr;
+}
+
+const dtNavMesh* ARecastNavMesh::GetRecastMesh() const
+{
+	return RecastNavMeshImpl ? RecastNavMeshImpl->GetRecastMesh() : nullptr;
+}
+
 #endif// WITH_RECAST
diff --git a/Engine/Source/Runtime/Engine/Private/AI/Navigation/RecastNavMeshGenerator.cpp b/Engine/Source/Runtime/Engine/Private/AI/Navigation/RecastNavMeshGenerator.cpp
--- a/Engine/Source/Runtime/Engine/Private/AI/Navigation/RecastNavMeshGenerator.cpp
+++ b/Engine/Source/Runtime/Engine/Private/AI/Navigation/RecastNavMeshGenerator.cpp
@@ -2599,6 +2599,12 @@ bool FRecastTileGenerator::GenerateNavigationData(FNavMeshBuildContext& BuildCon
 				BuildContext.log(RC_LOG_ERROR, "GenerateNavigationData: Failed to build regions.");
 				return false;
 			}
+
+			// skip empty layer
+			if (GenerationContext.Layer->regCount <= 0)
+			{
+				continue;
+			}
 		}
 	
 		{
@@ -2626,6 +2632,12 @@ bool FRecastTileGenerator::GenerateNavigationData(FNavMeshBuildContext& BuildCon
 				BuildContext.log(RC_LOG_ERROR, "GenerateNavigationData: Failed to generate contour set (0x%08X).", status);
 				return false;
 			}
+
+			// skip empty layer, sometimes there are regions assigned but all flagged as empty (id=0)
+			if (GenerationContext.ContourSet->nconts <= 0)
+			{
+				continue;
+			}
 		}
 
 		{
@@ -3376,11 +3388,13 @@ bool FRecastNavMeshGenerator::RebuildAll()
 
 void FRecastNavMeshGenerator::EnsureBuildCompletion()
 {
-	const bool bHasTasks = GetNumRemaningBuildTasks() > 0;
+	const bool bHadTasks = GetNumRemaningBuildTasks() > 0;
 	
+	const bool bDoAsyncDataGathering = (GatherGeometryOnGameThread() == false);
 	do 
 	{
-		ProcessTileTasks(16);
+		const int32 NumTasksToSubmit = (bDoAsyncDataGathering ? 1 : MaxTileGeneratorTasks) - RunningDirtyTiles.Num();
+		ProcessTileTasks(NumTasksToSubmit);
 		
 		// Block until tasks are finished
 		for (FRunningTileElement& Element : RunningDirtyTiles)
@@ -3391,7 +3405,7 @@ void FRecastNavMeshGenerator::EnsureBuildCompletion()
 	while (GetNumRemaningBuildTasks() > 0);
 
 	// Update navmesh drawing only if we had something to build
-	if (bHasTasks)
+	if (bHadTasks)
 	{
 		DestNavMesh->RequestDrawingUpdate();
 	}
diff --git a/Engine/Source/Runtime/Engine/Private/Actor.cpp b/Engine/Source/Runtime/Engine/Private/Actor.cpp
--- a/Engine/Source/Runtime/Engine/Private/Actor.cpp
+++ b/Engine/Source/Runtime/Engine/Private/Actor.cpp
@@ -2593,13 +2593,17 @@ void AActor::ClearInstanceComponents(const bool bDestroyComponents)
 
 UActorComponent* AActor::FindComponentByClass(const TSubclassOf<UActorComponent> ComponentClass) const
 {
-	UActorComponent* FoundComponent = NULL;
-	for (UActorComponent* Component : OwnedComponents)
+	UActorComponent* FoundComponent = nullptr;
+
+	if (UClass* TargetClass = ComponentClass.Get())
 	{
-		if (Component && Component->IsA(ComponentClass))
+		for (UActorComponent* Component : OwnedComponents)
 		{
-			FoundComponent = Component;
-			break;
+			if (Component && Component->IsA(TargetClass))
+			{
+				FoundComponent = Component;
+				break;
+			}
 		}
 	}
 
@@ -2615,7 +2619,7 @@ TArray<UActorComponent*> AActor::GetComponentsByClass(TSubclassOf<UActorComponen
 {
 	TArray<UActorComponent*> ValidComponents;
 
-        // In the UActorComponent case we can skip the IsA checks for a slight performance benefit
+	// In the UActorComponent case we can skip the IsA checks for a slight performance benefit
 	if (ComponentClass == UActorComponent::StaticClass())
 	{
 		for (UActorComponent* Component : OwnedComponents)
diff --git a/Engine/Source/Runtime/Engine/Private/Animation/AnimNode_SubInstance.cpp b/Engine/Source/Runtime/Engine/Private/Animation/AnimNode_SubInstance.cpp
--- a/Engine/Source/Runtime/Engine/Private/Animation/AnimNode_SubInstance.cpp
+++ b/Engine/Source/Runtime/Engine/Private/Animation/AnimNode_SubInstance.cpp
@@ -149,8 +149,8 @@ void FAnimNode_SubInstance::OnInitializeAnimInstance(const FAnimInstanceProxy* I
 		MeshComp->SubInstances.Add(InstanceToRun);
 
 		// Build property lists
-		InstanceProperties.Empty();
-		SubInstanceProperties.Empty();
+		InstanceProperties.Reset(SourcePropertyNames.Num());
+		SubInstanceProperties.Reset(SourcePropertyNames.Num());
 
 		check(SourcePropertyNames.Num() == DestPropertyNames.Num());
 
@@ -164,10 +164,11 @@ void FAnimNode_SubInstance::OnInitializeAnimInstance(const FAnimInstanceProxy* I
 			UProperty* SourceProperty = FindField<UProperty>(SourceClass, SourceName);
 			UProperty* DestProperty = FindField<UProperty>(*InstanceClass, DestName);
 
-			check(SourceProperty && DestProperty);
-
-			InstanceProperties.Add(SourceProperty);
-			SubInstanceProperties.Add(DestProperty);
+			if (SourceProperty && DestProperty)
+			{
+				InstanceProperties.Add(SourceProperty);
+				SubInstanceProperties.Add(DestProperty);
+			}
 		}
 		
 	}
diff --git a/Engine/Source/Runtime/Engine/Private/AssetManager.cpp b/Engine/Source/Runtime/Engine/Private/AssetManager.cpp
--- a/Engine/Source/Runtime/Engine/Private/AssetManager.cpp
+++ b/Engine/Source/Runtime/Engine/Private/AssetManager.cpp
@@ -18,6 +18,7 @@
 #include "Editor.h"
 #include "Widgets/Notifications/SNotificationList.h"
 #include "Framework/Notifications/NotificationManager.h"
+#include "Commandlets/ChunkDependencyInfo.h"
 #endif
 
 #define LOCTEXT_NAMESPACE "AssetManager"
@@ -2036,6 +2037,49 @@ bool UAssetManager::GetPackageManagers(FName PackageName, bool bRecurseToParents
 	return bFoundAny;
 }
 
+void UAssetManager::StartInitialLoading()
+{
+	ScanPrimaryAssetTypesFromConfig();
+}
+
+void UAssetManager::FinishInitialLoading()
+{
+	// See if we have pending scans, if so defer result
+	bool bWaitingOnDeferredScan = false;
+
+	for (const TPair<FName, TSharedRef<FPrimaryAssetTypeData>>& TypePair : AssetTypeMap)
+	{
+		const FPrimaryAssetTypeData& TypeData = TypePair.Value.Get();
+
+		if (TypeData.DeferredAssetScanPaths.Num())
+		{
+			bWaitingOnDeferredScan = true;
+		}
+	}
+
+	if (!bWaitingOnDeferredScan)
+	{
+		PostInitialAssetScan();
+	}
+}
+
+bool UAssetManager::IsPathExcludedFromScan(const FString& Path) const
+{
+	const UAssetManagerSettings& Settings = GetSettings();
+
+	for (const FDirectoryPath& ExcludedPath : GetSettings().DirectoriesToExclude)
+	{
+		if (Path.Contains(ExcludedPath.Path))
+		{
+			return true;
+		}
+	}
+
+	return false;
+}
+
+#if WITH_EDITOR
+
 EAssetSetManagerResult::Type UAssetManager::ShouldSetManager(const FAssetIdentifier& Manager, const FAssetIdentifier& Source, const FAssetIdentifier& Target, EAssetRegistryDependencyType::Type DependencyType, EAssetSetManagerFlags::Type Flags) const
 {
 	FPrimaryAssetId ManagerPrimaryAssetId = Manager.GetPrimaryAssetId();
@@ -2048,8 +2092,8 @@ EAssetSetManagerResult::Type UAssetManager::ShouldSetManager(const FAssetIdentif
 
 	const FString TargetPackageString = Target.PackageName.ToString();
 
-	// Ignore script and engine references
-	if (!TargetPackageString.StartsWith(TEXT("/Game/"), ESearchCase::CaseSensitive))
+	// Ignore script references
+	if (TargetPackageString.StartsWith(TEXT("/Script/"), ESearchCase::CaseSensitive))
 	{
 		return EAssetSetManagerResult::DoNotSet;
 	}
@@ -2077,7 +2121,7 @@ EAssetSetManagerResult::Type UAssetManager::ShouldSetManager(const FAssetIdentif
 			// References made by maps are ownership references, because there is no way to distinguish between sublevels and top level maps we "include" sublevels in parent maps via reference
 			bIsOwnershipReference = true;
 		}
-		
+
 		if (bIsOwnershipReference)
 		{
 			return EAssetSetManagerResult::SetButDoNotRecurse;
@@ -2090,48 +2134,6 @@ EAssetSetManagerResult::Type UAssetManager::ShouldSetManager(const FAssetIdentif
 	return EAssetSetManagerResult::SetAndRecurse;
 }
 
-void UAssetManager::StartInitialLoading()
-{
-	ScanPrimaryAssetTypesFromConfig();
-}
-
-void UAssetManager::FinishInitialLoading()
-{
-	// See if we have pending scans, if so defer result
-	bool bWaitingOnDeferredScan = false;
-
-	for (const TPair<FName, TSharedRef<FPrimaryAssetTypeData>>& TypePair : AssetTypeMap)
-	{
-		const FPrimaryAssetTypeData& TypeData = TypePair.Value.Get();
-
-		if (TypeData.DeferredAssetScanPaths.Num())
-		{
-			bWaitingOnDeferredScan = true;
-		}
-	}
-
-	if (!bWaitingOnDeferredScan)
-	{
-		PostInitialAssetScan();
-	}
-}
-
-bool UAssetManager::IsPathExcludedFromScan(const FString& Path) const
-{
-	const UAssetManagerSettings& Settings = GetSettings();
-
-	for (const FDirectoryPath& ExcludedPath : GetSettings().DirectoriesToExclude)
-	{
-		if (Path.Contains(ExcludedPath.Path))
-		{
-			return true;
-		}
-	}
-
-	return false;
-}
-
-#if WITH_EDITOR
 void UAssetManager::ScanPathsSynchronous(const TArray<FString>& PathsToScan) const
 {
 	TArray<FString> Directories;
@@ -2516,19 +2518,21 @@ bool UAssetManager::GetPackageChunkIds(FName PackageName, const ITargetPlatform*
 		}
 	}
 
-	// If 0 is in list at all, constrain to just that
-	if (OutChunkList.Contains(0) && OutChunkList.Num() > 1)
+	int32 HighestChunk = 0;
+	for (int32 ChunkId : OutChunkList)
 	{
-		OutChunkList.Empty();
-		OutChunkList.Add(0);
+		if (ChunkId > HighestChunk)
+		{
+			HighestChunk = ChunkId;
+		}
 	}
 
-	return bFoundAny;
-}
+	// Use chunk dependency info to remove redundant chunks
+	UChunkDependencyInfo *DependencyInfo = GetMutableDefault<UChunkDependencyInfo>();
+	DependencyInfo->GetOrBuildChunkDependencyGraph(HighestChunk);
+	DependencyInfo->RemoveRedundantChunks(OutChunkList);
 
-void UAssetManager::UpdatePackageSourceHash(FName PackageName, FMD5& PackageSourceHash) const
-{
-	// TODO Need to handle shader cache
+	return bFoundAny;
 }
 
 void UAssetManager::PreBeginPIE(bool bStartSimulate)
@@ -2583,6 +2587,12 @@ void UAssetManager::RefreshPrimaryAssetDirectory()
 	{
 		FPrimaryAssetTypeData& TypeData = TypePair.Value.Get();
 
+		// Rescan the runtime data, the class may have gotten changed by hot reload
+		if (!TypeData.Info.FillRuntimeData())
+		{
+			continue;
+		}
+
 		if (TypeData.Info.AssetScanPaths.Num())
 		{
 			// Clear old data
diff --git a/Engine/Source/Runtime/Engine/Private/AssetManagerTypes.cpp b/Engine/Source/Runtime/Engine/Private/AssetManagerTypes.cpp
--- a/Engine/Source/Runtime/Engine/Private/AssetManagerTypes.cpp
+++ b/Engine/Source/Runtime/Engine/Private/AssetManagerTypes.cpp
@@ -6,6 +6,8 @@
 
 bool FPrimaryAssetTypeInfo::FillRuntimeData()
 {
+	// Hot reload may have messed up asset pointer
+	AssetBaseClass.ResetWeakPtr();
 	AssetBaseClassLoaded = AssetBaseClass.LoadSynchronous();
 
 	if (!ensureMsgf(AssetBaseClassLoaded, TEXT("Failed to load Primary Asset Type class %s!"), *AssetBaseClass.ToString()))
diff --git a/Engine/Source/Runtime/Engine/Private/Blueprint.cpp b/Engine/Source/Runtime/Engine/Private/Blueprint.cpp
--- a/Engine/Source/Runtime/Engine/Private/Blueprint.cpp
+++ b/Engine/Source/Runtime/Engine/Private/Blueprint.cpp
@@ -629,6 +629,12 @@ UClass* UBlueprint::RegenerateClass(UClass* ClassToRegenerate, UObject* Previous
 		UPackage* Package = Cast<UPackage>(GetOutermost());
 		bool bIsPackageDirty = Package ? Package->IsDirty() : false;
 
+		if( Package )
+		{
+			// Tell the linker to try to find exports in memory first, so that it gets the new, regenerated versions
+			Package->FindExportsInMemoryFirst(true);
+		}
+
 		UClass* GeneratedClassResolved = GeneratedClass;
 
 		UBlueprint::ForceLoadMetaData(this);
@@ -653,6 +659,8 @@ UClass* UBlueprint::RegenerateClass(UClass* ClassToRegenerate, UObject* Previous
 
 		FBlueprintCompilationManager::NotifyBlueprintLoaded( this ); 
 		
+		FBlueprintEditorUtils::PreloadBlueprintSpecificData( this );
+
 		// clear this now that we're not in a re-entrrant context - bHasBeenRegenerated will guard against 'real' 
 		// double regeneration calls:
 		bIsRegeneratingOnLoad = false;
diff --git a/Engine/Source/Runtime/Engine/Private/CanvasRenderTarget2D.cpp b/Engine/Source/Runtime/Engine/Private/CanvasRenderTarget2D.cpp
--- a/Engine/Source/Runtime/Engine/Private/CanvasRenderTarget2D.cpp
+++ b/Engine/Source/Runtime/Engine/Private/CanvasRenderTarget2D.cpp
@@ -121,7 +121,7 @@ UCanvasRenderTarget2D* UCanvasRenderTarget2D::CreateCanvasRenderTarget2D(UObject
 		UCanvasRenderTarget2D* NewCanvasRenderTarget = NewObject<UCanvasRenderTarget2D>(GetTransientPackage(), CanvasRenderTarget2DClass);
 		if (NewCanvasRenderTarget)
 		{
-			NewCanvasRenderTarget->World = GEngine->GetWorldFromContextObject(WorldContextObject);
+			NewCanvasRenderTarget->World = GEngine->GetWorldFromContextObject(WorldContextObject, EGetWorldErrorMode::LogAndReturnNull);
 			NewCanvasRenderTarget->InitAutoFormat(Width, Height);
 			return NewCanvasRenderTarget;
 		}
diff --git a/Engine/Source/Runtime/Engine/Private/Character.cpp b/Engine/Source/Runtime/Engine/Private/Character.cpp
--- a/Engine/Source/Runtime/Engine/Private/Character.cpp
+++ b/Engine/Source/Runtime/Engine/Private/Character.cpp
@@ -226,7 +226,7 @@ float ACharacter::GetDefaultHalfHeight() const
 UActorComponent* ACharacter::FindComponentByClass(const TSubclassOf<UActorComponent> ComponentClass) const
 {
 	// If the character has a Mesh, treat it as the first 'hit' when finding components
-	if (Mesh && Mesh->IsA(ComponentClass))
+	if (Mesh && ComponentClass && Mesh->IsA(ComponentClass))
 	{
 		return Mesh;
 	}
diff --git a/Engine/Source/Runtime/Engine/Private/CheatManager.cpp b/Engine/Source/Runtime/Engine/Private/CheatManager.cpp
--- a/Engine/Source/Runtime/Engine/Private/CheatManager.cpp
+++ b/Engine/Source/Runtime/Engine/Private/CheatManager.cpp
@@ -1134,33 +1134,28 @@ void UCheatManager::CheatScript(FString ScriptName)
 	APlayerController* const PlayerController = GetOuterAPlayerController();
 	ULocalPlayer* const LocalPlayer = PlayerController ? Cast<ULocalPlayer>(PlayerController->Player) : nullptr;
 
-	if (LocalPlayer)
-	{
-		// Run commands from the ini
-		FConfigSection const* const CommandsToRun = GConfig->GetSectionPrivate(*FString::Printf(TEXT("CheatScript.%s"), *ScriptName), 0, 1, GGameIni);
+	UConsole* ConsoleToDisplayResults = (LocalPlayer && LocalPlayer->ViewportClient) ? LocalPlayer->ViewportClient->ViewportConsole : nullptr;
+
+	// Run commands from the ini
+	FConfigSection const* const CommandsToRun = GConfig->GetSectionPrivate(*FString::Printf(TEXT("CheatScript.%s"), *ScriptName), 0, 1, GGameIni);
 
-		if (CommandsToRun)
+	if (CommandsToRun)
+	{
+		for (FConfigSectionMap::TConstIterator It(*CommandsToRun); It; ++It)
 		{
-			for (FConfigSectionMap::TConstIterator It(*CommandsToRun); It; ++It)
+			// show user what commands ran
+			if (ConsoleToDisplayResults)
 			{
-				// show user what commands ran
-				if (LocalPlayer->ViewportClient && LocalPlayer->ViewportClient->ViewportConsole)
-				{
-					FString const S = FString::Printf(TEXT("> %s"), *It.Value().GetValue());
-					LocalPlayer->ViewportClient->ViewportConsole->OutputText(S);
-				}
-
-				LocalPlayer->Exec(GetWorld(), *It.Value().GetValue(), *GLog);
+				FString const S = FString::Printf(TEXT("> %s"), *It.Value().GetValue());
+				ConsoleToDisplayResults->OutputText(S);
 			}
-		}
-		else
-		{
-			UE_LOG(LogCheatManager, Warning, TEXT("Can't find section 'CheatScript.%s' in DefaultGame.ini"), *ScriptName);
+
+			PlayerController->ConsoleCommand(*It.Value().GetValue(), /*bWriteToLog=*/ true);
 		}
 	}
 	else
 	{
-		UE_LOG(LogCheatManager, Warning, TEXT("Can't find local player!"));
+		UE_LOG(LogCheatManager, Warning, TEXT("Can't find section 'CheatScript.%s' in DefaultGame.ini"), *ScriptName);
 	}
 }
 
