Mark FInterp structs blueprintype and properties as blueprint readwrite.
Mark just enough properties in UInterpTrack classes to fix UT warnings
Fix ArenaMutatorMenu's orphan pin warning

diff --git a/Engine/Source/Runtime/CoreUObject/Public/UObject/NoExportTypes.h b/Engine/Source/Runtime/CoreUObject/Public/UObject/NoExportTypes.h
--- a/Engine/Source/Runtime/CoreUObject/Public/UObject/NoExportTypes.h
+++ b/Engine/Source/Runtime/CoreUObject/Public/UObject/NoExportTypes.h
@@ -364,13 +364,13 @@ struct FVector2D
 
 
 
-USTRUCT(immutable, noexport)
+USTRUCT(immutable, BlueprintType, noexport)
 struct FTwoVectors
 {
-	UPROPERTY(EditAnywhere, Category=TwoVectors, SaveGame)
+	UPROPERTY(EditAnywhere, BlueprintReadWrite, Category=TwoVectors, SaveGame)
 	FVector v1;
 
-	UPROPERTY(EditAnywhere, Category=TwoVectors, SaveGame)
+	UPROPERTY(EditAnywhere, BlueprintReadWrite, Category=TwoVectors, SaveGame)
 	FVector v2;
 
 };
@@ -663,223 +663,223 @@ struct FMatrix
 
 
 
-USTRUCT(noexport)
+USTRUCT(noexport, BlueprintType)
 struct FInterpCurvePointFloat
 {
-	UPROPERTY(EditAnywhere, Category=InterpCurvePointFloat)
+	UPROPERTY(EditAnywhere, BlueprintReadWrite, Category=InterpCurvePointFloat)
 	float InVal;
 
-	UPROPERTY(EditAnywhere, Category=InterpCurvePointFloat)
+	UPROPERTY(EditAnywhere, BlueprintReadWrite, Category=InterpCurvePointFloat)
 	float OutVal;
 
-	UPROPERTY(EditAnywhere, Category=InterpCurvePointFloat)
+	UPROPERTY(EditAnywhere, BlueprintReadWrite, Category=InterpCurvePointFloat)
 	float ArriveTangent;
 
-	UPROPERTY(EditAnywhere, Category=InterpCurvePointFloat)
+	UPROPERTY(EditAnywhere, BlueprintReadWrite, Category=InterpCurvePointFloat)
 	float LeaveTangent;
 
-	UPROPERTY(EditAnywhere, Category=InterpCurvePointFloat)
+	UPROPERTY(EditAnywhere, BlueprintReadWrite, Category=InterpCurvePointFloat)
 	TEnumAsByte<enum EInterpCurveMode> InterpMode;
 
 };
 
 
 
-USTRUCT(noexport)
+USTRUCT(noexport, BlueprintType)
 struct FInterpCurveFloat
 {
-	UPROPERTY(EditAnywhere, Category=InterpCurveFloat)
+	UPROPERTY(EditAnywhere, BlueprintReadWrite, Category=InterpCurveFloat)
 	TArray<FInterpCurvePointFloat> Points;
 
-	UPROPERTY(EditAnywhere, Category=InterpCurveFloat)
+	UPROPERTY(EditAnywhere, BlueprintReadWrite, Category=InterpCurveFloat)
 	bool bIsLooped;
 
-	UPROPERTY(EditAnywhere, Category=InterpCurveFloat)
+	UPROPERTY(EditAnywhere, BlueprintReadWrite, Category=InterpCurveFloat)
 	float LoopKeyOffset;
 };
 
 
 
-USTRUCT(noexport)
+USTRUCT(noexport, BlueprintType)
 struct FInterpCurvePointVector2D
 {
-	UPROPERTY(EditAnywhere, Category=InterpCurvePointVector2D)
+	UPROPERTY(EditAnywhere, BlueprintReadWrite, Category=InterpCurvePointVector2D)
 	float InVal;
 
-	UPROPERTY(EditAnywhere, Category=InterpCurvePointVector2D)
+	UPROPERTY(EditAnywhere, BlueprintReadWrite, Category=InterpCurvePointVector2D)
 	FVector2D OutVal;
 
-	UPROPERTY(EditAnywhere, Category=InterpCurvePointVector2D)
+	UPROPERTY(EditAnywhere, BlueprintReadWrite, Category=InterpCurvePointVector2D)
 	FVector2D ArriveTangent;
 
-	UPROPERTY(EditAnywhere, Category=InterpCurvePointVector2D)
+	UPROPERTY(EditAnywhere, BlueprintReadWrite, Category=InterpCurvePointVector2D)
 	FVector2D LeaveTangent;
 
-	UPROPERTY(EditAnywhere, Category=InterpCurvePointVector2D)
+	UPROPERTY(EditAnywhere, BlueprintReadWrite, Category=InterpCurvePointVector2D)
 	TEnumAsByte<enum EInterpCurveMode> InterpMode;
 
 };
 
 
 
-USTRUCT(noexport)
+USTRUCT(noexport, BlueprintType)
 struct FInterpCurveVector2D
 {
-	UPROPERTY(EditAnywhere, Category=InterpCurveVector2D)
+	UPROPERTY(EditAnywhere, BlueprintReadWrite, Category=InterpCurveVector2D)
 	TArray<FInterpCurvePointVector2D> Points;
 
-	UPROPERTY(EditAnywhere, Category=InterpCurveVector2D)
+	UPROPERTY(EditAnywhere, BlueprintReadWrite, Category=InterpCurveVector2D)
 	bool bIsLooped;
 
-	UPROPERTY(EditAnywhere, Category=InterpCurveVector2D)
+	UPROPERTY(EditAnywhere, BlueprintReadWrite, Category=InterpCurveVector2D)
 	float LoopKeyOffset;
 };
 
 
 
-USTRUCT(noexport)
+USTRUCT(noexport, BlueprintType)
 struct FInterpCurvePointVector
 {
-	UPROPERTY(EditAnywhere, Category=InterpCurvePointVector)
+	UPROPERTY(EditAnywhere, BlueprintReadWrite, Category=InterpCurvePointVector)
 	float InVal;
 
-	UPROPERTY(EditAnywhere, Category=InterpCurvePointVector)
+	UPROPERTY(EditAnywhere, BlueprintReadWrite, Category=InterpCurvePointVector)
 	FVector OutVal;
 
-	UPROPERTY(EditAnywhere, Category=InterpCurvePointVector)
+	UPROPERTY(EditAnywhere, BlueprintReadWrite, Category=InterpCurvePointVector)
 	FVector ArriveTangent;
 
-	UPROPERTY(EditAnywhere, Category=InterpCurvePointVector)
+	UPROPERTY(EditAnywhere, BlueprintReadWrite, Category=InterpCurvePointVector)
 	FVector LeaveTangent;
 
-	UPROPERTY(EditAnywhere, Category=InterpCurvePointVector)
+	UPROPERTY(EditAnywhere, BlueprintReadWrite, Category=InterpCurvePointVector)
 	TEnumAsByte<enum EInterpCurveMode> InterpMode;
 
 };
 
 
 
-USTRUCT(noexport)
+USTRUCT(noexport, BlueprintType)
 struct FInterpCurveVector
 {
-	UPROPERTY(EditAnywhere, Category=InterpCurveVector)
+	UPROPERTY(EditAnywhere, BlueprintReadWrite, Category=InterpCurveVector)
 	TArray<FInterpCurvePointVector> Points;
 
-	UPROPERTY(EditAnywhere, Category=InterpCurveVector)
+	UPROPERTY(EditAnywhere, BlueprintReadWrite, Category=InterpCurveVector)
 	bool bIsLooped;
 
-	UPROPERTY(EditAnywhere, Category=InterpCurveVector)
+	UPROPERTY(EditAnywhere, BlueprintReadWrite, Category=InterpCurveVector)
 	float LoopKeyOffset;
 };
 
 
 
-USTRUCT(noexport)
+USTRUCT(noexport, BlueprintType)
 struct FInterpCurvePointQuat
 {
-	UPROPERTY(EditAnywhere, Category=InterpCurvePointQuat)
+	UPROPERTY(EditAnywhere, BlueprintReadWrite, Category=InterpCurvePointQuat)
 	float InVal;
 
-	UPROPERTY(EditAnywhere, Category=InterpCurvePointQuat)
+	UPROPERTY(EditAnywhere, BlueprintReadWrite, Category=InterpCurvePointQuat)
 	FQuat OutVal;
 
-	UPROPERTY(EditAnywhere, Category=InterpCurvePointQuat)
+	UPROPERTY(EditAnywhere, BlueprintReadWrite, Category=InterpCurvePointQuat)
 	FQuat ArriveTangent;
 
-	UPROPERTY(EditAnywhere, Category=InterpCurvePointQuat)
+	UPROPERTY(EditAnywhere, BlueprintReadWrite, Category=InterpCurvePointQuat)
 	FQuat LeaveTangent;
 
-	UPROPERTY(EditAnywhere, Category=InterpCurvePointQuat)
+	UPROPERTY(EditAnywhere, BlueprintReadWrite, Category=InterpCurvePointQuat)
 	TEnumAsByte<enum EInterpCurveMode> InterpMode;
 
 };
 
 
 
-USTRUCT(noexport)
+USTRUCT(noexport, BlueprintType)
 struct FInterpCurveQuat
 {
-	UPROPERTY(EditAnywhere, Category=InterpCurveQuat)
+	UPROPERTY(EditAnywhere, BlueprintReadWrite, Category=InterpCurveQuat)
 	TArray<FInterpCurvePointQuat> Points;
 
-	UPROPERTY(EditAnywhere, Category=InterpCurveQuat)
+	UPROPERTY(EditAnywhere, BlueprintReadWrite, Category=InterpCurveQuat)
 	bool bIsLooped;
 
-	UPROPERTY(EditAnywhere, Category=InterpCurveQuat)
+	UPROPERTY(EditAnywhere, BlueprintReadWrite, Category=InterpCurveQuat)
 	float LoopKeyOffset;
 };
 
 
 
-USTRUCT(noexport)
+USTRUCT(noexport, BlueprintType)
 struct FInterpCurvePointTwoVectors
 {
-	UPROPERTY(EditAnywhere, Category=InterpCurvePointTwoVectors)
+	UPROPERTY(EditAnywhere, BlueprintReadWrite, Category=InterpCurvePointTwoVectors)
 	float InVal;
 
-	UPROPERTY(EditAnywhere, Category=InterpCurvePointTwoVectors)
+	UPROPERTY(EditAnywhere, BlueprintReadWrite, Category=InterpCurvePointTwoVectors)
 	FTwoVectors OutVal;
 
-	UPROPERTY(EditAnywhere, Category=InterpCurvePointTwoVectors)
+	UPROPERTY(EditAnywhere, BlueprintReadWrite, Category=InterpCurvePointTwoVectors)
 	FTwoVectors ArriveTangent;
 
-	UPROPERTY(EditAnywhere, Category=InterpCurvePointTwoVectors)
+	UPROPERTY(EditAnywhere, BlueprintReadWrite, Category=InterpCurvePointTwoVectors)
 	FTwoVectors LeaveTangent;
 
-	UPROPERTY(EditAnywhere, Category=InterpCurvePointTwoVectors)
+	UPROPERTY(EditAnywhere, BlueprintReadWrite, Category=InterpCurvePointTwoVectors)
 	TEnumAsByte<enum EInterpCurveMode> InterpMode;
 
 };
 
 
 
-USTRUCT(noexport)
+USTRUCT(noexport, BlueprintType)
 struct FInterpCurveTwoVectors
 {
-	UPROPERTY(EditAnywhere, Category=InterpCurveTwoVectors)
+	UPROPERTY(EditAnywhere, BlueprintReadWrite, Category=InterpCurveTwoVectors)
 	TArray<FInterpCurvePointTwoVectors> Points;
 
-	UPROPERTY(EditAnywhere, Category=InterpCurveTwoVectors)
+	UPROPERTY(EditAnywhere, BlueprintReadWrite, Category=InterpCurveTwoVectors)
 	bool bIsLooped;
 
-	UPROPERTY(EditAnywhere, Category=InterpCurveTwoVectors)
+	UPROPERTY(EditAnywhere, BlueprintReadWrite, Category=InterpCurveTwoVectors)
 	float LoopKeyOffset;
 };
 
 
 
-USTRUCT(noexport)
+USTRUCT(noexport, BlueprintType)
 struct FInterpCurvePointLinearColor
 {
-	UPROPERTY(EditAnywhere, Category=InterpCurvePointLinearColor)
+	UPROPERTY(EditAnywhere, BlueprintReadWrite, Category=InterpCurvePointLinearColor)
 	float InVal;
 
-	UPROPERTY(EditAnywhere, Category=InterpCurvePointLinearColor)
+	UPROPERTY(EditAnywhere, BlueprintReadWrite, Category=InterpCurvePointLinearColor)
 	FLinearColor OutVal;
 
-	UPROPERTY(EditAnywhere, Category=InterpCurvePointLinearColor)
+	UPROPERTY(EditAnywhere, BlueprintReadWrite, Category=InterpCurvePointLinearColor)
 	FLinearColor ArriveTangent;
 
-	UPROPERTY(EditAnywhere, Category=InterpCurvePointLinearColor)
+	UPROPERTY(EditAnywhere, BlueprintReadWrite, Category=InterpCurvePointLinearColor)
 	FLinearColor LeaveTangent;
 
-	UPROPERTY(EditAnywhere, Category=InterpCurvePointLinearColor)
+	UPROPERTY(EditAnywhere, BlueprintReadWrite, Category=InterpCurvePointLinearColor)
 	TEnumAsByte<enum EInterpCurveMode> InterpMode;
 
 };
 
 
 
-USTRUCT(noexport)
+USTRUCT(noexport, BlueprintType)
 struct FInterpCurveLinearColor
 {
-	UPROPERTY(EditAnywhere, Category=InterpCurveLinearColor)
+	UPROPERTY(EditAnywhere, BlueprintReadWrite, Category=InterpCurveLinearColor)
 	TArray<FInterpCurvePointLinearColor> Points;
 
-	UPROPERTY(EditAnywhere, Category=InterpCurveLinearColor)
+	UPROPERTY(EditAnywhere, BlueprintReadWrite, Category=InterpCurveLinearColor)
 	bool bIsLooped;
 
-	UPROPERTY(EditAnywhere, Category=InterpCurveLinearColor)
+	UPROPERTY(EditAnywhere, BlueprintReadWrite, Category=InterpCurveLinearColor)
 	float LoopKeyOffset;
 };
 
diff --git a/Engine/Source/Runtime/Engine/Classes/Matinee/InterpData.h b/Engine/Source/Runtime/Engine/Classes/Matinee/InterpData.h
--- a/Engine/Source/Runtime/Engine/Classes/Matinee/InterpData.h
+++ b/Engine/Source/Runtime/Engine/Classes/Matinee/InterpData.h
@@ -27,7 +27,7 @@ class UInterpData : public UObject
 	float PathBuildTime;
 
 	/** Actual interpolation data. Groups of InterpTracks. */
-	UPROPERTY(export)
+	UPROPERTY(export, BlueprintReadOnly, Category=InterpData)
 	TArray<class UInterpGroup*> InterpGroups;
 
 	/** Used for curve editor to remember curve-editing setup. Only loaded in editor. */
diff --git a/Engine/Source/Runtime/Engine/Classes/Matinee/InterpGroup.h b/Engine/Source/Runtime/Engine/Classes/Matinee/InterpGroup.h
--- a/Engine/Source/Runtime/Engine/Classes/Matinee/InterpGroup.h
+++ b/Engine/Source/Runtime/Engine/Classes/Matinee/InterpGroup.h
@@ -86,7 +86,7 @@ class UInterpGroup : public UObject, public FInterpEdInputInterface
 {
 	GENERATED_UCLASS_BODY()
 
-	UPROPERTY(export)
+	UPROPERTY(export, BlueprintReadOnly, Category=InterpGroup)
 	TArray<class UInterpTrack*> InterpTracks;
 
 	/** 
diff --git a/Engine/Source/Runtime/Engine/Classes/Matinee/InterpTrack.h b/Engine/Source/Runtime/Engine/Classes/Matinee/InterpTrack.h
--- a/Engine/Source/Runtime/Engine/Classes/Matinee/InterpTrack.h
+++ b/Engine/Source/Runtime/Engine/Classes/Matinee/InterpTrack.h
@@ -105,7 +105,7 @@ class UInterpTrack : public UObject, public FCurveEdInterface, public FInterpEdI
 	GENERATED_UCLASS_BODY()
 
 	/** A list of subtracks that belong to this track */
-	UPROPERTY()
+	UPROPERTY(BlueprintReadOnly, Category=InterpTrack)
 	TArray<class UInterpTrack*> SubTracks;
 
 #if WITH_EDITORONLY_DATA
diff --git a/Engine/Source/Runtime/Engine/Classes/Matinee/InterpTrackMove.h b/Engine/Source/Runtime/Engine/Classes/Matinee/InterpTrackMove.h
--- a/Engine/Source/Runtime/Engine/Classes/Matinee/InterpTrackMove.h
+++ b/Engine/Source/Runtime/Engine/Classes/Matinee/InterpTrackMove.h
@@ -102,7 +102,7 @@ class UInterpTrackMove : public UInterpTrack
 	GENERATED_UCLASS_BODY()
 	
 	/** Actual position keyframe data. */
-	UPROPERTY()
+	UPROPERTY(BlueprintReadOnly, Category=InterpTrackMove)
 	FInterpCurveVector PosTrack;
 
 	/** Actual rotation keyframe data, stored as Euler angles in degrees, for easy editing on curve. */
