Copying //UE4/Dev-Framework to //UE4/Dev-Main (Source: //UE4/Dev-Framework @ 3544039)

diff --git a/Engine/Source/Runtime/Core/Public/Math/Vector4.h b/Engine/Source/Runtime/Core/Public/Math/Vector4.h
--- a/Engine/Source/Runtime/Core/Public/Math/Vector4.h
+++ b/Engine/Source/Runtime/Core/Public/Math/Vector4.h
@@ -292,9 +292,6 @@ public:
 	 */
 	FORCEINLINE FVector4 GetSafeNormal(float Tolerance = SMALL_NUMBER) const;
 
-	DEPRECATED(4.7, "Deprecated due to unclear name, use GetSafeNormal instead.")
-	FORCEINLINE FVector4 SafeNormal(float Tolerance=SMALL_NUMBER) const;
-
 	/**
 	 * Calculates normalized version of vector without checking if it is non-zero.
 	 *
@@ -302,9 +299,6 @@ public:
 	 */
 	FORCEINLINE FVector4 GetUnsafeNormal3() const;
 
-	DEPRECATED(4.7, "Deprecated due to unclear name, use GetUnsafeNormal3 instead.")
-	FORCEINLINE FVector4 UnsafeNormal3() const;
-
 	/**
 	 * Return the FRotator orientation corresponding to the direction in which the vector points.
 	 * Sets Yaw and Pitch to the proper numbers, and sets roll to zero because the roll can't be determined from a vector.
@@ -620,12 +614,6 @@ FORCEINLINE FVector4 FVector4::GetSafeNormal(float Tolerance) const
 }
 
 
-FORCEINLINE FVector4 FVector4::SafeNormal(float Tolerance) const
-{
-	return GetSafeNormal(Tolerance);
-}
-
-
 FORCEINLINE FVector4 FVector4::GetUnsafeNormal3() const
 {
 	const float Scale = FMath::InvSqrt(X*X+Y*Y+Z*Z);
@@ -633,12 +621,6 @@ FORCEINLINE FVector4 FVector4::GetUnsafeNormal3() const
 }
 
 
-FORCEINLINE FVector4 FVector4::UnsafeNormal3() const
-{
-	return GetUnsafeNormal3();
-}
-
-
 FORCEINLINE float FVector4::Size3() const
 {
 	return FMath::Sqrt(X*X + Y*Y + Z*Z);
diff --git a/Engine/Source/Runtime/Core/Public/Misc/AssertionMacros.h b/Engine/Source/Runtime/Core/Public/Misc/AssertionMacros.h
--- a/Engine/Source/Runtime/Core/Public/Misc/AssertionMacros.h
+++ b/Engine/Source/Runtime/Core/Public/Misc/AssertionMacros.h
@@ -200,17 +200,6 @@ struct CORE_API FDebug
 
 #endif	// DO_CHECK
 
-namespace UE4Asserts_Private
-{
-	DEPRECATED(4.9, "ensureMsg is deprecated, please use ensureMsgf instead.")
-	FORCEINLINE bool DeprecatedEnsure(bool bValue)
-	{
-		return bValue;
-	}
-}
-
-#define ensureMsg(InExpression, InMsg) UE4Asserts_Private::DeprecatedEnsure(ensureMsgf(InExpression, InMsg))
-
 namespace UE4Asserts_Private
 {
 	// A junk function to allow us to use sizeof on a member variable which is potentially a bitfield
diff --git a/Engine/Source/Runtime/Core/Public/Misc/Attribute.h b/Engine/Source/Runtime/Core/Public/Misc/Attribute.h
--- a/Engine/Source/Runtime/Core/Public/Misc/Attribute.h
+++ b/Engine/Source/Runtime/Core/Public/Misc/Attribute.h
@@ -317,30 +317,6 @@ public:
 		return false;
 	}
 
-	/**
-	 * Equal operator.
-	 *
-	 * @param InOther The other attribute to compare with.
-	 * @return true if the attributes are not equal, false otherwise.
-	 */
-	DEPRECATED(4.9, "Comparing an attribute directly to an attribute is deprecated.  Please call Get() on the attributes before comparing them.  If you were actually testing if the Attributes were the same functionally, call IdenticalTo")
-	bool operator==( const TAttribute& InOther ) const
-	{
-		return Get() == InOther.Get();
-	}
-
-	/**
-	 * Not equal operator.
-	 *
-	 * @param InOther The other attribute to compare with.
-	 * @return true if the attributes are not equal, false otherwise.
-	 */
-	DEPRECATED(4.9, "Comparing an attribute directly to an attribute is deprecated.  Please call Get() on the attributes before comparing them.  If you were actually testing if the Attributes were the same functionally, call IdenticalTo")
-	bool operator!=( const TAttribute& InOther ) const
-	{
-		return Get() != InOther.Get();
-	}
-	
 private:
 
 	/** Special explicit constructor for TAttribute::Create() */
diff --git a/Engine/Source/Runtime/Core/Public/Misc/IQueuedWork.h b/Engine/Source/Runtime/Core/Public/Misc/IQueuedWork.h
--- a/Engine/Source/Runtime/Core/Public/Misc/IQueuedWork.h
+++ b/Engine/Source/Runtime/Core/Public/Misc/IQueuedWork.h
@@ -39,7 +39,3 @@ public:
 	 */
 	virtual ~IQueuedWork() { }
 };
-
-
-DEPRECATED(4.8, "FQueuedWork has been renamed to IQueuedWork")
-typedef IQueuedWork FQueuedWork;
diff --git a/Engine/Source/Runtime/Core/Public/Templates/MemoryOps.h b/Engine/Source/Runtime/Core/Public/Templates/MemoryOps.h
--- a/Engine/Source/Runtime/Core/Public/Templates/MemoryOps.h
+++ b/Engine/Source/Runtime/Core/Public/Templates/MemoryOps.h
@@ -138,14 +138,6 @@ FORCEINLINE typename TEnableIf<TIsBitwiseConstructible<DestinationElementType, S
 }
 
 
-template <typename ElementType>
-DEPRECATED(4.6, "CopyConstructItems is deprecated, use ConstructItems<ElementType> instead.")
-FORCEINLINE void CopyConstructItems(void* Dest, const ElementType* Source, int32 Count)
-{
-	ConstructItems<ElementType>(Dest, Source, Count);
-}
-
-
 /**
  * Copy assigns a range of items.
  *
@@ -210,14 +202,6 @@ FORCEINLINE typename TEnableIf<UE4MemoryOps_Private::TCanBitwiseRelocate<Destina
 	FMemory::Memmove(Dest, Source, sizeof(SourceElementType) * Count);
 }
 
-template <typename ElementType>
-DEPRECATED(4.6, "RelocateItems is deprecated, use RelocateConstructItems<ElementType> instead.")
-FORCEINLINE void RelocateItems(void* Dest, const ElementType* Source, int32 Count)
-{
-	RelocateConstructItems<ElementType>(Dest, Source, Count);
-}
-
-
 /**
  * Move constructs a range of items into memory.
  *
diff --git a/Engine/Source/Runtime/Core/Public/Windows/WindowsPlatformProcess.h b/Engine/Source/Runtime/Core/Public/Windows/WindowsPlatformProcess.h
--- a/Engine/Source/Runtime/Core/Public/Windows/WindowsPlatformProcess.h
+++ b/Engine/Source/Runtime/Core/Public/Windows/WindowsPlatformProcess.h
@@ -22,9 +22,6 @@ public:
 	FORCEINLINE explicit FProcHandle( HandleType Other )
 		: TProcHandle( Other )
 	{}
-
-	DEPRECATED(4.8, "FProcHandle::Close() is redundant - handles created with FPlatformProcess::CreateProc() should be closed with FPlatformProcess::CloseProc().")
-	FORCEINLINE bool Close();
 };
 
 
@@ -246,14 +243,3 @@ private:
 
 
 typedef FWindowsPlatformProcess FPlatformProcess;
-
-inline bool FProcHandle::Close()
-{
-	if (IsValid())
-	{
-		FPlatformProcess::CloseProc(*this);
-		return true;
-	}
-	return false;
-}
-
diff --git a/Engine/Source/Runtime/CoreUObject/Private/Serialization/AsyncLoading.cpp b/Engine/Source/Runtime/CoreUObject/Private/Serialization/AsyncLoading.cpp
--- a/Engine/Source/Runtime/CoreUObject/Private/Serialization/AsyncLoading.cpp
+++ b/Engine/Source/Runtime/CoreUObject/Private/Serialization/AsyncLoading.cpp
@@ -6601,11 +6601,6 @@ int32 LoadPackageAsync(const FString& PackageName, FLoadPackageAsyncDelegate Com
 	return LoadPackageAsync(PackageName, Guid, PackageToLoadFrom, CompletionDelegate, InPackageFlags, InPIEInstanceID, InPackagePriority );
 }
 
-int32 LoadPackageAsync(const FString& InName, const FGuid* InGuid, FName InType /* Unused, deprecated */, const TCHAR* InPackageToLoadFrom /*= nullptr*/, FLoadPackageAsyncDelegate InCompletionDelegate /*= FLoadPackageAsyncDelegate()*/, EPackageFlags InPackageFlags /*= PKG_None*/, int32 InPIEInstanceID /*= INDEX_NONE*/, int32 InPackagePriority /*= 0*/)
-{
-	return LoadPackageAsync(InName, InGuid, InPackageToLoadFrom, InCompletionDelegate, InPackageFlags, InPIEInstanceID, InPackagePriority);
-}
-
 void CancelAsyncLoading()
 {
 	// Cancelling async loading while loading is suspend will result in infinite stall
@@ -6691,11 +6686,6 @@ void FlushAsyncLoading(int32 PackageID /* = INDEX_NONE */)
 	}
 }
 
-void FlushAsyncLoading(FName ExcludeType)
-{
-	FlushAsyncLoading();
-}
-
 EAsyncPackageState::Type ProcessAsyncLoadingUntilComplete(TFunctionRef<bool()> CompletionPredicate, float TimeLimit)
 {
 	if (!IsAsyncLoading())
diff --git a/Engine/Source/Runtime/CoreUObject/Private/UObject/PropertyClass.cpp b/Engine/Source/Runtime/CoreUObject/Private/UObject/PropertyClass.cpp
--- a/Engine/Source/Runtime/CoreUObject/Private/UObject/PropertyClass.cpp
+++ b/Engine/Source/Runtime/CoreUObject/Private/UObject/PropertyClass.cpp
@@ -148,6 +148,16 @@ bool UClassProperty::SameType(const UProperty* Other) const
 	return Super::SameType(Other) && (MetaClass == ((UClassProperty*)Other)->MetaClass);
 }
 
+bool UClassProperty::Identical( const void* A, const void* B, uint32 PortFlags ) const
+{
+	UObject* ObjectA = A ? GetObjectPropertyValue(A) : nullptr;
+	UObject* ObjectB = B ? GetObjectPropertyValue(B) : nullptr;
+
+	check(ObjectA == nullptr || ObjectA->IsA<UClass>());
+	check(ObjectB == nullptr || ObjectB->IsA<UClass>());
+	return (ObjectA == ObjectB);
+}
+
 IMPLEMENT_CORE_INTRINSIC_CLASS(UClassProperty, UObjectProperty,
 	{
 		Class->EmitObjectReference(STRUCT_OFFSET(UClassProperty, MetaClass), TEXT("MetaClass"));
diff --git a/Engine/Source/Runtime/CoreUObject/Private/UObject/PropertyMap.cpp b/Engine/Source/Runtime/CoreUObject/Private/UObject/PropertyMap.cpp
--- a/Engine/Source/Runtime/CoreUObject/Private/UObject/PropertyMap.cpp
+++ b/Engine/Source/Runtime/CoreUObject/Private/UObject/PropertyMap.cpp
@@ -843,10 +843,12 @@ bool UMapProperty::ConvertFromType(const FPropertyTag& Tag, FArchive& Ar, uint8*
 
 	const auto SerializeOrConvert = [](UProperty* CurrentType, const FPropertyTag& InTag, FArchive& InAr, uint8* InData, UStruct* InDefaultsStruct) -> bool
 	{
+		// Serialize wants the property address, while convert wants the container address. InData is the container address
 		bool bDummyAdvance = false;
 		if(CurrentType->GetID() == InTag.Type)
 		{
-			CurrentType->SerializeItem(InAr, InData, InDefaultsStruct);
+			uint8* DestAddress = CurrentType->ContainerPtrToValuePtr<uint8>(InData, InTag.ArrayIndex);
+			CurrentType->SerializeItem(InAr, DestAddress, InDefaultsStruct);
 			return true;
 		}
 		else if( CurrentType->ConvertFromType(InTag, InAr, InData, InDefaultsStruct, bDummyAdvance) )
@@ -948,7 +950,7 @@ bool UMapProperty::ConvertFromType(const FPropertyTag& Tag, FArchive& Ar, uint8*
 						KeyProp->CopyCompleteValue_InContainer(NextPairPtr, TempKeyStorage);
 
 						// Deserialize value
-						if( SerializeOrConvert( ValueProp, ValuePropertyTag, Ar, NextPairPtr + MapLayout.ValueOffset, DefaultsStruct ) )
+						if( SerializeOrConvert( ValueProp, ValuePropertyTag, Ar, NextPairPtr, DefaultsStruct ) )
 						{
 							// first entry went fine, convert the rest:
 							for(int32 I = 1; I < NumEntries; ++I)
@@ -963,7 +965,7 @@ bool UMapProperty::ConvertFromType(const FPropertyTag& Tag, FArchive& Ar, uint8*
 								NextPairPtr = MapHelper.GetPairPtrWithoutCheck(NextPairIndex);
 								// This copy is unnecessary when the key was already in the map:
 								KeyProp->CopyCompleteValue_InContainer(NextPairPtr, TempKeyStorage);
-								verify( SerializeOrConvert( ValueProp, ValuePropertyTag, Ar, NextPairPtr + MapLayout.ValueOffset, DefaultsStruct ) );
+								verify( SerializeOrConvert( ValueProp, ValuePropertyTag, Ar, NextPairPtr, DefaultsStruct ) );
 							}
 						}
 						else
diff --git a/Engine/Source/Runtime/CoreUObject/Private/UObject/ScriptCore.cpp b/Engine/Source/Runtime/CoreUObject/Private/UObject/ScriptCore.cpp
--- a/Engine/Source/Runtime/CoreUObject/Private/UObject/ScriptCore.cpp
+++ b/Engine/Source/Runtime/CoreUObject/Private/UObject/ScriptCore.cpp
@@ -2680,23 +2680,31 @@ IMPLEMENT_VM_FUNCTION( EX_SetArray, execSetArray );
 void UObject::execSetSet( FFrame& Stack, RESULT_DECL )
 {
 	// Get the set address
-	Stack.MostRecentPropertyAddress = NULL;
-	Stack.MostRecentProperty = NULL;
-	Stack.Step( Stack.Object, NULL ); // Set to set
-	int32 Num = Stack.ReadInt<int32>();
+	Stack.MostRecentPropertyAddress = nullptr;
+	Stack.MostRecentProperty = nullptr;
+	Stack.Step( Stack.Object, nullptr ); // Set to set
+	const int32 Num = Stack.ReadInt<int32>();
 
 	USetProperty* SetProperty = CastChecked<USetProperty>(Stack.MostRecentProperty);
  	FScriptSetHelper SetHelper(SetProperty, Stack.MostRecentPropertyAddress);
  	SetHelper.EmptyElements(Num);
  
- 	// Read in the parameters one at a time
- 	while(*Stack.Code != EX_EndSet)
- 	{
-		// needs to be an initialized/constructed value, in case the op is a literal that gets assigned over  
- 		int32 Index = SetHelper.AddDefaultValue_Invalid_NeedsRehash();
- 		Stack.Step(Stack.Object, SetHelper.GetElementPtr(Index));
- 	}
-	SetHelper.Rehash();
+	if (Num > 0)
+	{
+		FDefaultConstructedPropertyElement TempElement(SetProperty->ElementProp);
+
+		// Read in the parameters one at a time
+		while (*Stack.Code != EX_EndSet)
+		{
+			// needs to be an initialized/constructed value, in case the op is a literal that gets assigned over  
+			Stack.Step(Stack.Object, TempElement.GetObjAddress());
+			SetHelper.AddElement(TempElement.GetObjAddress());
+		}
+	}
+	else
+	{
+		check(*Stack.Code == EX_EndSet);
+	}
  
  	P_FINISH;
 }
@@ -2705,24 +2713,32 @@ IMPLEMENT_VM_FUNCTION( EX_SetSet, execSetSet );
 void UObject::execSetMap( FFrame& Stack, RESULT_DECL )
 {
 	// Get the map address
-	Stack.MostRecentPropertyAddress = NULL;
-	Stack.MostRecentProperty = NULL;
-	Stack.Step( Stack.Object, NULL ); // Map to set
-	int32 Num = Stack.ReadInt<int32>();
+	Stack.MostRecentPropertyAddress = nullptr;
+	Stack.MostRecentProperty = nullptr;
+	Stack.Step( Stack.Object, nullptr ); // Map to set
+	const int32 Num = Stack.ReadInt<int32>();
 	
 	UMapProperty* MapProperty = CastChecked<UMapProperty>(Stack.MostRecentProperty);
  	FScriptMapHelper MapHelper(MapProperty, Stack.MostRecentPropertyAddress);
  	MapHelper.EmptyValues(Num);
  
- 	// Read in the parameters one at a time
- 	while(*Stack.Code != EX_EndMap)
- 	{
-		// needs to be an initialized/constructed value, in case the op is a literal that gets assigned over 
-		int32 Index = MapHelper.AddDefaultValue_Invalid_NeedsRehash();
- 		Stack.Step(Stack.Object, MapHelper.GetKeyPtr(Index));
- 		Stack.Step(Stack.Object, MapHelper.GetValuePtr(Index));
- 	}
-	MapHelper.Rehash();
+	if (Num > 0)
+	{
+		FDefaultConstructedPropertyElement TempKey(MapProperty->KeyProp);
+		FDefaultConstructedPropertyElement TempValue(MapProperty->ValueProp);
+
+		// Read in the parameters one at a time
+		while (*Stack.Code != EX_EndMap)
+		{
+			Stack.Step(Stack.Object, TempKey.GetObjAddress());
+			Stack.Step(Stack.Object, TempValue.GetObjAddress());
+			MapHelper.AddPair(TempKey.GetObjAddress(), TempValue.GetObjAddress());
+		}
+	}
+	else
+	{
+		check(*Stack.Code == EX_EndMap);
+	}
  
  	P_FINISH;
 }
diff --git a/Engine/Source/Runtime/CoreUObject/Private/UObject/UObjectGlobals.cpp b/Engine/Source/Runtime/CoreUObject/Private/UObject/UObjectGlobals.cpp
--- a/Engine/Source/Runtime/CoreUObject/Private/UObject/UObjectGlobals.cpp
+++ b/Engine/Source/Runtime/CoreUObject/Private/UObject/UObjectGlobals.cpp
@@ -3129,20 +3129,6 @@ UObject* StaticConstructObject_Internal
 	return Result;
 }
 
-UObject* StaticConstructObject
-(
-	UClass*			InClass,
-	UObject*		InOuter								/*=GetTransientPackage()*/,
-	FName			InName								/*=NAME_None*/,
-	EObjectFlags	InFlags								/*=0*/,
-	UObject*		InTemplate							/*=NULL*/,
-	bool			bCopyTransientsFromClassDefaults	/*=false*/, 
-	FObjectInstancingGraph* InInstanceGraph				/*=NULL*/
-)
-{
-	return StaticConstructObject_Internal(InClass, InOuter, InName, InFlags, EInternalObjectFlags::None, InTemplate, bCopyTransientsFromClassDefaults, InInstanceGraph);
-}
-
 void FObjectInitializer::AssertIfInConstructor(UObject* Outer, const TCHAR* ErrorMessage)
 {
 	FUObjectThreadContext& ThreadContext = FUObjectThreadContext::Get();
diff --git a/Engine/Source/Runtime/CoreUObject/Public/Templates/Casts.h b/Engine/Source/Runtime/CoreUObject/Public/Templates/Casts.h
--- a/Engine/Source/Runtime/CoreUObject/Public/Templates/Casts.h
+++ b/Engine/Source/Runtime/CoreUObject/Public/Templates/Casts.h
@@ -260,21 +260,6 @@ FORCEINLINE T* ExactCast( UObject* Src )
 
 #endif
 
-
-template <typename InterfaceType>
-DEPRECATED(4.6, "InterfaceCast is deprecated, use Cast or dynamic_cast instead.")
-FORCEINLINE InterfaceType* InterfaceCast(UObject* Src)
-{
-	return Cast<InterfaceType>(Src);
-}
-
-template <typename InterfaceType>
-DEPRECATED(4.6, "InterfaceCast is deprecated, use Cast or dynamic_cast instead.")
-FORCEINLINE InterfaceType* InterfaceCast(const UObject* Src)
-{
-	return Cast<InterfaceType>(const_cast<UObject*>(Src));
-}
-
 // auto weak versions
 template< class T, class U > FORCEINLINE T* Cast       ( const TAutoWeakObjectPtr<U>& Src                                                                   ) { return Cast       <T>(Src.Get()); }
 template< class T, class U > FORCEINLINE T* ExactCast  ( const TAutoWeakObjectPtr<U>& Src                                                                   ) { return ExactCast  <T>(Src.Get()); }
diff --git a/Engine/Source/Runtime/CoreUObject/Public/UObject/Class.h b/Engine/Source/Runtime/CoreUObject/Public/UObject/Class.h
--- a/Engine/Source/Runtime/CoreUObject/Public/UObject/Class.h
+++ b/Engine/Source/Runtime/CoreUObject/Public/UObject/Class.h
@@ -2934,31 +2934,6 @@ FORCEINLINE bool UObject::Implements() const
 
 // UObjectGlobals.h
 
-/**
- * Construct an object of a particular class.
- * 
- * @param	Class		the class of object to construct
- * @param	Outer		the outer for the new object.  If not specified, object will be created in the transient package.
- * @param	Name		the name for the new object.  If not specified, the object will be given a transient name via
- *						MakeUniqueObjectName
- * @param	SetFlags	the object flags to apply to the new object
- * @param	Template	the object to use for initializing the new object.  If not specified, the class's default object will
- *						be used
- * @param	bInCopyTransientsFromClassDefaults - if true, copy transient from the class defaults instead of the pass in archetype ptr (often these are the same)
- * @param	InstanceGraph
- *						contains the mappings of instanced objects and components to their templates
- *
- * @return	a pointer of type T to a new object of the specified class
- */
-template< class T >
-DEPRECATED(4.8, "ConstructObject is deprecated. Use NewObject instead")
-T* ConstructObject(UClass* Class, UObject* Outer, FName Name, EObjectFlags SetFlags, UObject* Template, bool bCopyTransientsFromClassDefaults, struct FObjectInstancingGraph* InstanceGraph )
-{
-	checkf(Class, TEXT("ConstructObject called with a NULL class object"));
-	checkSlow(Class->IsChildOf(T::StaticClass()));
-	return (T*)StaticConstructObject_Internal(Class, Outer, Name, SetFlags, EInternalObjectFlags::None, Template, bCopyTransientsFromClassDefaults, InstanceGraph);
-}
-
 /**
  * Gets the default object of a class.
  *
diff --git a/Engine/Source/Runtime/CoreUObject/Public/UObject/ObjectMacros.h b/Engine/Source/Runtime/CoreUObject/Public/UObject/ObjectMacros.h
--- a/Engine/Source/Runtime/CoreUObject/Public/UObject/ObjectMacros.h
+++ b/Engine/Source/Runtime/CoreUObject/Public/UObject/ObjectMacros.h
@@ -1366,16 +1366,11 @@ public: \
 	{ \
 		return TStaticCastFlags; \
 	} \
-	DEPRECATED(4.7, "operator new has been deprecated for UObjects - please use NewObject or NewNamedObject instead") \
-	inline void* operator new( const size_t InSize, UObject* InOuter=(UObject*)GetTransientPackage(), FName InName=NAME_None, EObjectFlags InSetFlags=RF_NoFlags ) \
-	{ \
-		return StaticAllocateObject( StaticClass(), InOuter, InName, InSetFlags ); \
-	} \
 	/** For internal use only; use StaticConstructObject() to create new objects. */ \
 	inline void* operator new(const size_t InSize, EInternal InInternalOnly, UObject* InOuter = (UObject*)GetTransientPackage(), FName InName = NAME_None, EObjectFlags InSetFlags = RF_NoFlags) \
 	{ \
 		return StaticAllocateObject(StaticClass(), InOuter, InName, InSetFlags); \
-} \
+	} \
 	/** For internal use only; use StaticConstructObject() to create new objects. */ \
 	inline void* operator new( const size_t InSize, EInternal* InMem ) \
 	{ \
diff --git a/Engine/Source/Runtime/CoreUObject/Public/UObject/ScriptInterface.h b/Engine/Source/Runtime/CoreUObject/Public/UObject/ScriptInterface.h
--- a/Engine/Source/Runtime/CoreUObject/Public/UObject/ScriptInterface.h
+++ b/Engine/Source/Runtime/CoreUObject/Public/UObject/ScriptInterface.h
@@ -8,6 +8,7 @@
 
 #include "CoreMinimal.h"
 #include "UObjectGlobals.h"
+#include "Templates/Casts.h"
 
 /**
  * FScriptInterface
@@ -158,7 +159,7 @@ public:
 	{
 		SetObject(SourceObject);
 		
-		InterfaceType* SourceInterface = SourceObject;
+		InterfaceType* SourceInterface = Cast<InterfaceType>(SourceObject);
 		SetInterface( SourceInterface );
 
 		return *((InterfaceType*)GetInterface());
diff --git a/Engine/Source/Runtime/CoreUObject/Public/UObject/UObjectGlobals.h b/Engine/Source/Runtime/CoreUObject/Public/UObject/UObjectGlobals.h
--- a/Engine/Source/Runtime/CoreUObject/Public/UObject/UObjectGlobals.h
+++ b/Engine/Source/Runtime/CoreUObject/Public/UObject/UObjectGlobals.h
@@ -241,26 +241,6 @@ COREUOBJECT_API UClass* StaticLoadClass(UClass* BaseClass, UObject* InOuter, con
 */
 COREUOBJECT_API UObject* StaticConstructObject_Internal(UClass* Class, UObject* InOuter = (UObject*)GetTransientPackage(), FName Name = NAME_None, EObjectFlags SetFlags = RF_NoFlags, EInternalObjectFlags InternalSetFlags = EInternalObjectFlags::None, UObject* Template = NULL, bool bCopyTransientsFromClassDefaults = false, struct FObjectInstancingGraph* InstanceGraph = NULL, bool bAssumeTemplateIsArchetype = false);
 
-/**
- * Create a new instance of an object.  The returned object will be fully initialized.  If InFlags contains RF_NeedsLoad (indicating that the object still needs to load its object data from disk), components
- * are not instanced (this will instead occur in PostLoad()).  The different between StaticConstructObject and StaticAllocateObject is that StaticConstructObject will also call the class constructor on the object
- * and instance any components.
- * 
- * @param	Class		the class of the object to create
- * @param	InOuter		the object to create this object within (the Outer property for the new object will be set to the value specified here).
- * @param	Name		the name to give the new object. If no value (NAME_None) is specified, the object will be given a unique name in the form of ClassName_#.
- * @param	SetFlags	the ObjectFlags to assign to the new object. some flags can affect the behavior of constructing the object.
- * @param	Template	if specified, the property values from this object will be copied to the new object, and the new object's ObjectArchetype value will be set to this object.
- *						If NULL, the class default object is used instead.
- * @param	bInCopyTransientsFromClassDefaults - if true, copy transient from the class defaults instead of the pass in archetype ptr (often these are the same)
- * @param	InstanceGraph
- *						contains the mappings of instanced objects and components to their templates
- *
- * @return	a pointer to a fully initialized object of the specified class.
- */
-DEPRECATED(4.8, "StaticConstructObject is deprecated, please use NewObject instead. For internal CoreUObject module usage, please use StaticConstructObject_Internal.")
-COREUOBJECT_API UObject* StaticConstructObject( UClass* Class, UObject* InOuter=(UObject*)GetTransientPackage(), FName Name=NAME_None, EObjectFlags SetFlags=RF_NoFlags, UObject* Template=NULL, bool bCopyTransientsFromClassDefaults=false, struct FObjectInstancingGraph* InstanceGraph=NULL );
-
 /**
  * Creates a copy of SourceObject using the Outer and Name specified, as well as copies of all objects contained by SourceObject.  
  * Any objects referenced by SourceOuter or RootObject and contained by SourceOuter are also copied, maintaining their name relative to SourceOuter.  Any
@@ -327,22 +307,6 @@ typedef int32 TAsyncLoadPriority;
  */
 DECLARE_DELEGATE_ThreeParams(FLoadPackageAsyncDelegate, const FName& /*PackageName*/, UPackage* /*LoadedPackage*/, EAsyncLoadingResult::Type /*Result*/)
 
-/**
- * [Deprecated] Asynchronously load a package and all contained objects that match context flags. Non- blocking.
- *
- * @param	InName					Name of package to load
- * @param	InGuid					GUID of the package to load, or NULL for "don't care"
- * @param	InType					A type name associated with this package for later use
- * @param	InPackageToLoadFrom		If non-null, this is another package name. We load from this package name, into a (probably new) package named PackageName
- * @param	InCompletionDelegate	Delegate to be invoked when the packages has finished streaming
- * @param	InFlags					Package flags
- * @param	InPIEInstanceID			PIE instance ID
- * @param	InPackagePriority		Loading priority
- * @return Unique ID associated with this load request (the same package can be associated with multiple IDs).
- */
-DEPRECATED(4.9, "LoadPackageAsync override that takes package type parameter FName InType is deprecated.")
-COREUOBJECT_API int32 LoadPackageAsync(const FString& InName, const FGuid* InGuid, FName InType, const TCHAR* InPackageToLoadFrom = nullptr, FLoadPackageAsyncDelegate InCompletionDelegate = FLoadPackageAsyncDelegate(), EPackageFlags InPackageFlags = PKG_None, int32 InPIEInstanceID = INDEX_NONE, TAsyncLoadPriority InPackagePriority = 0);
-
 /**
 * Asynchronously load a package and all contained objects that match context flags. Non- blocking.
 *
@@ -477,14 +441,6 @@ COREUOBJECT_API bool IsReferenced( UObject*& Res, EObjectFlags KeepFlags, EInter
  */
 COREUOBJECT_API void FlushAsyncLoading(int32 PackageID = INDEX_NONE);
 
-/**
-* [Deprecated] Blocks till all pending package/ linker requests are fulfilled.
-*
-* @param	ExcludeType		Do not flush packages associated with this specific type name
-*/
-DEPRECATED(4.9, "FlushAsyncLoading override that takes package type parameter FName ExcludeType is deprecated.")
-COREUOBJECT_API void FlushAsyncLoading(FName ExcludeType);
-
 /**
  * @return number of active async load package requests
  */
@@ -1147,26 +1103,6 @@ public:
  * Helper class for deferred execution of 
 */
 
-/**
- * Construct an object of a particular class.
- * 
- * @param	Class		the class of object to construct
- * @param	Outer		the outer for the new object.  If not specified, object will be created in the transient package.
- * @param	Name		the name for the new object.  If not specified, the object will be given a transient name via
- *						MakeUniqueObjectName
- * @param	SetFlags	the object flags to apply to the new object
- * @param	Template	the object to use for initializing the new object.  If not specified, the class's default object will
- *						be used
- * @param	bInCopyTransientsFromClassDefaults - if true, copy transient from the class defaults instead of the pass in archetype ptr (often these are the same)
- * @param	InstanceGraph
- *						contains the mappings of instanced objects and components to their templates
- *
- * @return	a pointer of type T to a new object of the specified class
- */
-template< class T >
-DEPRECATED(4.8, "ConstructObject is deprecated. Use NewObject instead")
-T* ConstructObject(UClass* Class, UObject* Outer = (UObject*)GetTransientPackage(), FName Name=NAME_None, EObjectFlags SetFlags=RF_NoFlags, UObject* Template=NULL, bool bCopyTransientsFromClassDefaults=false, struct FObjectInstancingGraph* InstanceGraph=NULL );
-
 #if DO_CHECK
 /** Called by NewObject to make sure Child is actually a child of Parent */
 COREUOBJECT_API void CheckIsClassChildOf_Internal(UClass* Parent, UClass* Child);
@@ -1227,21 +1163,6 @@ FUNCTION_NON_NULL_RETURN_END
 	return static_cast<T*>(StaticConstructObject_Internal(T::StaticClass(), Outer, Name, Flags, EInternalObjectFlags::None, Template, bCopyTransientsFromClassDefaults, InInstanceGraph));
 }
 
-/**
- * Convenience template for constructing a named object.
- *
- * @param	Outer	The outer for the new object.
- * @param	Name	The name of the new object.
- * @param	Flags	The object flags for the new object.
- */
-template< class TClass >
-DEPRECATED(4.8, "NewNamedObject is deprecated. Use NewObject instead")
-TClass* NewNamedObject(UObject* Outer, FName Name, EObjectFlags Flags = RF_NoFlags, UObject* Template=NULL)
-{
-	return NewObject<TClass>(Outer, Name, Flags, Template);
-}
-
-
 /**
  * Convenience template for duplicating an object
  *
diff --git a/Engine/Source/Runtime/CoreUObject/Public/UObject/UnrealType.h b/Engine/Source/Runtime/CoreUObject/Public/UObject/UnrealType.h
--- a/Engine/Source/Runtime/CoreUObject/Public/UObject/UnrealType.h
+++ b/Engine/Source/Runtime/CoreUObject/Public/UObject/UnrealType.h
@@ -2218,6 +2218,7 @@ public:
 	// UProperty interface
 	virtual const TCHAR* ImportText_Internal( const TCHAR* Buffer, void* Data, int32 PortFlags, UObject* OwnerObject, FOutputDevice* ErrorText ) const override;
 	virtual bool SameType(const UProperty* Other) const override;
+	virtual bool Identical( const void* A, const void* B, uint32 PortFlags ) const override;
 	// End of UProperty interface
 
 	virtual FString GetCPPTypeCustom(FString* ExtendedTypeText, uint32 CPPExportFlags, const FString& InnerNativeTypeName)  const override;
diff --git a/Engine/Source/Runtime/Engine/Classes/AI/Navigation/NavRelevantComponent.h b/Engine/Source/Runtime/Engine/Classes/AI/Navigation/NavRelevantComponent.h
--- a/Engine/Source/Runtime/Engine/Classes/AI/Navigation/NavRelevantComponent.h
+++ b/Engine/Source/Runtime/Engine/Classes/AI/Navigation/NavRelevantComponent.h
@@ -24,9 +24,6 @@ class ENGINE_API UNavRelevantComponent : public UActorComponent, public INavRele
 	virtual void UpdateNavigationBounds() override;
 	virtual UObject* GetNavigationParent() const override;
 	//~ End INavRelevantInterface Interface
-
-	DEPRECATED(4.8, "This function is deprecated. Use CalcAndCacheBounds instead")
-	virtual void CalcBounds();
 	
 	virtual void CalcAndCacheBounds() const;
 
diff --git a/Engine/Source/Runtime/Engine/Classes/AI/Navigation/NavigationData.h b/Engine/Source/Runtime/Engine/Classes/AI/Navigation/NavigationData.h
--- a/Engine/Source/Runtime/Engine/Classes/AI/Navigation/NavigationData.h
+++ b/Engine/Source/Runtime/Engine/Classes/AI/Navigation/NavigationData.h
@@ -972,8 +972,6 @@ private:
 	static uint16 GetNextUniqueID();
 
 public:
-	DEPRECATED(4.8, "GetRandomPointInRadius is deprecated, please use GetRandomReachablePointInRadius")
-	virtual bool GetRandomPointInRadius(const FVector& Origin, float Radius, FNavLocation& OutResult, FSharedConstNavQueryFilter Filter = NULL, const UObject* Querier = NULL) const PURE_VIRTUAL(ANavigationData::GetRandomPointInRadius, return false;);
 	DEPRECATED(4.12, "This function is now deprecated, initial rebuild ignore should be handled by discarding dirty areas in UNavigationSystem::ConditionalPopulateNavOctree.")
 	void MarkAsNeedingUpdate() { }
 };
@@ -989,9 +987,6 @@ struct FAsyncPathFindingQuery : public FPathFindingQuery
 		: QueryID(INVALID_NAVQUERYID)
 	{ }
 
-	DEPRECATED(4.8, "This version of FAsyncPathFindingQuery constructor is deprecated. Please use ANavigationData reference rather than a pointer version")
-	FAsyncPathFindingQuery(const UObject* InOwner, const ANavigationData* InNavData, const FVector& Start, const FVector& End, const FNavPathQueryDelegate& Delegate, FSharedConstNavQueryFilter SourceQueryFilter);
-
 	FAsyncPathFindingQuery(const UObject* InOwner, const ANavigationData& InNavData, const FVector& Start, const FVector& End, const FNavPathQueryDelegate& Delegate, FSharedConstNavQueryFilter SourceQueryFilter);
 	FAsyncPathFindingQuery(const FPathFindingQuery& Query, const FNavPathQueryDelegate& Delegate, const EPathFindingMode::Type QueryMode);
 
diff --git a/Engine/Source/Runtime/Engine/Classes/AI/Navigation/RecastNavMeshDataChunk.h b/Engine/Source/Runtime/Engine/Classes/AI/Navigation/RecastNavMeshDataChunk.h
--- a/Engine/Source/Runtime/Engine/Classes/AI/Navigation/RecastNavMeshDataChunk.h
+++ b/Engine/Source/Runtime/Engine/Classes/AI/Navigation/RecastNavMeshDataChunk.h
@@ -77,12 +77,6 @@ private:
 	void SerializeRecastData(FArchive& Ar, int32 NavMeshVersion);
 #endif//WITH_RECAST
 
-	
-	DEPRECATED(4.8, "AttachTiles is deprecated. Use the version takin a reference instead.")
-	TArray<uint32> AttachTiles(FPImplRecastNavMesh* NavMeshImpl);
-	DEPRECATED(4.8, "DetachTiles is deprecated. Use the version takin a reference instead.")
-	TArray<uint32> DetachTiles(FPImplRecastNavMesh* NavMeshImpl);
-
 private:
 	TArray<FRecastTileData> Tiles;
 };
diff --git a/Engine/Source/Runtime/Engine/Classes/Animation/AnimInstance.h b/Engine/Source/Runtime/Engine/Classes/Animation/AnimInstance.h
--- a/Engine/Source/Runtime/Engine/Classes/Animation/AnimInstance.h
+++ b/Engine/Source/Runtime/Engine/Classes/Animation/AnimInstance.h
@@ -340,7 +340,7 @@ struct FMontageEvaluationState
 	bool bIsActive;
 };
 
-UCLASS(transient, Blueprintable, hideCategories=AnimInstance, BlueprintType, meta=(BlueprintThreadSafe))
+UCLASS(transient, Blueprintable, hideCategories=AnimInstance, BlueprintType, meta=(BlueprintThreadSafe), Within=SkeletalMeshComponent)
 class ENGINE_API UAnimInstance : public UObject
 {
 	GENERATED_UCLASS_BODY()
diff --git a/Engine/Source/Runtime/Engine/Classes/Camera/CameraComponent.h b/Engine/Source/Runtime/Engine/Classes/Camera/CameraComponent.h
--- a/Engine/Source/Runtime/Engine/Classes/Camera/CameraComponent.h
+++ b/Engine/Source/Runtime/Engine/Classes/Camera/CameraComponent.h
@@ -196,18 +196,4 @@ public:
 	/** DEPRECATED: use bUsePawnControlRotation instead */
 	UPROPERTY()
 	uint32 bUseControllerViewRotation_DEPRECATED:1;
-
-	/**
-	 * DEPRECATED variable: use "bUsePawnControlRotation" instead. Existing code using this value may not behave correctly.
-	 *
-	 * The correct way to be backwards-compatible for existing saved content using CameraComponents is:
-	 *  - add "_DEPRECATED" to any constructor initialization of bUseControllerViewRotation (which becomes bUseControllerViewRotation_DEPRECATED).
-	 *	- add initialization of the new "bUsePawnControlRotation" to your desired default value.
-	 *	- if there was no previous initialization of bUseControllerViewRotation in code, then you MUST add initialization of bUsePawnControlRotation
-	 *    to "true" if you wish to maintain the same behavior as before (since the default has changed).
-	 *
-	 * This is not a UPROPERTY, with good reason: we don't want to serialize in old values (bUseControllerViewRotation_DEPRECATED handles that).
-	 */
-	DEPRECATED(4.5, "This variable is deprecated; please see the upgrade notes. It only exists to allow compilation of old projects, and code should stop using it in favor of the new bUsePawnControlRotation.")
-	bool bUseControllerViewRotation;
 };
diff --git a/Engine/Source/Runtime/Engine/Classes/Components/ActorComponent.h b/Engine/Source/Runtime/Engine/Classes/Components/ActorComponent.h
--- a/Engine/Source/Runtime/Engine/Classes/Components/ActorComponent.h
+++ b/Engine/Source/Runtime/Engine/Classes/Components/ActorComponent.h
@@ -688,6 +688,9 @@ public:
 	/** If we belong to a world, mark this for a deferred update, otherwise do it now. */
 	void MarkForNeededEndOfFrameRecreate();
 
+	/** If we belong to a world, clear the request to do a deferred update. */
+	void ClearNeedEndOfFrameUpdate();
+
 	/** return true if this component requires end of frame updates to happen from the game thread. */
 	virtual bool RequiresGameThreadEndOfFrameUpdates() const;
 
@@ -739,7 +742,7 @@ public:
 	/** Give a readable name for this component, including asset name if applicable */
 	virtual UObject const* AdditionalStatObject() const
 	{
-		return NULL;
+		return nullptr;
 	}
 
 	// Always called immediately before properties are received from the remote.
@@ -865,14 +868,17 @@ protected:
 	 *		already been registered. Setting bForceUpdate to true overrides that check */
 	void HandleCanEverAffectNavigationChange(bool bForceUpdate = false);
 
+private:
+
 #if !(UE_BUILD_SHIPPING || UE_BUILD_TEST)
 
-private:
 	// this is the old name of the tick function. We just want to avoid mistakes with an attempt to override this
 	virtual void Tick( float DeltaTime ) final { check(0); }
 
 #endif
 
+	void ClearNeedEndOfFrameUpdate_Internal();
+
 public:
 
 	/** Prefix used to identify template component instances */
@@ -921,6 +927,14 @@ FORCEINLINE_DEBUGGABLE bool UActorComponent::IsNetMode(ENetMode Mode) const
 #endif // UE_EDITOR
 }
 
+FORCEINLINE void UActorComponent::ClearNeedEndOfFrameUpdate()
+{
+	if (MarkedForEndOfFrameUpdateState != 0)
+	{
+		ClearNeedEndOfFrameUpdate_Internal();
+	}
+}
+
 FORCEINLINE_DEBUGGABLE AActor* UActorComponent::GetOwner() const
 {
 #if WITH_EDITOR
diff --git a/Engine/Source/Runtime/Engine/Classes/Components/PrimitiveComponent.h b/Engine/Source/Runtime/Engine/Classes/Components/PrimitiveComponent.h
--- a/Engine/Source/Runtime/Engine/Classes/Components/PrimitiveComponent.h
+++ b/Engine/Source/Runtime/Engine/Classes/Components/PrimitiveComponent.h
@@ -2184,16 +2184,6 @@ public:
 	 */
 
 	virtual bool ComputePenetration(FMTDResult & OutMTD, const FCollisionShape& CollisionShape, const FVector& Pos, const FQuat& Rot);
-	
-	/**
-	 * Return true if the given Pawn can step up onto this component.
-	 * @param APawn is the pawn that wants to step onto this component.
-	 */
-	DEPRECATED(4.3, "UPrimitiveComponent::CanBeBaseForCharacter() is deprecated, use CanCharacterStepUp() instead.")
-	virtual bool CanBeBaseForCharacter(class APawn* Pawn) const
-	{
-		return CanCharacterStepUp(Pawn);
-	}
 
 	/**
 	 * Return true if the given Pawn can step up onto this component.
@@ -2203,12 +2193,6 @@ public:
 	 */
 	UFUNCTION(BlueprintCallable, Category=Collision)
 	virtual bool CanCharacterStepUp(class APawn* Pawn) const;
-	
-	DEPRECATED(4.5, "UPrimitiveComponent::DisableNavigationRelevance() is deprecated, use SetCanEverAffectNavigation() instead.")
-	void DisableNavigationRelevance()
-	{
-		SetCanEverAffectNavigation(false);
-	}
 
 	//~ Begin INavRelevantInterface Interface
 	virtual FBox GetNavigationBounds() const override;
@@ -2219,10 +2203,6 @@ public:
 
 	void SetCustomNavigableGeometry(const EHasCustomNavigableGeometry::Type InType);
 
-	/** Collects custom navigable geometry of component.
-	 *	@return true if regular navigable geometry exporting should be run as well */
-	DEPRECATED(4.8, "UPrimitiveComponent::DoCustomNavigableGeometryExport(FNavigableGeometryExport* GeomExport) is deprecated, use UPrimitiveComponent::DoCustomNavigableGeometryExport(FNavigableGeometryExport& GeomExport) instead (takes ref instead of a pointer)")
-	virtual bool DoCustomNavigableGeometryExport(FNavigableGeometryExport* GeomExport) const { return DoCustomNavigableGeometryExport(*GeomExport); }
 	/** Collects custom navigable geometry of component.
 	*	@return true if regular navigable geometry exporting should be run as well */
 	virtual bool DoCustomNavigableGeometryExport(FNavigableGeometryExport& GeomExport) const { return true; }
diff --git a/Engine/Source/Runtime/Engine/Classes/Components/SplineComponent.h b/Engine/Source/Runtime/Engine/Classes/Components/SplineComponent.h
--- a/Engine/Source/Runtime/Engine/Classes/Components/SplineComponent.h
+++ b/Engine/Source/Runtime/Engine/Classes/Components/SplineComponent.h
@@ -331,9 +331,6 @@ public:
 	/** Get scale at the provided input key value */
 	FVector GetScaleAtSplineInputKey(float InKey) const;
 
-	DEPRECATED(4.9, "This has been replaced by GetNumberOfSplinePoints(), as it is a clearer name for Blueprint users.")
-	int32 GetNumSplinePoints() const { return GetNumberOfSplinePoints(); }
-
 	/** Specify unselected spline component segment color in the editor */
 	UFUNCTION(BlueprintCallable, Category = Editor)
 	void SetUnselectedSplineSegmentColor(const FLinearColor& SegmentColor);
diff --git a/Engine/Source/Runtime/Engine/Classes/EdGraph/EdGraphNode.h b/Engine/Source/Runtime/Engine/Classes/EdGraph/EdGraphNode.h
--- a/Engine/Source/Runtime/Engine/Classes/EdGraph/EdGraphNode.h
+++ b/Engine/Source/Runtime/Engine/Classes/EdGraph/EdGraphNode.h
@@ -30,7 +30,7 @@ struct FDiffSingleResult;
   * a structure, rather than implicitly defining names for containers.
   */
 USTRUCT()
-struct FEdGraphTerminalType
+struct ENGINE_API FEdGraphTerminalType
 {
 	GENERATED_USTRUCT_BODY()
 
@@ -63,7 +63,10 @@ struct FEdGraphTerminalType
 	UPROPERTY()
 	bool bTerminalIsWeakPointer;
 
-	ENGINE_API friend FArchive& operator<<(FArchive& Ar, FEdGraphTerminalType& P);
+	/** Creates a TerminalType from the primary portion of the PinType */
+	static FEdGraphTerminalType FromPinType(const FEdGraphPinType& PinType);
+
+	friend FArchive& operator<<(FArchive& Ar, FEdGraphTerminalType& P);
 };
 
 /** Enum used to define which way data flows into or out of this pin. */
@@ -472,9 +475,6 @@ public:
 	 */
 	virtual bool CanPasteHere(const UEdGraph* TargetGraph) const { return IsCompatibleWithGraph(TargetGraph); }
 
-	DEPRECATED(4.5, "The UEdGraphNode::CanPasteHere() that takes a UEdGraphSchema parameter is deprecated, instead use the CanPasteHere() that only takes a single UEdGraph param.")
-	virtual bool CanPasteHere(const UEdGraph* TargetGraph, const UEdGraphSchema* Schema) const { return CanPasteHere(TargetGraph); }
-
 	/**
 	 * Determine if this node can be created under the specified schema
      */
@@ -511,9 +511,6 @@ public:
 	 */
 	virtual FText GetTooltipText() const;
 
-	DEPRECATED(4.5, "UEdGraphNode::GetTooltip() is deprecated, instead use GetTooltipText(), which returns localized text.")
-	virtual FString GetTooltip() const { return GetTooltipText().ToString(); }
-
 	/**
 	 * Returns the keywords that should be used when searching for this node
 	 *
diff --git a/Engine/Source/Runtime/Engine/Classes/Engine/AssetManager.h b/Engine/Source/Runtime/Engine/Classes/Engine/AssetManager.h
--- a/Engine/Source/Runtime/Engine/Classes/Engine/AssetManager.h
+++ b/Engine/Source/Runtime/Engine/Classes/Engine/AssetManager.h
@@ -8,12 +8,16 @@
 #include "StreamableManager.h"
 #include "AssetBundleData.h"
 #include "AssetRegistryModule.h"
+#include "GenericPlatform/GenericPlatformChunkInstall.h"
 #include "AssetManager.generated.h"
 
 /** Defined in C++ file */
 struct FPrimaryAssetTypeData;
 struct FPrimaryAssetData;
 
+/** Delegate called when acquiring resources/chunks for assets, parameter will be true if all resources were acquired, false if any failed */
+DECLARE_DELEGATE_OneParam(FAssetManagerAcquireResourceDelegate, bool);
+
 /** 
  * A singleton UObject that is responsible for loading and unloading PrimaryAssets, and maintaining game-specific asset references
  * Games should override this class and change the class reference
@@ -243,7 +247,7 @@ public:
 	virtual TSharedPtr<FStreamableHandle> PreloadPrimaryAssets(const TArray<FPrimaryAssetId>& AssetsToLoad, const TArray<FName>& LoadBundles, bool bLoadRecursive, FStreamableDelegate DelegateToCall = FStreamableDelegate(), TAsyncLoadPriority Priority = FStreamableManager::DefaultAsyncLoadPriority);
 
 	/** Quick wrapper to async load some non primary assets with the primary streamable manager. This will not auto release the handle, release it if needed */
-	virtual TSharedPtr<FStreamableHandle> LoadAssetList(const TArray<FStringAssetReference>& AssetList);
+	virtual TSharedPtr<FStreamableHandle> LoadAssetList(const TArray<FStringAssetReference>& AssetList, FStreamableDelegate DelegateToCall = FStreamableDelegate(), TAsyncLoadPriority Priority = FStreamableManager::DefaultAsyncLoadPriority, const FString& DebugName = TEXT("LoadAssetList"));
 
 	/** Returns a single AssetBundleInfo, matching Scope and Name */
 	virtual FAssetBundleEntry GetAssetBundleEntry(const FPrimaryAssetId& BundleScope, FName BundleName) const;
@@ -251,7 +255,36 @@ public:
 	/** Appends all AssetBundleInfos inside a given scope */
 	virtual bool GetAssetBundleEntries(const FPrimaryAssetId& BundleScope, TArray<FAssetBundleEntry>& OutEntries) const;
 
+	/** 
+	 * Returns the list of Chunks that are not currently mounted, and are required to load the referenced assets. Returns true if any chunks are missing
+	 *
+	 * @param AssetList				Asset Paths to check chunks for
+	 * @param OutMissingChunkList	Chunks that are known about but not yet installed
+	 * @param OutErrorChunkList		Chunks that do not exist at all and are not installable
+	 */
+	virtual bool FindMissingChunkList(const TArray<FStringAssetReference>& AssetList, TArray<int32>& OutMissingChunkList, TArray<int32>& OutErrorChunkList) const;
+
+	/** 
+	 * Acquires a set of chunks using the platform chunk layer, then calls the passed in callback
+	 *
+	 * @param AssetList				Asset Paths to get resources for
+	 * @param CompleteDelegate		Delegate called when chunks have been acquired or failed. If any chunks fail the entire operation is considered a failure
+	 * @param Priority				Priority to use when acquiring chunks
+	 */
+	virtual void AcquireResourcesForAssetList(const TArray<FStringAssetReference>& AssetList, FAssetManagerAcquireResourceDelegate CompleteDelegate, EChunkPriority::Type Priority = EChunkPriority::Immediate);
+
+	/** 
+	 * Acquires a set of chunks using the platform chunk layer, then calls the passed in callback. This will download all bundles of a primary asset
+	 *
+	 * @param PrimaryAssetList		Primary assets to get chunks for
+	 * @param CompleteDelegate		Delegate called when chunks have been acquired or failed. If any chunks fail the entire operation is considered a failure
+	 * @param Priority				Priority to use when acquiring chunks
+	 */
+	virtual void AcquireResourcesForPrimaryAssetList(const TArray<FPrimaryAssetId>& PrimaryAssetList, FAssetManagerAcquireResourceDelegate CompleteDelegate, EChunkPriority::Type Priority = EChunkPriority::Immediate);
 	
+	/** Returns the chunk download/install progress. AcquiredCount is the number of chunks that were requested and have already been insatlled, Requested includes both installed and pending */
+	virtual bool GetResourceAcquireProgress(int32& OutAcquiredCount, int32& OutRequestedCount) const;
+
 	// FUNCTIONS FOR MANAGEMENT/COOK RULES
 
 	/** Changes the default management rules for a specified type */
@@ -410,6 +443,12 @@ protected:
 	/** Returns true if path should be excluded from primary asset scans */
 	virtual bool IsPathExcludedFromScan(const FString& Path) const;
 
+	/** Call to start acquiring a list of chunks */
+	virtual void AcquireChunkList(const TArray<int32>& ChunkList, FAssetManagerAcquireResourceDelegate CompleteDelegate, EChunkPriority::Type Priority, TSharedPtr<FStreamableHandle> StalledHandle);
+
+	/** Called when a new chunk has been downloaded */
+	virtual void OnChunkDownloaded(uint32 ChunkId, bool bSuccess);
+
 #if WITH_EDITOR
 	/** Function used during creating Management references to decide when to recurse and set references */
 	virtual EAssetSetManagerResult::Type ShouldSetManager(const FAssetIdentifier& Manager, const FAssetIdentifier& Source, const FAssetIdentifier& Target, EAssetRegistryDependencyType::Type DependencyType, EAssetSetManagerFlags::Type Flags) const;
@@ -461,6 +500,25 @@ protected:
 	/** The streamable manager used for all primary asset loading */
 	FStreamableManager StreamableManager;
 
+	/** Defines a set of chunk installs that are waiting */
+	struct FPendingChunkInstall
+	{
+		/** Chunks we originally requested */
+		TArray<int32> RequestedChunks;
+
+		/** Chunks we are still waiting for */
+		TArray<int32> PendingChunks;
+
+		/** Stalled streamable handle waiting for this install, may be null */
+		TSharedPtr<FStreamableHandle> StalledStreamableHandle;
+
+		/** Delegate to call on completion, may be empty */
+		FAssetManagerAcquireResourceDelegate ManualCallback;
+	};
+
+	/** List of chunk installs that are being waited for */
+	TArray<FPendingChunkInstall> PendingChunkInstalls;
+
 	/** List of UObjects that are being kept from being GCd, derived from the asset type map. Arrays are currently more efficient than Sets */
 	UPROPERTY()
 	TArray<UObject*> ObjectReferenceList;
@@ -477,6 +535,18 @@ protected:
 	UPROPERTY()
 	bool bShouldUseSynchronousLoad;
 
+	/** True if we are loading from pak files */
+	UPROPERTY()
+	bool bIsLoadingFromPakFiles;
+
+	/** True if the chunk install interface should be queries before loading assets */
+	UPROPERTY()
+	bool bShouldAcquireMissingChunksOnLoad;
+
+	/** If true, DevelopmentCook assets will error when they are cooked */
+	UPROPERTY()
+	bool bOnlyCookProductionAssets;
+
 	/** True if we are currently in bulk scanning mode */
 	UPROPERTY()
 	bool bIsBulkScanning;
@@ -498,6 +568,9 @@ protected:
 	TMap<FString, FString> PrimaryAssetIdRedirects;
 	TMap<FName, FName> AssetPathRedirects;
 
+	/** Delegate bound to chunk install */
+	FDelegateHandle ChunkInstallDelegateHandle;
+
 private:
 	/** Per-type asset information, cannot be accessed by children as it is defined in CPP file */
 	TMap<FName, TSharedRef<FPrimaryAssetTypeData>> AssetTypeMap;
diff --git a/Engine/Source/Runtime/Engine/Classes/Engine/AssetManagerSettings.h b/Engine/Source/Runtime/Engine/Classes/Engine/AssetManagerSettings.h
--- a/Engine/Source/Runtime/Engine/Classes/Engine/AssetManagerSettings.h
+++ b/Engine/Source/Runtime/Engine/Classes/Engine/AssetManagerSettings.h
@@ -46,7 +46,7 @@ class ENGINE_API UAssetManagerSettings : public UDeveloperSettings
 	GENERATED_BODY()
 
 public:
-	UAssetManagerSettings() : bOnlyCookProductionAssets(false), bShouldGuessTypeAndNameInEditor(true) {}
+	UAssetManagerSettings() : bOnlyCookProductionAssets(false), bShouldGuessTypeAndNameInEditor(true), bShouldAcquireMissingChunksOnLoad(false) {}
 
 	/** List of asset types to scan at startup */
 	UPROPERTY(config, EditAnywhere, Category = "Asset Manager")
@@ -68,6 +68,10 @@ public:
 	UPROPERTY(config, EditAnywhere, Category = "Asset Manager")
 	bool bShouldGuessTypeAndNameInEditor;
 
+	/** If true, this will query the platform chunk install interface to request missing chunks for any requested primary asset loads */
+	UPROPERTY(config, EditAnywhere, Category = "Asset Manager")
+	bool bShouldAcquireMissingChunksOnLoad;
+
 	/** Redirect from Type:Name to Type:NameNew */
 	UPROPERTY(config, EditAnywhere, Category = "Redirects")
 	TArray<FAssetManagerRedirect> PrimaryAssetIdRedirects;
diff --git a/Engine/Source/Runtime/Engine/Classes/Engine/Brush.h b/Engine/Source/Runtime/Engine/Classes/Engine/Brush.h
--- a/Engine/Source/Runtime/Engine/Classes/Engine/Brush.h
+++ b/Engine/Source/Runtime/Engine/Classes/Engine/Brush.h
@@ -235,14 +235,6 @@ public:
 
 	// ABrush interface.
 
-	/** Returns the prepivot FVector of the RootComponent of this actor */
-	DEPRECATED(4.9, "This method has been superseded by AActor::GetPivotOffset.")
-	FVector GetPrePivot() const;
-
-	/** Set the PrePivot of the root component */
-	DEPRECATED(4.9, "This method has been superseded by AActor::SetPivotOffset.")
-	void SetPrePivot(const FVector& InPrePivot);
-
 	/** Figures out the best color to use for this brushes wireframe drawing.	*/
 	virtual FColor GetWireColor() const;
 
diff --git a/Engine/Source/Runtime/Engine/Classes/Engine/Engine.h b/Engine/Source/Runtime/Engine/Classes/Engine/Engine.h
--- a/Engine/Source/Runtime/Engine/Classes/Engine/Engine.h
+++ b/Engine/Source/Runtime/Engine/Classes/Engine/Engine.h
@@ -1516,6 +1516,22 @@ public:
 	// for IsInitialized()
 	bool bIsInitialized;
 
+private:
+
+	/** The last frame GC was run from ConditionalCollectGarbage to avoid multiple GCs in one frame */
+	uint64 LastGCFrame;
+
+	/** Time in seconds (game time so we respect time dilation) since the last time we purged references to pending kill objects */
+	float TimeSinceLastPendingKillPurge;
+
+	/** Whether a full purge has been triggered, so that the next GarbageCollect will do a full purge no matter what. */
+	bool bFullPurgeTriggered;
+
+	/** Whether we should delay GC for one frame to finish some pending operation */
+	bool bShouldDelayGarbageCollect; 
+
+public:
+
 	/**
 	 * Get the color to use for object selection
 	 */
@@ -2042,6 +2058,37 @@ public:
 
 	static void PreGarbageCollect();
 
+	/**
+	 *  Collect garbage once per frame driven by World ticks
+	 */
+	void ConditionalCollectGarbage();
+
+	/**
+	 *  Interface to allow WorldSettings to request immediate garbage collection
+	 */
+	void PerformGarbageCollectionAndCleanupActors();
+
+	/** Updates the timer between garbage collection such that at the next opportunity garbage collection will be run. */
+	void ForceGarbageCollection(bool bFullPurge = false);
+
+	/**
+	 *  Requests a one frame delay of Garbage Collection
+	 */
+	void DelayGarbageCollection();
+
+	/**
+	 * Updates the timer (as a one-off) that is used to trigger garbage collection; this should only be used for things
+	 * like performance tests, using it recklessly can dramatically increase memory usage and cost of the eventual GC.
+	 *
+	 * Note: Things that force a GC will still force a GC after using this method (and they will also reset the timer)
+	 */
+	void SetTimeUntilNextGarbageCollection(float MinTimeUntilNextPass);
+
+	/**
+	 * Returns the current desired time between garbage collection passes (not the time remaining)
+	 */
+	float GetTimeBetweenGarbageCollectionPasses() const;
+
 	/**
 	 * Returns whether we are running on a console platform or on the PC.
 	 * @param ConsoleType - if specified, only returns true if we're running on the specified platform
@@ -2163,10 +2210,6 @@ public:
 	/** @return the currently active audio device */
 	class FAudioDevice* GetActiveAudioDevice();
 
-	DEPRECATED(4.8, "GetAudioDevice is deprecated UEngine::GetMainAudioDevice instead.")
-	/** @return the main audio device. */
-	class FAudioDevice* GetAudioDevice();
-
 	/** @return whether we're currently running in split screen (more than one local player) */
 	bool IsSplitScreen(UWorld *InWorld);
 
diff --git a/Engine/Source/Runtime/Engine/Classes/Engine/EngineTypes.h b/Engine/Source/Runtime/Engine/Classes/Engine/EngineTypes.h
--- a/Engine/Source/Runtime/Engine/Classes/Engine/EngineTypes.h
+++ b/Engine/Source/Runtime/Engine/Classes/Engine/EngineTypes.h
@@ -3539,6 +3539,21 @@ struct FDirectoryPath
 };
 
 
+/**
+ * Reference to an editor collection of assets. This allows an editor-only picker UI
+ */
+USTRUCT(BlueprintType)
+struct FCollectionReference
+{
+	GENERATED_USTRUCT_BODY()
+
+	/**
+	 * Name of the collection
+	 */
+	UPROPERTY(EditAnywhere, BlueprintReadWrite, Category=CollectionReference)
+	FName CollectionName;
+};
+
 /**
 * This is used for redirecting old name to new name
 * We use manually parsing array, but that makes harder to modify from property setting
diff --git a/Engine/Source/Runtime/Engine/Classes/Engine/GameViewportClient.h b/Engine/Source/Runtime/Engine/Classes/Engine/GameViewportClient.h
--- a/Engine/Source/Runtime/Engine/Classes/Engine/GameViewportClient.h
+++ b/Engine/Source/Runtime/Engine/Classes/Engine/GameViewportClient.h
@@ -342,10 +342,6 @@ public:
 	/** @return Whether or not the main viewport is fullscreen or windowed. */
 	bool IsFullScreenViewport() const;
 
-	/** @return mouse position in game viewport coordinates (does not account for splitscreen) */
-	DEPRECATED(4.5, "Use GetMousePosition that returns a boolean if mouse is in window instead.")
-	FVector2D GetMousePosition() const;
-
 	/** @return mouse position in game viewport coordinates (does not account for splitscreen) */
 	bool GetMousePosition(FVector2D& MousePosition) const;
 
@@ -363,12 +359,6 @@ public:
 	 */
 	virtual ULocalPlayer* SetupInitialLocalPlayer(FString& OutError);
 
-	DEPRECATED(4.4, "CreatePlayer is deprecated UGameInstance::CreateLocalPlayer instead.")
-	virtual ULocalPlayer* CreatePlayer(int32 ControllerId, FString& OutError, bool bSpawnActor);
-
-	DEPRECATED(4.4, "RemovePlayer is deprecated UGameInstance::RemoveLocalPlayer instead.")
-	virtual bool RemovePlayer(class ULocalPlayer* ExPlayer);
-
 	/** @return Returns the splitscreen type that is currently being used */
 	FORCEINLINE ESplitScreenType::Type GetCurrentSplitscreenConfiguration() const
 	{
diff --git a/Engine/Source/Runtime/Engine/Classes/Engine/Level.h b/Engine/Source/Runtime/Engine/Classes/Engine/Level.h
--- a/Engine/Source/Runtime/Engine/Classes/Engine/Level.h
+++ b/Engine/Source/Runtime/Engine/Classes/Engine/Level.h
@@ -729,14 +729,6 @@ public:
 	 */
 	ENGINE_API static void BuildStreamingData(UWorld* World, ULevel* TargetLevel=NULL, UTexture2D* TargetTexture=NULL);
 
-	/**
-	* Deprecated. Returns the default brush for this level.
-	*
-	* @return		The default brush for this level.
-	*/
-	DEPRECATED(4.3, "GetBrush is deprecated use GetDefaultBrush instead.")
-	ENGINE_API ABrush* GetBrush() const;
-
 	/**
 	 * Returns the default brush for this level.
 	 *
diff --git a/Engine/Source/Runtime/Engine/Classes/Engine/LevelStreaming.h b/Engine/Source/Runtime/Engine/Classes/Engine/LevelStreaming.h
--- a/Engine/Source/Runtime/Engine/Classes/Engine/LevelStreaming.h
+++ b/Engine/Source/Runtime/Engine/Classes/Engine/LevelStreaming.h
@@ -75,7 +75,7 @@ DECLARE_DYNAMIC_MULTICAST_DELEGATE( FLevelStreamingVisibilityStatus );
  * interface for when a level should be streamed in and out of memory.
  *
  */
-UCLASS(abstract, editinlinenew,MinimalAPI,BlueprintType)
+UCLASS(abstract, editinlinenew,MinimalAPI,BlueprintType, Within=World)
 class ULevelStreaming : public UObject
 {
 	GENERATED_UCLASS_BODY()
diff --git a/Engine/Source/Runtime/Engine/Classes/Engine/MemberReference.h b/Engine/Source/Runtime/Engine/Classes/Engine/MemberReference.h
--- a/Engine/Source/Runtime/Engine/Classes/Engine/MemberReference.h
+++ b/Engine/Source/Runtime/Engine/Classes/Engine/MemberReference.h
@@ -308,6 +308,8 @@ public:
 	{
 		TFieldType* ReturnField = NULL;
 
+		bool bUseUpToDateClass = SelfScope && SelfScope->GetAuthoritativeClass() != SelfScope;
+
 		if(bSelfContext && SelfScope == NULL)
 		{
 			UE_LOG(LogBlueprint, Warning, TEXT("FMemberReference::ResolveMember (%s) bSelfContext == true, but no scope supplied!"), *MemberName.ToString() );
@@ -368,6 +370,9 @@ public:
 #endif
 			if(TargetScope != NULL)
 			{
+#if WITH_EDITOR
+				TargetScope = GetClassToUse(TargetScope, bUseUpToDateClass);
+#endif
 				// Find in target scope
 				ReturnField = FindField<TFieldType>(TargetScope, MemberName);
 
@@ -464,7 +469,9 @@ protected:
 
 	/** Init the field redirect map (if not already done) from .ini file entries */
 	ENGINE_API static void InitFieldRedirectMap();
-
+	
+	/** @return the 'real' generated class for blueprint classes, but only if we're already passed through CompileClassLayout */
+	ENGINE_API static UClass* GetClassToUse(UClass* InClass, bool bUseUpToDateClass);
 #endif
 
 public:
diff --git a/Engine/Source/Runtime/Engine/Classes/Engine/PrimaryAssetLabel.h b/Engine/Source/Runtime/Engine/Classes/Engine/PrimaryAssetLabel.h
--- a/Engine/Source/Runtime/Engine/Classes/Engine/PrimaryAssetLabel.h
+++ b/Engine/Source/Runtime/Engine/Classes/Engine/PrimaryAssetLabel.h
@@ -4,11 +4,12 @@
 
 #include "DataAsset.h"
 #include "Engine/AssetManagerTypes.h"
+#include "Engine/EngineTypes.h"
 #include "PrimaryAssetLabel.generated.h"
 
 /** A seed file that is created to mark referenced assets as part of this primary asset */
-UCLASS(MinimalAPI)
-class UPrimaryAssetLabel : public UPrimaryDataAsset
+UCLASS()
+class ENGINE_API UPrimaryAssetLabel : public UPrimaryDataAsset
 {
 	GENERATED_BODY()
 
@@ -16,6 +17,9 @@ public:
 	/** Bundle name used for directory-tagged assets */
 	static const FName DirectoryBundle;
 
+	/** Bundle used for collection assets */
+	static const FName CollectionBundle;
+
 	/** Constructor */
 	UPrimaryAssetLabel();
 		
@@ -39,17 +43,19 @@ public:
 	UPROPERTY(EditAnywhere, Category = PrimaryAssetLabel, meta = (AssetBundles = "Explicit", BlueprintBaseOnly))
 	TArray<TAssetSubclassOf<UObject>> ExplicitBlueprints;
 
+	/** Collection to load asset references out of */
+	UPROPERTY(EditAnywhere, Category = PrimaryAssetLabel)
+	FCollectionReference AssetCollection;
+
 	/** Set to editor only if this is not available in a cooked build */
 	virtual bool IsEditorOnly() const
 	{
 		return !bIsRuntimeLabel;
 	}
 
-	// TODO add collections
-
 #if WITH_EDITORONLY_DATA
 	/** This scans the class for AssetBundles metadata on asset properties and initializes the AssetBundleData with InitializeAssetBundlesFromMetadata */
-	ENGINE_API virtual void UpdateAssetBundleData();
+	virtual void UpdateAssetBundleData();
 #endif
 
 };
\ No newline at end of file
diff --git a/Engine/Source/Runtime/Engine/Classes/Engine/StreamableManager.h b/Engine/Source/Runtime/Engine/Classes/Engine/StreamableManager.h
--- a/Engine/Source/Runtime/Engine/Classes/Engine/StreamableManager.h
+++ b/Engine/Source/Runtime/Engine/Classes/Engine/StreamableManager.h
@@ -41,6 +41,12 @@ struct ENGINE_API FStreamableHandle : public TSharedFromThis<FStreamableHandle,
 		return !bCanceled && !bReleased;
 	}
 
+	/** If this handle is stalled and waiting for another event to occur before it is actually requested */
+	bool IsStalled() const
+	{
+		return bStalled;
+	}
+
 	/** Returns true if this is a combined handle that depends on child handles */
 	bool IsCombinedHandle() const
 	{
@@ -53,12 +59,21 @@ struct ENGINE_API FStreamableHandle : public TSharedFromThis<FStreamableHandle,
 		return DebugName;
 	}
 
+	/** Returns the streaming priority */
+	TAsyncLoadPriority GetPriority() const
+	{
+		return Priority;
+	}
+
 	/** Release this handle, called from normal gameplay code to indicate that the loaded assets are no longer needed. If called before completion will release on completion */
 	void ReleaseHandle();
 
 	/** Cancel a request, callable from within the manager or externally. Will stop delegate from being called */
 	void CancelHandle();
 
+	/** Tells a stalled handle to start it's actual request */
+	void StartStalledHandle();
+
 	/** Bind delegate that is called when load completes, only works if loading is in progress. This will overwrite any already bound delegate! */
 	bool BindCompleteDelegate(FStreamableDelegate NewDelegate);
 
@@ -108,7 +123,7 @@ private:
 	/** Callback when async load finishes, it's here so we can use a shared pointer for callback safety */
 	void AsyncLoadCallbackWrapper(const FName& PackageName, UPackage* LevelPackage, EAsyncLoadingResult::Type Result, FStringAssetReference TargetName);
 
-	/** Called on meta handle when a child handle has completed/cancelled */
+	/** Called on meta handle when a child handle has completed/canceled */
 	void UpdateCombinedHandle();
 
 	/** Call to call the update delegate if bound, will propagate to parents */
@@ -120,9 +135,12 @@ private:
 	/** True if this request was released, which will stop it from keeping hard GC references */
 	bool bReleased;
 
-	/** True if this request was explicitly cancelled, which stops it from calling the completion delegate and immediately releases it */
+	/** True if this request was explicitly canceled, which stops it from calling the completion delegate and immediately releases it */
 	bool bCanceled;
 
+	/** True if this handle has been created but not yet actually requested. This handle is probably waiting for a resource like a chunk to be available */
+	bool bStalled;
+
 	/** If true, this handle will be released when it finishes loading */
 	bool bReleaseWhenLoaded;
 
@@ -140,7 +158,10 @@ private:
 
 	/** Name of this handle, passed in by caller to help in debugging */
 	FString DebugName;
-	
+
+	/** The async priority for this request */
+	TAsyncLoadPriority Priority;
+
 	/** How many FStreamables is this waiting on to finish loading */
 	int32 StreamablesLoading;
 
@@ -160,8 +181,10 @@ private:
 		: bLoadCompleted(false)
 		, bReleased(false)
 		, bCanceled(false)
+		, bStalled(false)
 		, bReleaseWhenLoaded(false)
 		, bIsCombinedHandle(false)
+		, Priority(0)
 		, StreamablesLoading(0)
 		, OwningManager(nullptr)
 	{
@@ -183,14 +206,15 @@ struct ENGINE_API FStreamableManager : public FGCObject
 	 * @param DelegateToCall		Delegate to call when load finishes. Will be called on the next tick if asset is already loaded, or many seconds later
 	 * @param Priority				Priority to pass to the streaming system, higher priority will be loaded first
 	 * @param bManageActiveHandle	If true, the manager will keep the streamable handle active until explicitly released
+	 * @param bStartStalled			If true, the handle will start in a stalled state and will not attempt to actually async load until StartStalledHandle is called on it
 	 * @param DebugName				Name of this handle, will be reported in debug tools
 	 */
-	TSharedPtr<FStreamableHandle> RequestAsyncLoad(const TArray<FStringAssetReference>& TargetsToStream, FStreamableDelegate DelegateToCall = FStreamableDelegate(), TAsyncLoadPriority Priority = DefaultAsyncLoadPriority, bool bManageActiveHandle = false, const FString& DebugName = TEXT("RequestAsyncLoad ArrayDelegate"));
-	TSharedPtr<FStreamableHandle> RequestAsyncLoad(const FStringAssetReference& TargetToStream, FStreamableDelegate DelegateToCall = FStreamableDelegate(), TAsyncLoadPriority Priority = DefaultAsyncLoadPriority, bool bManageActiveHandle = false, const FString& DebugName = TEXT("RequestAsyncLoad SingleDelegate"));
+	TSharedPtr<FStreamableHandle> RequestAsyncLoad(const TArray<FStringAssetReference>& TargetsToStream, FStreamableDelegate DelegateToCall = FStreamableDelegate(), TAsyncLoadPriority Priority = DefaultAsyncLoadPriority, bool bManageActiveHandle = false, bool bStartStalled = false, const FString& DebugName = TEXT("RequestAsyncLoad ArrayDelegate"));
+	TSharedPtr<FStreamableHandle> RequestAsyncLoad(const FStringAssetReference& TargetToStream, FStreamableDelegate DelegateToCall = FStreamableDelegate(), TAsyncLoadPriority Priority = DefaultAsyncLoadPriority, bool bManageActiveHandle = false, bool bStartStalled = false, const FString& DebugName = TEXT("RequestAsyncLoad SingleDelegate"));
 
 	/** Lambda Wrappers. Be aware that Callback may go off multiple seconds in the future */
-	TSharedPtr<FStreamableHandle> RequestAsyncLoad(const TArray<FStringAssetReference>& TargetsToStream, TFunction<void()>&& Callback, TAsyncLoadPriority Priority = DefaultAsyncLoadPriority, bool bManageActiveHandle = false, const FString& DebugName = TEXT("RequestAsyncLoad ArrayLambda"));
-	TSharedPtr<FStreamableHandle> RequestAsyncLoad(const FStringAssetReference& TargetToStream, TFunction<void()>&& Callback, TAsyncLoadPriority Priority = DefaultAsyncLoadPriority, bool bManageActiveHandle = false, const FString& DebugName = TEXT("RequestAsyncLoad SingleLambda"));
+	TSharedPtr<FStreamableHandle> RequestAsyncLoad(const TArray<FStringAssetReference>& TargetsToStream, TFunction<void()>&& Callback, TAsyncLoadPriority Priority = DefaultAsyncLoadPriority, bool bManageActiveHandle = false, bool bStartStalled = false, const FString& DebugName = TEXT("RequestAsyncLoad ArrayLambda"));
+	TSharedPtr<FStreamableHandle> RequestAsyncLoad(const FStringAssetReference& TargetToStream, TFunction<void()>&& Callback, TAsyncLoadPriority Priority = DefaultAsyncLoadPriority, bool bManageActiveHandle = false, bool bStartStalled = false, const FString& DebugName = TEXT("RequestAsyncLoad SingleLambda"));
 
 	/** 
 	 * Synchronously load a set of assets, and return a handle.
@@ -290,6 +314,7 @@ private:
 	friend FStreamableHandle;
 
 	void RemoveReferencedAsset(const FStringAssetReference& Target, TSharedRef<FStreamableHandle> Handle);
+	void StartHandleRequests(TSharedRef<FStreamableHandle> Handle);
 	FStringAssetReference ResolveRedirects(const FStringAssetReference& Target) const;
 	void FindInMemory(FStringAssetReference& InOutTarget, struct FStreamable* Existing);
 	struct FStreamable* FindStreamable(const FStringAssetReference& Target) const;
diff --git a/Engine/Source/Runtime/Engine/Classes/Engine/World.h b/Engine/Source/Runtime/Engine/Classes/Engine/World.h
--- a/Engine/Source/Runtime/Engine/Classes/Engine/World.h
+++ b/Engine/Source/Runtime/Engine/Classes/Engine/World.h
@@ -1096,15 +1096,6 @@ public:
 	/** Counter for allocating game- unique controller player numbers															*/
 	int32										PlayerNum;
 
-	/** Time in seconds (game time so we respect time dilation) since the last time we purged references to pending kill objects */
-	float										TimeSinceLastPendingKillPurge;
-
-	/** Whether a full purge has been triggered, so that the next GarbageCollect will do a full purge no matter what.			*/
-	bool										FullPurgeTriggered;
-	
-	/** Whether we should delay GC for one frame to finish some pending operation												*/
-	bool										bShouldDelayGarbageCollect;
-
 	/** Whether world object has been initialized via Init()																	*/
 	bool										bIsWorldInitialized;
 	
@@ -1917,8 +1908,6 @@ private:
 
 public:
 
-	DEPRECATED(4.3, "GetBrush is deprecated use GetDefaultBrush instead.")
-	ABrush* GetBrush() const;
 	/** 
 	 * Returns the default brush for the persistent level.
 	 * This is usually the 'builder brush' for editor builds, undefined for non editor instances and may be NULL.
@@ -2229,17 +2218,6 @@ public:
 	 */
 	void FlushLevelStreaming(EFlushLevelStreamingType FlushType = EFlushLevelStreamingType::Full);
 
-	/**
-	 * [Deprecated] Flushes level streaming in blocking fashion and returns when all levels are loaded/ visible/ hidden
-	 * so further calls to UpdateLevelStreaming won't do any work unless state changes. Basically blocks
-	 * on all async operation like updating components.
-	 *
-	 * @param FlushType					Whether to only flush level visibility operations
-	 * @param ExcludeType				Exclude packages of this type from flushing
-	 */
-	DEPRECATED(4.9, "FlushLevelStreaming override that takes ExcludeType parameter is deprecated.")
-	void FlushLevelStreaming(EFlushLevelStreamingType FlushType, FName ExcludeType);
-
 	/**
 	 * Triggers a call to ULevel::BuildStreamingData(this,NULL,NULL) within a few seconds.
 	 */
@@ -2369,11 +2347,13 @@ public:
 	/**
 	 *  Interface to allow WorldSettings to request immediate garbage collection
 	 */
+	DEPRECATED(4.18, "Use GEngine->PerformGarbageCollectionAndCleanupActors instead.")
 	void PerformGarbageCollectionAndCleanupActors();
 
 	/**
 	 *  Requests a one frame delay of Garbage Collection
 	 */
+	DEPRECATED(4.18, "Use GEngine->DelayGarbageCollection instead.")
 	void DelayGarbageCollection();
 
 	/**
@@ -2382,15 +2362,15 @@ public:
 	 *
 	 * Note: Things that force a GC will still force a GC after using this method (and they will also reset the timer)
 	 */
+	DEPRECATED(4.18, "Use GEngine->SetTimeUntilNextGarbageCollection instead.")
 	void SetTimeUntilNextGarbageCollection(float MinTimeUntilNextPass);
 
 	/**
 	 * Returns the current desired time between garbage collection passes (not the time remaining)
 	 */
+	DEPRECATED(4.18, "Call GEngine->GetTimeBetweenGarbageCollectionPasses instead")
 	float GetTimeBetweenGarbageCollectionPasses() const;
 
-protected:
-
 	/**
 	 *	Remove NULL entries from actor list. Only does so for dynamic actors to avoid resorting. 
 	 *	In theory static actors shouldn't be deleted during gameplay.
@@ -2753,33 +2733,6 @@ public:
 	{
 		return CastChecked<T>(SpawnActorAbsolute(Class, Transform, SpawnParameters), ECastCheckedType::NullAllowed);
 	}
-	/**
-	* Spawns given class and returns class T pointer, forcibly sets world position. WILL NOT run Construction Script of Blueprints 
-	* to give caller an opportunity to set parameters beforehand.  Caller is responsible for invoking construction
-	* manually by calling UGameplayStatics::FinishSpawningActor (see AActor::OnConstruction).
-	*/
-	template< class T >
-	DEPRECATED(4.9, "This version of SpawnActorDeferred is deprecated. Please use the version that takes an FTransform and ESpawnActorCollisionHandlingMethod.")
-	T* SpawnActorDeferred(
-		UClass* Class,
-		FVector const& Location,
-		FRotator const& Rotation,
-		AActor* Owner = nullptr,
-		APawn* Instigator = nullptr,
-		bool bNoCollisionFail = false
-		)
-	{
-		if( Owner )
-		{
-			check(this==Owner->GetWorld());
-		}
-		FActorSpawnParameters SpawnInfo;
-		SpawnInfo.SpawnCollisionHandlingOverride = bNoCollisionFail ? ESpawnActorCollisionHandlingMethod::AlwaysSpawn : ESpawnActorCollisionHandlingMethod::AdjustIfPossibleButDontSpawnIfColliding;
-		SpawnInfo.Owner = Owner;
-		SpawnInfo.Instigator = Instigator;
-		SpawnInfo.bDeferConstruction = true;
-		return (Class != NULL) ? Cast<T>(SpawnActor(Class, &Location, &Rotation, SpawnInfo )) : NULL;
-	}
 
 	/**
 	 * Spawns given class and returns class T pointer, forcibly sets world transform (note this allows scale as well). WILL NOT run Construction Script of Blueprints 
@@ -3113,6 +3066,7 @@ public:
 	int32 GetDetailMode();
 
 	/** Updates the timer between garbage collection such that at the next opportunity garbage collection will be run. */
+	DEPRECATED(4.18, "Call GEngine->ForceGarbageCollection instead")
 	void ForceGarbageCollection( bool bFullPurge = false );
 
 	/** asynchronously loads the given levels in preparation for a streaming map transition.
diff --git a/Engine/Source/Runtime/Engine/Classes/GameFramework/Actor.h b/Engine/Source/Runtime/Engine/Classes/GameFramework/Actor.h
--- a/Engine/Source/Runtime/Engine/Classes/GameFramework/Actor.h
+++ b/Engine/Source/Runtime/Engine/Classes/GameFramework/Actor.h
@@ -1268,9 +1268,6 @@ public:
 	UFUNCTION(BlueprintCallable, Category="Utilities|Time")
 	float GetActorTimeDilation() const;
 
-	DEPRECATED(4.5, "Actor::SetTickPrerequisite() will be removed, use AddTickPrerequisiteActor().")
-	void SetTickPrerequisite(AActor* PrerequisiteActor);
-
 	/** Make this actor tick after PrerequisiteActor. This only applies to this actor's tick function; dependencies for owned components must be set up separately if desired. */
 	UFUNCTION(BlueprintCallable, Category="Utilities", meta=(Keywords = "dependency"))
 	virtual void AddTickPrerequisiteActor(AActor* PrerequisiteActor);
@@ -1640,10 +1637,6 @@ public:
 	 */
 	virtual USceneComponent* GetDefaultAttachComponent() const { return GetRootComponent(); }
 
-	/** Returns this actor's root component cast to a primitive component */
-	DEPRECATED(4.5, "Use GetRootComponent() and cast manually if needed")
-	class UPrimitiveComponent* GetRootPrimitiveComponent() const;
-
 	/**
 	 * Sets root component to be the specified component.  NewRootComponent's owner should be this actor.
 	 * @return true if successful
@@ -1872,9 +1865,6 @@ public:
 	 */
 	virtual float GetNetPriority(const FVector& ViewPos, const FVector& ViewDir, class AActor* Viewer, AActor* ViewTarget, UActorChannel* InChannel, float Time, bool bLowBandwidth);
 
-	DEPRECATED(4.8, "GetNetPriority now takes a ViewTarget, please override that version.")
-	virtual float GetNetPriority(const FVector& ViewPos, const FVector& ViewDir, class APlayerController* Viewer, UActorChannel* InChannel, float Time, bool bLowBandwidth);
-
 	/**
 	 * Similar to GetNetPriority, but will only be used for prioritizing actors while recording a replay.
 	 *
@@ -1890,9 +1880,6 @@ public:
 	/** Returns true if the actor should be dormant for a specific net connection. Only checked for DORM_DormantPartial */
 	virtual bool GetNetDormancy(const FVector& ViewPos, const FVector& ViewDir, class AActor* Viewer, AActor* ViewTarget, UActorChannel* InChannel, float Time, bool bLowBandwidth);
 
-	DEPRECATED(4.8, "GetNetDormancy changed from PlayerController parameter to Actor parameter, please override that version.")
-	virtual bool GetNetDormancy(const FVector& ViewPos, const FVector& ViewDir, class APlayerController* Viewer, AActor* ViewTarget, UActorChannel* InChannel, float Time, bool bLowBandwidth);
-
 	/** 
 	 * Allows for a specific response from the actor when the actor channel is opened (client side)
 	 * @param InBunch Bunch received at time of open
@@ -1996,10 +1983,6 @@ public:
 	UFUNCTION()
 	virtual void OnRep_ReplicatedMovement();
 
-	/** Update and smooth location, not called for simulated physics! */
-	DEPRECATED(4.4, "PostNetReceiveLocation() has been replaced by PostNetReceiveLocationAndRotation().")
-	virtual void PostNetReceiveLocation() {}
-
 	/** Update location and rotation from ReplicatedMovement. Not called for simulated physics! */
 	virtual void PostNetReceiveLocationAndRotation();
 
@@ -2108,9 +2091,6 @@ public:
 	*/
 	virtual bool IsReplayRelevantFor(const AActor* RealViewer, const AActor* ViewTarget, const FVector& SrcLocation, const float CullDistanceSquared) const;
 
-	DEPRECATED(4.8, "IsNetRelevantFor changed from PlayerController parameter to Actor parameter, please override that version.")
-	virtual bool IsNetRelevantFor(const APlayerController* RealViewer, const AActor* ViewTarget, const FVector& SrcLocation) const;
-
 	/**
 	 * Check if this actor is the owner when doing relevancy checks for actors marked bOnlyRelevantToOwner
 	 *
@@ -2950,12 +2930,6 @@ public:
 	/** A fence to track when the primitive is detached from the scene in the rendering thread. */
 	FRenderCommandFence DetachFence;
 
-// DEPRECATED FUNCTIONS
-
-	/** Get the class of this Actor. */
-	DEPRECATED(4.8, "Use GetClass() instead of GetActorClass()")
-	class UClass* GetActorClass() const;
-
 private:
 
 	// Helper that already assumes the Hit info is reversed, and avoids creating a temp FHitResult if possible.
diff --git a/Engine/Source/Runtime/Engine/Classes/GameFramework/Character.h b/Engine/Source/Runtime/Engine/Classes/GameFramework/Character.h
--- a/Engine/Source/Runtime/Engine/Classes/GameFramework/Character.h
+++ b/Engine/Source/Runtime/Engine/Classes/GameFramework/Character.h
@@ -149,14 +149,6 @@ namespace MovementBaseUtility
 
 	/** Get the transforms for the given MovementBase, optionally at the location of a bone. Returns false if MovementBase is NULL, or if BoneName is not a valid bone. */
 	ENGINE_API bool GetMovementBaseTransform(const UPrimitiveComponent* MovementBase, const FName BoneName, FVector& OutLocation, FQuat& OutQuat);
-
-
-	// Deprecated function
-	DEPRECATED(4.4, "MovementBaseUtility::UseRelativePosition() is deprecated, use MovementBaseUtility::UseRelativeLocation() instead.")
-	FORCEINLINE bool UseRelativePosition(const UPrimitiveComponent* MovementBase)
-	{
-		return UseRelativeLocation(MovementBase);
-	}
 }
 
 /** Struct to hold information about the "base" object the character is standing on. */
@@ -210,21 +202,8 @@ struct FBasedMovementInfo
 	{
 		return (MovementBase == NULL) && bServerHasBaseComponent;
 	}
-
-
-	// Deprecated function
-	DEPRECATED(4.4, "FBasedMovementInfo::HasRelativePosition() is deprecated, use FBasedMovementInfo::HasRelativeLocation() instead.")
-	FORCEINLINE bool HasRelativePosition() const
-	{
-		return HasRelativeLocation();
-	}
 };
 
-// Allow the old name to continue to work for one release
-DEPRECATED(4.4, "FRepRelativeMovement has been renamed to FBasedMovementInfo")
-typedef FBasedMovementInfo FRepRelativeMovement;
-
-
 
 /**
  * Characters are Pawns that have a mesh, collision, and built-in movement logic.
@@ -552,9 +531,6 @@ public:
 	UFUNCTION(BlueprintCallable, Category="Pawn|Character")
 	virtual bool IsJumpProvidingForce() const;
 
-	DEPRECATED(4.4, "IsJumping() has been renamed IsJumpProvidingForce().")
-	virtual bool IsJumping() const { return IsJumpProvidingForce(); }
-
 	/** Play Animation Montage on the character mesh **/
 	UFUNCTION(BlueprintCallable, Category=Animation)
 	virtual float PlayAnimMontage(class UAnimMontage* AnimMontage, float InPlayRate = 1.f, FName StartSectionName = NAME_None);
@@ -565,17 +541,7 @@ public:
 
 	/** Return current playing Montage **/
 	UFUNCTION(BlueprintCallable, Category=Animation)
-	class UAnimMontage * GetCurrentMontage();
-
-protected:
-
-	/**
-	 * Player Jumped. Called internally when a jump has been detected because bPressedJump was true.
-	 * @param	bReplayingMoves: true if this is being done as part of replaying moves on a locally controlled client after a server correction.
-	 * @return	True if the jump was allowed by CanJump() and if CharacterMovement->Jump() succeeded.
-	 **/ 
-	DEPRECATED(4.5, "Character::DoJump() will be removed, use CharacterMovementComponent::DoJump() instead.")
-	virtual bool DoJump(bool bReplayingMoves);
+	class UAnimMontage* GetCurrentMontage();
 
 public:
 
@@ -652,16 +618,6 @@ public:
 	void OnWalkingOffLedge(const FVector& PreviousFloorImpactNormal, const FVector& PreviousFloorContactNormal, const FVector& PreviousLocation, float TimeDelta);
 	virtual void OnWalkingOffLedge_Implementation(const FVector& PreviousFloorImpactNormal, const FVector& PreviousFloorContactNormal, const FVector& PreviousLocation, float TimeDelta);
 
-
-	// Deprecated, use the new version that takes multiple parameters.
-	DEPRECATED(4.8, "OnWalkingOffLedge() is deprecated and will not be called, use the new version that takes multiple parameters.")
-	void OnWalkingOffLedge() {}
-
-	// Deprecated, use the new version that takes multiple parameters.
-	DEPRECATED(4.8, "OnWalkingOffLedge_Implementation() is deprecated and will not be called, use the new version that takes multiple parameters.")
-	virtual void OnWalkingOffLedge_Implementation() {}
-
-
 	/** Called when pawn's movement is blocked
 		@PARAM Impact describes the blocking hit. */
 	virtual void MoveBlockedBy(const FHitResult& Impact) {};
@@ -759,9 +715,6 @@ public:
 	UPROPERTY(BlueprintAssignable, Category="Pawn|Character")
 	FCharacterMovementUpdatedSignature OnCharacterMovementUpdated;
 
-	DEPRECATED(4.4, "NotifyLanded() has been renamed ShouldNotifyLanded() instead.")
-	virtual bool NotifyLanded(const struct FHitResult& Hit) { return ShouldNotifyLanded(Hit); }
-
 	/** 
 	 * Returns true if the Landed() event should be called. Used by CharacterMovement to prevent notifications while playing back network moves.
 	 */
@@ -783,9 +736,6 @@ public:
 	 */
 	virtual float GetJumpMaxHoldTime() const;
 
-	DEPRECATED(4.4, "UpdateFromCompressedFlags has moved to UCharacterMovementComponent")
-	virtual void UpdateFromCompressedFlags(uint8 Flags) { }
-
 public:
 
 	UFUNCTION(reliable, client)
diff --git a/Engine/Source/Runtime/Engine/Classes/GameFramework/CharacterMovementComponent.h b/Engine/Source/Runtime/Engine/Classes/GameFramework/CharacterMovementComponent.h
--- a/Engine/Source/Runtime/Engine/Classes/GameFramework/CharacterMovementComponent.h
+++ b/Engine/Source/Runtime/Engine/Classes/GameFramework/CharacterMovementComponent.h
@@ -1137,9 +1137,6 @@ public:
 	*/
 	virtual void PerformAirControlForPathFollowing(FVector Direction, float ZDiff);
 
-	DEPRECATED(4.5, "PerformAirControl has been deprecated, use PerformAirControlForPathFollowing instead.")
-	virtual void PerformAirControl(FVector Direction, float ZDiff);
-
 	/** Transition from walking to falling */
 	virtual void StartFalling(int32 Iterations, float remainingTime, float timeTick, const FVector& Delta, const FVector& subLoc);
 
@@ -1425,21 +1422,6 @@ protected:
 	 */
 	virtual float BoostAirControl(float DeltaTime, float TickAirControl, const FVector& FallAcceleration);
 
-	/**
-	 * (DEPRECATED) Checks if air control will cause the player collision shape to hit something given the current location.
-	 *
-	 * @param DeltaTime			Time step for the current update.
-	 * @param AdditionalTime	Time to look ahead further, applying acceleration and gravity.
-	 * @param FallVelocity		Velocity to test with.
-	 * @param FallAcceleration	Acceleration used during movement.
-	 * @param Gravity			Gravity to apply to any additional time.
-	 * @param OutHitResult		Result of impact, valid if this function returns true.
-	 * @return True if there is an impact, in which case OutHitResult contains the result of that impact.
-	 * @see GetAirControl()
-	 */
-	DEPRECATED(4.9, "FindAirControlImpact is no longer used by engine code.")
-	virtual bool FindAirControlImpact(float DeltaTime, float AdditionalTime, const FVector& FallVelocity, const FVector& FallAcceleration, const FVector& Gravity, FHitResult& OutHitResult);
-
 	/**
 	 * Limits the air control to use during falling movement, given an impact while falling.
 	 * This function is used internally by PhysFalling().
@@ -1726,12 +1708,6 @@ protected:
 	 */
 	virtual FVector ComputeSlideVector(const FVector& Delta, const float Time, const FVector& Normal, const FHitResult& Hit) const override;
 
-	/**
-	 * (DEPRECATED) Given an upward impact on the top of the capsule, allows calculation of a different movement delta.
-	 */
-	DEPRECATED(4.6, "AdjustUpperHemisphereImpact() is deprecated, for custom behavior override ComputeSlideVector() instead.")
-	virtual FVector AdjustUpperHemisphereImpact(const FVector& Delta, const FHitResult& Hit) const;
-
 	/** 
 	 * Limit the slide vector when falling if the resulting slide might boost the character faster upwards.
 	 * @param SlideResult:	Vector of movement for the slide (usually the result of ComputeSlideVector)
diff --git a/Engine/Source/Runtime/Engine/Classes/GameFramework/Controller.h b/Engine/Source/Runtime/Engine/Classes/GameFramework/Controller.h
--- a/Engine/Source/Runtime/Engine/Classes/GameFramework/Controller.h
+++ b/Engine/Source/Runtime/Engine/Classes/GameFramework/Controller.h
@@ -183,9 +183,6 @@ public:
 	UFUNCTION(BlueprintCallable, Category="Pawn", meta=(DisplayName="Get Controlled Pawn"))
 	APawn* K2_GetPawn() const;
 
-	DEPRECATED(4.4, "Use GetPawn() instead of GetControlledPawn()")
-	APawn* GetControlledPawn() const { return GetPawn(); }
-
 public:
 
 	virtual void GetLifetimeReplicatedProps(TArray<FLifetimeProperty>& OutLifetimeProps) const override;
@@ -260,9 +257,6 @@ public:
 	 */
 	virtual void PawnPendingDestroy(APawn* inPawn);
 
-	DEPRECATED(4.5, "NotifyKilled is unused in engine code and will be removed from AController.")
-	virtual void NotifyKilled(AController* Killer, AController* KilledPlayer, APawn* KilledPawn, const UDamageType* DamageType) {}
-
 	/** Called when this controller instigates ANY damage */
 	virtual void InstigatedAnyDamage(float Damage, const class UDamageType* DamageType, class AActor* DamagedActor, class AActor* DamageCauser);
 
diff --git a/Engine/Source/Runtime/Engine/Classes/GameFramework/MovementComponent.h b/Engine/Source/Runtime/Engine/Classes/GameFramework/MovementComponent.h
--- a/Engine/Source/Runtime/Engine/Classes/GameFramework/MovementComponent.h
+++ b/Engine/Source/Runtime/Engine/Classes/GameFramework/MovementComponent.h
@@ -221,10 +221,6 @@ public:
 	 */
 	virtual bool ShouldSkipUpdate(float DeltaTime) const;
 
-	DEPRECATED(4.5, "SkipUpdate has been renamed to ShouldSkipUpdate")
-	virtual bool SkipUpdate(float DeltaTime) { return ShouldSkipUpdate(DeltaTime); }
-
-
 	/** @return PhysicsVolume this MovementComponent is using, or the world's default physics volume if none. **/
 	UFUNCTION(BlueprintCallable, Category="Components|Movement")
 	virtual APhysicsVolume* GetPhysicsVolume() const;
diff --git a/Engine/Source/Runtime/Engine/Classes/GameFramework/NavMovementComponent.h b/Engine/Source/Runtime/Engine/Classes/GameFramework/NavMovementComponent.h
--- a/Engine/Source/Runtime/Engine/Classes/GameFramework/NavMovementComponent.h
+++ b/Engine/Source/Runtime/Engine/Classes/GameFramework/NavMovementComponent.h
@@ -158,20 +158,6 @@ public:
 	/** @return true if currently flying (moving through a non-fluid volume without resting on the ground) */
 	UFUNCTION(BlueprintCallable, Category="AI|Components|NavMovement")
 	virtual bool IsFlying() const;
-
-	//----------------------------------------------------------------------//
-	// DEPRECATED
-	//----------------------------------------------------------------------//
-public:
-	DEPRECATED(4.7, "This function is deprecated. Please use GetNavAgentPropertiesRef instead.")
-	const FNavAgentProperties* GetNavAgentProperties() const;
-	DEPRECATED(4.7, "This function is deprecated. Please use GetNavAgentPropertiesRef instead.")
-	FNavAgentProperties* GetNavAgentProperties();
-
-	DEPRECATED(4.8, "This function is deprecated. Please use UpdateNavAgent version that's accepring a reference instead.")
-	void UpdateNavAgent(AActor* InOwner);
-	DEPRECATED(4.8, "This function is deprecated. Please use UpdateNavAgent version that's accepring a reference instead.")
-	void UpdateNavAgent(UCapsuleComponent* CapsuleComponent);
 };
 
 
diff --git a/Engine/Source/Runtime/Engine/Classes/GameFramework/Pawn.h b/Engine/Source/Runtime/Engine/Classes/GameFramework/Pawn.h
--- a/Engine/Source/Runtime/Engine/Classes/GameFramework/Pawn.h
+++ b/Engine/Source/Runtime/Engine/Classes/GameFramework/Pawn.h
@@ -33,7 +33,7 @@ ENGINE_API DECLARE_LOG_CATEGORY_EXTERN(LogDamage, Warning, All);
  *
  * @see https://docs.unrealengine.com/latest/INT/Gameplay/Framework/Pawn/
  */
-UCLASS(config=Game, BlueprintType, Blueprintable, hideCategories=(Navigation), meta=(ShortTooltip="A Pawn is an actor that can be 'possessed' and receieve input from a controller."))
+UCLASS(config=Game, BlueprintType, Blueprintable, hideCategories=(Navigation), meta=(ShortTooltip="A Pawn is an actor that can be 'possessed' and receive input from a controller."))
 class ENGINE_API APawn : public AActor, public INavAgentInterface
 {
 	GENERATED_BODY()
@@ -462,10 +462,6 @@ public:
 	UFUNCTION(BlueprintCallable, Category="Pawn", meta=(DeprecatedFunction, DeprecationMessage="Use Character.LaunchCharacter instead"))
 	void LaunchPawn(FVector LaunchVelocity, bool bXYOverride, bool bZOverride);
 
-	/** Internal function meant for use only within Pawn or by a PawnMovementComponent. */
-	DEPRECATED(4.5, "Internal_GetMovementInputVector is deprecated, use Internal_GetPendingMovementInputVector() or Internal_GetLastMovementInputVector() instead.")
-	inline FVector Internal_GetMovementInputVector() const { return Internal_GetPendingMovementInputVector(); }
-
 	/** (Deprecated) Return the input vector in world space. */
 	DEPRECATED(4.5, "GetMovementInputVector() has been deprecated, use either GetPendingMovementInputVector() or GetLastMovementInputVector().")
 	FVector GetMovementInputVector() const;
diff --git a/Engine/Source/Runtime/Engine/Classes/GameFramework/PlayerInput.h b/Engine/Source/Runtime/Engine/Classes/GameFramework/PlayerInput.h
--- a/Engine/Source/Runtime/Engine/Classes/GameFramework/PlayerInput.h
+++ b/Engine/Source/Runtime/Engine/Classes/GameFramework/PlayerInput.h
@@ -312,10 +312,6 @@ public:
 	UFUNCTION(exec)
 	void SetMouseSensitivity(const float Sensitivity);
 
-	/** Exec function to return the mouse sensitivity to its default value */
-	DEPRECATED(4.8, "SetMouseSensitivityToDefault is deprecated, use SetAxisProperties instead")
-	void SetMouseSensitivityToDefault();
-
 	/** Exec function to add a debug exec command */
 	UFUNCTION(exec)
 	void SetBind(FName BindName, const FString& Command);
@@ -333,10 +329,6 @@ public:
 	UFUNCTION(exec)
 	void InvertAxisKey(const FKey AxisKey);
 
-	/** Backwards compatibility exec function for people used to it instead of using InvertAxisKey */
-	DEPRECATED(4.8, "InvertMouse is deprecated, use InvertAxisKey(EKeys::MouseY); instead")
-	void InvertMouse();
-
 	/** Exec function to invert an axis mapping */
 	UFUNCTION(exec)
 	void InvertAxis(const FName AxisName);
@@ -523,6 +515,9 @@ public:
 	/** Returns the list of keys mapped to the specified Action Name */
 	const TArray<FInputActionKeyMapping>& GetKeysForAction(const FName ActionName);
 
+	/** Returns the list of keys mapped to the specified Axis Name */
+	const TArray<FInputAxisKeyMapping>& GetKeysForAxis(const FName AxisName);
+
 	static const TArray<FInputActionKeyMapping>& GetEngineDefinedActionMappings() { return EngineDefinedActionMappings; }
 	static const TArray<FInputAxisKeyMapping>& GetEngineDefinedAxisMappings() { return EngineDefinedAxisMappings; }
 
@@ -608,6 +603,9 @@ private:
 	/** Static empty array to be able to return from GetKeysFromAction when there are no keys mapped to the requested action name */
 	static const TArray<FInputActionKeyMapping> NoKeyMappings;
 
+	/** Static empty array to be able to return from GetKeysFromAxis when there are no axis mapped to the requested axis name */
+	static const TArray<FInputAxisKeyMapping> NoAxisMappings;
+
 	/** Action Mappings defined by engine systems that cannot be remapped by users */
 	static TArray<FInputActionKeyMapping> EngineDefinedActionMappings;
 
diff --git a/Engine/Source/Runtime/Engine/Classes/GameFramework/ProjectileMovementComponent.h b/Engine/Source/Runtime/Engine/Classes/GameFramework/ProjectileMovementComponent.h
--- a/Engine/Source/Runtime/Engine/Classes/GameFramework/ProjectileMovementComponent.h
+++ b/Engine/Source/Runtime/Engine/Classes/GameFramework/ProjectileMovementComponent.h
@@ -164,10 +164,6 @@ class ENGINE_API UProjectileMovementComponent : public UMovementComponent
 
 	bool ShouldApplyGravity() const { return ProjectileGravityScale != 0.f; }
 
-	/** DEPRECATED: renamed to ComputeVelocity(). */
-	DEPRECATED(4.6, "CalculateVelocity() is deprecated, use ComputeVelocity() instead.")
-	virtual FVector CalculateVelocity(FVector OldVelocity, float DeltaTime, bool bGravityEnabled_UNUSED) const;
-
 	/**
 	 * Given an initial velocity and a time step, compute a new velocity.
 	 * Default implementation applies the result of ComputeAcceleration() to velocity.
@@ -301,9 +297,6 @@ public:
 	virtual float GetGravityZ() const override;
 
 protected:
-	DEPRECATED(4.9, "GetEffectiveGravityZ() is deprecated, use GetGravityZ() instead, or UMovementComponent::GetGravityZ() if you want only the unmodified gravity.")
-	float GetEffectiveGravityZ() const;
-
 	/** Minimum delta time considered when ticking. Delta times below this are not considered. This is a very small non-zero positive value to avoid potential divide-by-zero in simulation code. */
 	static const float MIN_TICK_TIME;
 };
diff --git a/Engine/Source/Runtime/Engine/Classes/GameFramework/SpringArmComponent.h b/Engine/Source/Runtime/Engine/Classes/GameFramework/SpringArmComponent.h
--- a/Engine/Source/Runtime/Engine/Classes/GameFramework/SpringArmComponent.h
+++ b/Engine/Source/Runtime/Engine/Classes/GameFramework/SpringArmComponent.h
@@ -144,14 +144,4 @@ protected:
 	 * by default it returns bHitSomething ? TraceHitLocation : DesiredArmLocation
 	 */
 	virtual FVector BlendLocations(const FVector& DesiredArmLocation, const FVector& TraceHitLocation, bool bHitSomething, float DeltaTime);
-
-public:
-
-	/**
-	* DEPRECATED variable: use "bUsePawnControlRotation" instead. Existing code using this value may not behave correctly.
-	* This is not a UPROPERTY, with good reason: we don't want to serialize in old values.
-	*/
-	DEPRECATED(4.5, "This variable is deprecated; existing code using this value may not behave correctly. It only exists to allow compilation of old projects, and code should stop using it in favor of the new bUsePawnControlRotation.")
-	bool bUseControllerViewRotation;
-
 };
diff --git a/Engine/Source/Runtime/Engine/Classes/Kismet/BlueprintMapLibrary.h b/Engine/Source/Runtime/Engine/Classes/Kismet/BlueprintMapLibrary.h
--- a/Engine/Source/Runtime/Engine/Classes/Kismet/BlueprintMapLibrary.h
+++ b/Engine/Source/Runtime/Engine/Classes/Kismet/BlueprintMapLibrary.h
@@ -15,6 +15,7 @@ class ENGINE_API UBlueprintMapLibrary : public UBlueprintFunctionLibrary
 {
 	GENERATED_BODY()
 
+public:
 	/** 
 	 * Adds a key and value to the map. If something already uses the provided key it will be overwritten with the new value.
 	 * After calling Key is guaranteed to be associated with Value until a subsequent mutation of the Map.
@@ -93,6 +94,13 @@ class ENGINE_API UBlueprintMapLibrary : public UBlueprintFunctionLibrary
 	UFUNCTION(BlueprintCallable, CustomThunk, meta=(DisplayName = "Clear", CompactNodeTitle = "CLEAR", MapParam = "TargetMap" ), Category = "Utilities|Map")
 	static void Map_Clear(const TMap<int32, int32>& TargetMap);
 
+	/** 
+	* Not exposed to users. Supports setting a map property on an object by name.
+	*/
+	UFUNCTION(BlueprintCallable, CustomThunk, meta=(BlueprintInternalUseOnly = "true", MapParam = "Value"))
+	static void SetMapPropertyByName(UObject* Object, FName PropertyName, const TMap<int32, int32>& Value);
+
+
 	DECLARE_FUNCTION(execMap_Add)
 	{
 		Stack.MostRecentProperty = nullptr;
@@ -322,6 +330,21 @@ class ENGINE_API UBlueprintMapLibrary : public UBlueprintFunctionLibrary
 		P_NATIVE_END
 	}
 
+	DECLARE_FUNCTION(execSetMapPropertyByName)
+	{
+		P_GET_OBJECT(UObject, OwnerObject);
+		P_GET_PROPERTY(UNameProperty, MapPropertyName);
+
+		Stack.StepCompiledIn<UMapProperty>(nullptr);
+		void* SrcMapAddr = Stack.MostRecentPropertyAddress;
+
+		P_FINISH;
+
+		P_NATIVE_BEGIN;
+		GenericMap_SetMapPropertyByName(OwnerObject, MapPropertyName, SrcMapAddr);
+		P_NATIVE_END;
+	}
+
 	static bool GenericMap_Add(const void* TargetMap, const UMapProperty* MapProperty, const void* KeyPtr, const void* ValuePtr);
 	static bool GenericMap_Remove(const void* TargetMap, const UMapProperty* MapProperty, const void* KeyPtr);
 	static bool GenericMap_Find(const void* TargetMap, const UMapProperty* MapProperty, const void* KeyPtr, void* ValuePtr);
@@ -329,5 +352,6 @@ class ENGINE_API UBlueprintMapLibrary : public UBlueprintFunctionLibrary
 	static void GenericMap_Values(const void* MapAddr, const UMapProperty* MapProperty, const void* ArrayAddr, const UArrayProperty* ArrayProperty);
 	static int32 GenericMap_Length(const void* TargetMap, const UMapProperty* MapProperty);
 	static void GenericMap_Clear(const void* TargetMap, const UMapProperty* MapProperty);
+	static void GenericMap_SetMapPropertyByName(UObject* OwnerObject, FName MapPropertyName, const void* SrcMapAddr);
 };
 
diff --git a/Engine/Source/Runtime/Engine/Classes/Kismet/BlueprintSetLibrary.h b/Engine/Source/Runtime/Engine/Classes/Kismet/BlueprintSetLibrary.h
--- a/Engine/Source/Runtime/Engine/Classes/Kismet/BlueprintSetLibrary.h
+++ b/Engine/Source/Runtime/Engine/Classes/Kismet/BlueprintSetLibrary.h
@@ -130,6 +130,12 @@ public:
 	UFUNCTION(BlueprintCallable, CustomThunk, meta = (DisplayName = "Difference", CompactNodeTitle = "DIFFERENCE", SetParam = "A|B|Result"), Category = "Utilities|Set")
 	static void Set_Difference(const TSet<int32>& A, const TSet<int32>& B, TSet<int32>& Result );
 
+	/** 
+	* Not exposed to users. Supports setting a set property on an object by name.
+	*/
+	UFUNCTION(BlueprintCallable, CustomThunk, meta=(BlueprintInternalUseOnly = "true", SetParam = "Value"))
+	static void SetSetPropertyByName(UObject* Object, FName PropertyName, const TSet<int32>& Value);
+
 	DECLARE_FUNCTION(execSet_Add)
 	{
 		Stack.MostRecentProperty = nullptr;
@@ -469,6 +475,21 @@ public:
 		P_NATIVE_END;
 	}
 
+	DECLARE_FUNCTION(execSetSetPropertyByName)
+	{
+		P_GET_OBJECT(UObject, OwnerObject);
+		P_GET_PROPERTY(UNameProperty, SetPropertyName);
+
+		Stack.StepCompiledIn<USetProperty>(nullptr);
+		void* SrcSetAddr = Stack.MostRecentPropertyAddress;
+
+		P_FINISH;
+
+		P_NATIVE_BEGIN;
+		GenericSet_SetSetPropertyByName(OwnerObject, SetPropertyName, SrcSetAddr);
+		P_NATIVE_END;
+	}
+
 	static bool GenericSet_Add(const void* TargetSet, const USetProperty* SetProperty, const void* ItemPtr);
 	static void GenericSet_AddItems(const void* TargetSet, const USetProperty* SetProperty, const void* TargetArray, const UArrayProperty* ArrayProperty);
 	static bool GenericSet_Remove(const void* TargetSet, const USetProperty* SetProperty, const void* ItemPtr);
@@ -480,5 +501,6 @@ public:
 	static void GenericSet_Intersect(const void* SetA, const USetProperty* SetPropertyA, const void* SetB, const USetProperty* SetPropertyB, const void* SetResult, const USetProperty* SetPropertyResult);
 	static void GenericSet_Union(const void* SetA, const USetProperty* SetPropertyA, const void* SetB, const USetProperty* SetPropertyB, const void* SetResult, const USetProperty* SetPropertyResult);
 	static void GenericSet_Difference(const void* SetA, const USetProperty* SetPropertyA, const void* SetB, const USetProperty* SetPropertyB, const void* SetResult, const USetProperty* SetPropertyResult);
+	static void GenericSet_SetSetPropertyByName(UObject* OwnerObject, FName SetPropertyName, const void* SrcSetAddr);
 };
 
