Fixed a critical crash bug relating to a race condition in async package summary reading.

diff --git a/Engine/Source/Runtime/CoreUObject/Private/Serialization/AsyncLoading.cpp b/Engine/Source/Runtime/CoreUObject/Private/Serialization/AsyncLoading.cpp
--- a/Engine/Source/Runtime/CoreUObject/Private/Serialization/AsyncLoading.cpp
+++ b/Engine/Source/Runtime/CoreUObject/Private/Serialization/AsyncLoading.cpp
@@ -7243,7 +7243,10 @@ bool FArchiveAsync2::WaitForIntialPhases(float InTimeLimit)
 bool FArchiveAsync2::PrecacheInternal(int64 RequestOffset, int64 RequestSize, bool bApplyMinReadSize)
 {
 	// CAUTION! This is possibly called the first time from a random IO thread.
-	if (LoadPhase != ELoadPhase::WaitingForSummary)
+
+	bool bIsWaitingForSummary =( LoadPhase == ELoadPhase::WaitingForSummary);
+
+	if (!bIsWaitingForSummary)
 	{
 		if (RequestSize == 0 || (RequestOffset >= PrecacheStartPos && RequestOffset + RequestSize <= PrecacheEndPos))
 		{
@@ -7273,7 +7276,7 @@ bool FArchiveAsync2::PrecacheInternal(int64 RequestOffset, int64 RequestSize, bo
 	ReadRequestSize = RequestSize;
 
 
-	if (bApplyMinReadSize && LoadPhase != ELoadPhase::WaitingForSummary)
+	if (bApplyMinReadSize && !bIsWaitingForSummary)
 	{
 #if WITH_EDITOR
 		static int64 MinimumReadSize = 1024 * 1024;
@@ -7301,17 +7304,19 @@ bool FArchiveAsync2::PrecacheInternal(int64 RequestOffset, int64 RequestSize, bo
 	}
 #endif
 	check(ReadRequestOffset - HeaderSizeWhenReadingExportsFromSplitFile >= 0);
+
+	// caution, this callback can fire before this even returns....and so bIsWaitingForSummary must be a local variable or we could get all confused by concurrency!
 	ReadRequestPtr = Handle->ReadRequest(ReadRequestOffset - HeaderSizeWhenReadingExportsFromSplitFile, ReadRequestSize, AIOP_Normal
-		, (GEventDrivenLoaderEnabled && LoadPhase == ELoadPhase::WaitingForSummary) ? &ReadCallbackFunctionForLinkerLoad : nullptr
+		, (GEventDrivenLoaderEnabled && bIsWaitingForSummary) ? &ReadCallbackFunctionForLinkerLoad : nullptr
 		);
-	if (LoadPhase != ELoadPhase::WaitingForSummary && ReadRequestPtr->PollCompletion())
+	if (!bIsWaitingForSummary && ReadRequestPtr->PollCompletion())
 	{
 		LogItem(TEXT("Read Start Hot"), ReadRequestOffset - HeaderSizeWhenReadingExportsFromSplitFile, ReadRequestSize, StartTime);
 		CompleteRead();
 		check(RequestOffset >= PrecacheStartPos && RequestOffset + RequestSize <= PrecacheEndPos);
 		return true;
 	}
-	else if (LoadPhase == ELoadPhase::WaitingForSummary)
+	else if (bIsWaitingForSummary)
 	{
 		LogItem(TEXT("Read Start Summary"), ReadRequestOffset - HeaderSizeWhenReadingExportsFromSplitFile, ReadRequestSize, StartTime);
 	}
@@ -7607,7 +7612,7 @@ void FArchiveAsync2::Serialize(void* Data, int64 Count)
 		LogItem(TEXT("Sync After Block"), AfterBlockOffset, AfterBlockSize);
 		PrecacheInternal(AfterBlockOffset, AfterBlockSize);
 		WaitRead();
-		check(AfterBlockOffset >= PrecacheStartPos && AfterBlockOffset + AfterBlockSize <= PrecacheEndPos);
+		checkf(AfterBlockOffset >= PrecacheStartPos && AfterBlockOffset + AfterBlockSize <= PrecacheEndPos, TEXT("Sync After Block???   %lld %lld %lld %lld"), AfterBlockOffset, AfterBlockSize, PrecacheStartPos, PrecacheEndPos);
 		FMemory::Memcpy(((uint8*)Data) + Count - AfterBlockSize, PrecacheBuffer + AfterBlockOffset - PrecacheStartPos, AfterBlockSize);
 	}
 #if DEVIRTUALIZE_FLinkerLoad_Serialize
