Fixing foliage not updating to match the new landscape height

diff --git a/Engine/Source/Runtime/CoreUObject/Private/UObject/LazyObjectPtr.cpp b/Engine/Source/Runtime/CoreUObject/Private/UObject/LazyObjectPtr.cpp
--- a/Engine/Source/Runtime/CoreUObject/Private/UObject/LazyObjectPtr.cpp
+++ b/Engine/Source/Runtime/CoreUObject/Private/UObject/LazyObjectPtr.cpp
@@ -44,12 +44,6 @@ FUniqueObjectGuid FUniqueObjectGuid::FixupForPIE(int32 PlayInEditorID) const
 UObject* FUniqueObjectGuid::ResolveObject() const
 {
 	UObject* Result = GuidAnnotation.Find(*this);
-	// If the object is still being loaded, we can't return it yet as it may
-	// cause undesired side-effects if client code starts changing it.
-	if (Result && Result->HasAnyFlags(RF_NeedLoad | RF_NeedPostLoad))
-	{
-		Result = nullptr;
-	}
 	return Result;
 }
 
