Copying //UE4/Dev-Framework to //UE4/Dev-Main (Source: //UE4/Dev-Framework @ 3582324)

diff --git a/Engine/Source/Runtime/CoreUObject/Private/UObject/PropertySoftClassPtr.cpp b/Engine/Source/Runtime/CoreUObject/Private/UObject/PropertySoftClassPtr.cpp
new file mode 100644
--- /dev/null
+++ b/Engine/Source/Runtime/CoreUObject/Private/UObject/PropertySoftClassPtr.cpp
@@ -0,0 +1,107 @@
+// Copyright 1998-2017 Epic Games, Inc. All Rights Reserved.
+
+#include "CoreMinimal.h"
+#include "UObject/ObjectMacros.h"
+#include "Templates/Casts.h"
+#include "UObject/UnrealType.h"
+#include "UObject/LinkerPlaceholderClass.h"
+
+/*-----------------------------------------------------------------------------
+	USoftClassProperty.
+-----------------------------------------------------------------------------*/
+
+void USoftClassProperty::BeginDestroy()
+{
+#if USE_CIRCULAR_DEPENDENCY_LOAD_DEFERRING
+	if (ULinkerPlaceholderClass* PlaceholderClass = Cast<ULinkerPlaceholderClass>(MetaClass))
+	{
+		PlaceholderClass->RemoveReferencingProperty(this);
+	}
+#endif // USE_CIRCULAR_DEPENDENCY_LOAD_DEFERRING
+
+	Super::BeginDestroy();
+}
+
+FString USoftClassProperty::GetCPPType(FString* ExtendedTypeText, uint32 CPPExportFlags) const
+{
+	check(MetaClass);
+	return GetCPPTypeCustom(ExtendedTypeText, CPPExportFlags, 
+		FString::Printf(TEXT("%s%s"), MetaClass->GetPrefixCPP(), *MetaClass->GetName()));
+}
+FString USoftClassProperty::GetCPPTypeCustom(FString* ExtendedTypeText, uint32 CPPExportFlags, const FString& InnerNativeTypeName) const
+{
+	ensure(!InnerNativeTypeName.IsEmpty());
+	return FString::Printf(TEXT("TSoftClassPtr<%s> "), *InnerNativeTypeName);
+}
+FString USoftClassProperty::GetCPPMacroType( FString& ExtendedTypeText ) const
+{
+	ExtendedTypeText = FString::Printf(TEXT("TSoftClassPtr<%s%s> "),MetaClass->GetPrefixCPP(),*MetaClass->GetName());
+	return TEXT("SOFTOBJECT");
+}
+
+FString USoftClassProperty::GetCPPTypeForwardDeclaration() const
+{
+	return FString::Printf(TEXT("class %s%s;"), MetaClass->GetPrefixCPP(), *MetaClass->GetName());
+}
+
+void USoftClassProperty::Serialize( FArchive& Ar )
+{
+	Super::Serialize( Ar );
+	Ar << MetaClass;
+
+#if USE_CIRCULAR_DEPENDENCY_LOAD_DEFERRING
+	if (Ar.IsLoading() || Ar.IsObjectReferenceCollector())
+	{
+		if (ULinkerPlaceholderClass* PlaceholderClass = Cast<ULinkerPlaceholderClass>(MetaClass))
+		{
+			PlaceholderClass->AddReferencingProperty(this);
+		}
+	}
+#endif // USE_CIRCULAR_DEPENDENCY_LOAD_DEFERRING
+
+	if( !(MetaClass||HasAnyFlags(RF_ClassDefaultObject)) )
+	{
+		// If we failed to load the MetaClass and we're not a CDO, that means we relied on a class that has been removed or doesn't exist.
+		// The most likely cause for this is either an incomplete recompile, or if content was migrated between games that had native class dependencies
+		// that do not exist in this game.  We allow blueprint classes to continue, because compile on load will error out, and stub the class that was using it
+		UClass* TestClass = dynamic_cast<UClass*>(GetOwnerStruct());
+		if( TestClass && TestClass->HasAllClassFlags(CLASS_Native) && !TestClass->HasAllClassFlags(CLASS_NewerVersionExists) && (TestClass->GetOutermost() != GetTransientPackage()) )
+		{
+			checkf(false, TEXT("Class property tried to serialize a missing class.  Did you remove a native class and not fully recompile?"));
+		}
+	}
+}
+
+#if USE_CIRCULAR_DEPENDENCY_LOAD_DEFERRING
+void USoftClassProperty::SetMetaClass(UClass* NewMetaClass)
+{
+	if (ULinkerPlaceholderClass* NewPlaceholderClass = Cast<ULinkerPlaceholderClass>(NewMetaClass))
+	{
+		NewPlaceholderClass->AddReferencingProperty(this);
+	}
+
+	if (ULinkerPlaceholderClass* OldPlaceholderClass = Cast<ULinkerPlaceholderClass>(MetaClass))
+	{
+		OldPlaceholderClass->RemoveReferencingProperty(this);
+	}
+	MetaClass = NewMetaClass;
+}
+#endif // USE_CIRCULAR_DEPENDENCY_LOAD_DEFERRING
+
+void USoftClassProperty::AddReferencedObjects(UObject* InThis, FReferenceCollector& Collector)
+{
+	USoftClassProperty* This = CastChecked<USoftClassProperty>(InThis);
+	Collector.AddReferencedObject( This->MetaClass, This );
+	Super::AddReferencedObjects( This, Collector );
+}
+
+bool USoftClassProperty::SameType(const UProperty* Other) const
+{
+	return Super::SameType(Other) && (MetaClass == ((USoftClassProperty*)Other)->MetaClass);
+}
+
+IMPLEMENT_CORE_INTRINSIC_CLASS(USoftClassProperty, USoftObjectProperty,
+	{
+		Class->EmitObjectReference(STRUCT_OFFSET(USoftClassProperty, MetaClass), TEXT("MetaClass"));
+	}
+);
diff --git a/Engine/Source/Runtime/CoreUObject/Private/UObject/PropertySoftObjectPtr.cpp b/Engine/Source/Runtime/CoreUObject/Private/UObject/PropertySoftObjectPtr.cpp
new file mode 100644
--- /dev/null
+++ b/Engine/Source/Runtime/CoreUObject/Private/UObject/PropertySoftObjectPtr.cpp
@@ -0,0 +1,204 @@
+// Copyright 1998-2017 Epic Games, Inc. All Rights Reserved.
+
+#include "CoreMinimal.h"
+#include "UObject/ObjectMacros.h"
+#include "UObject/SoftObjectPtr.h"
+#include "UObject/PropertyPortFlags.h"
+#include "UObject/UnrealType.h"
+
+/*-----------------------------------------------------------------------------
+	USoftObjectProperty.
+-----------------------------------------------------------------------------*/
+
+FString USoftObjectProperty::GetCPPTypeCustom(FString* ExtendedTypeText, uint32 CPPExportFlags, const FString& InnerNativeTypeName) const
+{
+	ensure(!InnerNativeTypeName.IsEmpty());
+	return FString::Printf(TEXT("TSoftObjectPtr<%s>"), *InnerNativeTypeName);
+}
+FString USoftObjectProperty::GetCPPMacroType( FString& ExtendedTypeText ) const
+{
+	ExtendedTypeText = FString::Printf(TEXT("TSoftObjectPtr<%s%s>"), PropertyClass->GetPrefixCPP(), *PropertyClass->GetName());
+	return TEXT("SOFTOBJECT");
+}
+
+FString USoftObjectProperty::GetCPPTypeForwardDeclaration() const
+{
+	return FString::Printf(TEXT("class %s%s;"), PropertyClass->GetPrefixCPP(), *PropertyClass->GetName());
+}
+
+FName USoftObjectProperty::GetID() const
+{
+	// SoftClass shares the same tag, they are binary compatible
+	return NAME_SoftObjectProperty;
+}
+
+// this is always shallow, can't see that we would want it any other way
+bool USoftObjectProperty::Identical( const void* A, const void* B, uint32 PortFlags ) const
+{
+	FSoftObjectPtr ObjectA = A ? *((FSoftObjectPtr*)A) : FSoftObjectPtr();
+	FSoftObjectPtr ObjectB = B ? *((FSoftObjectPtr*)B) : FSoftObjectPtr();
+
+	return ObjectA.GetUniqueID() == ObjectB.GetUniqueID();
+}
+
+void USoftObjectProperty::SerializeItem( FArchive& Ar, void* Value, void const* Defaults ) const
+{
+	// We never serialize our reference while the garbage collector is harvesting references
+	// to objects, because we don't want soft object pointers to keep objects from being garbage collected
+	// Allow persistent archives so they can keep track of string references. (e.g. FArchiveSaveTagImports)
+	if( !Ar.IsObjectReferenceCollector() || Ar.IsModifyingWeakAndStrongReferences() || Ar.IsPersistent() )
+	{
+		FSoftObjectPtr OldValue = *(FSoftObjectPtr*)Value;
+		Ar << *(FSoftObjectPtr*)Value;
+
+		if (Ar.IsLoading() || Ar.IsModifyingWeakAndStrongReferences()) 
+		{
+			if (OldValue.GetUniqueID() != ((FSoftObjectPtr*)Value)->GetUniqueID())
+			{
+				CheckValidObject(Value);
+			}
+		}
+	}
+}
+
+void USoftObjectProperty::ExportTextItem( FString& ValueStr, const void* PropertyValue, const void* DefaultValue, UObject* Parent, int32 PortFlags, UObject* ExportRootScope ) const
+{
+	FSoftObjectPtr& SoftObjectPtr = *(FSoftObjectPtr*)PropertyValue;
+
+	FSoftObjectPath SoftObjectPath;
+	UObject *Object = SoftObjectPtr.Get();
+
+	if (Object)
+	{
+		// Use object in case name has changed.
+		SoftObjectPath = FSoftObjectPath(Object);
+	}
+	else
+	{
+		SoftObjectPath = SoftObjectPtr.GetUniqueID();
+	}
+
+	if (0 != (PortFlags & PPF_ExportCpp))
+	{
+		ValueStr += FString::Printf(TEXT("FSoftObjectPath(TEXT(\"%s\"))"), *SoftObjectPath.ToString().ReplaceCharWithEscapedChar());
+		return;
+	}
+
+	SoftObjectPath.ExportTextItem(ValueStr, SoftObjectPath, Parent, PortFlags, ExportRootScope);
+}
+
+const TCHAR* USoftObjectProperty::ImportText_Internal( const TCHAR* InBuffer, void* Data, int32 PortFlags, UObject* Parent, FOutputDevice* ErrorText ) const
+{
+	FSoftObjectPtr& SoftObjectPtr = *(FSoftObjectPtr*)Data;
+
+	FSoftObjectPath SoftObjectPath;
+
+	if (SoftObjectPath.ImportTextItem(InBuffer, PortFlags, Parent, ErrorText))
+	{
+		SoftObjectPtr = SoftObjectPath;
+		return InBuffer;
+	}
+
+	else
+	{
+		SoftObjectPtr = nullptr;
+		return nullptr;
+	}
+}
+
+bool USoftObjectProperty::ConvertFromType(const FPropertyTag& Tag, FArchive& Ar, uint8* Data, UStruct* DefaultsStruct, bool& bOutAdvanceProperty)
+{
+	static FName NAME_AssetObjectProperty = "AssetObjectProperty";
+	static FName NAME_SoftObjectPath = "SoftObjectPath";
+	static FName NAME_SoftClassPath = "SoftClassPath";
+	static FName NAME_StringAssetReference = "StringAssetReference";
+	static FName NAME_StringClassReference = "StringClassReference";
+	bOutAdvanceProperty = true;
+
+	if (Tag.Type == NAME_AssetObjectProperty)
+	{
+		// Old name of soft object property, serialize normally
+		uint8* DestAddress = ContainerPtrToValuePtr<uint8>(Data, Tag.ArrayIndex);
+
+		Tag.SerializeTaggedProperty(Ar, this, DestAddress, nullptr);
+
+		if (Ar.IsCriticalError())
+		{
+			bOutAdvanceProperty = false;
+		}
+
+		return true;
+	}
+	else if (Tag.Type == NAME_ObjectProperty)
+	{
+		// This property used to be a raw UObjectProperty Foo* but is now a TSoftObjectPtr<Foo>;
+		// Serialize from mismatched tag directly into the FSoftObjectPtr's soft object path to ensure that the delegates needed for cooking
+		// are fired
+		FSoftObjectPtr* PropertyValue = GetPropertyValuePtr_InContainer(Data, Tag.ArrayIndex);
+		check(PropertyValue);
+
+		return PropertyValue->GetUniqueID().SerializeFromMismatchedTag(Tag, Ar);
+	}
+	else if (Tag.Type == NAME_StructProperty && (Tag.StructName == NAME_SoftObjectPath || Tag.StructName == NAME_SoftClassPath || Tag.StructName == NAME_StringAssetReference || Tag.StructName == NAME_StringClassReference))
+	{
+		// This property used to be a FSoftObjectPath but is now a TSoftObjectPtr<Foo>
+		FSoftObjectPath PreviousValue;
+		// explicitly call Serialize to ensure that the various delegates needed for cooking are fired
+		PreviousValue.Serialize(Ar);
+
+		// now copy the value into the object's address space
+		FSoftObjectPtr PreviousValueSoftObjectPtr;
+		PreviousValueSoftObjectPtr = PreviousValue;
+		SetPropertyValue_InContainer(Data, PreviousValueSoftObjectPtr, Tag.ArrayIndex);
+
+		return true;
+	}
+
+	return false;
+}
+
+UObject* USoftObjectProperty::GetObjectPropertyValue(const void* PropertyValueAddress) const
+{
+	return GetPropertyValue(PropertyValueAddress).Get();
+}
+
+void USoftObjectProperty::SetObjectPropertyValue(void* PropertyValueAddress, UObject* Value) const
+{
+	SetPropertyValue(PropertyValueAddress, TCppType(Value));
+}
+
+bool USoftObjectProperty::AllowCrossLevel() const
+{
+	return true;
+}
+
+uint32 USoftObjectProperty::GetValueTypeHashInternal(const void* Src) const
+{
+	return GetTypeHash(GetPropertyValue(Src));
+}
+
+void USoftObjectProperty::CopySingleValueToScriptVM(void* Dest, void const* Src) const
+{
+	CopySingleValue(Dest, Src);
+}
+
+void USoftObjectProperty::CopyCompleteValueToScriptVM(void* Dest, void const* Src) const
+{
+	CopyCompleteValue(Dest, Src);
+}
+
+void USoftObjectProperty::CopySingleValueFromScriptVM(void* Dest, void const* Src) const
+{
+	CopySingleValue(Dest, Src);
+}
+
+void USoftObjectProperty::CopyCompleteValueFromScriptVM(void* Dest, void const* Src) const
+{
+	CopyCompleteValue(Dest, Src);
+}
+
+IMPLEMENT_CORE_INTRINSIC_CLASS(USoftObjectProperty, UObjectPropertyBase,
+	{
+	}
+);
+
diff --git a/Engine/Source/Runtime/CoreUObject/Private/UObject/PropertyStruct.cpp b/Engine/Source/Runtime/CoreUObject/Private/UObject/PropertyStruct.cpp
--- a/Engine/Source/Runtime/CoreUObject/Private/UObject/PropertyStruct.cpp
+++ b/Engine/Source/Runtime/CoreUObject/Private/UObject/PropertyStruct.cpp
@@ -61,17 +61,14 @@ void UStructProperty::LinkInternal(FArchive& Ar)
 		GetLinker()->Preload(this);
 	}
 
-#if !(UE_BUILD_SHIPPING || UE_BUILD_TEST)
-	if (Struct == NULL)
+	if (Struct)
 	{
-		UE_LOG(LogProperty, Error, TEXT("Struct type unknown for property '%s'; perhaps the USTRUCT() was renamed or deleted?"), *GetFullName());
+		// Preload is required here in order to load the value of Struct->PropertiesSize
+		Ar.Preload(Struct);
 	}
-#endif
-
-	// Preload is required here in order to load the value of Struct->PropertiesSize
-	Ar.Preload(Struct);
-	if ( !ensure(Struct) )
+	else
 	{
+		UE_LOG(LogProperty, Error, TEXT("Struct type unknown for property '%s'; perhaps the USTRUCT() was renamed or deleted?"), *GetFullName());
 		Struct = GetFallbackStruct();
 	}
 	PreloadInnerStructMembers(this);
diff --git a/Engine/Source/Runtime/CoreUObject/Private/UObject/PropertyTag.cpp b/Engine/Source/Runtime/CoreUObject/Private/UObject/PropertyTag.cpp
--- a/Engine/Source/Runtime/CoreUObject/Private/UObject/PropertyTag.cpp
+++ b/Engine/Source/Runtime/CoreUObject/Private/UObject/PropertyTag.cpp
@@ -167,7 +167,7 @@ FArchive& operator<<( FArchive& Ar, FPropertyTag& Tag )
 }
 
 // Property serializer.
-void FPropertyTag::SerializeTaggedProperty( FArchive& Ar, UProperty* Property, uint8* Value, uint8* Defaults )
+void FPropertyTag::SerializeTaggedProperty( FArchive& Ar, UProperty* Property, uint8* Value, uint8* Defaults ) const
 {
 	if (Property->GetClass() == UBoolProperty::StaticClass())
 	{
diff --git a/Engine/Source/Runtime/CoreUObject/Private/UObject/SavePackage.cpp b/Engine/Source/Runtime/CoreUObject/Private/UObject/SavePackage.cpp
--- a/Engine/Source/Runtime/CoreUObject/Private/UObject/SavePackage.cpp
+++ b/Engine/Source/Runtime/CoreUObject/Private/UObject/SavePackage.cpp
@@ -28,7 +28,7 @@
 #include "UObject/Package.h"
 #include "Templates/Casts.h"
 #include "UObject/LazyObjectPtr.h"
-#include "Misc/StringAssetReference.h"
+#include "UObject/SoftObjectPtr.h"
 #include "UObject/PropertyPortFlags.h"
 #include "UObject/UnrealType.h"
 #include "UObject/TextProperty.h"
@@ -686,7 +686,7 @@ public:
  * or 
  * - It's an editor only object
 */
-bool IsEditorOnlyObject(const UObject* InObject)
+bool IsEditorOnlyObject(const UObject* InObject, bool bCheckRecursive)
 {
 	DECLARE_SCOPE_CYCLE_COUNTER(TEXT("IsEditorOnlyObject"), STAT_IsEditorOnlyObject, STATGROUP_LoadTime);
 
@@ -721,7 +721,39 @@ bool IsEditorOnlyObject(const UObject* InObject)
 		return true;
 	}
 
-	// Recursion is checked with ConditionallyExcludeObjectForTarget
+	if (bCheckRecursive && !InObject->IsNative())
+	{
+		UObject* Outer = InObject->GetOuter();
+		if (Outer && Outer != Package)
+		{
+			if (IsEditorOnlyObject(Outer, true))
+			{
+				return true;
+			}
+		}
+		const UStruct* InStruct = Cast<UStruct>(InObject);
+		if (InStruct)
+		{
+			const UStruct* SuperStruct = InStruct->GetSuperStruct();
+			if (SuperStruct && IsEditorOnlyObject(SuperStruct, true))
+			{
+				return true;
+			}
+		}
+		else
+		{
+			if (IsEditorOnlyObject(InObject->GetClass(), true))
+			{
+				return true;
+			}
+
+			UObject* Archetype = InObject->GetArchetype();
+			if (Archetype && IsEditorOnlyObject(Archetype, true))
+			{
+				return true;
+			}
+		}
+	}
 	return false;
 }
 
@@ -828,7 +860,7 @@ static void ConditionallyExcludeObjectForTarget(UObject* Obj, EObjectMark Exclud
 	}
 	else
 	{
-		if (!Obj->HasAnyMarks(OBJECTMARK_EditorOnly) && ((InheritedMarks & OBJECTMARK_EditorOnly) || IsEditorOnlyObject(Obj)))
+		if (!Obj->HasAnyMarks(OBJECTMARK_EditorOnly) && ((InheritedMarks & OBJECTMARK_EditorOnly) || IsEditorOnlyObject(Obj, false)))
 		{
 			Obj->Mark(OBJECTMARK_EditorOnly);
 		}
@@ -1109,7 +1141,7 @@ public:
 	virtual FArchive& operator<<(UObject*& Obj) override;
 	virtual FArchive& operator<< (struct FWeakObjectPtr& Value) override;
 	virtual FArchive& operator<<(FLazyObjectPtr& LazyObjectPtr) override;
-	virtual FArchive& operator<<(FStringAssetReference& Value) override;
+	virtual FArchive& operator<<(FSoftObjectPath& Value) override;
 	virtual FArchive& operator<<(FName& Name) override;
 	
 	virtual void MarkSearchableName(const UObject* TypeObject, const FName& ValueName) const override;
@@ -1249,26 +1281,17 @@ FArchive& FArchiveSaveTagImports::operator<<( FLazyObjectPtr& LazyObjectPtr)
 	return *this << ID;
 }
 
-FArchive& FArchiveSaveTagImports::operator<<(FStringAssetReference& Value)
+FArchive& FArchiveSaveTagImports::operator<<(FSoftObjectPath& Value)
 {
 	if (Value.IsValid())
 	{
-		Value.PreSavePath();
+		Value.SerializePath(*this);
+
 		FString Path = Value.ToString();
+		FName PackageName = FName(*FPackageName::ObjectPathToPackageName(Path));
 
-		if (GetIniFilenameFromObjectsReference(Path) != nullptr)
-		{
-			Linker->StringAssetReferencesMap.AddUnique(Path);
-		}
-		else
-		{
-			FString NormalizedPath = FPackageName::GetNormalizedObjectPath(Path);
-			if (!NormalizedPath.IsEmpty())
-			{
-				Linker->StringAssetReferencesMap.AddUnique(FPackageName::ObjectPathToPackageName(NormalizedPath));
-				Value.SetPath(MoveTemp(NormalizedPath));
-			}
-		}
+		SavePackageState->MarkNameAsReferenced(PackageName);
+		Linker->SoftPackageReferenceList.AddUnique(PackageName);	
 	}
 	return *this;
 }
@@ -1928,8 +1951,8 @@ class FExportReferenceSorter : public FArchiveUObject
 				UWeakObjectProperty::StaticClass(),
 				UObjectPropertyBase::StaticClass(),
 				ULazyObjectProperty::StaticClass(),
-				UAssetObjectProperty::StaticClass(),
-				UAssetClassProperty::StaticClass(),
+				USoftObjectProperty::StaticClass(),
+				USoftClassProperty::StaticClass(),
 				UMapProperty::StaticClass(),
 				USetProperty::StaticClass(),
 				UEnumProperty::StaticClass()
@@ -2036,8 +2059,8 @@ class FExportReferenceSorter : public FArchiveUObject
 			UWeakObjectProperty::StaticClass(),
 			UObjectPropertyBase::StaticClass(),
 			ULazyObjectProperty::StaticClass(),
-			UAssetObjectProperty::StaticClass(),
-			UAssetClassProperty::StaticClass(),
+			USoftObjectProperty::StaticClass(),
+			USoftClassProperty::StaticClass(),
 			UMapProperty::StaticClass(),
 			USetProperty::StaticClass(),
 			UEnumProperty::StaticClass()
@@ -4543,16 +4566,16 @@ FSavePackageResultStruct UPackage::Save(UPackage* InOuter, UObject* Base, EObjec
 				// Only save string asset and searchable name map if saving for editor
 				if (!(Linker->Summary.PackageFlags & PKG_FilterEditorOnly))
 				{
-					// Save string asset reference map
-					Linker->Summary.StringAssetReferencesOffset = Linker->Tell();
-					Linker->Summary.StringAssetReferencesCount = Linker->StringAssetReferencesMap.Num();
+					// Save soft package references
+					Linker->Summary.SoftPackageReferencesOffset = Linker->Tell();
+					Linker->Summary.SoftPackageReferencesCount = Linker->SoftPackageReferenceList.Num();
 					{
 #if WITH_EDITOR
 						FArchive::FScopeSetDebugSerializationFlags S(*Linker, DSF_IgnoreDiff, true);
 #endif
-						for (FString& StringAssetReference : Linker->StringAssetReferencesMap)
+						for (FName& SoftPackageName : Linker->SoftPackageReferenceList)
 						{
-							*Linker << StringAssetReference;
+							*Linker << SoftPackageName;
 						}
 					}
 
@@ -4562,8 +4585,8 @@ FSavePackageResultStruct UPackage::Save(UPackage* InOuter, UObject* Base, EObjec
 				}
 				else
 				{
-					Linker->Summary.StringAssetReferencesCount = 0;
-					Linker->Summary.StringAssetReferencesOffset = 0;
+					Linker->Summary.SoftPackageReferencesCount = 0;
+					Linker->Summary.SoftPackageReferencesOffset = 0;
 					Linker->Summary.SearchableNamesOffset = 0;
 				}
 
@@ -4690,7 +4713,7 @@ FSavePackageResultStruct UPackage::Save(UPackage* InOuter, UObject* Base, EObjec
 							{
 								UE_LOG(LogSavePackage, Warning, TEXT("A dependency '%s' of '%s' is in the linker table, but is pending kill. We will keep the dependency anyway (%d)."), *ToTest->GetFullName(), *ForObj->GetFullName(), CallSite);
 							}
-							bool bNotFiltered = (ExcludedObjectMarks == OBJECTMARK_NOMARKS || !ToTest->HasAnyMarks(ExcludedObjectMarks)) && (!(Linker->Summary.PackageFlags & PKG_FilterEditorOnly) || !IsEditorOnlyObject(ToTest));
+							bool bNotFiltered = (ExcludedObjectMarks == OBJECTMARK_NOMARKS || !ToTest->HasAnyMarks(ExcludedObjectMarks)) && (!(Linker->Summary.PackageFlags & PKG_FilterEditorOnly) || !IsEditorOnlyObject(ToTest, false));
 							if (bMandatory && !bNotFiltered)
 							{
 								UE_LOG(LogSavePackage, Warning, TEXT("A dependency '%s' of '%s' was filtered, but is mandatory. This indicates a problem with editor only stripping. We will keep the dependency anyway (%d)."), *ToTest->GetFullName(), *ForObj->GetFullName(), CallSite);
diff --git a/Engine/Source/Runtime/CoreUObject/Private/UObject/ScriptCore.cpp b/Engine/Source/Runtime/CoreUObject/Private/UObject/ScriptCore.cpp
--- a/Engine/Source/Runtime/CoreUObject/Private/UObject/ScriptCore.cpp
+++ b/Engine/Source/Runtime/CoreUObject/Private/UObject/ScriptCore.cpp
@@ -20,7 +20,7 @@
 #include "UObject/CoreNative.h"
 #include "UObject/Class.h"
 #include "Templates/Casts.h"
-#include "Misc/StringAssetReference.h"
+#include "UObject/SoftObjectPtr.h"
 #include "UObject/PropertyPortFlags.h"
 #include "UObject/UnrealType.h"
 #include "UObject/Stack.h"
@@ -1212,24 +1212,6 @@ void UObject::ProcessEvent( UFunction* Function, void* Parms )
 	FScopeCycleCounterUObject ContextScope(bShouldTrackObject ? this : nullptr);
 #endif
 
-#if !(UE_BUILD_SHIPPING || UE_BUILD_TEST)
-	bool bInstrumentScriptEvent = false;
-	if (GetClass()->HasInstrumentation())
-	{
-		// Don't instrument native events that have not been implemented/overridden in script (BP). These events will not have had any profiler data generated for them at compile time.
-		if (!Function->HasAnyFunctionFlags(FUNC_Native) || Function->GetOuter()->GetClass()->HasAnyClassFlags(CLASS_CompiledFromBlueprint))
-		{
-			if (Function->HasAnyFunctionFlags(FUNC_Event|FUNC_BlueprintEvent))
-			{
-				// Don't handle latent actions here, let the latent action manager handle them.
-				FScriptInstrumentationSignal EventInstrumentationInfo(EScriptInstrumentation::Event, this, Function);
-				FBlueprintCoreDelegates::InstrumentScriptEvent(EventInstrumentationInfo);
-				bInstrumentScriptEvent = true;
-			}
-		}
-	}
-#endif
-
 #if DO_BLUEPRINT_GUARD
 	FBlueprintExceptionTracker& BlueprintExceptionTracker = FBlueprintExceptionTracker::Get();
 	BlueprintExceptionTracker.ScriptEntryTag++;
@@ -1350,11 +1332,6 @@ void UObject::ProcessEvent( UFunction* Function, void* Parms )
 	}
 
 #if !(UE_BUILD_SHIPPING || UE_BUILD_TEST)
-	if (bInstrumentScriptEvent)
-	{
-		FScriptInstrumentationSignal EventInstrumentationInfo(EScriptInstrumentation::Stop, this, Function);
-		FBlueprintCoreDelegates::InstrumentScriptEvent(EventInstrumentationInfo);
-	}
 #if WITH_EDITORONLY_DATA
 	FBlueprintCoreDelegates::OnScriptExecutionEnd.Broadcast();
 #endif
@@ -2539,13 +2516,13 @@ void UObject::execObjectConst( FFrame& Stack, RESULT_DECL )
 }
 IMPLEMENT_VM_FUNCTION( EX_ObjectConst, execObjectConst );
 
-void UObject::execAssetConst(FFrame& Stack, RESULT_DECL)
+void UObject::execSoftObjectConst(FFrame& Stack, RESULT_DECL)
 {
 	FString LongPath;
 	Stack.Step(Stack.Object, &LongPath);
-	*(FAssetPtr*)RESULT_PARAM = FStringAssetReference(LongPath);
+	*(FSoftObjectPtr*)RESULT_PARAM = FSoftObjectPath(LongPath);
 }
-IMPLEMENT_VM_FUNCTION(EX_AssetConst, execAssetConst);
+IMPLEMENT_VM_FUNCTION( EX_SoftObjectConst, execSoftObjectConst);
 
 void UObject::execInstanceDelegate( FFrame& Stack, RESULT_DECL )
 {
diff --git a/Engine/Source/Runtime/CoreUObject/Private/UObject/SoftObjectPath.cpp b/Engine/Source/Runtime/CoreUObject/Private/UObject/SoftObjectPath.cpp
new file mode 100644
--- /dev/null
+++ b/Engine/Source/Runtime/CoreUObject/Private/UObject/SoftObjectPath.cpp
@@ -0,0 +1,465 @@
+// Copyright 1998-2017 Epic Games, Inc. All Rights Reserved.
+
+#include "UObject/SoftObjectPath.h"
+#include "UObject/PropertyPortFlags.h"
+#include "UObject/UnrealType.h"
+#include "UObject/ObjectRedirector.h"
+#include "Misc/PackageName.h"
+#include "UObject/LinkerLoad.h"
+#include "UObject/UObjectThreadContext.h"
+#include "Misc/RedirectCollector.h"
+
+FSoftObjectPath::FSoftObjectPath(const UObject* InObject)
+{
+	if (InObject)
+	{
+		SetPath(InObject->GetPathName());
+	}
+}
+
+FString FSoftObjectPath::ToString() const
+{
+	// Most of the time there is no sub path so we can do a single string allocation
+	FString AssetPathString = GetAssetPathString();
+	if (SubPathString.IsEmpty())
+	{
+		return AssetPathString;
+	}
+	FString FullPathString;
+
+	// Preallocate to correct size and then append strings
+	FullPathString.Empty(AssetPathString.Len() + SubPathString.Len() + 1);
+	FullPathString.Append(AssetPathString);
+	FullPathString.AppendChar(':');
+	FullPathString.Append(SubPathString);
+	return FullPathString;
+}
+
+void FSoftObjectPath::SetPath(FString Path)
+{
+	if (Path.IsEmpty() || Path.Equals(TEXT("None"), ESearchCase::CaseSensitive))
+	{
+		// Empty path, just empty the pathname.
+		Reset();
+	}
+	else if (ensureMsgf(!FPackageName::IsShortPackageName(Path), TEXT("Cannot create SoftObjectPath with short package name '%s'! You must pass in fully qualified package names"), *Path))
+	{
+		if (Path[0] != '/')
+		{
+			// Possibly an ExportText path. Trim the ClassName.
+			Path = FPackageName::ExportTextPathToObjectPath(Path);
+		}
+
+		int32 ColonIndex = INDEX_NONE;
+
+		if (Path.FindChar(':', ColonIndex))
+		{
+			// Has a subobject, split on that then create a name from the temporary path
+			SubPathString = Path.Mid(ColonIndex + 1);
+			Path.RemoveAt(ColonIndex, Path.Len() - ColonIndex);
+			AssetPathName = *Path;
+		}
+		else
+		{
+			// No Subobject
+			AssetPathName = *Path;
+			SubPathString.Empty();
+		}
+	}
+}
+
+bool FSoftObjectPath::PreSavePath()
+{
+#if WITH_EDITOR
+	FName FoundRedirection = GRedirectCollector.GetAssetPathRedirection(AssetPathName);
+
+	if (FoundRedirection != NAME_None)
+	{
+		AssetPathName = FoundRedirection;
+		return true;
+	}
+#endif // WITH_EDITOR
+	return false;
+}
+
+void FSoftObjectPath::PostLoadPath() const
+{
+#if WITH_EDITOR
+	GRedirectCollector.OnSoftObjectPathLoaded(*this);
+#endif // WITH_EDITOR
+}
+
+bool FSoftObjectPath::Serialize(FArchive& Ar)
+{
+	// Archivers will call back into SerializePath for the various fixups
+	Ar << *this;
+
+	return true;
+}
+
+void FSoftObjectPath::SerializePath(FArchive& Ar, bool bSkipSerializeIfArchiveHasSize)
+{
+#if WITH_EDITOR
+	if (Ar.IsSaving())
+	{
+		PreSavePath();
+	}
+#endif // WITH_EDITOR
+
+	if (!bSkipSerializeIfArchiveHasSize || Ar.IsObjectReferenceCollector() || Ar.Tell() < 0)
+	{
+		if (Ar.IsLoading() && Ar.UE4Ver() < VER_UE4_ADDED_SOFT_OBJECT_PATH)
+		{
+			FString Path;
+			Ar << Path;
+
+			if (Ar.UE4Ver() < VER_UE4_KEEP_ONLY_PACKAGE_NAMES_IN_STRING_ASSET_REFERENCES_MAP)
+			{
+				Path = FPackageName::GetNormalizedObjectPath(Path);
+			}
+
+			SetPath(MoveTemp(Path));
+		}
+		else
+		{
+			Ar << AssetPathName;
+			Ar << SubPathString;
+		}
+	}
+
+#if WITH_EDITOR
+	if (Ar.IsLoading())
+	{
+		if (Ar.IsPersistent())
+		{
+			PostLoadPath();
+		}
+		if (Ar.GetPortFlags()&PPF_DuplicateForPIE)
+		{
+			// Remap unique ID if necessary
+			// only for fixing up cross-level references, inter-level references handled in FDuplicateDataReader
+			FixupForPIE();
+		}
+	}
+#endif // WITH_EDITOR
+}
+
+bool FSoftObjectPath::operator==(FSoftObjectPath const& Other) const
+{
+	return AssetPathName == Other.AssetPathName && SubPathString == Other.SubPathString;
+}
+
+FSoftObjectPath& FSoftObjectPath::operator=(FSoftObjectPath const& Other)
+{
+	AssetPathName = Other.AssetPathName;
+	SubPathString = Other.SubPathString;
+	return *this;
+}
+
+bool FSoftObjectPath::ExportTextItem(FString& ValueStr, FSoftObjectPath const& DefaultValue, UObject* Parent, int32 PortFlags, UObject* ExportRootScope) const
+{
+	if (0 != (PortFlags & EPropertyPortFlags::PPF_ExportCpp))
+	{
+		return false;
+	}
+
+	if (IsValid())
+	{
+		// Fixup any redirectors
+		FSoftObjectPath Temp = *this;
+		Temp.PreSavePath();
+
+		ValueStr += Temp.ToString();
+	}
+	else
+	{
+		ValueStr += TEXT("None");
+	}
+	return true;
+}
+
+bool FSoftObjectPath::ImportTextItem(const TCHAR*& Buffer, int32 PortFlags, UObject* Parent, FOutputDevice* ErrorText)
+{
+	FString ImportedPath = TEXT("");
+	const TCHAR* NewBuffer = UPropertyHelpers::ReadToken(Buffer, ImportedPath, 1);
+	if (!NewBuffer)
+	{
+		return false;
+	}
+	Buffer = NewBuffer;
+	if (ImportedPath == TEXT("None"))
+	{
+		ImportedPath = TEXT("");
+	}
+	else
+	{
+		if (*Buffer == TCHAR('\''))
+		{
+			// A ' token likely means we're looking at a path string in the form "Texture2d'/Game/UI/HUD/Actions/Barrel'" and we need to read and append the path part
+			// We have to skip over the first ' as UPropertyHelpers::ReadToken doesn't read single-quoted strings correctly, but does read a path correctly
+			Buffer++; // Skip the leading '
+			ImportedPath.Reset();
+			NewBuffer = UPropertyHelpers::ReadToken(Buffer, ImportedPath, 1);
+			if (!NewBuffer)
+			{
+				return false;
+			}
+			Buffer = NewBuffer;
+			if (*Buffer++ != TCHAR('\''))
+			{
+				return false;
+			}
+		}
+	}
+
+	SetPath(MoveTemp(ImportedPath));
+
+	// Consider this a load, so Config string references get cooked
+	PostLoadPath();
+
+	return true;
+}
+
+/**
+ * Serializes from mismatched tag.
+ *
+ * @template_param TypePolicy The policy should provide two things:
+ *	- GetTypeName() method that returns registered name for this property type,
+ *	- typedef Type, which is a C++ type to serialize if property matched type name.
+ * @param Tag Property tag to match type.
+ * @param Ar Archive to serialize from.
+ */
+template <class TypePolicy>
+bool SerializeFromMismatchedTagTemplate(FString& Output, const FPropertyTag& Tag, FArchive& Ar)
+{
+	if (Tag.Type == TypePolicy::GetTypeName())
+	{
+		typename TypePolicy::Type* ObjPtr = nullptr;
+		Ar << ObjPtr;
+		if (ObjPtr)
+		{
+			Output = ObjPtr->GetPathName();
+		}
+		else
+		{
+			Output = FString();
+		}
+		return true;
+	}
+	else if (Tag.Type == NAME_StrProperty)
+	{
+		FString String;
+		Ar << String;
+
+		Output = String;
+		return true;
+	}
+	return false;
+}
+
+bool FSoftObjectPath::SerializeFromMismatchedTag(struct FPropertyTag const& Tag, FArchive& Ar)
+{
+	struct UObjectTypePolicy
+	{
+		typedef UObject Type;
+		static const FName FORCEINLINE GetTypeName() { return NAME_ObjectProperty; }
+	};
+
+	FString Path = ToString();
+
+	bool bReturn = SerializeFromMismatchedTagTemplate<UObjectTypePolicy>(Path, Tag, Ar);
+
+	if (Ar.IsLoading())
+	{
+		SetPath(MoveTemp(Path));
+		PostLoadPath();
+	}
+
+	return bReturn;
+}
+
+UObject* FSoftObjectPath::TryLoad() const
+{
+	UObject* LoadedObject = nullptr;
+
+	if ( IsValid() )
+	{
+		LoadedObject = LoadObject<UObject>(nullptr, *ToString());
+		while (UObjectRedirector* Redirector = Cast<UObjectRedirector>(LoadedObject))
+		{
+			LoadedObject = Redirector->DestinationObject;
+		}
+	}
+
+	return LoadedObject;
+}
+
+UObject* FSoftObjectPath::ResolveObject() const
+{
+	// Don't try to resolve if we're saving a package because StaticFindObject can't be used here
+	// and we usually don't want to force references to weak pointers while saving.
+	if (!IsValid() || GIsSavingPackage)
+	{
+		return nullptr;
+	}
+
+	UObject* FoundObject = FindObject<UObject>(nullptr, *ToString());
+	while (UObjectRedirector* Redirector = Cast<UObjectRedirector>(FoundObject))
+	{
+		FoundObject = Redirector->DestinationObject;
+	}
+
+	return FoundObject;
+}
+
+FSoftObjectPath FSoftObjectPath::GetOrCreateIDForObject(const class UObject *Object)
+{
+	check(Object);
+	return FSoftObjectPath(Object);
+}
+
+void FSoftObjectPath::AddPIEPackageName(FName NewPIEPackageName)
+{
+	PIEPackageNames.Add(NewPIEPackageName);
+}
+
+void FSoftObjectPath::ClearPIEPackageNames()
+{
+	PIEPackageNames.Empty();
+}
+
+void FSoftObjectPath::FixupForPIE()
+{
+#if WITH_EDITOR
+	if (GPlayInEditorID != INDEX_NONE && IsValid())
+	{
+		FString Path = ToString();
+
+		// Determine if this reference has already been fixed up for PIE
+		const FString ShortPackageOuterAndName = FPackageName::GetLongPackageAssetName(Path);
+		if (!ShortPackageOuterAndName.StartsWith(PLAYWORLD_PACKAGE_PREFIX))
+		{
+			// Name of the ULevel subobject of UWorld, set in InitializeNewWorld
+			bool bIsChildOfLevel = SubPathString.StartsWith(TEXT("PersistentLevel."));
+
+			FString PIEPath = FString::Printf(TEXT("%s/%s_%d_%s"), *FPackageName::GetLongPackagePath(Path), PLAYWORLD_PACKAGE_PREFIX, GPlayInEditorID, *ShortPackageOuterAndName);
+			FName PIEPackage = FName(*FPackageName::ObjectPathToPackageName(PIEPath));
+
+			// Duplicate if this an already registered PIE package or this looks like a level subobject reference
+			if (bIsChildOfLevel || PIEPackageNames.Contains(PIEPackage))
+			{
+				// Need to prepend PIE prefix, as we're in PIE and this refers to an object in a PIE package
+				SetPath(MoveTemp(PIEPath));
+			}
+		}
+	}
+#endif
+}
+
+bool FSoftClassPath::SerializeFromMismatchedTag(struct FPropertyTag const& Tag, FArchive& Ar)
+{
+	struct UClassTypePolicy
+	{
+		typedef UClass Type;
+		// Class property shares the same tag id as Object property
+		static const FName FORCEINLINE GetTypeName() { return NAME_ObjectProperty; }
+	};
+
+	FString Path = ToString();
+
+	bool bReturn = SerializeFromMismatchedTagTemplate<UClassTypePolicy>(Path, Tag, Ar);
+
+	if (Ar.IsLoading())
+	{
+		SetPath(MoveTemp(Path));
+		PostLoadPath();
+	}
+
+	return bReturn;
+}
+
+UClass* FSoftClassPath::ResolveClass() const
+{
+	return Cast<UClass>(ResolveObject());
+}
+
+FSoftClassPath FSoftClassPath::GetOrCreateIDForClass(const UClass *InClass)
+{
+	check(InClass);
+	return FSoftClassPath(InClass);
+}
+
+bool FSoftObjectPathThreadContext::GetSerializationOptions(FName& OutPackageName, FName& OutPropertyName, ESoftObjectPathCollectType& OutCollectType) const
+{
+	FName CurrentPackageName, CurrentPropertyName;
+	ESoftObjectPathCollectType CurrentCollectType = ESoftObjectPathCollectType::AlwaysCollect;
+	bool bFoundAnything = false;
+	if (OptionStack.Num() > 0)
+	{
+		// Go from the top of the stack down
+		for (int32 i = OptionStack.Num() - 1; i >= 0; i--)
+		{
+			const FSerializationOptions& Options = OptionStack[i];
+			// Find first valid package/property names. They may not necessarily match
+			if (Options.PackageName != NAME_None && CurrentPackageName == NAME_None)
+			{
+				CurrentPackageName = Options.PackageName;
+			}
+			if (Options.PropertyName != NAME_None && CurrentPropertyName == NAME_None)
+			{
+				CurrentPropertyName = Options.PropertyName;
+			}
+
+			// Restrict based on lowest/most restrictive collect type
+			if (Options.CollectType < CurrentCollectType)
+			{
+				CurrentCollectType = Options.CollectType;
+			}
+		}
+
+		bFoundAnything = true;
+	}
+
+	// Check UObject thread context as a backup
+	FUObjectThreadContext& ThreadContext = FUObjectThreadContext::Get();
+	if (ThreadContext.SerializedObject)
+	{
+		FLinkerLoad* Linker = ThreadContext.SerializedObject->GetLinker();
+		if (Linker)
+		{
+			if (CurrentPackageName == NAME_None)
+			{
+				CurrentPackageName = FName(*FPackageName::FilenameToLongPackageName(Linker->Filename));
+			}
+			if (Linker->GetSerializedProperty() && CurrentPropertyName == NAME_None)
+			{
+				CurrentPropertyName = Linker->GetSerializedProperty()->GetFName();
+			}
+#if WITH_EDITORONLY_DATA
+			bool bEditorOnly = Linker->IsEditorOnlyPropertyOnTheStack();
+#else
+			bool bEditorOnly = false;
+#endif
+			// If we were always collect before and not overridden by stack options, set to editor only
+			if (bEditorOnly && CurrentCollectType == ESoftObjectPathCollectType::AlwaysCollect)
+			{
+				CurrentCollectType = ESoftObjectPathCollectType::EditorOnlyCollect;
+			}
+
+			bFoundAnything = true;
+		}
+	}
+
+	if (bFoundAnything)
+	{
+		OutPackageName = CurrentPackageName;
+		OutPropertyName = CurrentPropertyName;
+		OutCollectType = CurrentCollectType;
+		return true;
+	}
+
+	return bFoundAnything;
+}
+
+FThreadSafeCounter FSoftObjectPath::CurrentTag(1);
+TSet<FName> FSoftObjectPath::PIEPackageNames;
diff --git a/Engine/Source/Runtime/CoreUObject/Private/UObject/UObjectBase.cpp b/Engine/Source/Runtime/CoreUObject/Private/UObject/UObjectBase.cpp
--- a/Engine/Source/Runtime/CoreUObject/Private/UObject/UObjectBase.cpp
+++ b/Engine/Source/Runtime/CoreUObject/Private/UObject/UObjectBase.cpp
@@ -538,34 +538,6 @@ static TArray<FFieldCompiledInInfo*>& GetDeferredClassRegistration()
 }
 
 #if WITH_HOT_RELOAD
-TMap<UObject*, UObject*>& GetDuplicatedCDOMap()
-{
-	/**
-	 * GC referencer object, so duplicated CDOs aren't destroyed if still
-	 * referenced by the map.
-	 */
-	class FDuplicatedCDOMap : public FGCObject
-	{
-	public:
-		/**
-		 * Add referenced objects to reference collector.
-		 */
-		virtual void AddReferencedObjects(FReferenceCollector& Collector) override
-		{
-			for (auto& KVP : Map)
-			{
-				Collector.AddReferencedObject(KVP.Value);
-			}
-		}
-
-		/** Duplicated CDOs map. */
-		TMap<UObject*, UObject*> Map;
-	};
-
-	static FDuplicatedCDOMap Cache;
-	return Cache.Map;
-}
-
 /** Map of deferred class registration info (including size and reflection info) */
 static TMap<FName, FFieldCompiledInInfo*>& GetDeferRegisterClassMap()
 {
@@ -747,41 +719,6 @@ void UClassReplaceHotReloadClasses()
 	HotReloadClasses.Empty();
 }
 
-/**
- * Creates a cache of cpp-only-changed UClasses' CDOs, which are going to be
- * used later during BP reinstancing.
- */
-static void UClassGenerateCDODuplicatesForHotReload()
-{
-	if (!GIsHotReload)
-	{
-		return;
-	}
-
-	const TArray<FFieldCompiledInInfo*>& HotReloadClasses = GetHotReloadClasses();
-	TMap<UObject*, UObject*> DuplicatedCDOMap = GetDuplicatedCDOMap();
-	UPackage* TransientPackage = GetTransientPackage();
-
-	for (UClass* Class : TObjectRange<UClass>())
-	{
-		UObject* CDO = Class->GetDefaultObject();
-
-		for (const FFieldCompiledInInfo* HotReloadedClass : HotReloadClasses)
-		{
-			if (!HotReloadedClass->bHasChanged && Class->IsChildOf(HotReloadedClass->OldClass))
-			{
-				GIsDuplicatingClassForReinstancing = true;
-
-				FName UniqueName = MakeUniqueObjectName(TransientPackage, Class, TEXT("HOTRELOAD_CDO_DUPLICATE"));
-				UObject* DupCDO = StaticDuplicateObject(CDO, TransientPackage, UniqueName);
-
-				GIsDuplicatingClassForReinstancing = false;
-
-				DuplicatedCDOMap.Add(CDO, DupCDO);
-			}
-		}
-	}
-}
 #endif
 
 /**
@@ -874,11 +811,6 @@ void ProcessNewlyLoadedUObjects()
 {
 	DECLARE_SCOPE_CYCLE_COUNTER(TEXT("ProcessNewlyLoadedUObjects"), STAT_ProcessNewlyLoadedUObjects, STATGROUP_ObjectVerbose);
 
-
-
-#if WITH_HOT_RELOAD
-	UClassGenerateCDODuplicatesForHotReload();
-#endif
 	UClassRegisterAllCompiledInClasses();
 
 	const TArray<UClass* (*)()>& DeferredCompiledInRegistration = GetDeferredCompiledInRegistration();
diff --git a/Engine/Source/Runtime/CoreUObject/Private/UObject/UObjectGlobals.cpp b/Engine/Source/Runtime/CoreUObject/Private/UObject/UObjectGlobals.cpp
--- a/Engine/Source/Runtime/CoreUObject/Private/UObject/UObjectGlobals.cpp
+++ b/Engine/Source/Runtime/CoreUObject/Private/UObject/UObjectGlobals.cpp
@@ -130,15 +130,14 @@ FCoreUObjectDelegates::FPreLoadMapDelegate FCoreUObjectDelegates::PreLoadMap;
 FCoreUObjectDelegates::FPostLoadMapDelegate FCoreUObjectDelegates::PostLoadMapWithWorld;
 PRAGMA_DISABLE_DEPRECATION_WARNINGS
 FSimpleMulticastDelegate FCoreUObjectDelegates::PostLoadMap;
-FCoreUObjectDelegates::FStringAssetReferenceLoaded FCoreUObjectDelegates::StringAssetReferenceLoaded;
-FCoreUObjectDelegates::FStringAssetReferenceSaving FCoreUObjectDelegates::StringAssetReferenceSaving;
+FCoreUObjectDelegates::FSoftObjectPathLoaded FCoreUObjectDelegates::StringAssetReferenceLoaded;
+FCoreUObjectDelegates::FSoftObjectPathSaving FCoreUObjectDelegates::StringAssetReferenceSaving;
 FCoreUObjectDelegates::FOnRedirectorFollowed FCoreUObjectDelegates::RedirectorFollowed;
 PRAGMA_ENABLE_DEPRECATION_WARNINGS
 FSimpleMulticastDelegate FCoreUObjectDelegates::PostDemoPlay;
 FCoreUObjectDelegates::FOnLoadObjectsOnTop FCoreUObjectDelegates::ShouldLoadOnTop;
 
 FCoreUObjectDelegates::FPackageCreatedForLoad FCoreUObjectDelegates::PackageCreatedForLoad;
-FCoreUObjectDelegates::FPackageLoadedFromStringAssetReference FCoreUObjectDelegates::PackageLoadedFromStringAssetReference;
 FCoreUObjectDelegates::FGetPrimaryAssetIdForObject FCoreUObjectDelegates::GetPrimaryAssetIdForObject;
 
 /** Check whether we should report progress or not */
@@ -652,6 +651,16 @@ UPackage* CreatePackage( UObject* InOuter, const TCHAR* PackageName )
 
 FString ResolveIniObjectsReference(const FString& ObjectReference, const FString* IniFilename, bool bThrow)
 {
+	if (!IniFilename)
+	{
+		IniFilename = GetIniFilenameFromObjectsReference(ObjectReference);
+	}
+
+	if (!IniFilename)
+	{
+		return ObjectReference;
+	}
+
 	// Get .ini key and section.
 	FString Section = ObjectReference.Mid(1 + ObjectReference.Find(TEXT(":"), ESearchCase::CaseSensitive));
 	int32 i = Section.Find(TEXT("."), ESearchCase::CaseSensitive, ESearchDir::FromEnd);
@@ -668,7 +677,7 @@ FString ResolveIniObjectsReference(const FString& ObjectReference, const FString
 	if (!GConfig->GetString(*Section, *Key, Output, *IniFilename))
 	{
 		if (bThrow == true)
-{
+		{
 			UE_LOG(LogUObjectGlobals, Error, TEXT(" %s %s "), *FString::Printf(TEXT("Can't find '%s' in configuration file section=%s key=%s"), *ObjectReference, *Section, *Key), **IniFilename);
 		}
 	}
@@ -704,13 +713,6 @@ const FString* GetIniFilenameFromObjectsReference(const FString& Name)
 //
 bool ResolveName(UObject*& InPackage, FString& InOutName, bool Create, bool Throw, uint32 LoadFlags /*= LOAD_None*/)
 {
-	const FString* IniFilename = GetIniFilenameFromObjectsReference(InOutName);
-
-	if (IniFilename && InOutName.Contains(TEXT("."), ESearchCase::CaseSensitive))
-	{
-		InOutName = ResolveIniObjectsReference(InOutName, IniFilename, Throw);
-	}
-
 	// Strip off the object class.
 	ConstructorHelpers::StripObjectClass( InOutName );
 
@@ -1683,7 +1685,7 @@ void EndLoad()
 	}
 
 	// Loaded new objects, so allow reaccessing asset ptrs
-	FStringAssetReference::InvalidateTag();
+	FSoftObjectPath::InvalidateTag();
 }
 
 /*-----------------------------------------------------------------------------
@@ -2245,10 +2247,8 @@ UObject* StaticAllocateObject
 	if (bCreatingCDO)
 	{
 		check(InClass->GetClass());
-		if( !GIsDuplicatingClassForReinstancing || InClass->HasAnyClassFlags(CLASS_Native) )
-		{
-			InName = InClass->GetDefaultObjectName();
-		}
+		ensure(!GIsDuplicatingClassForReinstancing || InClass->HasAnyClassFlags(CLASS_Native));
+		InName = InClass->GetDefaultObjectName();
 		// never call PostLoad on class default objects
 		InFlags &= ~(RF_NeedPostLoad|RF_NeedPostLoadSubobjects);
 	}
@@ -4047,10 +4047,10 @@ namespace UE4CodeGen_Private
 			}
 			break;
 
-			case EPropertyClass::AssetObject:
+			case EPropertyClass::SoftObject:
 			{
-				const FAssetObjectPropertyParams* Prop = (const FAssetObjectPropertyParams*)PropBase;
-				NewProp = new (EC_InternalUseOnlyConstructor, Outer, UTF8_TO_TCHAR(Prop->NameUTF8), Prop->ObjectFlags) UAssetObjectProperty(FObjectInitializer(), EC_CppProperty, Prop->Offset, Prop->PropertyFlags, Prop->ClassFunc ? Prop->ClassFunc() : nullptr);
+				const FSoftObjectPropertyParams* Prop = (const FSoftObjectPropertyParams*)PropBase;
+				NewProp = new (EC_InternalUseOnlyConstructor, Outer, UTF8_TO_TCHAR(Prop->NameUTF8), Prop->ObjectFlags) USoftObjectProperty(FObjectInitializer(), EC_CppProperty, Prop->Offset, Prop->PropertyFlags, Prop->ClassFunc ? Prop->ClassFunc() : nullptr);
 
 #if WITH_METADATA
 				MetaDataArray = Prop->MetaDataArray;
@@ -4071,10 +4071,10 @@ namespace UE4CodeGen_Private
 			}
 			break;
 
-			case EPropertyClass::AssetClass:
+			case EPropertyClass::SoftClass:
 			{
-				const FAssetClassPropertyParams* Prop = (const FAssetClassPropertyParams*)PropBase;
-				NewProp = new (EC_InternalUseOnlyConstructor, Outer, UTF8_TO_TCHAR(Prop->NameUTF8), Prop->ObjectFlags) UAssetClassProperty(FObjectInitializer(), EC_CppProperty, Prop->Offset, Prop->PropertyFlags, Prop->MetaClassFunc ? Prop->MetaClassFunc() : nullptr);
+				const FSoftClassPropertyParams* Prop = (const FSoftClassPropertyParams*)PropBase;
+				NewProp = new (EC_InternalUseOnlyConstructor, Outer, UTF8_TO_TCHAR(Prop->NameUTF8), Prop->ObjectFlags) USoftClassProperty(FObjectInitializer(), EC_CppProperty, Prop->Offset, Prop->PropertyFlags, Prop->MetaClassFunc ? Prop->MetaClassFunc() : nullptr);
 
 #if WITH_METADATA
 				MetaDataArray = Prop->MetaDataArray;
diff --git a/Engine/Source/Runtime/CoreUObject/Public/Blueprint/BlueprintSupport.h b/Engine/Source/Runtime/CoreUObject/Public/Blueprint/BlueprintSupport.h
--- a/Engine/Source/Runtime/CoreUObject/Public/Blueprint/BlueprintSupport.h
+++ b/Engine/Source/Runtime/CoreUObject/Public/Blueprint/BlueprintSupport.h
@@ -275,7 +275,7 @@ struct FBlueprintDependencyType
 struct COREUOBJECT_API FCompactBlueprintDependencyData
 {
 	int16 ObjectRefIndex;
-	FBlueprintDependencyType ClassDependency;
+	FBlueprintDependencyType StructDependency;
 	FBlueprintDependencyType CDODependency;
 
 	FCompactBlueprintDependencyData()
@@ -283,10 +283,10 @@ struct COREUOBJECT_API FCompactBlueprintDependencyData
 	{}
 
 	FCompactBlueprintDependencyData(int16 InObjectRefIndex
-		, FBlueprintDependencyType InClassDependency
-		, FBlueprintDependencyType InCDODependency)
+		, FBlueprintDependencyType InStructDependency
+		, FBlueprintDependencyType InCDODependency = FBlueprintDependencyType())
 		: ObjectRefIndex(InObjectRefIndex)
-		, ClassDependency(InClassDependency)
+		, StructDependency(InStructDependency)
 		, CDODependency(InCDODependency)
 	{}
 };
@@ -310,21 +310,19 @@ struct COREUOBJECT_API FBlueprintDependencyObjectRef
 struct COREUOBJECT_API FBlueprintDependencyData
 {
 	FBlueprintDependencyObjectRef ObjectRef;
-	// 0 - dependency type for dynamic class
+	// 0 - dependency type for dynamic class or UDS
 	// 1 - dependency type for CD0
 	FBlueprintDependencyType DependencyTypes[2];
 
 	int16 ObjectRefIndex; // NativizationWithoutEDLBT
 
 	FBlueprintDependencyData(const FBlueprintDependencyObjectRef& InObjectRef
-		, FBlueprintDependencyType InClassDependency
-		, FBlueprintDependencyType InCDODependency
-		, int16 InObjectRefIndex)
+		, const FCompactBlueprintDependencyData& InCompactDependencyData)
 		: ObjectRef(InObjectRef)
-		, ObjectRefIndex(InObjectRefIndex)
+		, ObjectRefIndex(InCompactDependencyData.ObjectRefIndex)
 	{
-		DependencyTypes[0] = InClassDependency;
-		DependencyTypes[1] = InCDODependency;
+		DependencyTypes[0] = InCompactDependencyData.StructDependency;
+		DependencyTypes[1] = InCompactDependencyData.CDODependency;
 	}
 
 	bool operator==(const FBlueprintDependencyData& Other) const
diff --git a/Engine/Source/Runtime/CoreUObject/Public/CoreUObject.h b/Engine/Source/Runtime/CoreUObject/Public/CoreUObject.h
--- a/Engine/Source/Runtime/CoreUObject/Public/CoreUObject.h
+++ b/Engine/Source/Runtime/CoreUObject/Public/CoreUObject.h
@@ -36,8 +36,7 @@ MONOLITHIC_HEADER_BOILERPLATE()
 #include "Templates/Casts.h"
 #include "UObject/PersistentObjectPtr.h"
 #include "UObject/LazyObjectPtr.h"
-#include "Misc/StringAssetReference.h"
-#include "UObject/AssetPtr.h"
+#include "UObject/SoftObjectPtr.h"
 #include "UObject/Interface.h"
 #include "UObject/LevelGuids.h"
 #include "UObject/Package.h"
@@ -94,4 +93,3 @@ MONOLITHIC_HEADER_BOILERPLATE()
 #include "Misc/RedirectCollector.h"
 #include "UObject/ScriptStackTracker.h"
 #include "Misc/WorldCompositionUtility.h"
-#include "Misc/StringClassReference.h"
diff --git a/Engine/Source/Runtime/CoreUObject/Public/Misc/RedirectCollector.h b/Engine/Source/Runtime/CoreUObject/Public/Misc/RedirectCollector.h
--- a/Engine/Source/Runtime/CoreUObject/Public/Misc/RedirectCollector.h
+++ b/Engine/Source/Runtime/CoreUObject/Public/Misc/RedirectCollector.h
@@ -24,7 +24,6 @@ private:
 		: Package(InPackage)
 		, Property(InProperty)
 		, bReferencedByEditorOnlyProperty(false)
-		//, PackageCache(NAME_None)
 		{}
 
 		bool operator==(const FPackagePropertyPair& Other) const
@@ -36,14 +35,8 @@ private:
 		const FName& GetCachedPackageName() const
 		{
 			return Package;
-			/*if (PackageCache == NAME_None)
-			{
-				PackageCache = FName(*Package);
-			}
-			return PackageCache;*/
 		}
 
-
 		void SetPackage(const FName& InPackage)
 		{
 			Package = InPackage;
@@ -76,67 +69,69 @@ private:
 		FName Package;
 		FName Property;
 		bool bReferencedByEditorOnlyProperty;
-		// mutable FName PackageCache;
 	};
 
 public:
 
 	/**
-	 * Called from FStringAssetReference::PostLoadPath, registers this for later querying
-	 * @param InString Name of the asset that was loaded
+	 * Called from FSoftObjectPath::PostLoadPath, registers this for later querying
+	 * @param InPath The soft object path that was loaded
 	 */
-	void OnStringAssetReferenceLoaded(const FString& InString);
+	void OnSoftObjectPathLoaded(const struct FSoftObjectPath& InPath);
 
 	/**
-	 * Load the string asset references to resolve them, add that to the remap table, and empty the array
-	 * @param FilterPackage If set, only fixup references that were created by FilterPackage. If empty, clear all of them
-	 * @param bProcessAlreadyResolvedPackages If a package has already been resolved in a previous call of this function, should we process this package again and resolve it again?  
-	 *             True = process it again, False = don't process it again
+	 * Load all soft object paths to resolve them, add that to the remap table, and empty the array
+	 * @param FilterPackage If set, only load references that were created by FilterPackage. If empty, resolve  all of them
 	 */
-	void ResolveStringAssetReference(FName FilterPackage = NAME_None, bool bProcessAlreadyResolvedPackages = true);
+	void ResolveAllSoftObjectPaths(FName FilterPackage = NAME_None);
 
 	/**
-	 * Returns the list of packages referenced by string asset references loaded by FilterPackage, and remove them from the internal list
+	 * Returns the list of packages referenced by soft object paths loaded by FilterPackage, and remove them from the internal list
 	 * @param FilterPackage Return references made by loading this package. If passed null will return all references made with no explicit package
-	 * @param OutReferencedPackages Return list of packages referenced by FilterPackage
 	 * @param bGetEditorOnly If true will return references loaded by editor only objects, if false it will not
+	 * @param OutReferencedPackages Return list of packages referenced by FilterPackage
 	 */
-	void ProcessStringAssetReferencePackageList(FName FilterPackage, bool bGetEditorOnly, TSet<FName>& OutReferencedPackages);
+	void ProcessSoftObjectPathPackageList(FName FilterPackage, bool bGetEditorOnly, TSet<FName>& OutReferencedPackages);
 
 	/** Adds a new mapping for redirector path to destination path, this is called from the Asset Registry to register all redirects it knows about */
-	void AddAssetPathRedirection(const FString& OriginalPath, const FString& RedirectedPath);
+	void AddAssetPathRedirection(FName OriginalPath, FName RedirectedPath);
 
 	/** Removes an asset path redirection, call this when deleting redirectors */
-	void RemoveAssetPathRedirection(const FString& OriginalPath);
+	void RemoveAssetPathRedirection(FName OriginalPath);
 
 	/** Returns a remapped asset path, if it returns null there is no relevant redirector */
-	FString* GetAssetPathRedirection(const FString& OriginalPath);
+	FName GetAssetPathRedirection(FName OriginalPath);
 
 	/**
 	 * Do we have any references to resolve.
 	 * @return true if we have references to resolve
 	 */
-	bool HasAnyStringAssetReferencesToResolve() const
+	bool HasAnySoftObjectPathsToResolve() const
 	{
-		return StringAssetReferences.Num() > 0;
+		return SoftObjectPathMap.Num() > 0;
 	}
 
-	/** Writes out log info for how long string asset references took to resolve */
-	void LogTimers() const;
-
 	DEPRECATED(4.17, "OnStringAssetReferenceSaved is deprecated, call GetAssetPathRedirection")
 	FString OnStringAssetReferenceSaved(const FString& InString);
 
+	DEPRECATED(4.18, "OnStringAssetReferenceLoaded is deprecated, call OnSoftObjectPathLoaded")
+	void OnStringAssetReferenceLoaded(const FString& InString);
+
+	DEPRECATED(4.18, "ResolveStringAssetReference is deprecated, call ResolveAllSoftObjectPaths")
+	void ResolveStringAssetReference(FName FilterPackage = NAME_None, bool bProcessAlreadyResolvedPackages = true)
+	{
+		ResolveAllSoftObjectPaths(FilterPackage);
+	}
+
 private:
 
-	/** A gathered list string asset references , with the key being the string reference (GetPathName()) and the value equal to the package with the reference */
-	TMultiMap<FName, FPackagePropertyPair> StringAssetReferences;
+	/** A map of assets referenced by soft object paths, with the key being the asset being referenced and the value equal to the package with the reference */
+	TMultiMap<FName, FPackagePropertyPair> SoftObjectPathMap;
 
 	/** When saving, apply this remapping to all string asset references */
-	TMap<FString, FString> AssetPathRedirectionMap;
-	TSet<FName> AlreadyRemapped;
+	TMap<FName, FName> AssetPathRedirectionMap;
 
-	/** For StringAssetReferences map */
+	/** For SoftObjectPackageMap map */
 	FCriticalSection CriticalSection;
 };
 
diff --git a/Engine/Source/Runtime/CoreUObject/Public/Misc/StringAssetReference.h b/Engine/Source/Runtime/CoreUObject/Public/Misc/StringAssetReference.h
--- a/Engine/Source/Runtime/CoreUObject/Public/Misc/StringAssetReference.h
+++ b/Engine/Source/Runtime/CoreUObject/Public/Misc/StringAssetReference.h
@@ -2,210 +2,5 @@
 
 #pragma once
 
-#include "CoreMinimal.h"
-#include "HAL/ThreadSafeCounter.h"
-#include "UObject/Class.h"
-#include "HAL/ThreadSingleton.h"
-
-/**
- * A struct that contains a string reference to an asset on disk.
- * This can be used to make soft references to assets that are loaded on demand.
- * This is stored internally as a string of the form /package/path.assetname[.objectname]
- * If the MetaClass metadata is applied to a UProperty with this the UI will restrict to that type of asset.
- */
-struct COREUOBJECT_API FStringAssetReference
-{
-	FStringAssetReference() {}
-
-	/** Construct from another asset reference */
-	FStringAssetReference(const FStringAssetReference& Other)
-	{
-		SetPath(Other.ToString());
-	}
-
-	/** Construct from a path string */
-	FStringAssetReference(FString PathString)
-	{
-		SetPath(MoveTemp(PathString));
-	}
-
-	/** Construct from an existing object in memory */
-	FStringAssetReference(const UObject* InObject);
-
-	~FStringAssetReference() {}
-
-	/** Returns string representation of reference, in form /package/path.assetname[.objectname] */
-	FORCEINLINE const FString& ToString() const
-	{
-		return AssetLongPathname;		
-	}
-
-	/** Returns /package/path, leaving off the asset name */
-	FString GetLongPackageName() const
-	{
-		FString PackageName;
-		ToString().Split(TEXT("."), &PackageName, nullptr, ESearchCase::CaseSensitive, ESearchDir::FromStart);
-		return PackageName;
-	}
-	
-	/** Returns assetname string, leaving off the /package/path part */
-	const FString GetAssetName() const
-	{
-		FString AssetName;
-		ToString().Split(TEXT("."), nullptr, &AssetName, ESearchCase::CaseSensitive, ESearchDir::FromStart);
-		return AssetName;
-	}	
-
-	/** Sets asset path of this reference based on a string path */
-	void SetPath(FString Path);
-
-	/**
-	 * Attempts to load the asset, this will call LoadObject which can be very slow
-	 *
-	 * @return Loaded UObject, or nullptr if the reference is null or the asset fails to load
-	 */
-	UObject* TryLoad() const;
-
-	/**
-	 * Attempts to find a currently loaded object that matches this path
-	 *
-	 * @return Found UObject, or nullptr if not currently in memory
-	 */
-	UObject* ResolveObject() const;
-
-	/** Resets reference to point to null */
-	void Reset()
-	{		
-		SetPath(TEXT(""));
-	}
-	
-	/** Check if this could possibly refer to a real object, or was initialized to null */
-	bool IsValid() const
-	{
-		return ToString().Len() > 0;
-	}
-
-	/** Checks to see if this is initialized to null */
-	bool IsNull() const
-	{
-		return ToString().Len() == 0;
-	}
-
-	/** Struct overrides */
-	bool Serialize(FArchive& Ar);
-	bool operator==(FStringAssetReference const& Other) const;
-	bool operator!=(FStringAssetReference const& Other) const
-	{
-		return !(*this == Other);
-	}
-	FStringAssetReference& operator=(FStringAssetReference const& Other);
-	bool ExportTextItem(FString& ValueStr, FStringAssetReference const& DefaultValue, UObject* Parent, int32 PortFlags, UObject* ExportRootScope) const;
-	bool ImportTextItem( const TCHAR*& Buffer, int32 PortFlags, UObject* Parent, FOutputDevice* ErrorText );
-	bool SerializeFromMismatchedTag(struct FPropertyTag const& Tag, FArchive& Ar);
-
-	/** Serializes the internal path and also handles save/PIE fixups. Call this from the archiver overrides */
-	void SerializePath(FArchive& Ar, bool bSkipSerializeIfArchiveHasSize = false);
-
-	/** Fixes up string asset reference for saving, call if saving with a method that skips SerializePath. This can modify the path, it will return true if it was modified */
-	bool PreSavePath();
-
-	/** Handles when a string asset reference has been loaded, call if loading with a method that skips SerializePath. This does not modify path but might call callbacks */
-	void PostLoadPath() const;
-
-	FORCEINLINE friend uint32 GetTypeHash(FStringAssetReference const& This)
-	{
-		return GetTypeHash(This.ToString());
-	}
-
-	/** Code needed by FAssetPtr internals */
-	static int32 GetCurrentTag()
-	{
-		return CurrentTag.GetValue();
-	}
-	static int32 InvalidateTag()
-	{
-		return CurrentTag.Increment();
-	}
-	static FStringAssetReference GetOrCreateIDForObject(const UObject* Object);
-	
-	/** Enables special PIE path handling, to remove pie prefix as needed */
-	static void SetPackageNamesBeingDuplicatedForPIE(const TArray<FString>& InPackageNamesBeingDuplicatedForPIE);
-	
-	/** Disables special PIE path handling */
-	static void ClearPackageNamesBeingDuplicatedForPIE();
-
-private:
-	/** Asset path */
-	FString AssetLongPathname;
-
-	/** Fixes up this StringAssetReference to add or remove the PIE prefix depending on what is currently active */
-	void FixupForPIE();
-
-	/** Global counter that determines when we need to re-search based on path because more objects have been loaded **/
-	static FThreadSafeCounter CurrentTag;
-
-	/** Package names currently being duplicated, needed by FixupForPIE */
-	static TArray<FString> PackageNamesBeingDuplicatedForPIE;
-
-	friend COREUOBJECT_API UScriptStruct* Z_Construct_UScriptStruct_FStringAssetReference();
-};
-
-enum class EStringAssetReferenceCollectType : uint8
-{
-	/** References is not tracked in any situation, transient reference */
-	NeverCollect,
-	/** Editor only reference, this is tracked for redirector fixup but not for cooking */
-	EditorOnlyCollect,
-	/** Game reference, this is gathered for both redirector fixup and cooking */
-	AlwaysCollect,
-};
-
-class COREUOBJECT_API FStringAssetReferenceThreadContext : public TThreadSingleton<FStringAssetReferenceThreadContext>
-{
-	friend TThreadSingleton<FStringAssetReferenceThreadContext>;
-	friend struct FStringAssetReferenceSerializationScope;
-
-	FStringAssetReferenceThreadContext() {}
-
-	struct FSerializationOptions
-	{
-		FName PackageName;
-		FName PropertyName;
-		EStringAssetReferenceCollectType CollectType;
-
-		FSerializationOptions() : CollectType(EStringAssetReferenceCollectType::AlwaysCollect) {}
-		FSerializationOptions(FName InPackageName, FName InPropertyName, EStringAssetReferenceCollectType InCollectType) : PackageName(InPackageName), PropertyName(InPropertyName), CollectType(InCollectType) {}
-	};
-
-	TArray<FSerializationOptions> OptionStack;
-public:
-	/** 
-	 * Returns the current serialization options that were added using SerializationScope or LinkerLoad
-	 *
-	 * @param OutPackageName Package that this string asset belongs to
-	 * @param OutPropertyName Property that this string asset reference belongs to
-	 * @param OutCollectType Type of collecting that should be done
-	 */
-	bool GetSerializationOptions(FName& OutPackageName, FName& OutPropertyName, EStringAssetReferenceCollectType& OutCollectType) const;
-};
-
-/** Helper class to set and restore serialization options for string asset references */
-struct FStringAssetReferenceSerializationScope
-{
-	/** 
-	 * Create a new serialization scope, which affects the way that string asset references are saved
-	 *
-	 * @param SerializingPackageName Package that this string asset belongs to
-	 * @param SerializingPropertyName Property that this string asset reference belongs to
-	 * @param CollectType Set type of collecting that should be done, can be used to disable tracking entirely
-	 */
-	FStringAssetReferenceSerializationScope(FName SerializingPackageName, FName SerializingPropertyName, EStringAssetReferenceCollectType CollectType)
-	{
-		FStringAssetReferenceThreadContext::Get().OptionStack.Emplace(SerializingPackageName, SerializingPropertyName, CollectType);
-	}
-
-	~FStringAssetReferenceSerializationScope()
-	{
-		FStringAssetReferenceThreadContext::Get().OptionStack.Pop();
-	}
-};
\ No newline at end of file
+COMPILE_WARNING("FStringAssetReference has been renamed to FSoftObjectPath, change to #include \"UObject/SoftObjectPath.h\" and rename references")
+#include "UObject/SoftObjectPath.h"
diff --git a/Engine/Source/Runtime/CoreUObject/Public/Misc/StringClassReference.h b/Engine/Source/Runtime/CoreUObject/Public/Misc/StringClassReference.h
--- a/Engine/Source/Runtime/CoreUObject/Public/Misc/StringClassReference.h
+++ b/Engine/Source/Runtime/CoreUObject/Public/Misc/StringClassReference.h
@@ -2,68 +2,5 @@
 
 #pragma once
 
-#include "CoreMinimal.h"
-#include "UObject/Class.h"
-#include "Misc/StringAssetReference.h"
-
-struct FPropertyTag;
-
-/**
- * A struct that contains a string reference to a class, can be used to make soft references to classes
- */
-struct COREUOBJECT_API FStringClassReference : public FStringAssetReference
-{
-	FStringClassReference()
-	{ }
-
-	FStringClassReference(FStringClassReference const& Other)
-		: FStringAssetReference(Other)
-	{ }
-
-	/**
-	 * Construct from a path string
-	 */
-	FStringClassReference(const FString& PathString)
-		: FStringAssetReference(PathString)
-	{ }
-
-	/**
-	 * Construct from an existing class, will do some string processing
-	 */
-	FStringClassReference(const UClass* InClass)
-		: FStringAssetReference(InClass)
-	{ }
-
-	/**
-	* Attempts to load the class.
-	* @return Loaded UObject, or null if the class fails to load, or if the reference is not valid.
-	*/
-	template<typename T>
-	UClass* TryLoadClass() const
-	{
-		if ( IsValid() )
-		{
-			return LoadClass<T>(nullptr, *ToString(), nullptr, LOAD_None, nullptr);
-		}
-
-		return nullptr;
-	}
-
-	/**
-	 * Attempts to find a currently loaded object that matches this object ID
-	 * @return Found UClass, or NULL if not currently loaded
-	 */
-	UClass* ResolveClass() const;
-
-	bool SerializeFromMismatchedTag(const FPropertyTag& Tag, FArchive& Ar);
-
-	static FStringClassReference GetOrCreateIDForClass(const UClass *InClass);
-
-private:
-	/* Forbid creation for UObject. This class is for UClass only. Use FStringAssetReference instead. */
-	FStringClassReference(const UObject* InObject) { }
-
-	/* Forbidden. This class is for UClass only. Use FStringAssetReference instead. */
-	static FStringAssetReference GetOrCreateIDForObject(const UObject *Object);
-};
-
+COMPILE_WARNING("FStringClassReference has been renamed to FSoftClassPath and merged headers, change to #include \"UObject/SoftObjectPath.h\" and rename references")
+#include "UObject/SoftObjectPath.h"
\ No newline at end of file
diff --git a/Engine/Source/Runtime/CoreUObject/Public/Serialization/ArchiveUObject.h b/Engine/Source/Runtime/CoreUObject/Public/Serialization/ArchiveUObject.h
--- a/Engine/Source/Runtime/CoreUObject/Public/Serialization/ArchiveUObject.h
+++ b/Engine/Source/Runtime/CoreUObject/Public/Serialization/ArchiveUObject.h
@@ -4,8 +4,13 @@
 
 #include "CoreMinimal.h"
 
+struct FLazyObjectPtr;
+struct FSoftObjectPtr;
+struct FSoftObjectPath;
+struct FWeakObjectPtr;
+
 /**
- * Base FArchive for serializing UObjects. Supports FLazyObjectPtr and FAssetPtr serialization.
+ * Base FArchive for serializing UObjects. Supports FLazyObjectPtr and FSoftObjectPtr serialization.
  */
 class COREUOBJECT_API FArchiveUObject : public FArchive
 {
@@ -14,9 +19,9 @@ public:
 	using FArchive::operator<<; // For visibility of the overloads we don't override
 
 	//~ Begin FArchive Interface
-	virtual FArchive& operator<< (class FLazyObjectPtr& Value) override;
-	virtual FArchive& operator<< (class FAssetPtr& Value) override;
-	virtual FArchive& operator<< (struct FStringAssetReference& Value) override;
-	virtual FArchive& operator<< (struct FWeakObjectPtr& Value) override;
+	virtual FArchive& operator<<(FLazyObjectPtr& Value) override;
+	virtual FArchive& operator<<(FSoftObjectPtr& Value) override;
+	virtual FArchive& operator<<(FSoftObjectPath& Value) override;
+	virtual FArchive& operator<<(FWeakObjectPtr& Value) override;
 	//~ End FArchive Interface
 };
diff --git a/Engine/Source/Runtime/CoreUObject/Public/Serialization/DuplicatedDataReader.h b/Engine/Source/Runtime/CoreUObject/Public/Serialization/DuplicatedDataReader.h
--- a/Engine/Source/Runtime/CoreUObject/Public/Serialization/DuplicatedDataReader.h
+++ b/Engine/Source/Runtime/CoreUObject/Public/Serialization/DuplicatedDataReader.h
@@ -7,9 +7,6 @@
 #include "UObject/UObjectAnnotation.h"
 #include "Serialization/DuplicatedObject.h"
 
-class FLazyObjectPtr;
-struct FStringAssetReference;
-
 /*----------------------------------------------------------------------------
 	FDuplicateDataReader.
 ----------------------------------------------------------------------------*/
@@ -30,7 +27,7 @@ private:
 	virtual FArchive& operator<<(FName& N);
 	virtual FArchive& operator<<(UObject*& Object);
 	virtual FArchive& operator<<(FLazyObjectPtr& LazyObjectPtr);
-	virtual FArchive& operator<<(FStringAssetReference& StringAssetReference);
+	virtual FArchive& operator<<(FSoftObjectPath& SoftObjectPath);
 	
 	void SerializeFail();
 
diff --git a/Engine/Source/Runtime/CoreUObject/Public/Serialization/DuplicatedDataWriter.h b/Engine/Source/Runtime/CoreUObject/Public/Serialization/DuplicatedDataWriter.h
--- a/Engine/Source/Runtime/CoreUObject/Public/Serialization/DuplicatedDataWriter.h
+++ b/Engine/Source/Runtime/CoreUObject/Public/Serialization/DuplicatedDataWriter.h
@@ -7,7 +7,6 @@
 #include "UObject/UObjectAnnotation.h"
 #include "Serialization/DuplicatedObject.h"
 
-class FLazyObjectPtr;
 struct FObjectInstancingGraph;
 
 /*----------------------------------------------------------------------------
diff --git a/Engine/Source/Runtime/CoreUObject/Public/Serialization/ObjectReader.h b/Engine/Source/Runtime/CoreUObject/Public/Serialization/ObjectReader.h
--- a/Engine/Source/Runtime/CoreUObject/Public/Serialization/ObjectReader.h
+++ b/Engine/Source/Runtime/CoreUObject/Public/Serialization/ObjectReader.h
@@ -8,9 +8,10 @@
 #include "Internationalization/TextPackageNamespaceUtil.h"
 #include "UObject/PropertyPortFlags.h"
 
-class FAssetPtr;
-class FLazyObjectPtr;
-struct FStringAssetReference;
+struct FLazyObjectPtr;
+struct FSoftObjectPtr;
+struct FSoftObjectPath;
+struct FWeakObjectPtr;
 
 /**
  * UObject Memory Reader Archive. Reads from InBytes, writes to Obj.
@@ -60,12 +61,12 @@ public:
 		}
 	}
 
-	COREUOBJECT_API virtual FArchive& operator<<( class FName& N ) override;
-	COREUOBJECT_API virtual FArchive& operator<<( class UObject*& Res ) override;
-	COREUOBJECT_API virtual FArchive& operator<<( FLazyObjectPtr& LazyObjectPtr ) override;
-	COREUOBJECT_API virtual FArchive& operator<<( FAssetPtr& AssetPtr ) override;
-	COREUOBJECT_API virtual FArchive& operator<<(FStringAssetReference& AssetPtr) override;
-	COREUOBJECT_API virtual FArchive& operator<<(struct FWeakObjectPtr& Value) override;
+	COREUOBJECT_API virtual FArchive& operator<<(FName& N) override;
+	COREUOBJECT_API virtual FArchive& operator<<(UObject*& Res) override;
+	COREUOBJECT_API virtual FArchive& operator<<(FLazyObjectPtr& LazyObjectPtr) override;
+	COREUOBJECT_API virtual FArchive& operator<<(FSoftObjectPtr& Value) override;
+	COREUOBJECT_API virtual FArchive& operator<<(FSoftObjectPath& Value) override;
+	COREUOBJECT_API virtual FArchive& operator<<(FWeakObjectPtr& Value) override;
 	COREUOBJECT_API virtual FString GetArchiveName() const override;
 	//~ End FArchive Interface
 
diff --git a/Engine/Source/Runtime/CoreUObject/Public/Serialization/ObjectWriter.h b/Engine/Source/Runtime/CoreUObject/Public/Serialization/ObjectWriter.h
--- a/Engine/Source/Runtime/CoreUObject/Public/Serialization/ObjectWriter.h
+++ b/Engine/Source/Runtime/CoreUObject/Public/Serialization/ObjectWriter.h
@@ -8,9 +8,10 @@
 #include "Internationalization/TextPackageNamespaceUtil.h"
 #include "UObject/PropertyPortFlags.h"
 
-class FAssetPtr;
-class FLazyObjectPtr;
-struct FStringAssetReference;
+struct FLazyObjectPtr;
+struct FSoftObjectPtr;
+struct FSoftObjectPath;
+struct FWeakObjectPtr;
 
 /**
  * UObject Memory Writer Archive.
@@ -38,12 +39,12 @@ public:
 	}
 
 	//~ Begin FArchive Interface
-	COREUOBJECT_API virtual FArchive& operator<<( class FName& N ) override;
-	COREUOBJECT_API virtual FArchive& operator<<( class UObject*& Res ) override;
-	COREUOBJECT_API virtual FArchive& operator<<( FLazyObjectPtr& LazyObjectPtr ) override;
-	COREUOBJECT_API virtual FArchive& operator<<( FAssetPtr& AssetPtr ) override;
-	COREUOBJECT_API virtual FArchive& operator<<(FStringAssetReference& AssetPtr) override;
-	COREUOBJECT_API virtual FArchive& operator<<(struct FWeakObjectPtr& Value) override;
+	COREUOBJECT_API virtual FArchive& operator<<(FName& N) override;
+	COREUOBJECT_API virtual FArchive& operator<<(UObject*& Res) override;
+	COREUOBJECT_API virtual FArchive& operator<<(FLazyObjectPtr& Value) override;
+	COREUOBJECT_API virtual FArchive& operator<<(FSoftObjectPtr& Value) override;
+	COREUOBJECT_API virtual FArchive& operator<<(FSoftObjectPath& Value) override;
+	COREUOBJECT_API virtual FArchive& operator<<(FWeakObjectPtr& Value) override;
 	COREUOBJECT_API virtual FString GetArchiveName() const override;
 	//~ End FArchive Interface
 
diff --git a/Engine/Source/Runtime/CoreUObject/Public/Templates/Casts.h b/Engine/Source/Runtime/CoreUObject/Public/Templates/Casts.h
--- a/Engine/Source/Runtime/CoreUObject/Public/Templates/Casts.h
+++ b/Engine/Source/Runtime/CoreUObject/Public/Templates/Casts.h
@@ -10,7 +10,7 @@
 class AActor;
 class APawn;
 class APlayerController;
-class UAssetClassProperty;
+class USoftClassProperty;
 class UBlueprint;
 class ULevel;
 class UPrimitiveComponent;
@@ -312,8 +312,8 @@ DECLARE_CAST_BY_FLAG(UObjectPropertyBase)				\
 DECLARE_CAST_BY_FLAG(UObjectProperty)					\
 DECLARE_CAST_BY_FLAG(UWeakObjectProperty)				\
 DECLARE_CAST_BY_FLAG(ULazyObjectProperty)				\
-DECLARE_CAST_BY_FLAG(UAssetObjectProperty)				\
-DECLARE_CAST_BY_FLAG(UAssetClassProperty)				\
+DECLARE_CAST_BY_FLAG(USoftObjectProperty)				\
+DECLARE_CAST_BY_FLAG(USoftClassProperty)				\
 DECLARE_CAST_BY_FLAG(UBoolProperty)						\
 DECLARE_CAST_BY_FLAG(UFunction)							\
 DECLARE_CAST_BY_FLAG(UStructProperty)					\
diff --git a/Engine/Source/Runtime/CoreUObject/Public/UObject/AssetPtr.h b/Engine/Source/Runtime/CoreUObject/Public/UObject/AssetPtr.h
--- a/Engine/Source/Runtime/CoreUObject/Public/UObject/AssetPtr.h
+++ b/Engine/Source/Runtime/CoreUObject/Public/UObject/AssetPtr.h
@@ -1,596 +1,6 @@
 // Copyright 1998-2017 Epic Games, Inc. All Rights Reserved.
 
-/*=============================================================================
-	AssetPtr.h: Pointer to UObject asset, keeps extra information so that it is works even if the asset is not in memory
-=============================================================================*/
-
 #pragma once
 
-#include "CoreMinimal.h"
-#include "UObject/Object.h"
-#include "Templates/Casts.h"
-#include "UObject/PersistentObjectPtr.h"
-#include "Misc/StringAssetReference.h"
-#include "PackageName.h"
-
-/**
- * FAssetPtr is a type of weak pointer to a UObject, that also keeps track of the path to the object on disk.
- * It will change back and forth between being Valid and Pending as the referenced object loads or unloads.
- * It has no impact on if the object is garbage collected or not.
- *
- * This is useful to specify assets that you may want to asynchronously load on demand.
- */
-class FAssetPtr : public TPersistentObjectPtr<FStringAssetReference>
-{
-public:	
-	/** Default constructor, will be null */
-	FORCEINLINE FAssetPtr()
-	{
-	}
-
-	/** Construct from another asset pointer */
-	FORCEINLINE FAssetPtr(const FAssetPtr& Other)
-	{
-		(*this)=Other;
-	}
-
-	/** Construct from a string asset reference */
-	explicit FORCEINLINE FAssetPtr(const FStringAssetReference& AssetRef)
-		: TPersistentObjectPtr<FStringAssetReference>(AssetRef)
-	{
-	}
-
-	/** Construct from an object already in memory */
-	explicit FORCEINLINE FAssetPtr(const UObject* Object)
-	{
-		(*this)=Object;
-	}
-
-	/** Synchronously load (if necessary) and return the asset object represented by this asset ptr */
-	UObject* LoadSynchronous() const
-	{
-		UObject* Asset = Get();
-		if (Asset == nullptr && IsPending())
-		{
-			ToStringReference().TryLoad();
-			
-			// TryLoad will have loaded this pointer if it is valid
-			Asset = Get();
-		}
-		return Asset;
-	}
-
-	/** Returns the StringAssetReference that is wrapped by this AssetPtr */
-	FORCEINLINE const FStringAssetReference& ToStringReference() const
-	{
-		return GetUniqueID();
-	}
-
-	/** Returns string representation of reference, in form /package/path.assetname */
-	FORCEINLINE const FString& ToString() const
-	{
-		return ToStringReference().ToString();
-	}
-
-	/** Returns /package/path string, leaving off the asset name */
-	FORCEINLINE FString GetLongPackageName() const
-	{
-		return ToStringReference().GetLongPackageName();
-	}
-	
-	/** Returns assetname string, leaving off the /package/path. part */
-	FORCEINLINE FString GetAssetName() const
-	{
-		return ToStringReference().GetAssetName();
-	}
-
-	using TPersistentObjectPtr<FStringAssetReference>::operator=;
-};
-
-template <> struct TIsPODType<FAssetPtr> { enum { Value = TIsPODType<TPersistentObjectPtr<FStringAssetReference> >::Value }; };
-template <> struct TIsWeakPointerType<FAssetPtr> { enum { Value = TIsWeakPointerType<TPersistentObjectPtr<FStringAssetReference> >::Value }; };
-template <> struct THasGetTypeHash<FAssetPtr> { enum { Value = THasGetTypeHash<TPersistentObjectPtr<FStringAssetReference> >::Value }; };
-
-/**
- * TAssetPtr is templatized wrapper of the generic FAssetPtr, it can be used in UProperties
- */
-template<class T=UObject>
-class TAssetPtr
-{
-	template <class U>
-	friend class TAssetPtr;
-
-public:
-	/** Default constructor, will be null */
-	FORCEINLINE TAssetPtr()
-	{
-	}
-	
-	/** Construct from another asset pointer */
-	template <class U, class = typename TEnableIf<TPointerIsConvertibleFromTo<U, T>::Value>::Type>
-	FORCEINLINE TAssetPtr(const TAssetPtr<U>& Other)
-		: AssetPtr(Other.AssetPtr)
-	{
-	}
-
-	/** Construct from an object already in memory */
-	template <typename U>
-	FORCEINLINE TAssetPtr(const U* Object)
-		: AssetPtr(Object)
-	{
-	}
-
-	/** Construct from a nullptr */
-	FORCEINLINE TAssetPtr(TYPE_OF_NULLPTR)
-		: AssetPtr(nullptr)
-	{
-	}
-
-	/** Construct from a string asset reference */
-	explicit FORCEINLINE TAssetPtr(const FStringAssetReference& AssetRef)
-		: AssetPtr(AssetRef)
-	{
-	}
-
-	/** Reset the asset pointer back to the null state */
-	FORCEINLINE void Reset()
-	{
-		AssetPtr.Reset();
-	}
-
-	/** Resets the weak ptr only, call this when ObjectId may change */
-	FORCEINLINE void ResetWeakPtr()
-	{
-		AssetPtr.ResetWeakPtr();
-	}
-
-	/** Copy from an object already in memory */
-	template <typename U>
-	FORCEINLINE TAssetPtr& operator=(const U* Object)
-	{
-		AssetPtr = Object;
-		return *this;
-	}
-
-	/** Assign from a nullptr */
-	FORCEINLINE TAssetPtr& operator=(TYPE_OF_NULLPTR)
-	{
-		AssetPtr = nullptr;
-		return *this;
-	}
-
-	/** Copy from a string asset reference */
-	FORCEINLINE TAssetPtr& operator=(const FStringAssetReference& AssetRef)
-	{
-		AssetPtr = AssetRef;
-		return *this;
-	}
-
-	/** Copy from a weak pointer to an object already in memory */
-	template <class U, class = typename TEnableIf<TPointerIsConvertibleFromTo<U, T>::Value>::Type>
-	FORCEINLINE TAssetPtr& operator=(const TWeakObjectPtr<U>& Other)
-	{
-		AssetPtr = Other;
-		return *this;
-	}
-
-	/** Copy from another asset pointer */
-	template <class U, class = typename TEnableIf<TPointerIsConvertibleFromTo<U, T>::Value>::Type>
-	FORCEINLINE TAssetPtr& operator=(const TAssetPtr<U>& Other)
-	{
-		AssetPtr = Other.AssetPtr;
-		return *this;
-	}
-
-	/**
-	 * Compare asset pointers for equality
-	 * Caution: Two asset pointers might not be equal to each other, but they both might return nullptr
-	 *
-	 * @param Other asset pointer to compare to
-	 */
-	FORCEINLINE friend bool operator==(const TAssetPtr& Lhs, const TAssetPtr& Rhs)
-	{
-		return Lhs.AssetPtr == Rhs.AssetPtr;
-	}
-
-	/**
-	 * Compare asset pointers for equality
-	 * Caution: Two asset pointers might not be equal to each other, but they both might return nullptr
-	 *
-	 * @param Other asset pointer to compare to
-	 */
-	FORCEINLINE friend bool operator==(const TAssetPtr& Lhs, TYPE_OF_NULLPTR)
-	{
-		return Lhs.AssetPtr == nullptr;
-	}
-
-	/**
-	 * Compare asset pointers for equality
-	 * Caution: Two asset pointers might not be equal to each other, but they both might return nullptr
-	 *
-	 * @param Other asset pointer to compare to
-	 */
-	FORCEINLINE friend bool operator==(TYPE_OF_NULLPTR, const TAssetPtr& Rhs)
-	{
-		return nullptr == Rhs.AssetPtr;
-	}
-
-	/**
-	 * Compare asset pointers for inequality
-	 * Caution: Two asset pointers might not be equal to each other, but they both might return nullptr
-	 *
-	 * @param Other asset pointer to compare to
-	 */
-	FORCEINLINE friend bool operator!=(const TAssetPtr& Lhs, const TAssetPtr& Rhs)
-	{
-		return Lhs.AssetPtr != Rhs.AssetPtr;
-	}
-
-	/**
-	 * Compare asset pointers for inequality
-	 * Caution: Two asset pointers might not be equal to each other, but they both might return nullptr
-	 *
-	 * @param Other asset pointer to compare to
-	 */
-	FORCEINLINE friend bool operator!=(const TAssetPtr& Lhs, TYPE_OF_NULLPTR)
-	{
-		return Lhs.AssetPtr != nullptr;
-	}
-
-	/**
-	 * Compare asset pointers for inequality
-	 * Caution: Two asset pointers might not be equal to each other, but they both might return nullptr
-	 *
-	 * @param Other asset pointer to compare to
-	 */
-	FORCEINLINE friend bool operator!=(TYPE_OF_NULLPTR, const TAssetPtr& Rhs)
-	{
-		return nullptr != Rhs.AssetPtr;
-	}
-
-	/**
-	 * Dereference the asset pointer.
-	 *
-	 * @return nullptr if this object is gone or the lazy pointer was null, otherwise a valid UObject pointer
-	 */
-	FORCEINLINE T* Get() const
-	{
-		return dynamic_cast<T*>(AssetPtr.Get());
-	}
-
-	/** Dereference the asset pointer */
-	FORCEINLINE T& operator*() const
-	{
-		return *Get();
-	}
-
-	/** Dereference the asset pointer */
-	FORCEINLINE T* operator->() const
-	{
-		return Get();
-	}
-
-	/** Synchronously load (if necessary) and return the asset object represented by this asset ptr */
-	T* LoadSynchronous() const
-	{
-		UObject* Asset = AssetPtr.LoadSynchronous();
-		return Cast<T>(Asset);
-	}
-
-	/**  
-	 * Test if this points to a live UObject
-	 *
-	 * @return true if Get() would return a valid non-null pointer
-	 */
-	FORCEINLINE bool IsValid() const
-	{
-		// This does the runtime type check
-		return Get() != nullptr;
-	}
-
-	/**
-	* Test if this points to an asset on disk
-	*
-	* @return true if the file exists
-	*/
-	FORCEINLINE bool DoesAssetExist() const
-	{
-		// This does the runtime type check
-		return FPackageName::DoesPackageExist(GetLongPackageName());
-	}
-
-	/**  
-	 * Test if this does not point to a live UObject, but may in the future
-	 *
-	 * @return true if this does not point to a real object, but could possibly
-	 */
-	FORCEINLINE bool IsPending() const
-	{
-		return AssetPtr.IsPending();
-	}
-
-	/**  
-	 * Test if this can never point to a live UObject
-	 *
-	 * @return true if this is explicitly pointing to no object
-	 */
-	FORCEINLINE bool IsNull() const
-	{
-		return AssetPtr.IsNull();
-	}
-
-	/** Returns the StringAssetReference that is wrapped by this AssetPtr */
-	FORCEINLINE const FStringAssetReference& GetUniqueID() const
-	{
-		return AssetPtr.GetUniqueID();
-	}
-
-	/** Returns the StringAssetReference that is wrapped by this AssetPtr */
-	FORCEINLINE const FStringAssetReference& ToStringReference() const
-	{
-		return AssetPtr.ToStringReference();
-	}
-
-	/** Returns string representation of reference, in form /package/path.assetname */
-	FORCEINLINE const FString& ToString() const
-	{
-		return ToStringReference().ToString();
-	}
-
-	/** Returns /package/path string, leaving off the asset name */
-	FORCEINLINE FString GetLongPackageName() const
-	{
-		return ToStringReference().GetLongPackageName();
-	}
-	
-	/** Returns assetname string, leaving off the /package/path part */
-	FORCEINLINE FString GetAssetName() const
-	{
-		return ToStringReference().GetAssetName();
-	}
-
-	/** Dereference asset pointer to see if it points somewhere valid */
-	FORCEINLINE explicit operator bool() const
-	{
-		return IsValid();
-	}
-
-	/** Hash function */
-	FORCEINLINE friend uint32 GetTypeHash(const TAssetPtr<T>& Other)
-	{
-		return GetTypeHash(static_cast<const TPersistentObjectPtr<FStringAssetReference>&>(Other.AssetPtr));
-	}
-
-	friend FArchive& operator<<(FArchive& Ar, TAssetPtr<T>& Other)
-	{
-		Ar << Other.AssetPtr;
-		return Ar;
-	}
-
-private:
-	FAssetPtr AssetPtr;
-};
-
-template<class T> struct TIsPODType<TAssetPtr<T> > { enum { Value = TIsPODType<FAssetPtr>::Value }; };
-template<class T> struct TIsWeakPointerType<TAssetPtr<T> > { enum { Value = TIsWeakPointerType<FAssetPtr>::Value }; };
-
-/**
- * TAssetSubclassOf is a templatized wrapper around FAssetPtr that works like a TSubclassOf, it can be used in UProperties for blueprint subclasses
- */
-template<class TClass>
-class TAssetSubclassOf
-{
-	template <class TClassA>
-	friend class TAssetSubclassOf;
-
-public:
-	/** Default constructor, will be null */
-	FORCEINLINE TAssetSubclassOf()
-	{
-	}
-		
-	/** Construct from another asset pointer */
-	template <class TClassA, class = typename TEnableIf<TPointerIsConvertibleFromTo<TClassA, TClass>::Value>::Type>
-	FORCEINLINE TAssetSubclassOf(const TAssetSubclassOf<TClassA>& Other)
-		: AssetPtr(Other.AssetPtr)
-	{
-	}
-
-	/** Construct from a class already in memory */
-	FORCEINLINE TAssetSubclassOf(const UClass* From)
-		: AssetPtr(From)
-	{
-	}
-
-	/** Construct from a string asset reference */
-	explicit FORCEINLINE TAssetSubclassOf(const FStringAssetReference& AssetRef)
-		: AssetPtr(AssetRef)
-	{
-	}
-
-	/** Reset the asset pointer back to the null state */
-	FORCEINLINE void Reset()
-	{
-		AssetPtr.Reset();
-	}
-
-	/** Resets the weak ptr only, call this when ObjectId may change */
-	FORCEINLINE void ResetWeakPtr()
-	{
-		AssetPtr.ResetWeakPtr();
-	}
-
-	/** Copy from a class already in memory */
-	FORCEINLINE void operator=(const UClass* From)
-	{
-		AssetPtr = From;
-	}
-
-	/** Copy from a string asset reference */
-	FORCEINLINE void operator=(const FStringAssetReference& AssetRef)
-	{
-		AssetPtr = AssetRef;
-	}
-
-	/** Copy from a weak pointer already in memory */
-	template<class TClassA, class = typename TEnableIf<TPointerIsConvertibleFromTo<TClassA, UClass>::Value>::Type>
-	FORCEINLINE TAssetSubclassOf& operator=(const TWeakObjectPtr<TClassA>& Other)
-	{
-		AssetPtr = Other;
-		return *this;
-	}
-
-	/** Copy from another asset pointer */
-	template<class TClassA, class = typename TEnableIf<TPointerIsConvertibleFromTo<TClassA, TClass>::Value>::Type>
-	FORCEINLINE TAssetSubclassOf& operator=(const TAssetPtr<TClassA>& Other)
-	{
-		AssetPtr = Other.AssetPtr;
-		return *this;
-	}
-
-	/**  
-	 * Compare asset pointers for equality
-	 * Caution: Two asset pointers might not be equal to each other, but they both might return nullptr
-	 *
-	 * @param Other asset pointer to compare to 
-	 */
-	FORCEINLINE friend bool operator==(const TAssetSubclassOf& Lhs, const TAssetSubclassOf& Rhs)
-	{
-		return Lhs.AssetPtr == Rhs.AssetPtr;
-	}
-
-	/**  
-	 * Compare asset pointers for inequality
-	 * Caution: Two asset pointers might not be equal to each other, but they both might return nullptr
-	 *
-	 * @param Other asset pointer to compare to
-	 */
-	FORCEINLINE friend bool operator!=(const TAssetSubclassOf& Lhs, const TAssetSubclassOf& Rhs)
-	{
-		return Lhs.AssetPtr != Rhs.AssetPtr;
-	}
-
-	/**  
-	 * Dereference the asset pointer
-	 *
-	 * @return nullptr if this object is gone or the asset pointer was null, otherwise a valid UClass pointer
-	 */
-	FORCEINLINE UClass* Get() const
-	{
-		UClass* Class = dynamic_cast<UClass*>(AssetPtr.Get());
-		if (!Class || !Class->IsChildOf(TClass::StaticClass()))
-		{
-			return nullptr;
-		}
-		return Class;
-	}
-
-	/** Dereference the asset pointer */
-	FORCEINLINE UClass& operator*() const
-	{
-		return *Get();
-	}
-
-	/** Dereference the asset pointer */
-	FORCEINLINE UClass* operator->() const
-	{
-		return Get();
-	}
-
-	/**  
-	 * Test if this points to a live UObject
-	 *
-	 * @return true if Get() would return a valid non-null pointer
-	 */
-	FORCEINLINE bool IsValid() const
-	{
-		// This also does the UClass type check
-		return Get() != nullptr;
-	}
-
-	/**  
-	 * Test if this does not point to a live UObject, but may in the future
-	 *
-	 * @return true if this does not point to a real object, but could possibly
-	 */
-	FORCEINLINE bool IsPending() const
-	{
-		return AssetPtr.IsPending();
-	}
-
-	/**  
-	 * Test if this can never point to a live UObject
-	 *
-	 * @return true if this is explicitly pointing to no object
-	 */
-	FORCEINLINE bool IsNull() const
-	{
-		return AssetPtr.IsNull();
-	}
-
-	/** Returns the StringAssetReference that is wrapped by this AssetPtr */
-	FORCEINLINE const FStringAssetReference& GetUniqueID() const
-	{
-		return AssetPtr.GetUniqueID();
-	}
-
-	/** Returns the StringAssetReference that is wrapped by this AssetPtr */
-	FORCEINLINE const FStringAssetReference& ToStringReference() const
-	{
-		return AssetPtr.GetUniqueID();
-	}
-
-	/** Returns string representation of reference, in form /package/path.assetname  */
-	FORCEINLINE const FString& ToString() const
-	{
-		return ToStringReference().ToString();
-	}
-
-	/** Returns /package/path string, leaving off the asset name */
-	FORCEINLINE FString GetLongPackageName() const
-	{
-		return ToStringReference().GetLongPackageName();
-	}
-	
-	/** Returns assetname string, leaving off the /package/path part */
-	FORCEINLINE FString GetAssetName() const
-	{
-		return ToStringReference().GetAssetName();
-	}
-
-	/** Dereference asset pointer to see if it points somewhere valid */
-	FORCEINLINE explicit operator bool() const
-	{
-		return IsValid();
-	}
-
-	/** Hash function */
-	FORCEINLINE friend uint32 GetTypeHash(const TAssetSubclassOf<TClass>& Other)
-	{
-		return GetTypeHash(static_cast<const TPersistentObjectPtr<FStringAssetReference>&>(Other.AssetPtr));
-	}
-
-	/** Synchronously load (if necessary) and return the asset object represented by this asset ptr */
-	UClass* LoadSynchronous() const
-	{
-		UObject* Asset = AssetPtr.LoadSynchronous();
-		UClass* Class = dynamic_cast<UClass*>(Asset);
-		if (!Class || !Class->IsChildOf(TClass::StaticClass()))
-		{
-			return nullptr;
-		}
-		return Class;
-	}
-
-	friend FArchive& operator<<(FArchive& Ar, TAssetSubclassOf<TClass>& Other)
-	{
-		Ar << static_cast<FAssetPtr&>(Other.AssetPtr);
-		return Ar;
-	}
-
-private:
-	FAssetPtr AssetPtr;
-};
-
-template <class T> struct TIsPODType<TAssetSubclassOf<T> > { enum { Value = TIsPODType<FAssetPtr>::Value }; };
-template <class T> struct TIsWeakPointerType<TAssetSubclassOf<T> > { enum { Value = TIsWeakPointerType<FAssetPtr>::Value }; };
+COMPILE_WARNING("TAssetPtr has been renamed to TSoftObjectPtr, change to #include \"UObject/SoftObjectPtr.h\" and rename references")
+#include "UObject/SoftObjectPtr.h"
\ No newline at end of file
diff --git a/Engine/Source/Runtime/CoreUObject/Public/UObject/Class.h b/Engine/Source/Runtime/CoreUObject/Public/UObject/Class.h
--- a/Engine/Source/Runtime/CoreUObject/Public/UObject/Class.h
+++ b/Engine/Source/Runtime/CoreUObject/Public/UObject/Class.h
@@ -2559,9 +2559,6 @@ public:
 	 */
 	virtual void GetRequiredPreloadDependencies(TArray<UObject*>& DependenciesOut) {}
 
-	/** Returns true if this class implements script instrumentation. */
-	virtual bool HasInstrumentation() const { return false; }
-
 private:
 	#if UCLASS_FAST_ISA_IMPL == UCLASS_ISA_INDEXTREE
 		// For UObjectBaseUtility
@@ -3082,14 +3079,14 @@ template<> struct TBaseStructure<FInt32Interval>
 	COREUOBJECT_API static UScriptStruct* Get();
 };
 
-struct FStringAssetReference;
-template<> struct TBaseStructure<FStringAssetReference>
+struct FSoftObjectPath;
+template<> struct TBaseStructure<FSoftObjectPath>
 {
 	COREUOBJECT_API static UScriptStruct* Get();
 };
 
-struct FStringClassReference;
-template<> struct TBaseStructure<FStringClassReference>
+struct FSoftClassPath;
+template<> struct TBaseStructure<FSoftClassPath>
 {
 	COREUOBJECT_API static UScriptStruct* Get();
 };
diff --git a/Engine/Source/Runtime/CoreUObject/Public/UObject/CoreNet.h b/Engine/Source/Runtime/CoreUObject/Public/UObject/CoreNet.h
--- a/Engine/Source/Runtime/CoreUObject/Public/UObject/CoreNet.h
+++ b/Engine/Source/Runtime/CoreUObject/Public/UObject/CoreNet.h
@@ -11,10 +11,10 @@
 #include "Serialization/BitWriter.h"
 #include "Misc/NetworkGuid.h"
 #include "UObject/CoreNetTypes.h"
+#include "UObject/SoftObjectPath.h"
 
 class FOutBunch;
 class INetDeltaBaseState;
-struct FStringAssetReference;
 
 DECLARE_DELEGATE_RetVal_OneParam( bool, FNetObjectIsDynamic, const UObject*);
 
@@ -280,7 +280,7 @@ public:
 
 	virtual FArchive& operator<<(FName& Name) override;
 	virtual FArchive& operator<<(UObject*& Object) override;
-	virtual FArchive& operator<<(FStringAssetReference& Value) override;
+	virtual FArchive& operator<<(FSoftObjectPath& Value) override;
 	virtual FArchive& operator<<(struct FWeakObjectPtr& Value) override;
 };
 
@@ -299,7 +299,7 @@ public:
 
 	virtual FArchive& operator<<(FName& Name) override;
 	virtual FArchive& operator<<(UObject*& Object) override;
-	virtual FArchive& operator<<(FStringAssetReference& Value) override;
+	virtual FArchive& operator<<(FSoftObjectPath& Value) override;
 	virtual FArchive& operator<<(struct FWeakObjectPtr& Value) override;
 };
 
diff --git a/Engine/Source/Runtime/CoreUObject/Public/UObject/CoreRedirects.h b/Engine/Source/Runtime/CoreUObject/Public/UObject/CoreRedirects.h
--- a/Engine/Source/Runtime/CoreUObject/Public/UObject/CoreRedirects.h
+++ b/Engine/Source/Runtime/CoreUObject/Public/UObject/CoreRedirects.h
@@ -34,7 +34,7 @@ ENUM_CLASS_FLAGS(ECoreRedirectFlags);
 
 
 /**
- * An object path extracted into component names for matching. TODO merge with FStringAssetReference?
+ * An object path extracted into component names for matching. TODO merge with FSoftObjectPath?
  */
 struct COREUOBJECT_API FCoreRedirectObjectName
 {
diff --git a/Engine/Source/Runtime/CoreUObject/Public/UObject/FastReferenceCollector.h b/Engine/Source/Runtime/CoreUObject/Public/UObject/FastReferenceCollector.h
--- a/Engine/Source/Runtime/CoreUObject/Public/UObject/FastReferenceCollector.h
+++ b/Engine/Source/Runtime/CoreUObject/Public/UObject/FastReferenceCollector.h
@@ -32,21 +32,32 @@ class FGCArrayPool
 public:
 
 	/**
-	* Gets the singleton instance of the FObjectArrayPool
-	* @return Pool singleton.
-	*/
+	 * Gets the singleton instance of the FObjectArrayPool
+	 * @return Pool singleton.
+	 */
 	FORCEINLINE static FGCArrayPool& Get()
 	{
-		static FGCArrayPool Singleton;
-		return Singleton;
+		static FAutoConsoleCommandWithOutputDevice GCDumpPoolCommand(
+			TEXT("gc.DumpPoolStats"),
+			TEXT("Dumps count and size of GC Pools"),
+			FConsoleCommandWithOutputDeviceDelegate::CreateStatic(&FGCArrayPool::DumpStats)
+		);
+
+		static FGCArrayPool* Singleton = nullptr;
+
+		if (!Singleton)
+		{
+			Singleton = new FGCArrayPool();
+		}
+		return *Singleton;
 	}
 
 	/**
-	* Gets an event from the pool or creates one if necessary.
-	*
-	* @return The array.
-	* @see ReturnToPool
-	*/
+	 * Gets an event from the pool or creates one if necessary.
+	 *
+	 * @return The array.
+	 * @see ReturnToPool
+	 */
 	FORCEINLINE FGCArrayStruct* GetArrayStructFromPool()
 	{
 		FGCArrayStruct* Result = Pool.Pop();
@@ -62,11 +73,11 @@ public:
 	}
 
 	/**
-	* Returns an array to the pool.
-	*
-	* @param Array The array to return.
-	* @see GetArrayFromPool
-	*/
+	 * Returns an array to the pool.
+	 *
+	 * @param Array The array to return.
+	 * @see GetArrayFromPool
+	 */
 	FORCEINLINE void ReturnToPool(FGCArrayStruct* ArrayStruct)
 	{
 #if UE_BUILD_DEBUG
@@ -78,7 +89,10 @@ public:
 		Pool.Push(ArrayStruct);
 	}
 
-	/** Performs memory cleanup */
+	/** 
+	 * Performs manual memory cleanup. 
+	 * Generally the pools will be cleaned up when ClearWeakReferences is called on a full GC purge
+	 */
 	void Cleanup()
 	{
 #if UE_BUILD_DEBUG
@@ -100,7 +114,83 @@ public:
 		UE_LOG(LogGarbage, Log, TEXT("Freed %ub from %d GC array pools."), FreedMemory, AllArrays.Num());
 	}
 
-	/** Clears weak references for everything in the pool */
+	/**
+	 * Writes out info about the makeup of the pool. called by 'gc.DumpPoolStats'
+	 *
+	 * @param Array The array to return.
+	 * @see GetArrayFromPool
+	 */
+	static void DumpStats(FOutputDevice& OutputDevice)
+	{
+		FGCArrayPool& Instance = Get();
+
+		TArray<FGCArrayStruct*> PoppedItems;
+
+		TMap<int32, int32> Buckets;
+
+		int32 TotalSize = 0;
+		int32 MaxSize = 0;
+		int32 TotalItems = 0;
+
+		do
+		{
+			FGCArrayStruct* Item = Instance.Pool.Pop();
+
+			if (Item)
+			{
+				PoppedItems.Push(Item);
+
+				// Inc our bucket
+				Buckets.FindOrAdd(Item->ObjectsToSerialize.Max()) += 1;
+
+				TotalSize += Item->ObjectsToSerialize.Max();
+				TotalSize += Item->WeakReferences.Max();
+				TotalItems++;
+			}
+			else
+			{
+				break;
+			}
+
+		} while (true);
+
+		// return everything to the pool
+		while (PoppedItems.Num())
+		{
+			Instance.Pool.Push(PoppedItems.Pop());
+		}
+
+		// One of these lists is used by the main GC and is huge, so remove it so that it doesn't
+		// pollute the stats and we can accurately see what the task pools are using.
+		int32 TotalSizeKB = (TotalSize * sizeof(UObject*)) / 1024;
+
+		OutputDevice.Logf(TEXT("GCPoolStats: %d Pools totaling %d KB. Avg: Objs=%d, Size=%d KB."),
+			TotalItems,
+			TotalSizeKB,
+			TotalSize / FMath::Max(TotalItems, 1),
+			TotalSizeKB / FMath::Max(TotalItems, 1));
+
+		// long form output...
+		TArray<int32> Keys;
+
+		Buckets.GetKeys(Keys);
+
+		Keys.Sort([&](int32 lhs, int32 rhs) {
+			return lhs > rhs;
+		});
+
+		for (int Key : Keys)
+		{
+			const int32 Value = Buckets[Key];
+			int32 ItemSize = (Key * sizeof(UObject*)) / 1024;
+			OutputDevice.Logf(TEXT("\t%d\t\t(%d Items @ %d KB = %d KB)"), Key, Value, ItemSize, Value * ItemSize);
+		}
+	}
+
+	/** 
+	 * Clears weak references for everything in the pool. 
+	 * If bClearPools is true it will clear all of the pools as well, which is used during a full purge 
+	 */
 	void ClearWeakReferences(bool bClearPools)
 	{
 		TArray<FGCArrayStruct*> AllArrays;
diff --git a/Engine/Source/Runtime/CoreUObject/Public/UObject/LazyObjectPtr.h b/Engine/Source/Runtime/CoreUObject/Public/UObject/LazyObjectPtr.h
--- a/Engine/Source/Runtime/CoreUObject/Public/UObject/LazyObjectPtr.h
+++ b/Engine/Source/Runtime/CoreUObject/Public/UObject/LazyObjectPtr.h
@@ -119,7 +119,7 @@ template<> struct TIsPODType<FUniqueObjectGuid> { enum { Value = true }; };
  *
  * This is useful for cross level references or places where you need to point to an object whose name changes often.
  */
-class FLazyObjectPtr : public TPersistentObjectPtr<FUniqueObjectGuid>
+struct FLazyObjectPtr : public TPersistentObjectPtr<FUniqueObjectGuid>
 {
 public:	
 	/** Default constructor, sets to null */
@@ -172,7 +172,7 @@ template <> struct THasGetTypeHash<FLazyObjectPtr> { enum { Value = THasGetTypeH
  * TLazyObjectPtr is templatized version of the generic FLazyObjectPtr
  */
 template<class T=UObject>
-class TLazyObjectPtr : private FLazyObjectPtr
+struct TLazyObjectPtr : private FLazyObjectPtr
 {
 public:
 #if PLATFORM_COMPILER_HAS_DEFAULTED_FUNCTIONS
diff --git a/Engine/Source/Runtime/CoreUObject/Public/UObject/Linker.h b/Engine/Source/Runtime/CoreUObject/Public/UObject/Linker.h
--- a/Engine/Source/Runtime/CoreUObject/Public/UObject/Linker.h
+++ b/Engine/Source/Runtime/CoreUObject/Public/UObject/Linker.h
@@ -42,8 +42,8 @@ public:
 	TArray<FObjectExport> ExportMap;
 	/** List of dependency lists for each export */
 	TArray<TArray<FPackageIndex> > DependsMap;
-	/** Map that holds info about string asset references from the package. */
-	TArray<FString> StringAssetReferencesMap;
+	/** List of packages that are soft referenced by this package */
+	TArray<FName> SoftPackageReferenceList;
 	/** List of Searchable Names, by object containing them. Not in MultiMap to allow sorting, and sizes are usually small enough where TArray makes sense */
 	TMap<FPackageIndex, TArray<FName> > SearchableNamesMap;
 
diff --git a/Engine/Source/Runtime/CoreUObject/Public/UObject/LinkerLoad.h b/Engine/Source/Runtime/CoreUObject/Public/UObject/LinkerLoad.h
--- a/Engine/Source/Runtime/CoreUObject/Public/UObject/LinkerLoad.h
+++ b/Engine/Source/Runtime/CoreUObject/Public/UObject/LinkerLoad.h
@@ -5,8 +5,7 @@
 #include "CoreMinimal.h"
 #include "Serialization/ArchiveUObject.h"
 #include "UObject/LazyObjectPtr.h"
-#include "Misc/StringAssetReference.h"
-#include "UObject/AssetPtr.h"
+#include "UObject/SoftObjectPtr.h"
 #include "UObject/ObjectResource.h"
 #include "UObject/Linker.h"
 
@@ -725,12 +724,12 @@ private:
 		return Ar;
 	}
 
-	FORCEINLINE virtual FArchive& operator<<(FAssetPtr& AssetPtr) override
+	FORCEINLINE virtual FArchive& operator<<(FSoftObjectPtr& Value) override
 	{
 		FArchive& Ar = *this;
-		FStringAssetReference ID;
+		FSoftObjectPath ID;
 		ID.Serialize(Ar);
-		AssetPtr = ID;
+		Value = ID;
 		return Ar;
 	}
 	void BadNameIndexError(NAME_INDEX NameIndex);
diff --git a/Engine/Source/Runtime/CoreUObject/Public/UObject/LinkerSave.h b/Engine/Source/Runtime/CoreUObject/Public/UObject/LinkerSave.h
--- a/Engine/Source/Runtime/CoreUObject/Public/UObject/LinkerSave.h
+++ b/Engine/Source/Runtime/CoreUObject/Public/UObject/LinkerSave.h
@@ -7,8 +7,6 @@
 #include "UObject/ObjectResource.h"
 #include "UObject/Linker.h"
 
-class FAssetPtr;
-class FLazyObjectPtr;
 struct FUntypedBulkData;
 
 /*----------------------------------------------------------------------------
@@ -82,7 +80,7 @@ public:
 	FArchive& operator<<( FName& InName );
 	FArchive& operator<<( UObject*& Obj );
 	FArchive& operator<<( FLazyObjectPtr& LazyObjectPtr );
-	FArchive& operator<<( FAssetPtr& AssetPtr );
+	FArchive& operator<<( FSoftObjectPtr& SoftObjectPtr );
 
 #if WITH_EDITOR
 	// proxy for debugdata
