Copying //UE4/Dev-Framework to //UE4/Dev-Main (Source: //UE4/Dev-Framework @ 3582324)

diff --git a/Engine/Source/Runtime/CoreUObject/Public/UObject/NoExportTypes.h b/Engine/Source/Runtime/CoreUObject/Public/UObject/NoExportTypes.h
--- a/Engine/Source/Runtime/CoreUObject/Public/UObject/NoExportTypes.h
+++ b/Engine/Source/Runtime/CoreUObject/Public/UObject/NoExportTypes.h
@@ -18,8 +18,7 @@
 #include "Misc/Guid.h"
 #include "Misc/DateTime.h"
 #include "Misc/Timespan.h"
-#include "Misc/StringAssetReference.h"
-#include "Misc/StringClassReference.h"
+#include "UObject/SoftObjectPath.h"
 
 #include "Math/InterpCurvePoint.h"
 #include "Math/UnitConversion.h"
@@ -938,19 +937,22 @@ struct FTimespan
 };
 
 
-// A string asset reference
-
-USTRUCT(noexport, BlueprintType, meta=(HasNativeMake="Engine.BlueprintFunctionLibrary.MakeStringAssetReference"))
-struct FStringAssetReference
+/** An object path, this is saved as a name/string pair */
+USTRUCT(noexport, BlueprintType, meta=(HasNativeMake="Engine.KismetSystemLibrary.MakeSoftObjectPath", HasNativeBreak="Engine.KismetSystemLibrary.BreakSoftObjectPath"))
+struct FSoftObjectPath
 {
-	UPROPERTY(EditAnywhere, BlueprintReadWrite, Category=StringAssetReference)
-	FString AssetLongPathname;
-};
+	/** Asset path, patch to a top level object in a package */
+	UPROPERTY()
+	FName AssetPathName;
 
-// A string class reference
+	/** Optional FString for subobject within an asset */
+	UPROPERTY()
+	FString SubPathString;
+};
 
+/** Subclass of FSoftObjectPath that is only valid to use with UClass* */
 USTRUCT(noexport)
-struct FStringClassReference : public FStringAssetReference
+struct FSoftClassPath : public FSoftObjectPath
 {
 };
 
diff --git a/Engine/Source/Runtime/CoreUObject/Public/UObject/Object.h b/Engine/Source/Runtime/CoreUObject/Public/UObject/Object.h
--- a/Engine/Source/Runtime/CoreUObject/Public/UObject/Object.h
+++ b/Engine/Source/Runtime/CoreUObject/Public/UObject/Object.h
@@ -1208,7 +1208,7 @@ public:
 	DECLARE_FUNCTION(execUnicodeStringConst);
 	DECLARE_FUNCTION(execTextConst);
 	DECLARE_FUNCTION(execObjectConst);
-	DECLARE_FUNCTION(execAssetConst);
+	DECLARE_FUNCTION(execSoftObjectConst);
 
 	// @todo delegate: Multi-cast versions needed for script execution!		(Need Add, Remove, Clear/Empty)
 	DECLARE_FUNCTION(execInstanceDelegate);
diff --git a/Engine/Source/Runtime/CoreUObject/Public/UObject/ObjectMacros.h b/Engine/Source/Runtime/CoreUObject/Public/UObject/ObjectMacros.h
--- a/Engine/Source/Runtime/CoreUObject/Public/UObject/ObjectMacros.h
+++ b/Engine/Source/Runtime/CoreUObject/Public/UObject/ObjectMacros.h
@@ -331,11 +331,11 @@ typedef uint64 EClassCastFlags;
 #define CASTCLASS_UObjectPropertyBase			DECLARE_UINT64(0x0000000004000000)
 #define CASTCLASS_UWeakObjectProperty			DECLARE_UINT64(0x0000000008000000)
 #define CASTCLASS_ULazyObjectProperty			DECLARE_UINT64(0x0000000010000000)
-#define CASTCLASS_UAssetObjectProperty			DECLARE_UINT64(0x0000000020000000)
+#define CASTCLASS_USoftObjectProperty			DECLARE_UINT64(0x0000000020000000)
 #define CASTCLASS_UTextProperty					DECLARE_UINT64(0x0000000040000000)
 #define CASTCLASS_UInt16Property				DECLARE_UINT64(0x0000000080000000)
 #define CASTCLASS_UDoubleProperty				DECLARE_UINT64(0x0000000100000000)
-#define CASTCLASS_UAssetClassProperty			DECLARE_UINT64(0x0000000200000000)
+#define CASTCLASS_USoftClassProperty			DECLARE_UINT64(0x0000000200000000)
 #define CASTCLASS_UPackage						DECLARE_UINT64(0x0000000400000000)
 #define CASTCLASS_ULevel						DECLARE_UINT64(0x0000000800000000)
 #define CASTCLASS_AActor						DECLARE_UINT64(0x0000001000000000)
@@ -1047,10 +1047,10 @@ namespace UM
 	// Metadata usable in UPROPERTY for customizing the behavior when displaying the property in a property panel
 	enum
 	{
-		/// [PropertyMetadata] Used for Subclass and StringClassReference properties.  Indicates whether abstract class types should be shown in the class picker.
+		/// [PropertyMetadata] Used for Subclass and SoftClass properties.  Indicates whether abstract class types should be shown in the class picker.
 		AllowAbstract,
 
-		/// [PropertyMetadata] Used for FStringAssetReference properties.  Comma delimited list that indicates the class type(s) of assets to be displayed in the asset picker.
+		/// [PropertyMetadata] Used for FSoftObjectPath properties.  Comma delimited list that indicates the class type(s) of assets to be displayed in the asset picker.
 		AllowedClasses,
 
 		/// [PropertyMetadata] Used for FVector properties.  It causes a ratio lock to be added when displaying this property in details panels.
@@ -1059,10 +1059,10 @@ namespace UM
 		/// [PropertyMetadata] Used for integer properties.  Clamps the valid values that can be entered in the UI to be between 0 and the length of the array specified.
 		ArrayClamp,
 
-		/// [PropertyMetadata] Used for AssetPtr/StringAssetReference properties. Comma separated list of Bundle names used inside PrimaryDataAssets to specify which bundles this reference is part of
+		/// [PropertyMetadata] Used for SoftObjectPtr/SoftObjectPath properties. Comma separated list of Bundle names used inside PrimaryDataAssets to specify which bundles this reference is part of
 		AssetBundles,
 
-		/// [PropertyMetadata] Used for Subclass and StringClassReference properties.  Indicates whether only blueprint classes should be shown in the class picker.
+		/// [PropertyMetadata] Used for Subclass and SoftClass properties.  Indicates whether only blueprint classes should be shown in the class picker.
 		BlueprintBaseOnly,
 
 		/// [PropertyMetadata] Property defaults are generated by the Blueprint compiler and will not be copied when CopyPropertiesForUnrelatedObjects is called post-compile.
@@ -1089,7 +1089,7 @@ namespace UM
 		/// [PropertyMetadata] Specifies a boolean property that is used to indicate whether editing of this property is disabled.
 		EditCondition,
 		
-		/// [PropertyMetadata] Used for FStringAssetReference properties in conjunction with AllowedClasses. Indicates whether only the exact classes specified in AllowedClasses can be used or whether subclasses are valid.
+		/// [PropertyMetadata] Used for FSoftObjectPath properties in conjunction with AllowedClasses. Indicates whether only the exact classes specified in AllowedClasses can be used or whether subclasses are valid.
 		ExactClass,
 
 		/// [PropertyMetadata] Specifies a list of categories whose functions should be exposed when building a function list in the Blueprint Editor.
@@ -1107,7 +1107,7 @@ namespace UM
 		/// [PropertyMetadata] Used for FColor and FLinearColor properties. Indicates that the Alpha property should be hidden when displaying the property widget in the details.
 		HideAlphaChannel,
 
-		/// [PropertyMetadata] Used for Subclass and StringClassReference properties. Specifies to hide the ability to change view options in the class picker
+		/// [PropertyMetadata] Used for Subclass and SoftClass properties. Specifies to hide the ability to change view options in the class picker
 		HideViewOptions,
 
 		/// [PropertyMetadata] Signifies that the bool property is only displayed inline as an edit condition toggle in other properties, and should not be shown on its own row.
@@ -1122,10 +1122,10 @@ namespace UM
 		/// [PropertyMetadata] For properties in a structure indicates the default value of the property in a blueprint make structure node.
 		MakeStructureDefaultValue,
 
-		/// [PropertyMetadata] Used FStringClassReference properties. Indicates the parent class that the class picker will use when filtering which classes to display.
+		/// [PropertyMetadata] Used FSoftClassPath properties. Indicates the parent class that the class picker will use when filtering which classes to display.
 		MetaClass,
 
-		/// [PropertyMetadata] Used for Subclass and StringClassReference properties. Indicates the selected class must implement a specific interface
+		/// [PropertyMetadata] Used for Subclass and SoftClass properties. Indicates the selected class must implement a specific interface
 		MustImplement,
 
 		/// [PropertyMetadata] Used for numeric properties. Stipulates that the value must be a multiple of the metadata value.
@@ -1158,7 +1158,7 @@ namespace UM
 		// [PropertyMetadata] Used by struct properties. Indicates that the inner properties will not be shown inside an expandable struct, but promoted up a level.
 		ShowOnlyInnerProperties,
 
-		/// [PropertyMetadata] Used for Subclass and StringClassReference properties. Shows the picker as a tree view instead of as a list
+		/// [PropertyMetadata] Used for Subclass and SoftClass properties. Shows the picker as a tree view instead of as a list
 		ShowTreeView,
 
 		// [PropertyMetadata] Used by numeric properties. Indicates how rapidly the value will grow when moving an unbounded slider.
diff --git a/Engine/Source/Runtime/CoreUObject/Public/UObject/Package.h b/Engine/Source/Runtime/CoreUObject/Public/UObject/Package.h
--- a/Engine/Source/Runtime/CoreUObject/Public/UObject/Package.h
+++ b/Engine/Source/Runtime/CoreUObject/Public/UObject/Package.h
@@ -131,8 +131,6 @@ public:
 	/** Whether this package has been fully loaded (aka had all it's exports created) at some point.															*/
 	mutable bool	bHasBeenFullyLoaded;
 private:
-	/** Returns whether exports should be found in memory first before trying to load from disk from within CreateExport.										*/
-	bool	bShouldFindExportsInMemoryFirst;
 	/** Indicates which folder to display this package under in the Generic Browser's list of packages. If not specified, package is added to the root level.	*/
 	FName	FolderName;
 	/** Time in seconds it took to fully load this package. 0 if package is either in process of being loaded or has never been fully loaded.					*/
@@ -295,27 +293,6 @@ public:
 	 */
 	virtual void TagSubobjects(EObjectFlags NewFlags) override;
 
-	/**
-	 * Sets whether exports should be found in memory first or  not.
-	 *
-	 * @param bInShouldFindExportsInMemoryFirst	Whether to find in memory first or not
-	 */
-	void FindExportsInMemoryFirst( bool bInShouldFindExportsInMemoryFirst )
-	{
-		bShouldFindExportsInMemoryFirst = bInShouldFindExportsInMemoryFirst;
-	}
-
-	/**
-	 * Returns whether exports should be found in memory first before trying to load from disk
-	 * from within CreateExport.
-	 *
-	 * @return true if exports should be found via FindObject first, false otherwise.
-	 */
-	bool ShouldFindExportsInMemoryFirst()
-	{
-		return bShouldFindExportsInMemoryFirst;
-	}
-
 	/**
 	 * Called to indicate that this package contains a ULevel or UWorld object.
 	 */
diff --git a/Engine/Source/Runtime/CoreUObject/Public/UObject/PackageFileSummary.h b/Engine/Source/Runtime/CoreUObject/Public/UObject/PackageFileSummary.h
--- a/Engine/Source/Runtime/CoreUObject/Public/UObject/PackageFileSummary.h
+++ b/Engine/Source/Runtime/CoreUObject/Public/UObject/PackageFileSummary.h
@@ -119,14 +119,14 @@ public:
 	int32		DependsOffset;
 
 	/**
-	* Number of references contained in this package
+	* Number of soft package references contained in this package
 	*/
-	int32		StringAssetReferencesCount;
+	int32		SoftPackageReferencesCount;
 
 	/**
-	* Location into the file on disk for the string asset references map data
+	* Location into the file on disk for the soft package reference list
 	*/
-	int32		StringAssetReferencesOffset;
+	int32		SoftPackageReferencesOffset;
 
 	/**
 	* Location into the file on disk for the SearchableNamesMap data
diff --git a/Engine/Source/Runtime/CoreUObject/Public/UObject/PersistentObjectPtr.h b/Engine/Source/Runtime/CoreUObject/Public/UObject/PersistentObjectPtr.h
--- a/Engine/Source/Runtime/CoreUObject/Public/UObject/PersistentObjectPtr.h
+++ b/Engine/Source/Runtime/CoreUObject/Public/UObject/PersistentObjectPtr.h
@@ -10,10 +10,10 @@
 #include "UObject/WeakObjectPtr.h"
 
 /**
- * TPersistentObjectPtr is a template base class for FLazyObjectPtr and FAssetPtr
+ * TPersistentObjectPtr is a template base class for FLazyObjectPtr and FSoftObjectPtr
  */
 template<class TObjectID>
-class TPersistentObjectPtr
+struct TPersistentObjectPtr
 {
 public:	
 
@@ -120,7 +120,12 @@ public:
 		{
 			Object = ObjectID.ResolveObject();
 			WeakPtr = Object;
-			TagAtLastTest = TObjectID::GetCurrentTag();
+
+			// Not safe to update tag during save as ResolveObject may have failed accidentally
+			if (Object || !GIsSavingPackage)
+			{
+				TagAtLastTest = TObjectID::GetCurrentTag();
+			}
 
 			// If this object is pending kill or otherwise invalid, this will return nullptr as expected
 			Object = WeakPtr.Get();
@@ -141,7 +146,12 @@ public:
 		{
 			Object = ObjectID.ResolveObject();
 			WeakPtr = Object;
-			TagAtLastTest = TObjectID::GetCurrentTag();
+
+			// Not safe to update tag during save as ResolveObject may have failed accidentally
+			if (Object || !GIsSavingPackage)
+			{
+				TagAtLastTest = TObjectID::GetCurrentTag();
+			}
 
 			// If this object is pending kill or otherwise invalid, this will return nullptr as expected
 			Object = WeakPtr.Get(bEvenIfPendingKill);
@@ -229,7 +239,7 @@ public:
 	 */
 	FORCEINLINE bool IsNull() const
 	{
-		return Get() == nullptr && !ObjectID.IsValid();
+		return !ObjectID.IsValid();
 	}
 
 	/** Hash function */
diff --git a/Engine/Source/Runtime/CoreUObject/Public/UObject/PropertyTag.h b/Engine/Source/Runtime/CoreUObject/Public/UObject/PropertyTag.h
--- a/Engine/Source/Runtime/CoreUObject/Public/UObject/PropertyTag.h
+++ b/Engine/Source/Runtime/CoreUObject/Public/UObject/PropertyTag.h
@@ -40,5 +40,5 @@ struct FPropertyTag
 	friend FArchive& operator<<( FArchive& Ar, FPropertyTag& Tag );
 
 	// Property serializer.
-	void SerializeTaggedProperty( FArchive& Ar, UProperty* Property, uint8* Value, uint8* Defaults );
+	void SerializeTaggedProperty( FArchive& Ar, UProperty* Property, uint8* Value, uint8* Defaults ) const;
 };
diff --git a/Engine/Source/Runtime/CoreUObject/Public/UObject/Script.h b/Engine/Source/Runtime/CoreUObject/Public/UObject/Script.h
--- a/Engine/Source/Runtime/CoreUObject/Public/UObject/Script.h
+++ b/Engine/Source/Runtime/CoreUObject/Public/UObject/Script.h
@@ -265,7 +265,7 @@ enum EExprToken
 	EX_LetValueOnPersistentFrame = 0x64,
 	EX_ArrayConst			= 0x65,
 	EX_EndArrayConst		= 0x66,
-	EX_AssetConst			= 0x67,
+	EX_SoftObjectConst		= 0x67,
 	EX_CallMath				= 0x68, // static pure function from on local call space
 	EX_SwitchValue			= 0x69,
 	EX_InstrumentationEvent	= 0x6A, // Instrumentation event
diff --git a/Engine/Source/Runtime/CoreUObject/Public/UObject/ScriptMacros.h b/Engine/Source/Runtime/CoreUObject/Public/UObject/ScriptMacros.h
--- a/Engine/Source/Runtime/CoreUObject/Public/UObject/ScriptMacros.h
+++ b/Engine/Source/Runtime/CoreUObject/Public/UObject/ScriptMacros.h
@@ -78,8 +78,8 @@ enum {MAX_VARIABLE_SIZE = 0x0FFF };
 #define P_GET_TINTERFACE(ObjectType,ParamName)		PARAM_PASSED_BY_VAL(ParamName, UInterfaceProperty, TScriptInterface<ObjectType>)
 #define P_GET_TINTERFACE_REF(ObjectType,ParamName)	PARAM_PASSED_BY_REF(ParamName, UInterfaceProperty, TScriptInterface<ObjectType>)
 
-#define P_GET_ASSETOBJECT(StructType,ParamName)		PARAM_PASSED_BY_VAL(ParamName, UAssetObjectProperty, StructType)
-#define P_GET_ASSETOBJECT_REF(StructType,ParamName)	PARAM_PASSED_BY_REF(ParamName, UAssetObjectProperty, StructType)
+#define P_GET_SOFTOBJECT(ObjectType,ParamName)		PARAM_PASSED_BY_VAL(ParamName, USoftObjectProperty, ObjectType)
+#define P_GET_SOFTOBJECT_REF(ObjectType,ParamName)	PARAM_PASSED_BY_REF(ParamName, USoftObjectProperty, ObjectType)
 
 #define P_GET_ARRAY(ElementType,ParamName)			ElementType ParamName[(MAX_VARIABLE_SIZE/sizeof(ElementType))+1];		Stack.StepCompiledIn<UProperty>(ParamName);
 #define P_GET_ARRAY_REF(ElementType,ParamName)		ElementType ParamName##Temp[(MAX_VARIABLE_SIZE/sizeof(ElementType))+1]; ElementType* ParamName = Stack.StepCompiledInRef<UProperty, ElementType*>(ParamName##Temp);
diff --git a/Engine/Source/Runtime/CoreUObject/Public/UObject/ScriptSerialization.h b/Engine/Source/Runtime/CoreUObject/Public/UObject/ScriptSerialization.h
--- a/Engine/Source/Runtime/CoreUObject/Public/UObject/ScriptSerialization.h
+++ b/Engine/Source/Runtime/CoreUObject/Public/UObject/ScriptSerialization.h
@@ -372,7 +372,7 @@
 
 			break;
 		}
-		case EX_AssetConst:
+		case EX_SoftObjectConst:
 		{
 			SerializeExpr(iCode, Ar);
 			break;
diff --git a/Engine/Source/Runtime/CoreUObject/Public/UObject/SoftObjectPath.h b/Engine/Source/Runtime/CoreUObject/Public/UObject/SoftObjectPath.h
new file mode 100644
--- /dev/null
+++ b/Engine/Source/Runtime/CoreUObject/Public/UObject/SoftObjectPath.h
@@ -0,0 +1,325 @@
+// Copyright 1998-2017 Epic Games, Inc. All Rights Reserved.
+
+#pragma once
+
+#include "CoreMinimal.h"
+#include "HAL/ThreadSafeCounter.h"
+#include "UObject/Class.h"
+#include "HAL/ThreadSingleton.h"
+
+/**
+ * A struct that contains a string reference to an object, either a top level asset or a subobject.
+ * This can be used to make soft references to assets that are loaded on demand.
+ * This is stored internally as an FName pointing to the top level asset (/package/path.assetname) and an option a string subobject path.
+ * If the MetaClass metadata is applied to a UProperty with this the UI will restrict to that type of asset.
+ */
+struct COREUOBJECT_API FSoftObjectPath
+{
+	FSoftObjectPath() {}
+
+	/** Construct from another soft object path */
+	FSoftObjectPath(const FSoftObjectPath& Other)
+		: AssetPathName(Other.AssetPathName)
+		, SubPathString(Other.SubPathString)
+	{
+	}
+
+	/** Construct from a path string */
+	FSoftObjectPath(FString PathString)
+	{
+		SetPath(MoveTemp(PathString));
+	}
+
+	/** Construct from an asset FName and subobject pair */
+	FSoftObjectPath(FName InAssetPathName, FString InSubPathString)
+		: AssetPathName(InAssetPathName)
+		, SubPathString(InSubPathString)
+	{}
+	
+	/** Construct from an existing object in memory */
+	FSoftObjectPath(const UObject* InObject);
+
+	~FSoftObjectPath() {}
+
+	/** Returns string representation of reference, in form /package/path.assetname[:subpath] */
+	FString ToString() const;
+
+	/** Returns the entire asset path as an FName, including both package and asset but not sub object */
+	FORCEINLINE FName GetAssetPathName() const
+	{
+		return AssetPathName;
+	}
+
+	/** Returns string version of asset path, including both package and asset but not sub object */
+	FORCEINLINE FString GetAssetPathString() const
+	{
+		if (AssetPathName == NAME_None)
+		{
+			return FString();
+		}
+
+		return AssetPathName.ToString();
+	}
+
+	/** Returns the sub path, which is often empty */
+	FORCEINLINE FString GetSubPathString() const
+	{
+		return SubPathString;
+	}
+
+	/** Returns /package/path, leaving off the asset name and sub object */
+	FString GetLongPackageName() const
+	{
+		FString PackageName;
+		GetAssetPathString().Split(TEXT("."), &PackageName, nullptr, ESearchCase::CaseSensitive, ESearchDir::FromStart);
+		return PackageName;
+	}
+
+	/** Returns assetname string, leaving off the /package/path part and sub object */
+	FString GetAssetName() const
+	{
+		FString AssetName;
+		GetAssetPathString().Split(TEXT("."), nullptr, &AssetName, ESearchCase::CaseSensitive, ESearchDir::FromStart);
+		return AssetName;
+	}
+
+	/** Sets asset path of this reference based on a string path */
+	void SetPath(FString Path);
+
+	/**
+	 * Attempts to load the asset, this will call LoadObject which can be very slow
+	 *
+	 * @return Loaded UObject, or nullptr if the reference is null or the asset fails to load
+	 */
+	UObject* TryLoad() const;
+
+	/**
+	 * Attempts to find a currently loaded object that matches this path
+	 *
+	 * @return Found UObject, or nullptr if not currently in memory
+	 */
+	UObject* ResolveObject() const;
+
+	/** Resets reference to point to null */
+	void Reset()
+	{		
+		AssetPathName = NAME_None;
+		SubPathString.Reset();
+	}
+	
+	/** Check if this could possibly refer to a real object, or was initialized to null */
+	FORCEINLINE bool IsValid() const
+	{
+		return AssetPathName != NAME_None;
+	}
+
+	/** Checks to see if this is initialized to null */
+	FORCEINLINE bool IsNull() const
+	{
+		return AssetPathName == NAME_None;
+	}
+
+	/** Check if this represents an asset, meaning it is not null but does not have a sub path */
+	FORCEINLINE bool IsAsset() const
+	{
+		return AssetPathName != NAME_None && SubPathString.IsEmpty();
+	}
+
+	/** Check if this represents a sub object, meaning it has a sub path */
+	FORCEINLINE bool IsSubobject() const
+	{
+		return AssetPathName != NAME_None && !SubPathString.IsEmpty();
+	}
+
+	/** Struct overrides */
+	bool Serialize(FArchive& Ar);
+	bool operator==(FSoftObjectPath const& Other) const;
+	bool operator!=(FSoftObjectPath const& Other) const
+	{
+		return !(*this == Other);
+	}
+	FSoftObjectPath& operator=(FSoftObjectPath const& Other);
+	bool ExportTextItem(FString& ValueStr, FSoftObjectPath const& DefaultValue, UObject* Parent, int32 PortFlags, UObject* ExportRootScope) const;
+	bool ImportTextItem( const TCHAR*& Buffer, int32 PortFlags, UObject* Parent, FOutputDevice* ErrorText );
+	bool SerializeFromMismatchedTag(struct FPropertyTag const& Tag, FArchive& Ar);
+
+	/** Serializes the internal path and also handles save/PIE fixups. Call this from the archiver overrides */
+	void SerializePath(FArchive& Ar, bool bSkipSerializeIfArchiveHasSize = false);
+
+	/** Fixes up string asset reference for saving, call if saving with a method that skips SerializePath. This can modify the path, it will return true if it was modified */
+	bool PreSavePath();
+
+	/** Handles when a string asset reference has been loaded, call if loading with a method that skips SerializePath. This does not modify path but might call callbacks */
+	void PostLoadPath() const;
+
+	/** Fixes up this SoftObjectPath to add or remove the PIE prefix depending on what is currently active */
+	void FixupForPIE();
+
+	FORCEINLINE friend uint32 GetTypeHash(FSoftObjectPath const& This)
+	{
+		uint32 Hash = 0;
+
+		Hash = HashCombine(Hash, GetTypeHash(This.AssetPathName));
+		Hash = HashCombine(Hash, GetTypeHash(This.SubPathString));
+		return Hash;
+	}
+
+	/** Code needed by FSoftObjectPtr internals */
+	static int32 GetCurrentTag()
+	{
+		return CurrentTag.GetValue();
+	}
+	static int32 InvalidateTag()
+	{
+		return CurrentTag.Increment();
+	}
+	static FSoftObjectPath GetOrCreateIDForObject(const UObject* Object);
+	
+	/** Adds list of packages names that have been created specifically for PIE, this is used for editor fixup */
+	static void AddPIEPackageName(FName NewPIEPackageName);
+	
+	/** Disables special PIE path handling, call when PIE finishes to clear list */
+	static void ClearPIEPackageNames();
+
+private:
+	/** Asset path, patch to a top level object in a package. This is /package/path.assetname */
+	FName AssetPathName;
+
+	/** Optional FString for subobject within an asset. This is the sub path after the : */
+	FString SubPathString;
+
+	/** Global counter that determines when we need to re-search based on path because more objects have been loaded **/
+	static FThreadSafeCounter CurrentTag;
+
+	/** Package names currently being duplicated, needed by FixupForPIE */
+	static TSet<FName> PIEPackageNames;
+
+	friend COREUOBJECT_API UScriptStruct* Z_Construct_UScriptStruct_FSoftObjectPath();
+};
+
+/**
+ * A struct that contains a string reference to a class, can be used to make soft references to classes
+ */
+struct COREUOBJECT_API FSoftClassPath : public FSoftObjectPath
+{
+	FSoftClassPath()
+	{ }
+
+	FSoftClassPath(FSoftClassPath const& Other)
+		: FSoftObjectPath(Other)
+	{ }
+
+	/**
+	 * Construct from a path string
+	 */
+	FSoftClassPath(const FString& PathString)
+		: FSoftObjectPath(PathString)
+	{ }
+
+	/**
+	 * Construct from an existing class, will do some string processing
+	 */
+	FSoftClassPath(const UClass* InClass)
+		: FSoftObjectPath(InClass)
+	{ }
+
+	/**
+	* Attempts to load the class.
+	* @return Loaded UObject, or null if the class fails to load, or if the reference is not valid.
+	*/
+	template<typename T>
+	UClass* TryLoadClass() const
+	{
+		if ( IsValid() )
+		{
+			return LoadClass<T>(nullptr, *ToString(), nullptr, LOAD_None, nullptr);
+		}
+
+		return nullptr;
+	}
+
+	/**
+	 * Attempts to find a currently loaded object that matches this object ID
+	 * @return Found UClass, or NULL if not currently loaded
+	 */
+	UClass* ResolveClass() const;
+
+	bool SerializeFromMismatchedTag(const FPropertyTag& Tag, FArchive& Ar);
+
+	static FSoftClassPath GetOrCreateIDForClass(const UClass *InClass);
+
+private:
+	/** Forbid creation for UObject. This class is for UClass only. Use FSoftObjectPath instead. */
+	FSoftClassPath(const UObject* InObject) { }
+
+	/** Forbidden. This class is for UClass only. Use FSoftObjectPath instead. */
+	static FSoftObjectPath GetOrCreateIDForObject(const UObject *Object);
+};
+
+// Not deprecating these yet as it will lead to too many warnings in games
+//DEPRECATED(4.18, "FStringAssetReference was renamed to FSoftObjectPath as it is now not always a string and can also refer to a subobject")
+typedef FSoftObjectPath FStringAssetReference;
+
+//DEPRECATED(4.18, "FStringClassReference was renamed to FSoftClassPath")
+typedef FSoftClassPath FStringClassReference;
+
+/** OPtions for how to set soft object path collection */
+enum class ESoftObjectPathCollectType : uint8
+{
+	/** References is not tracked in any situation, transient reference */
+	NeverCollect,
+	/** Editor only reference, this is tracked for redirector fixup but not for cooking */
+	EditorOnlyCollect,
+	/** Game reference, this is gathered for both redirector fixup and cooking */
+	AlwaysCollect,
+};
+
+class COREUOBJECT_API FSoftObjectPathThreadContext : public TThreadSingleton<FSoftObjectPathThreadContext>
+{
+	friend TThreadSingleton<FSoftObjectPathThreadContext>;
+	friend struct FSoftObjectPathSerializationScope;
+
+	FSoftObjectPathThreadContext() {}
+
+	struct FSerializationOptions
+	{
+		FName PackageName;
+		FName PropertyName;
+		ESoftObjectPathCollectType CollectType;
+
+		FSerializationOptions() : CollectType(ESoftObjectPathCollectType::AlwaysCollect) {}
+		FSerializationOptions(FName InPackageName, FName InPropertyName, ESoftObjectPathCollectType InCollectType) : PackageName(InPackageName), PropertyName(InPropertyName), CollectType(InCollectType) {}
+	};
+
+	TArray<FSerializationOptions> OptionStack;
+public:
+	/** 
+	 * Returns the current serialization options that were added using SerializationScope or LinkerLoad
+	 *
+	 * @param OutPackageName Package that this string asset belongs to
+	 * @param OutPropertyName Property that this string asset reference belongs to
+	 * @param OutCollectType Type of collecting that should be done
+	 */
+	bool GetSerializationOptions(FName& OutPackageName, FName& OutPropertyName, ESoftObjectPathCollectType& OutCollectType) const;
+};
+
+/** Helper class to set and restore serialization options for string asset references */
+struct FSoftObjectPathSerializationScope
+{
+	/** 
+	 * Create a new serialization scope, which affects the way that string asset references are saved
+	 *
+	 * @param SerializingPackageName Package that this string asset belongs to
+	 * @param SerializingPropertyName Property that this string asset reference belongs to
+	 * @param CollectType Set type of collecting that should be done, can be used to disable tracking entirely
+	 */
+	FSoftObjectPathSerializationScope(FName SerializingPackageName, FName SerializingPropertyName, ESoftObjectPathCollectType CollectType)
+	{
+		FSoftObjectPathThreadContext::Get().OptionStack.Emplace(SerializingPackageName, SerializingPropertyName, CollectType);
+	}
+
+	~FSoftObjectPathSerializationScope()
+	{
+		FSoftObjectPathThreadContext::Get().OptionStack.Pop();
+	}
+};
\ No newline at end of file
diff --git a/Engine/Source/Runtime/CoreUObject/Public/UObject/SoftObjectPtr.h b/Engine/Source/Runtime/CoreUObject/Public/UObject/SoftObjectPtr.h
new file mode 100644
--- /dev/null
+++ b/Engine/Source/Runtime/CoreUObject/Public/UObject/SoftObjectPtr.h
@@ -0,0 +1,615 @@
+// Copyright 1998-2017 Epic Games, Inc. All Rights Reserved.
+
+/*=============================================================================
+	SoftObjectPtr.h: Pointer to UObject asset, keeps extra information so that it is works even if the asset is not in memory
+=============================================================================*/
+
+#pragma once
+
+#include "CoreMinimal.h"
+#include "UObject/Object.h"
+#include "Templates/Casts.h"
+#include "UObject/PersistentObjectPtr.h"
+#include "UObject/SoftObjectPath.h"
+
+/**
+ * FSoftObjectPtr is a type of weak pointer to a UObject, that also keeps track of the path to the object on disk.
+ * It will change back and forth between being Valid and Pending as the referenced object loads or unloads.
+ * It has no impact on if the object is garbage collected or not.
+ *
+ * This is useful to specify assets that you may want to asynchronously load on demand.
+ */
+struct FSoftObjectPtr : public TPersistentObjectPtr<FSoftObjectPath>
+{
+public:	
+	/** Default constructor, will be null */
+	FORCEINLINE FSoftObjectPtr()
+	{
+	}
+
+	/** Construct from another soft pointer */
+	FORCEINLINE FSoftObjectPtr(const FSoftObjectPtr& Other)
+	{
+		(*this)=Other;
+	}
+
+	/** Construct from a soft object path */
+	explicit FORCEINLINE FSoftObjectPtr(const FSoftObjectPath& ObjectPath)
+		: TPersistentObjectPtr<FSoftObjectPath>(ObjectPath)
+	{
+	}
+
+	/** Construct from an object already in memory */
+	explicit FORCEINLINE FSoftObjectPtr(const UObject* Object)
+	{
+		(*this)=Object;
+	}
+
+	/** Synchronously load (if necessary) and return the asset object represented by this asset ptr */
+	UObject* LoadSynchronous() const
+	{
+		UObject* Asset = Get();
+		if (Asset == nullptr && IsPending())
+		{
+			ToSoftObjectPath().TryLoad();
+			
+			// TryLoad will have loaded this pointer if it is valid
+			Asset = Get();
+		}
+		return Asset;
+	}
+
+	/** Returns the StringObjectPath that is wrapped by this SoftObjectPtr */
+	FORCEINLINE const FSoftObjectPath& ToSoftObjectPath() const
+	{
+		return GetUniqueID();
+	}
+
+	DEPRECATED(4.18, "ToStringReference was renamed to ToSoftObjectPath")
+	FORCEINLINE const FSoftObjectPath& ToStringReference() const
+	{
+		return GetUniqueID();
+	}
+
+	/** Returns string representation of reference, in form /package/path.assetname */
+	FORCEINLINE FString ToString() const
+	{
+		return ToSoftObjectPath().ToString();
+	}
+
+	/** Returns /package/path string, leaving off the asset name */
+	FORCEINLINE FString GetLongPackageName() const
+	{
+		return ToSoftObjectPath().GetLongPackageName();
+	}
+	
+	/** Returns assetname string, leaving off the /package/path. part */
+	FORCEINLINE FString GetAssetName() const
+	{
+		return ToSoftObjectPath().GetAssetName();
+	}
+
+	using TPersistentObjectPtr<FSoftObjectPath>::operator=;
+};
+
+template <> struct TIsPODType<FSoftObjectPtr> { enum { Value = TIsPODType<TPersistentObjectPtr<FSoftObjectPath> >::Value }; };
+template <> struct TIsWeakPointerType<FSoftObjectPtr> { enum { Value = TIsWeakPointerType<TPersistentObjectPtr<FSoftObjectPath> >::Value }; };
+template <> struct THasGetTypeHash<FSoftObjectPtr> { enum { Value = THasGetTypeHash<TPersistentObjectPtr<FSoftObjectPath> >::Value }; };
+
+/**
+ * TSoftObjectPtr is templatized wrapper of the generic FSoftObjectPtr, it can be used in UProperties
+ */
+template<class T=UObject>
+struct TSoftObjectPtr
+{
+	template <class U>
+	friend struct TSoftObjectPtr;
+
+public:
+	/** Default constructor, will be null */
+	FORCEINLINE TSoftObjectPtr()
+	{
+	}
+	
+	/** Construct from another soft pointer */
+	template <class U, class = typename TEnableIf<TPointerIsConvertibleFromTo<U, T>::Value>::Type>
+	FORCEINLINE TSoftObjectPtr(const TSoftObjectPtr<U>& Other)
+		: SoftObjectPtr(Other.SoftObjectPtr)
+	{
+	}
+
+	/** Construct from an object already in memory */
+	template <typename U>
+	FORCEINLINE TSoftObjectPtr(const U* Object)
+		: SoftObjectPtr(Object)
+	{
+	}
+
+	/** Construct from a nullptr */
+	FORCEINLINE TSoftObjectPtr(TYPE_OF_NULLPTR)
+		: SoftObjectPtr(nullptr)
+	{
+	}
+
+	/** Construct from a soft object path */
+	explicit FORCEINLINE TSoftObjectPtr(const FSoftObjectPath& ObjectPath)
+		: SoftObjectPtr(ObjectPath)
+	{
+	}
+
+	/** Reset the soft pointer back to the null state */
+	FORCEINLINE void Reset()
+	{
+		SoftObjectPtr.Reset();
+	}
+
+	/** Resets the weak ptr only, call this when ObjectId may change */
+	FORCEINLINE void ResetWeakPtr()
+	{
+		SoftObjectPtr.ResetWeakPtr();
+	}
+
+	/** Copy from an object already in memory */
+	template <typename U>
+	FORCEINLINE TSoftObjectPtr& operator=(const U* Object)
+	{
+		SoftObjectPtr = Object;
+		return *this;
+	}
+
+	/** Assign from a nullptr */
+	FORCEINLINE TSoftObjectPtr& operator=(TYPE_OF_NULLPTR)
+	{
+		SoftObjectPtr = nullptr;
+		return *this;
+	}
+
+	/** Copy from a soft object path */
+	FORCEINLINE TSoftObjectPtr& operator=(const FSoftObjectPath& ObjectPath)
+	{
+		SoftObjectPtr = ObjectPath;
+		return *this;
+	}
+
+	/** Copy from a weak pointer to an object already in memory */
+	template <class U, class = typename TEnableIf<TPointerIsConvertibleFromTo<U, T>::Value>::Type>
+	FORCEINLINE TSoftObjectPtr& operator=(const TWeakObjectPtr<U>& Other)
+	{
+		SoftObjectPtr = Other;
+		return *this;
+	}
+
+	/** Copy from another soft pointer */
+	template <class U, class = typename TEnableIf<TPointerIsConvertibleFromTo<U, T>::Value>::Type>
+	FORCEINLINE TSoftObjectPtr& operator=(const TSoftObjectPtr<U>& Other)
+	{
+		SoftObjectPtr = Other.SoftObjectPtr;
+		return *this;
+	}
+
+	/**
+	 * Compare soft pointers for equality
+	 * Caution: Two soft pointers might not be equal to each other, but they both might return nullptr
+	 *
+	 * @param Other soft pointer to compare to
+	 */
+	FORCEINLINE friend bool operator==(const TSoftObjectPtr& Lhs, const TSoftObjectPtr& Rhs)
+	{
+		return Lhs.SoftObjectPtr == Rhs.SoftObjectPtr;
+	}
+
+	/**
+	 * Compare soft pointers for equality
+	 * Caution: Two soft pointers might not be equal to each other, but they both might return nullptr
+	 *
+	 * @param Other soft pointer to compare to
+	 */
+	FORCEINLINE friend bool operator==(const TSoftObjectPtr& Lhs, TYPE_OF_NULLPTR)
+	{
+		return Lhs.SoftObjectPtr == nullptr;
+	}
+
+	/**
+	 * Compare soft pointers for equality
+	 * Caution: Two soft pointers might not be equal to each other, but they both might return nullptr
+	 *
+	 * @param Other soft pointer to compare to
+	 */
+	FORCEINLINE friend bool operator==(TYPE_OF_NULLPTR, const TSoftObjectPtr& Rhs)
+	{
+		return nullptr == Rhs.SoftObjectPtr;
+	}
+
+	/**
+	 * Compare soft pointers for inequality
+	 * Caution: Two soft pointers might not be equal to each other, but they both might return nullptr
+	 *
+	 * @param Other soft pointer to compare to
+	 */
+	FORCEINLINE friend bool operator!=(const TSoftObjectPtr& Lhs, const TSoftObjectPtr& Rhs)
+	{
+		return Lhs.SoftObjectPtr != Rhs.SoftObjectPtr;
+	}
+
+	/**
+	 * Compare soft pointers for inequality
+	 * Caution: Two soft pointers might not be equal to each other, but they both might return nullptr
+	 *
+	 * @param Other soft pointer to compare to
+	 */
+	FORCEINLINE friend bool operator!=(const TSoftObjectPtr& Lhs, TYPE_OF_NULLPTR)
+	{
+		return Lhs.SoftObjectPtr != nullptr;
+	}
+
+	/**
+	 * Compare soft pointers for inequality
+	 * Caution: Two soft pointers might not be equal to each other, but they both might return nullptr
+	 *
+	 * @param Other soft pointer to compare to
+	 */
+	FORCEINLINE friend bool operator!=(TYPE_OF_NULLPTR, const TSoftObjectPtr& Rhs)
+	{
+		return nullptr != Rhs.SoftObjectPtr;
+	}
+
+	/**
+	 * Dereference the soft pointer.
+	 *
+	 * @return nullptr if this object is gone or the lazy pointer was null, otherwise a valid UObject pointer
+	 */
+	FORCEINLINE T* Get() const
+	{
+		return dynamic_cast<T*>(SoftObjectPtr.Get());
+	}
+
+	/** Dereference the soft pointer */
+	FORCEINLINE T& operator*() const
+	{
+		return *Get();
+	}
+
+	/** Dereference the soft pointer */
+	FORCEINLINE T* operator->() const
+	{
+		return Get();
+	}
+
+	/** Synchronously load (if necessary) and return the asset object represented by this asset ptr */
+	T* LoadSynchronous() const
+	{
+		UObject* Asset = SoftObjectPtr.LoadSynchronous();
+		return Cast<T>(Asset);
+	}
+
+	/**  
+	 * Test if this points to a live UObject
+	 *
+	 * @return true if Get() would return a valid non-null pointer
+	 */
+	FORCEINLINE bool IsValid() const
+	{
+		// This does the runtime type check
+		return Get() != nullptr;
+	}
+
+	/**  
+	 * Test if this does not point to a live UObject, but may in the future
+	 *
+	 * @return true if this does not point to a real object, but could possibly
+	 */
+	FORCEINLINE bool IsPending() const
+	{
+		return SoftObjectPtr.IsPending();
+	}
+
+	/**  
+	 * Test if this can never point to a live UObject
+	 *
+	 * @return true if this is explicitly pointing to no object
+	 */
+	FORCEINLINE bool IsNull() const
+	{
+		return SoftObjectPtr.IsNull();
+	}
+
+	/** Returns the StringObjectPath that is wrapped by this SoftObjectPtr */
+	FORCEINLINE const FSoftObjectPath& GetUniqueID() const
+	{
+		return SoftObjectPtr.GetUniqueID();
+	}
+
+	/** Returns the StringObjectPath that is wrapped by this SoftObjectPtr */
+	FORCEINLINE const FSoftObjectPath& ToSoftObjectPath() const
+	{
+		return SoftObjectPtr.GetUniqueID();
+	}
+
+	DEPRECATED(4.18, "ToStringReference was renamed to ToSoftObjectPath")
+	FORCEINLINE const FSoftObjectPath& ToStringReference() const
+	{
+		return SoftObjectPtr.GetUniqueID();
+	}
+
+	/** Returns string representation of reference, in form /package/path.assetname */
+	FORCEINLINE FString ToString() const
+	{
+		return ToSoftObjectPath().ToString();
+	}
+
+	/** Returns /package/path string, leaving off the asset name */
+	FORCEINLINE FString GetLongPackageName() const
+	{
+		return ToSoftObjectPath().GetLongPackageName();
+	}
+	
+	/** Returns assetname string, leaving off the /package/path part */
+	FORCEINLINE FString GetAssetName() const
+	{
+		return ToSoftObjectPath().GetAssetName();
+	}
+
+	/** Dereference soft pointer to see if it points somewhere valid */
+	FORCEINLINE explicit operator bool() const
+	{
+		return IsValid();
+	}
+
+	/** Hash function */
+	FORCEINLINE friend uint32 GetTypeHash(const TSoftObjectPtr<T>& Other)
+	{
+		return GetTypeHash(static_cast<const TPersistentObjectPtr<FSoftObjectPath>&>(Other.SoftObjectPtr));
+	}
+
+	friend FArchive& operator<<(FArchive& Ar, TSoftObjectPtr<T>& Other)
+	{
+		Ar << Other.SoftObjectPtr;
+		return Ar;
+	}
+
+private:
+	FSoftObjectPtr SoftObjectPtr;
+};
+
+template<class T> struct TIsPODType<TSoftObjectPtr<T> > { enum { Value = TIsPODType<FSoftObjectPtr>::Value }; };
+template<class T> struct TIsWeakPointerType<TSoftObjectPtr<T> > { enum { Value = TIsWeakPointerType<FSoftObjectPtr>::Value }; };
+
+/**
+ * TSoftClassPtr is a templatized wrapper around FSoftObjectPtr that works like a TSubclassOf, it can be used in UProperties for blueprint subclasses
+ */
+template<class TClass=UObject>
+class TSoftClassPtr
+{
+	template <class TClassA>
+	friend class TSoftClassPtr;
+
+public:
+	/** Default constructor, will be null */
+	FORCEINLINE TSoftClassPtr()
+	{
+	}
+		
+	/** Construct from another soft pointer */
+	template <class TClassA, class = typename TEnableIf<TPointerIsConvertibleFromTo<TClassA, TClass>::Value>::Type>
+	FORCEINLINE TSoftClassPtr(const TSoftClassPtr<TClassA>& Other)
+		: SoftObjectPtr(Other.SoftObjectPtr)
+	{
+	}
+
+	/** Construct from a class already in memory */
+	FORCEINLINE TSoftClassPtr(const UClass* From)
+		: SoftObjectPtr(From)
+	{
+	}
+
+	/** Construct from a soft object path */
+	explicit FORCEINLINE TSoftClassPtr(const FSoftObjectPath& ObjectPath)
+		: SoftObjectPtr(ObjectPath)
+	{
+	}
+
+	/** Reset the soft pointer back to the null state */
+	FORCEINLINE void Reset()
+	{
+		SoftObjectPtr.Reset();
+	}
+
+	/** Resets the weak ptr only, call this when ObjectId may change */
+	FORCEINLINE void ResetWeakPtr()
+	{
+		SoftObjectPtr.ResetWeakPtr();
+	}
+
+	/** Copy from a class already in memory */
+	FORCEINLINE void operator=(const UClass* From)
+	{
+		SoftObjectPtr = From;
+	}
+
+	/** Copy from a soft object path */
+	FORCEINLINE void operator=(const FSoftObjectPath& ObjectPath)
+	{
+		SoftObjectPtr = ObjectPath;
+	}
+
+	/** Copy from a weak pointer already in memory */
+	template<class TClassA, class = typename TEnableIf<TPointerIsConvertibleFromTo<TClassA, UClass>::Value>::Type>
+	FORCEINLINE TSoftClassPtr& operator=(const TWeakObjectPtr<TClassA>& Other)
+	{
+		SoftObjectPtr = Other;
+		return *this;
+	}
+
+	/** Copy from another soft pointer */
+	template<class TClassA, class = typename TEnableIf<TPointerIsConvertibleFromTo<TClassA, TClass>::Value>::Type>
+	FORCEINLINE TSoftClassPtr& operator=(const TSoftObjectPtr<TClassA>& Other)
+	{
+		SoftObjectPtr = Other.SoftObjectPtr;
+		return *this;
+	}
+
+	/**  
+	 * Compare soft pointers for equality
+	 * Caution: Two soft pointers might not be equal to each other, but they both might return nullptr
+	 *
+	 * @param Other soft pointer to compare to 
+	 */
+	FORCEINLINE friend bool operator==(const TSoftClassPtr& Lhs, const TSoftClassPtr& Rhs)
+	{
+		return Lhs.SoftObjectPtr == Rhs.SoftObjectPtr;
+	}
+
+	/**  
+	 * Compare soft pointers for inequality
+	 * Caution: Two soft pointers might not be equal to each other, but they both might return nullptr
+	 *
+	 * @param Other soft pointer to compare to
+	 */
+	FORCEINLINE friend bool operator!=(const TSoftClassPtr& Lhs, const TSoftClassPtr& Rhs)
+	{
+		return Lhs.SoftObjectPtr != Rhs.SoftObjectPtr;
+	}
+
+	/**  
+	 * Dereference the soft pointer
+	 *
+	 * @return nullptr if this object is gone or the soft pointer was null, otherwise a valid UClass pointer
+	 */
+	FORCEINLINE UClass* Get() const
+	{
+		UClass* Class = dynamic_cast<UClass*>(SoftObjectPtr.Get());
+		if (!Class || !Class->IsChildOf(TClass::StaticClass()))
+		{
+			return nullptr;
+		}
+		return Class;
+	}
+
+	/** Dereference the soft pointer */
+	FORCEINLINE UClass& operator*() const
+	{
+		return *Get();
+	}
+
+	/** Dereference the soft pointer */
+	FORCEINLINE UClass* operator->() const
+	{
+		return Get();
+	}
+
+	/**  
+	 * Test if this points to a live UObject
+	 *
+	 * @return true if Get() would return a valid non-null pointer
+	 */
+	FORCEINLINE bool IsValid() const
+	{
+		// This also does the UClass type check
+		return Get() != nullptr;
+	}
+
+	/**  
+	 * Test if this does not point to a live UObject, but may in the future
+	 *
+	 * @return true if this does not point to a real object, but could possibly
+	 */
+	FORCEINLINE bool IsPending() const
+	{
+		return SoftObjectPtr.IsPending();
+	}
+
+	/**  
+	 * Test if this can never point to a live UObject
+	 *
+	 * @return true if this is explicitly pointing to no object
+	 */
+	FORCEINLINE bool IsNull() const
+	{
+		return SoftObjectPtr.IsNull();
+	}
+
+	/** Returns the StringObjectPath that is wrapped by this SoftObjectPtr */
+	FORCEINLINE const FSoftObjectPath& GetUniqueID() const
+	{
+		return SoftObjectPtr.GetUniqueID();
+	}
+
+	/** Returns the StringObjectPath that is wrapped by this SoftObjectPtr */
+	FORCEINLINE const FSoftObjectPath& ToSoftObjectPath() const
+	{
+		return SoftObjectPtr.GetUniqueID();
+	}
+
+	DEPRECATED(4.18, "ToStringReference was renamed to ToSoftObjectPath")
+	FORCEINLINE const FSoftObjectPath& ToStringReference() const
+	{
+		return SoftObjectPtr.GetUniqueID();
+	}
+
+	/** Returns string representation of reference, in form /package/path.assetname  */
+	FORCEINLINE FString ToString() const
+	{
+		return ToSoftObjectPath().ToString();
+	}
+
+	/** Returns /package/path string, leaving off the asset name */
+	FORCEINLINE FString GetLongPackageName() const
+	{
+		return ToSoftObjectPath().GetLongPackageName();
+	}
+	
+	/** Returns assetname string, leaving off the /package/path part */
+	FORCEINLINE FString GetAssetName() const
+	{
+		return ToSoftObjectPath().GetAssetName();
+	}
+
+	/** Dereference soft pointer to see if it points somewhere valid */
+	FORCEINLINE explicit operator bool() const
+	{
+		return IsValid();
+	}
+
+	/** Hash function */
+	FORCEINLINE friend uint32 GetTypeHash(const TSoftClassPtr<TClass>& Other)
+	{
+		return GetTypeHash(static_cast<const TPersistentObjectPtr<FSoftObjectPath>&>(Other.SoftObjectPtr));
+	}
+
+	/** Synchronously load (if necessary) and return the asset object represented by this asset ptr */
+	UClass* LoadSynchronous() const
+	{
+		UObject* Asset = SoftObjectPtr.LoadSynchronous();
+		UClass* Class = dynamic_cast<UClass*>(Asset);
+		if (!Class || !Class->IsChildOf(TClass::StaticClass()))
+		{
+			return nullptr;
+		}
+		return Class;
+	}
+
+	friend FArchive& operator<<(FArchive& Ar, TSoftClassPtr<TClass>& Other)
+	{
+		Ar << static_cast<FSoftObjectPtr&>(Other.SoftObjectPtr);
+		return Ar;
+	}
+
+private:
+	FSoftObjectPtr SoftObjectPtr;
+};
+
+template <class T> struct TIsPODType<TSoftClassPtr<T> > { enum { Value = TIsPODType<FSoftObjectPtr>::Value }; };
+template <class T> struct TIsWeakPointerType<TSoftClassPtr<T> > { enum { Value = TIsWeakPointerType<FSoftObjectPtr>::Value }; };
+
+DEPRECATED(4.18, "FAssetPtr was renamed to FSoftObjectPtr as it is not necessarily an asset")
+typedef FSoftObjectPtr FAssetPtr;
+
+// Not deprecating these yet as it will lead to too many warnings in games
+//DEPRECATED(4.18, "TAssetPtr was renamed to TSoftObjectPtr as it is not necessarily an asset")
+template<class T=UObject>
+using TAssetPtr = TSoftObjectPtr<T>;
+
+//DEPRECATED(4.18, "TAssetSubclassOf was renamed to TSoftClassPtr")
+template<class TClass = UObject>
+using TAssetSubclassOf = TSoftClassPtr<TClass>;
+
diff --git a/Engine/Source/Runtime/CoreUObject/Public/UObject/Stack.h b/Engine/Source/Runtime/CoreUObject/Public/UObject/Stack.h
--- a/Engine/Source/Runtime/CoreUObject/Public/UObject/Stack.h
+++ b/Engine/Source/Runtime/CoreUObject/Public/UObject/Stack.h
@@ -50,7 +50,7 @@ enum EPropertyType
 	CPT_MulticastDelegate,
 	CPT_WeakObjectReference,
 	CPT_LazyObjectReference,
-	CPT_AssetObjectReference,
+	CPT_SoftObjectReference,
 	CPT_Double,
 	CPT_Map,
 	CPT_Set,
diff --git a/Engine/Source/Runtime/CoreUObject/Public/UObject/UObjectGlobals.h b/Engine/Source/Runtime/CoreUObject/Public/UObject/UObjectGlobals.h
--- a/Engine/Source/Runtime/CoreUObject/Public/UObject/UObjectGlobals.h
+++ b/Engine/Source/Runtime/CoreUObject/Public/UObject/UObjectGlobals.h
@@ -151,7 +151,7 @@ COREUOBJECT_API TArray<const TCHAR*> ParsePropertyFlags(uint64 Flags);
 COREUOBJECT_API UPackage* GetTransientPackage();
 
 /**
- * Gets INI file name from object's reference if it contains one.
+ * Gets INI file name from object's reference if it contains one. 
  *
  * @returns If object reference doesn't contain any INI reference the function
  *		returns nullptr. Otherwise a ptr to INI's file name.
@@ -159,15 +159,15 @@ COREUOBJECT_API UPackage* GetTransientPackage();
 COREUOBJECT_API const FString* GetIniFilenameFromObjectsReference(const FString& ObjectsReferenceString);
 
 /**
- * Resolves ini object path to string object path.
+ * Resolves ini object path to string object path. This used to happen automatically in ResolveName but now must be called manually
  *
- * @param ObjectReference Ini reference.
- * @param IniFilename Ini filename.
- * @param bThrow Can this function throw?
+ * @param ObjectReference Ini reference, of the form engine-ini:/Script/Engine.Engine.DefaultMaterialName
+ * @param IniFilename Ini filename. If null it will call GetIniFilenameFromObjectsReference
+ * @param bThrow If true, will print an error if it can't find the file
  *
  * @returns Resolved object path.
  */
-COREUOBJECT_API FString ResolveIniObjectsReference(const FString& ObjectReference, const FString* IniFilename, bool bThrow = false);
+COREUOBJECT_API FString ResolveIniObjectsReference(const FString& ObjectReference, const FString* IniFilename = nullptr, bool bThrow = false);
 
 COREUOBJECT_API bool ResolveName(UObject*& Outer, FString& ObjectsReferenceString, bool Create, bool Throw, uint32 LoadFlags = LOAD_None);
 COREUOBJECT_API void SafeLoadError( UObject* Outer, uint32 LoadFlags, const TCHAR* ErrorMessage);
@@ -1881,10 +1881,6 @@ struct COREUOBJECT_API FCoreUObjectDelegates
 	/** Queries whether an object should be loaded on top ( replace ) an already existing one */
 	static FOnLoadObjectsOnTop ShouldLoadOnTop;
 
-	/** called when loading a string asset reference */
-	DECLARE_MULTICAST_DELEGATE_OneParam(FPackageLoadedFromStringAssetReference, const FName&);
-	static FPackageLoadedFromStringAssetReference PackageLoadedFromStringAssetReference;
-
 	/** called when path to world root is changed */
 	DECLARE_MULTICAST_DELEGATE_OneParam(FPackageCreatedForLoad, class UPackage*);
 	static FPackageCreatedForLoad PackageCreatedForLoad;
@@ -1893,13 +1889,13 @@ struct COREUOBJECT_API FCoreUObjectDelegates
 	DECLARE_DELEGATE_RetVal_OneParam(FPrimaryAssetId, FGetPrimaryAssetIdForObject, const UObject*);
 	static FGetPrimaryAssetIdForObject GetPrimaryAssetIdForObject;
 
-	DECLARE_DELEGATE_OneParam(FStringAssetReferenceLoaded, const FString&);
-	DEPRECATED(4.17, "StringAssetReferenceLoaded is deprecated, call FStringAssetReference::PostLoadPath instead")
-	static FStringAssetReferenceLoaded StringAssetReferenceLoaded;
+	DECLARE_DELEGATE_OneParam(FSoftObjectPathLoaded, const FString&);
+	DEPRECATED(4.17, "StringAssetReferenceLoaded is deprecated, call FSoftObjectPath::PostLoadPath instead")
+	static FSoftObjectPathLoaded StringAssetReferenceLoaded;
 
-	DECLARE_DELEGATE_RetVal_OneParam(FString, FStringAssetReferenceSaving, FString const& /*SavingAssetLongPathname*/);
-	DEPRECATED(4.17, "StringAssetReferenceSaving is deprecated, call FStringAssetReference::PreSavePath instead")
-	static FStringAssetReferenceSaving StringAssetReferenceSaving;
+	DECLARE_DELEGATE_RetVal_OneParam(FString, FSoftObjectPathSaving, FString const& /*SavingAssetLongPathname*/);
+	DEPRECATED(4.17, "StringAssetReferenceSaving is deprecated, call FSoftObjectPath::PreSavePath instead")
+	static FSoftObjectPathSaving StringAssetReferenceSaving;
 
 	DECLARE_MULTICAST_DELEGATE_TwoParams(FOnRedirectorFollowed, const FString&, UObject*);
 	DEPRECATED(4.17, "RedirectorFollowed is deprecated, FixeupRedirects was replaced with ResavePackages -FixupRedirect")
@@ -1989,15 +1985,13 @@ struct FAssetMsg
 #if WITH_EDITOR
 /** 
  * Returns if true if the object is editor-only:
- * - it's a package marked as PKG_EditorOnly
- * or
- * - it's a class from a package marked as PKG_EditorOnly
- * or
- * - its class is from a package marked as PKG_EditorOnly
+ * - it's a package marked as PKG_EditorOnly or inside one
  * or
- * - its outer is editor-only
+ * - IsEditorOnly returns true
+ * or 
+ * - if bCheckRecursive is true, if it's class, outer, or archetypes are editor only
  */
-COREUOBJECT_API bool IsEditorOnlyObject(const UObject* InObject);
+COREUOBJECT_API bool IsEditorOnlyObject(const UObject* InObject, bool bCheckRecursive = true);
 #endif //WITH_EDITOR
 
 struct FClassFunctionLinkInfo;
@@ -2020,10 +2014,10 @@ namespace UE4CodeGen_Private
 		Float,
 		Double,
 		Bool,
-		AssetClass,
+		SoftClass,
 		WeakObject,
 		LazyObject,
-		AssetObject,
+		SoftObject,
 		Class,
 		Object,
 		Interface,
@@ -2197,7 +2191,7 @@ namespace UE4CodeGen_Private
 #endif
 	};
 
-	struct FAssetClassPropertyParams // : FPropertyParamsBaseWithOffset
+	struct FSoftClassPropertyParams // : FPropertyParamsBaseWithOffset
 	{
 		EPropertyClass   Type;
 		const char*      NameUTF8;
@@ -2313,7 +2307,7 @@ namespace UE4CodeGen_Private
 	typedef FGenericPropertyParams FTextPropertyParams;
 	typedef FObjectPropertyParams  FWeakObjectPropertyParams;
 	typedef FObjectPropertyParams  FLazyObjectPropertyParams;
-	typedef FObjectPropertyParams  FAssetObjectPropertyParams;
+	typedef FObjectPropertyParams  FSoftObjectPropertyParams;
 
 	struct FFunctionParams
 	{
diff --git a/Engine/Source/Runtime/CoreUObject/Public/UObject/UnrealType.h b/Engine/Source/Runtime/CoreUObject/Public/UObject/UnrealType.h
--- a/Engine/Source/Runtime/CoreUObject/Public/UObject/UnrealType.h
+++ b/Engine/Source/Runtime/CoreUObject/Public/UObject/UnrealType.h
@@ -21,7 +21,7 @@
 #include "Templates/Greater.h"
 #include "Containers/List.h"
 #include "UObject/LazyObjectPtr.h"
-#include "UObject/AssetPtr.h"
+#include "UObject/SoftObjectPtr.h"
 #include "UObject/PropertyTag.h"
 #include "Serialization/SerializedPropertyScope.h"
 
@@ -2132,15 +2132,15 @@ public:
 //
 // Describes a reference variable to another object which may be nil, and will become valid or invalid at any point
 //
-class COREUOBJECT_API UAssetObjectProperty : public TUObjectPropertyBase<FAssetPtr>
+class COREUOBJECT_API USoftObjectProperty : public TUObjectPropertyBase<FSoftObjectPtr>
 {
-	DECLARE_CASTED_CLASS_INTRINSIC(UAssetObjectProperty, TUObjectPropertyBase<FAssetPtr>, 0, TEXT("/Script/CoreUObject"), CASTCLASS_UAssetObjectProperty)
+	DECLARE_CASTED_CLASS_INTRINSIC(USoftObjectProperty, TUObjectPropertyBase<FSoftObjectPtr>, 0, TEXT("/Script/CoreUObject"), CASTCLASS_USoftObjectProperty)
 
-	UAssetObjectProperty(ECppProperty, int32 InOffset, uint64 InFlags, UClass* InClass)
+	USoftObjectProperty(ECppProperty, int32 InOffset, uint64 InFlags, UClass* InClass)
 		: TUObjectPropertyBase(FObjectInitializer::Get(), EC_CppProperty, InOffset, InFlags, InClass)
 	{}
 
-	UAssetObjectProperty( const FObjectInitializer& ObjectInitializer, ECppProperty, int32 InOffset, uint64 InFlags, UClass* InClass )
+	USoftObjectProperty( const FObjectInitializer& ObjectInitializer, ECppProperty, int32 InOffset, uint64 InFlags, UClass* InClass )
 		: TUObjectPropertyBase(ObjectInitializer, EC_CppProperty, InOffset, InFlags, InClass)
 	{}
 
@@ -2168,7 +2168,7 @@ private:
 	virtual uint32 GetValueTypeHashInternal(const void* Src) const override;
 public:
 
-	// ScriptVM should store Asset as FAssetPtr not as UObject.
+	// ScriptVM should store Asset as FSoftObjectPtr not as UObject.
 
 	virtual void CopySingleValueToScriptVM(void* Dest, void const* Src) const override;
 	virtual void CopyCompleteValueToScriptVM(void* Dest, void const* Src) const override;
@@ -2240,25 +2240,25 @@ public:
 };
 
 /*-----------------------------------------------------------------------------
-	UAssetSubclassOfProperty.
+	USoftClassProperty.
 -----------------------------------------------------------------------------*/
 
 //
 // Describes a reference variable to another class which may be nil, and will become valid or invalid at any point
 //
-class COREUOBJECT_API UAssetClassProperty : public UAssetObjectProperty
+class COREUOBJECT_API USoftClassProperty : public USoftObjectProperty
 {
-	DECLARE_CASTED_CLASS_INTRINSIC(UAssetClassProperty, UAssetObjectProperty, 0, TEXT("/Script/CoreUObject"), CASTCLASS_UAssetClassProperty)
+	DECLARE_CASTED_CLASS_INTRINSIC(USoftClassProperty, USoftObjectProperty, 0, TEXT("/Script/CoreUObject"), CASTCLASS_USoftClassProperty)
 
 	// Variables.
 	class UClass* MetaClass;
 public:
-	UAssetClassProperty(ECppProperty, int32 InOffset, uint64 InFlags, UClass* InMetaClass)
+	USoftClassProperty(ECppProperty, int32 InOffset, uint64 InFlags, UClass* InMetaClass)
 		: Super(FObjectInitializer::Get(), EC_CppProperty, InOffset, InFlags, UClass::StaticClass())
 		, MetaClass(InMetaClass)
 	{}
 
-	UAssetClassProperty( const FObjectInitializer& ObjectInitializer, ECppProperty, int32 InOffset, uint64 InFlags, UClass* InMetaClass )
+	USoftClassProperty( const FObjectInitializer& ObjectInitializer, ECppProperty, int32 InOffset, uint64 InFlags, UClass* InMetaClass )
 		:	Super(ObjectInitializer, EC_CppProperty, InOffset, InFlags, UClass::StaticClass() )
 		,	MetaClass( InMetaClass )
 	{}
diff --git a/Engine/Source/Runtime/Engine/Classes/AI/AISystemBase.h b/Engine/Source/Runtime/Engine/Classes/AI/AISystemBase.h
--- a/Engine/Source/Runtime/Engine/Classes/AI/AISystemBase.h
+++ b/Engine/Source/Runtime/Engine/Classes/AI/AISystemBase.h
@@ -6,7 +6,7 @@
 #include "UObject/ObjectMacros.h"
 #include "UObject/Object.h"
 #include "Engine/EngineBaseTypes.h"
-#include "Misc/StringClassReference.h"
+#include "UObject/SoftObjectPath.h"
 #include "Modules/ModuleInterface.h"
 #include "AISystemBase.generated.h"
 
@@ -45,7 +45,7 @@ class ENGINE_API UAISystemBase : public UObject
 
 private:
 	UPROPERTY(globalconfig, noclear, meta = (MetaClass = "AISystem", DisplayName = "AISystem Class"))
-	FStringClassReference AISystemClassName;
+	FSoftClassPath AISystemClassName;
 
 	UPROPERTY(globalconfig, noclear, meta = (MetaClass = "AISystem", DisplayName = "AISystem Module"))
 	FName AISystemModuleName;
@@ -54,7 +54,7 @@ private:
 	bool bInstantiateAISystemOnClient;
 	
 public:
-	static FStringClassReference GetAISystemClassName();
+	static FSoftClassPath GetAISystemClassName();
 	static FName GetAISystemModuleName();
 	static bool ShouldInstantiateInNetMode(ENetMode NetMode);
 };
diff --git a/Engine/Source/Runtime/Engine/Classes/AI/Navigation/NavigationTypes.h b/Engine/Source/Runtime/Engine/Classes/AI/Navigation/NavigationTypes.h
--- a/Engine/Source/Runtime/Engine/Classes/AI/Navigation/NavigationTypes.h
+++ b/Engine/Source/Runtime/Engine/Classes/AI/Navigation/NavigationTypes.h
@@ -11,7 +11,7 @@
 #include "Templates/SubclassOf.h"
 #include "UObject/WeakObjectPtr.h"
 #include "Misc/CoreStats.h"
-#include "Misc/StringClassReference.h"
+#include "UObject/SoftObjectPath.h"
 #include "AI/Navigation/NavFilters/NavigationQueryFilter.h"
 #include "NavigationTypes.generated.h"
 
@@ -540,7 +540,7 @@ struct ENGINE_API FNavDataConfig : public FNavAgentProperties
 	TSubclassOf<ANavigationData> NavigationDataClass;
 
 	UPROPERTY(config)
-	FStringClassReference NavigationDataClassName;
+	FSoftClassPath NavigationDataClassName;
 
 	FNavDataConfig(float Radius = FNavigationSystem::FallbackAgentRadius, float Height = FNavigationSystem::FallbackAgentHeight);
 };
diff --git a/Engine/Source/Runtime/Engine/Classes/Animation/AnimBlueprint.h b/Engine/Source/Runtime/Engine/Classes/Animation/AnimBlueprint.h
--- a/Engine/Source/Runtime/Engine/Classes/Animation/AnimBlueprint.h
+++ b/Engine/Source/Runtime/Engine/Classes/Animation/AnimBlueprint.h
@@ -178,6 +178,6 @@ public:
 private:
 	/** The default skeletal mesh to use when previewing this asset - this only applies when you open Persona using this asset*/
 	UPROPERTY(duplicatetransient, AssetRegistrySearchable)
-	TAssetPtr<class USkeletalMesh> PreviewSkeletalMesh;
+	TSoftObjectPtr<class USkeletalMesh> PreviewSkeletalMesh;
 #endif
 };
diff --git a/Engine/Source/Runtime/Engine/Classes/Animation/AnimInstance.h b/Engine/Source/Runtime/Engine/Classes/Animation/AnimInstance.h
--- a/Engine/Source/Runtime/Engine/Classes/Animation/AnimInstance.h
+++ b/Engine/Source/Runtime/Engine/Classes/Animation/AnimInstance.h
@@ -903,7 +903,7 @@ public:
 	/** Name of Class to do Post Compile Validation.
 	* See Class UAnimBlueprintPostCompileValidation. */
 	UPROPERTY()
-	FStringClassReference PostCompileValidationClassName;
+	FSoftClassPath PostCompileValidationClassName;
 
 	/** Warn if AnimNodes are not using fast path during AnimBP compilation. */
 	virtual bool PCV_ShouldWarnAboutNodesNotUsingFastPath() const { return false; }
diff --git a/Engine/Source/Runtime/Engine/Classes/Animation/AnimationAsset.h b/Engine/Source/Runtime/Engine/Classes/Animation/AnimationAsset.h
--- a/Engine/Source/Runtime/Engine/Classes/Animation/AnimationAsset.h
+++ b/Engine/Source/Runtime/Engine/Classes/Animation/AnimationAsset.h
@@ -908,7 +908,7 @@ public:
 private:
 	/** The default skeletal mesh to use when previewing this asset - this only applies when you open Persona using this asset*/
 	UPROPERTY(duplicatetransient, AssetRegistrySearchable)
-	TAssetPtr<class USkeletalMesh> PreviewSkeletalMesh;
+	TSoftObjectPtr<class USkeletalMesh> PreviewSkeletalMesh;
 #endif //WITH_EDITORONLY_DATA
 
 protected:
diff --git a/Engine/Source/Runtime/Engine/Classes/Animation/PreviewAssetAttachComponent.h b/Engine/Source/Runtime/Engine/Classes/Animation/PreviewAssetAttachComponent.h
--- a/Engine/Source/Runtime/Engine/Classes/Animation/PreviewAssetAttachComponent.h
+++ b/Engine/Source/Runtime/Engine/Classes/Animation/PreviewAssetAttachComponent.h
@@ -11,8 +11,7 @@
 #include "CoreMinimal.h"
 #include "UObject/ObjectMacros.h"
 #include "UObject/Object.h"
-#include "Misc/StringAssetReference.h"
-#include "UObject/AssetPtr.h"
+#include "UObject/SoftObjectPtr.h"
 #include "PreviewAssetAttachComponent.generated.h"
 
 /** Preview items that are attached to the skeleton **/
@@ -24,7 +23,7 @@ struct FPreviewAttachedObjectPair
 private:
 	/** the object to be attached */
 	UPROPERTY()
-	TAssetPtr<class UObject> AttachedObject;
+	TSoftObjectPtr<class UObject> AttachedObject;
 
 	UPROPERTY()
 	UObject* Object_DEPRECATED;
@@ -48,18 +47,7 @@ public:
 
 	UObject* GetAttachedObject() const
 	{
-		UObject* AttachedObjectPtr = AttachedObject.Get();
-		if (!AttachedObjectPtr)
-		{
-			// if preview mesh isn't loaded, see if we have set
-			FStringAssetReference AttachedObjectStringRef = AttachedObject.ToStringReference();
-			// load it since now is the time to load
-			if (!AttachedObjectStringRef.ToString().IsEmpty())
-			{
-				AttachedObjectPtr = Cast<UObject>(StaticLoadObject(UObject::StaticClass(), NULL, *AttachedObjectStringRef.ToString(), NULL, LOAD_None, NULL));
-			}
-		}
-		return AttachedObjectPtr;
+		return AttachedObject.LoadSynchronous();
 	}
 
 	void SetAttachedObject(UObject* InObject)
diff --git a/Engine/Source/Runtime/Engine/Classes/Animation/Skeleton.h b/Engine/Source/Runtime/Engine/Classes/Animation/Skeleton.h
--- a/Engine/Source/Runtime/Engine/Classes/Animation/Skeleton.h
+++ b/Engine/Source/Runtime/Engine/Classes/Animation/Skeleton.h
@@ -467,11 +467,11 @@ private:
 private:
 	/** The default skeletal mesh to use when previewing this skeleton */
 	UPROPERTY(duplicatetransient, AssetRegistrySearchable)
-	TAssetPtr<class USkeletalMesh> PreviewSkeletalMesh;
+	TSoftObjectPtr<class USkeletalMesh> PreviewSkeletalMesh;
 
 	/** The additional skeletal meshes to use when previewing this skeleton */
 	UPROPERTY(duplicatetransient, AssetRegistrySearchable)
-	TAssetPtr<class UDataAsset> AdditionalPreviewSkeletalMeshes;
+	TSoftObjectPtr<class UDataAsset> AdditionalPreviewSkeletalMeshes;
 
 	UPROPERTY()
 	FRigConfiguration RigConfig;
diff --git a/Engine/Source/Runtime/Engine/Classes/EdGraph/EdGraphNode.h b/Engine/Source/Runtime/Engine/Classes/EdGraph/EdGraphNode.h
--- a/Engine/Source/Runtime/Engine/Classes/EdGraph/EdGraphNode.h
+++ b/Engine/Source/Runtime/Engine/Classes/EdGraph/EdGraphNode.h
@@ -175,6 +175,7 @@ class ENGINE_API UEdGraphNode : public UObject
 {
 	GENERATED_UCLASS_BODY()
 
+public:
 	TArray<UEdGraphPin*> Pins;
 
 	/** List of connector pins */
@@ -201,13 +202,16 @@ class ENGINE_API UEdGraphNode : public UObject
 	UPROPERTY()
 	TEnumAsByte<ENodeAdvancedPins::Type> AdvancedPinDisplay;
 
+private:
 	/** Indicates in what state the node is enabled, which may eliminate it from being compiled */
 	UPROPERTY()
 	ENodeEnabledState EnabledState;
 
+public:
 	/** When reconstructing a node should the orphaned pins be retained and transfered to the new pin list. */
 	ESaveOrphanPinMode OrphanedPinSaveMode;
 
+private:
 	/** Indicates whether or not the user explicitly set the enabled state */
 	UPROPERTY()
 	uint8 bUserSetEnabledState:1;
@@ -278,29 +282,40 @@ public:
 	FGuid NodeGuid;
 
 public:
-	/** Enables this node. */
-	FORCEINLINE void EnableNode()
+	/** Determines whether or not the node is enabled. */
+	bool IsNodeEnabled() const
 	{
-		bUserSetEnabledState = false;
-		EnabledState = ENodeEnabledState::Enabled;
+		return (EnabledState == ENodeEnabledState::Enabled) || ((EnabledState == ENodeEnabledState::DevelopmentOnly) && IsInDevelopmentMode());
 	}
 
-	/** Disables this node. */
-	FORCEINLINE void DisableNode()
+	/** Returns the specific sort of enable state this node wants */
+	ENodeEnabledState GetDesiredEnabledState() const
 	{
-		bUserSetEnabledState = false;
-		EnabledState = ENodeEnabledState::Disabled;
+		return EnabledState;
 	}
 
-	/** Determines whether or not the node is enabled. */
-	FORCEINLINE bool IsNodeEnabled() const
+	/** Set the enabled state of the node to a new value */
+	void SetEnabledState(ENodeEnabledState NewState, bool bUserAction = true)
+	{
+		EnabledState = NewState;
+		bUserSetEnabledState = bUserAction;
+	}
+
+	/** Has the user set the enabled state or is it still using the automatic settings? */
+	bool HasUserSetTheEnabledState() const
 	{
-		return (EnabledState == ENodeEnabledState::Enabled) || (EnabledState == ENodeEnabledState::DevelopmentOnly && IsInDevelopmentMode());
+		return bUserSetEnabledState;
 	}
 
 	/** Determines whether or not the node will compile in development mode. */
 	virtual bool IsInDevelopmentMode() const;
 
+	/** Returns true if this is a disabled automatically placed ghost node (see the DefaultEventNodes ini section) */
+	bool IsAutomaticallyPlacedGhostNode() const;
+
+	/** Marks this node as an automatically placed ghost node (see the DefaultEventNodes ini section) */
+	void MakeAutomaticallyPlacedGhostNode();
+
 #if WITH_EDITOR
 
 private:
@@ -598,6 +613,12 @@ public:
 	// (the object can be an actor, which selects it in the world, or a node/graph/pin)
 	virtual UObject* GetJumpTargetForDoubleClick() const;
 
+	// Returns true if it is possible to jump to the definition of this node (e.g., if it's a variable get or a function call)
+	virtual bool CanJumpToDefinition() const;
+
+	// Jump to the definition of this node (should only be called if CanJumpToDefinition() return true)
+	virtual void JumpToDefinition() const;
+
 	/** Create a new unique Guid for this node */
 	void CreateNewGuid();
 
@@ -628,9 +649,6 @@ public:
 	// called when a pin is removed
 	virtual void OnPinRemoved( UEdGraphPin* InRemovedPin ) {}
 
-	/** Return whether to draw this node as a comment node */
-	virtual bool ShouldDrawNodeAsComment() const { return false; }
-
 	/** 
 	* Returns whether to draw this node as a control point only (knot/reroute node). Note that this means that the node should only have on input and output pin.
 	* @param OutInputPinIndex The index in the pins array associated with the control point input pin.
diff --git a/Engine/Source/Runtime/Engine/Classes/EdGraph/EdGraphPin.h b/Engine/Source/Runtime/Engine/Classes/EdGraph/EdGraphPin.h
--- a/Engine/Source/Runtime/Engine/Classes/EdGraph/EdGraphPin.h
+++ b/Engine/Source/Runtime/Engine/Classes/EdGraph/EdGraphPin.h
@@ -560,7 +560,7 @@ private:
 	bool Serialize(FArchive& Ar);
 
 	// Helper functions
-	static void EnableAllConnectedNodes(class UEdGraphNode* InNode);
+	static void ConvertConnectedGhostNodesToRealNodes(UEdGraphNode* InNode);
 	static void ResolveReferencesToPin(UEdGraphPin* Pin, bool bStrictValidation = true);
 	static void SerializePinArray(FArchive& Ar, TArray<UEdGraphPin*>& ArrayRef, UEdGraphPin* RequestingPin, EPinResolveType ResolveType);
 	static bool SerializePin(FArchive& Ar, UEdGraphPin*& PinRef, int32 ArrayIdx, UEdGraphPin* RequestingPin, EPinResolveType ResolveType, TArray<UEdGraphPin*>& OldPins);
diff --git a/Engine/Source/Runtime/Engine/Classes/EdGraph/EdGraphSchema.h b/Engine/Source/Runtime/Engine/Classes/EdGraph/EdGraphSchema.h
--- a/Engine/Source/Runtime/Engine/Classes/EdGraph/EdGraphSchema.h
+++ b/Engine/Source/Runtime/Engine/Classes/EdGraph/EdGraphSchema.h
@@ -8,6 +8,7 @@
 #include "EdGraph/EdGraphNode.h"
 #include "EdGraph/EdGraphPin.h"
 #include "AssetData.h"
+#include "UObject/ObjectKey.h"
 #include "EdGraphSchema.generated.h"
 
 class FSlateRect;
@@ -50,6 +51,38 @@ enum ECanCreateConnectionResponse
 	CONNECT_RESPONSE_MAX,
 };
 
+// Used to opaquely verify that two different persistent entries backing actions are part of the same section/category (e.g., both are variables in the same Blueprint)
+struct FEdGraphSchemaActionDefiningObject
+{
+	FEdGraphSchemaActionDefiningObject(UObject* InObject, void* AdditionalPointer = nullptr, FName AdditionalName = NAME_None)
+		: DefiningObject(InObject)
+		, DefiningPointer(AdditionalPointer)
+		, DefiningName(AdditionalName)
+		, bIsEditable((InObject != nullptr) ? !InObject->IsNative() : false)
+	{
+	}
+
+	friend bool operator==(const FEdGraphSchemaActionDefiningObject& A, const FEdGraphSchemaActionDefiningObject& B)
+	{
+		return (A.DefiningObject == B.DefiningObject) && (A.DefiningPointer == B.DefiningPointer) && (A.DefiningName == B.DefiningName) && (A.bIsEditable == B.bIsEditable);
+	}
+
+	friend bool operator!=(const FEdGraphSchemaActionDefiningObject& A, const FEdGraphSchemaActionDefiningObject& B)
+	{
+		return !(A == B);
+	}
+
+	bool IsPotentiallyEditable() const
+	{
+		return bIsEditable;
+	}
+private:
+	FObjectKey DefiningObject;
+	void* DefiningPointer;
+	FName DefiningName;
+	bool bIsEditable;
+};
+
 /** This structure represents a context dependent action, with sufficient information for the schema to perform it. */
 USTRUCT()
 struct ENGINE_API FEdGraphSchemaAction
@@ -148,7 +181,9 @@ public:
 		return NewNode;
 	}
 
-	void UpdateCategory(FText NewCategory);
+	// Updates the category of the *action* and refreshes the search text; does not change the persistent backing item
+	// (e.g., it will not actually move a user added variable or function to a new category)
+	void CosmeticUpdateCategory(FText NewCategory);
 
 	void UpdateSearchData(FText NewMenuDescription, FText NewToolTipDescription, FText NewCategory, FText NewKeywords);
 
@@ -236,6 +271,19 @@ public:
 	// GC.
 	virtual void AddReferencedObjects(FReferenceCollector& Collector) {}
 
+	// Moves the item backing this action to the specified category if it is possible (does nothing for native-introduced variables/functions/etc...)
+	virtual void MovePersistentItemToCategory(const FText& NewCategoryName) {}
+
+	// Returns the ordering index of this action in the parent container (if the item cannot be reordered then this will return INDEX_NONE)
+	virtual int32 GetReorderIndexInContainer() const { return INDEX_NONE; }
+
+	// Reorders this action to be before the other item in the parent container (returns false if they are not in the same container or cannot be reordered)
+	virtual bool ReorderToBeforeAction(TSharedRef<FEdGraphSchemaAction> OtherAction) { return false; }
+
+	// Returns an opaque handle that can be used to confirm that two different persistent entries backing actions are part of the same section/category
+	// (e.g., both are variables in the same Blueprint)
+	virtual FEdGraphSchemaActionDefiningObject GetPersistentItemDefiningObject() const { return FEdGraphSchemaActionDefiningObject(nullptr); }
+
 private:
 	void UpdateSearchText();
 };
@@ -786,7 +834,7 @@ class ENGINE_API UEdGraphSchema : public UObject
 	virtual void BreakSinglePinLink(UEdGraphPin* SourcePin, UEdGraphPin* TargetPin);
 
 	/** Split a pin in to subelements */
-	virtual void SplitPin(UEdGraphPin* Pin) const { };
+	virtual void SplitPin(UEdGraphPin* Pin, bool bNotify = true) const { };
 
 	/** Collapses a pin and its siblings back in to the original pin */
 	virtual void RecombinePin(UEdGraphPin* Pin) const { };
@@ -796,7 +844,7 @@ class ENGINE_API UEdGraphSchema : public UObject
 
 	/** Break links on this pin and create links instead on MoveToPin */
 	virtual FPinConnectionResponse MovePinLinks(UEdGraphPin& MoveFromPin, UEdGraphPin& MoveToPin, bool bIsIntermediateMove = false) const;
-
+	 
 	/** Copies pin links from one pin to another without breaking the original links */
 	virtual FPinConnectionResponse CopyPinLinks(UEdGraphPin& CopyFromPin, UEdGraphPin& CopyToPin, bool bIsIntermediateCopy = false) const;
 
@@ -831,7 +879,7 @@ class ENGINE_API UEdGraphSchema : public UObject
 	 * @param	InOutExtraNames	List of extra names that are in-use from the substitution should be added to this list to prevent other substitutions from attempting to use them
 	 * @return					NULL if a substitute node cannot be created; otherwise, the substitute node instance
 	 */
-	virtual UEdGraphNode* CreateSubstituteNode(UEdGraphNode* Node, const UEdGraph* Graph, FObjectInstancingGraph* InstanceGraph, TArray<FName>& InOutExtraNames) const { return nullptr; }
+	virtual UEdGraphNode* CreateSubstituteNode(UEdGraphNode* Node, const UEdGraph* Graph, FObjectInstancingGraph* InstanceGraph, TSet<FName>& InOutExtraNames) const { return nullptr; }
 
 
 	/**
diff --git a/Engine/Source/Runtime/Engine/Classes/Engine/AssetManager.h b/Engine/Source/Runtime/Engine/Classes/Engine/AssetManager.h
--- a/Engine/Source/Runtime/Engine/Classes/Engine/AssetManager.h
+++ b/Engine/Source/Runtime/Engine/Classes/Engine/AssetManager.h
@@ -80,7 +80,7 @@ public:
 	 * @param BundleData		List of Name->asset paths that represent the possible bundle states for this asset
 	 * @return					True if added
 	 */
-	virtual bool AddDynamicAsset(const FPrimaryAssetId& PrimaryAssetId, const FStringAssetReference& AssetPath, const FAssetBundleData& BundleData);
+	virtual bool AddDynamicAsset(const FPrimaryAssetId& PrimaryAssetId, const FSoftObjectPath& AssetPath, const FAssetBundleData& BundleData);
 
 	/** This will expand out references in the passed in AssetBundleData that are pointing to other primary assets with bundles. This is useful to preload entire webs of assets */
 	virtual void RecursivelyExpandBundleData(FAssetBundleData& BundleData);
@@ -115,17 +115,17 @@ public:
 	/** Gets list of all loaded objects for a primary asset type, returns true if any were found. Will return blueprint class for blueprint assets. This works even if the asset wasn't loaded explicitly */
 	virtual bool GetPrimaryAssetObjectList(FPrimaryAssetType PrimaryAssetType, TArray<UObject*>& ObjectList) const;
 
-	/** Gets the FStringAssetReference for a primary asset type and name, returns invalid if not found */
-	virtual FStringAssetReference GetPrimaryAssetPath(const FPrimaryAssetId& PrimaryAssetId) const;
+	/** Gets the FSoftObjectPath for a primary asset type and name, returns invalid if not found */
+	virtual FSoftObjectPath GetPrimaryAssetPath(const FPrimaryAssetId& PrimaryAssetId) const;
 
-	/** Gets the list of all FStringAssetReferences for a given type, returns true if any found */
-	virtual bool GetPrimaryAssetPathList(FPrimaryAssetType PrimaryAssetType, TArray<FStringAssetReference>& AssetPathList) const;
+	/** Gets the list of all FSoftObjectPaths for a given type, returns true if any found */
+	virtual bool GetPrimaryAssetPathList(FPrimaryAssetType PrimaryAssetType, TArray<FSoftObjectPath>& AssetPathList) const;
 
 	/** Sees if the passed in object is a registered primary asset, if so return it. Returns invalid Identifier if not found */
 	virtual FPrimaryAssetId GetPrimaryAssetIdForObject(UObject* Object) const;
 
 	/** Sees if the passed in object path is a registered primary asset, if so return it. Returns invalid Identifier if not found */
-	virtual FPrimaryAssetId GetPrimaryAssetIdForPath(const FStringAssetReference& ObjectPath) const;
+	virtual FPrimaryAssetId GetPrimaryAssetIdForPath(const FSoftObjectPath& ObjectPath) const;
 	virtual FPrimaryAssetId GetPrimaryAssetIdForPath(FName ObjectPath) const;
 
 	/** Sees if the package has a primary asset, useful if only the package name is available */
@@ -247,7 +247,7 @@ public:
 	virtual TSharedPtr<FStreamableHandle> PreloadPrimaryAssets(const TArray<FPrimaryAssetId>& AssetsToLoad, const TArray<FName>& LoadBundles, bool bLoadRecursive, FStreamableDelegate DelegateToCall = FStreamableDelegate(), TAsyncLoadPriority Priority = FStreamableManager::DefaultAsyncLoadPriority);
 
 	/** Quick wrapper to async load some non primary assets with the primary streamable manager. This will not auto release the handle, release it if needed */
-	virtual TSharedPtr<FStreamableHandle> LoadAssetList(const TArray<FStringAssetReference>& AssetList, FStreamableDelegate DelegateToCall = FStreamableDelegate(), TAsyncLoadPriority Priority = FStreamableManager::DefaultAsyncLoadPriority, const FString& DebugName = TEXT("LoadAssetList"));
+	virtual TSharedPtr<FStreamableHandle> LoadAssetList(const TArray<FSoftObjectPath>& AssetList, FStreamableDelegate DelegateToCall = FStreamableDelegate(), TAsyncLoadPriority Priority = FStreamableManager::DefaultAsyncLoadPriority, const FString& DebugName = TEXT("LoadAssetList"));
 
 	/** Returns a single AssetBundleInfo, matching Scope and Name */
 	virtual FAssetBundleEntry GetAssetBundleEntry(const FPrimaryAssetId& BundleScope, FName BundleName) const;
@@ -262,7 +262,7 @@ public:
 	 * @param OutMissingChunkList	Chunks that are known about but not yet installed
 	 * @param OutErrorChunkList		Chunks that do not exist at all and are not installable
 	 */
-	virtual bool FindMissingChunkList(const TArray<FStringAssetReference>& AssetList, TArray<int32>& OutMissingChunkList, TArray<int32>& OutErrorChunkList) const;
+	virtual bool FindMissingChunkList(const TArray<FSoftObjectPath>& AssetList, TArray<int32>& OutMissingChunkList, TArray<int32>& OutErrorChunkList) const;
 
 	/** 
 	 * Acquires a set of chunks using the platform chunk layer, then calls the passed in callback
@@ -271,7 +271,7 @@ public:
 	 * @param CompleteDelegate		Delegate called when chunks have been acquired or failed. If any chunks fail the entire operation is considered a failure
 	 * @param Priority				Priority to use when acquiring chunks
 	 */
-	virtual void AcquireResourcesForAssetList(const TArray<FStringAssetReference>& AssetList, FAssetManagerAcquireResourceDelegate CompleteDelegate, EChunkPriority::Type Priority = EChunkPriority::Immediate);
+	virtual void AcquireResourcesForAssetList(const TArray<FSoftObjectPath>& AssetList, FAssetManagerAcquireResourceDelegate CompleteDelegate, EChunkPriority::Type Priority = EChunkPriority::Immediate);
 
 	/** 
 	 * Acquires a set of chunks using the platform chunk layer, then calls the passed in callback. This will download all bundles of a primary asset
@@ -309,10 +309,10 @@ public:
 	virtual FPrimaryAssetId ExtractPrimaryAssetIdFromData(const FAssetData& AssetData, FPrimaryAssetType SuggestedType = NAME_None) const;
 
 	/** Gets the FAssetData at a specific path, handles redirectors and blueprint classes correctly. Returns true if it found a valid data */
-	virtual bool GetAssetDataForPath(const FStringAssetReference& ObjectPath, FAssetData& AssetData) const;
+	virtual bool GetAssetDataForPath(const FSoftObjectPath& ObjectPath, FAssetData& AssetData) const;
 
-	/** Turns an FAssetData into FStringAssetReference, handles adding _C as necessary */
-	virtual FStringAssetReference GetAssetPathForData(const FAssetData& AssetData) const;
+	/** Turns an FAssetData into FSoftObjectPath, handles adding _C as necessary */
+	virtual FSoftObjectPath GetAssetPathForData(const FAssetData& AssetData) const;
 
 	/** Tries to redirect a Primary Asset Id, using list in AssetManagerSettings */
 	virtual FPrimaryAssetId GetRedirectedPrimaryAssetId(const FPrimaryAssetId& OldId) const;
@@ -322,10 +322,10 @@ public:
 
 	/** Reads AssetManagerSettings for specifically redirected asset paths. This is useful if you need to convert older saved data */
 	virtual FName GetRedirectedAssetPath(FName OldPath) const;
-	virtual FStringAssetReference GetRedirectedAssetPath(const FStringAssetReference& OldPath) const;
+	virtual FSoftObjectPath GetRedirectedAssetPath(const FSoftObjectPath& OldPath) const;
 
-	/** Extracts all FStringAssetReferences from a Class/Struct */
-	virtual void ExtractStringAssetReferences(const UStruct* Struct, const void* StructValue, TArray<FStringAssetReference>& FoundAssetReferences, const TArray<FName>& PropertiesToSkip = TArray<FName>()) const;
+	/** Extracts all FSoftObjectPaths from a Class/Struct */
+	virtual void ExtractSoftObjectPaths(const UStruct* Struct, const void* StructValue, TArray<FSoftObjectPath>& FoundAssetReferences, const TArray<FName>& PropertiesToSkip = TArray<FName>()) const;
 
 	/** Dumps out summary of managed types to log */
 	static void DumpAssetTypeSummary();
@@ -388,7 +388,7 @@ public:
 	 * Initializes asset bundle data from a passed in struct or class, this will read the AssetBundles metadata off the UProperties. As an example this property definition:
 	 *
 	 * UPROPERTY(EditDefaultsOnly, Category=Data, meta = (AssetBundles = "Client,Server"))
-	 * TAssetPtr<class UCurveTable> CurveTableReference;
+	 * TSoftObjectPtr<class UCurveTable> CurveTableReference;
 	 *
 	 * Would add the value of CurveTableReference to both the Client and Server asset bundles
 	 *
@@ -559,6 +559,10 @@ protected:
 	UPROPERTY()
 	bool bUpdateManagementDatabaseAfterScan;
 
+	/** True if only on-disk assets should be searched by the asset registry */
+	UPROPERTY()
+	bool bIncludeOnlyOnDiskAssets;
+
 	/** Number of notifications seen in this update */
 	UPROPERTY()
 	int32 NumberOfSpawnedNotifications;
diff --git a/Engine/Source/Runtime/Engine/Classes/Engine/AssetManagerTypes.h b/Engine/Source/Runtime/Engine/Classes/Engine/AssetManagerTypes.h
--- a/Engine/Source/Runtime/Engine/Classes/Engine/AssetManagerTypes.h
+++ b/Engine/Source/Runtime/Engine/Classes/Engine/AssetManagerTypes.h
@@ -86,7 +86,7 @@ struct FPrimaryAssetTypeInfo
 private:
 	/** Base Class of all assets of this type */
 	UPROPERTY(EditAnywhere, Category = AssetType, meta = (AllowAbstract))
-	TAssetSubclassOf<UObject> AssetBaseClass;
+	TSoftClassPtr<UObject> AssetBaseClass;
 public:
 
 	/** Base Class of all assets of this type */
@@ -108,7 +108,7 @@ private:
 
 	/** Individual assets to scan */
 	UPROPERTY(EditAnywhere, Category = AssetType)
-	TArray<FStringAssetReference> SpecificAssets;
+	TArray<FSoftObjectPath> SpecificAssets;
 public:
 
 	/** Default management rules for this type, individual assets can be overridden */
diff --git a/Engine/Source/Runtime/Engine/Classes/Engine/Blueprint.h b/Engine/Source/Runtime/Engine/Classes/Engine/Blueprint.h
--- a/Engine/Source/Runtime/Engine/Classes/Engine/Blueprint.h
+++ b/Engine/Source/Runtime/Engine/Classes/Engine/Blueprint.h
@@ -11,7 +11,7 @@
 #include "Engine/EngineTypes.h"
 #include "EdGraph/EdGraphPin.h"
 #include "Engine/BlueprintCore.h"
-#include "Misc/StringAssetReference.h"
+#include "UObject/SoftObjectPath.h"
 #include "Blueprint.generated.h"
 
 class FCompilerResultsLog;
@@ -104,7 +104,7 @@ struct FCompilerNativizationOptions
 
 	// Individually excluded assets
 	UPROPERTY()
-	TSet<FStringAssetReference> ExcludedAssets;
+	TSet<FSoftObjectPath> ExcludedAssets;
 
 	// Excluded folders. It excludes only BPGCs, enums and structures are still converted.
 	UPROPERTY()
@@ -146,9 +146,6 @@ public:
 	/** Whether or not this compile is for a duplicated blueprint */
 	bool bIsDuplicationInstigated;
 
-	/** Whether or not this compile should emit instrumentation events */
-	bool bAddInstrumentation;
-
 	/** Whether or not to reinstance and stub if the blueprint fails to compile */
 	bool bReinstanceAndStubOnFailure;
 
@@ -168,11 +165,6 @@ public:
 			|| (CompileType == EKismetCompileType::Cpp);
 	}
 
-	bool IsInstrumentationActive() const
-	{
-		return bAddInstrumentation && DoesRequireBytecodeGeneration();
-	}
-
 	/** Whether or not this compile type should operate on the generated class of the blueprint, as opposed to just the skeleton */
 	bool IsGeneratedClassCompileType() const
 	{
@@ -184,7 +176,6 @@ public:
 		, bSaveIntermediateProducts(false)
 		, bRegenerateSkelton(true)
 		, bIsDuplicationInstigated(false)
-		, bAddInstrumentation(false)
 		, bReinstanceAndStubOnFailure(true)
 	{
 	};
diff --git a/Engine/Source/Runtime/Engine/Classes/Engine/BlueprintGeneratedClass.h b/Engine/Source/Runtime/Engine/Classes/Engine/BlueprintGeneratedClass.h
--- a/Engine/Source/Runtime/Engine/Classes/Engine/BlueprintGeneratedClass.h
+++ b/Engine/Source/Runtime/Engine/Classes/Engine/BlueprintGeneratedClass.h
@@ -647,10 +647,6 @@ public:
 	UPROPERTY()
 	TArray<FEventGraphFastCallPair> FastCallPairs_DEPRECATED;
 
-	// If this Generated Class has instrumentation
-	UPROPERTY()
-	bool bHasInstrumentation;
-
 #if WITH_EDITORONLY_DATA
 	UPROPERTY(Transient)
 	UObject* OverridenArchetypeForCDO;
@@ -720,10 +716,6 @@ public:
 	virtual void Bind() override;
 	virtual void GetRequiredPreloadDependencies(TArray<UObject*>& DependenciesOut) override;
 	virtual UObject* FindArchetype(UClass* ArchetypeClass, const FName ArchetypeName) const override;
-	virtual bool HasInstrumentation() const override 
-	{
-		return bHasInstrumentation; 
-	}
 
 	virtual void InitPropertiesFromCustomList(uint8* DataPtr, const uint8* DefaultDataPtr) override;
 
diff --git a/Engine/Source/Runtime/Engine/Classes/Engine/Engine.h b/Engine/Source/Runtime/Engine/Classes/Engine/Engine.h
--- a/Engine/Source/Runtime/Engine/Classes/Engine/Engine.h
+++ b/Engine/Source/Runtime/Engine/Classes/Engine/Engine.h
@@ -10,8 +10,7 @@
 #include "Templates/SubclassOf.h"
 #include "Engine/EngineTypes.h"
 #include "Engine/EngineBaseTypes.h"
-#include "Misc/StringAssetReference.h"
-#include "Misc/StringClassReference.h"
+#include "UObject/SoftObjectPath.h"
 #include "Engine/World.h"
 #include "Misc/BufferedOutputDevice.h"
 #include "Engine.generated.h"
@@ -317,12 +316,13 @@ struct FWorldContext
 	UPROPERTY(transient)
 	TArray<FNamedNetDriver> ActiveNetDrivers;
 
+	/** The PIE instance of this world, -1 is default */
 	int32	PIEInstance;
 
+	/** The Prefix in front of PIE level names, empty is default */
 	FString	PIEPrefix;
 
-	FString PIERemapPrefix;
-
+	/** Is this running as a dedicated server */
 	bool	RunAsDedicated;
 
 	/** Is this world context waiting for an online login to complete (for PIE) */
@@ -599,7 +599,7 @@ private:
 public:
 	/** @todo document */
 	UPROPERTY(globalconfig, EditAnywhere, Category=Fonts, meta=(AllowedClasses="Font", DisplayName="Tiny Font"))
-	FStringAssetReference TinyFontName;
+	FSoftObjectPath TinyFontName;
 
 private:
 	/** @todo document */
@@ -609,7 +609,7 @@ private:
 public:
 	/** @todo document */
 	UPROPERTY(globalconfig, EditAnywhere, Category=Fonts, meta=(AllowedClasses="Font", DisplayName="Small Font"))
-	FStringAssetReference SmallFontName;
+	FSoftObjectPath SmallFontName;
 
 private:
 	/** @todo document */
@@ -619,7 +619,7 @@ private:
 public:
 	/** @todo document */
 	UPROPERTY(globalconfig, EditAnywhere, Category=Fonts, meta=(AllowedClasses="Font", DisplayName="Medium Font"))
-	FStringAssetReference MediumFontName;
+	FSoftObjectPath MediumFontName;
 
 private:
 	/** @todo document */
@@ -629,7 +629,7 @@ private:
 public:
 	/** @todo document */
 	UPROPERTY(globalconfig, EditAnywhere, Category=Fonts, meta=(AllowedClasses="Font", DisplayName="Large Font"))
-	FStringAssetReference LargeFontName;
+	FSoftObjectPath LargeFontName;
 
 private:
 	/** @todo document */
@@ -639,7 +639,7 @@ private:
 public:
 	/** @todo document */
 	UPROPERTY(globalconfig, EditAnywhere, Category=Fonts, meta=(AllowedClasses="Font", DisplayName="Subtitle Font"), AdvancedDisplay)
-	FStringAssetReference SubtitleFontName;
+	FSoftObjectPath SubtitleFontName;
 
 private:
 	/** Any additional fonts that script may use without hard-referencing the font. */
@@ -660,7 +660,7 @@ public:
 
 	/** @todo document */
 	UPROPERTY(globalconfig, noclear, EditAnywhere, Category=DefaultClasses, meta=(MetaClass="Console", DisplayName="Console Class"))
-	FStringClassReference ConsoleClassName;
+	FSoftClassPath ConsoleClassName;
 
 	/** The class to use for the game viewport client. */
 	UPROPERTY()
@@ -668,7 +668,7 @@ public:
 
 	/** @todo document */
 	UPROPERTY(globalconfig, noclear, EditAnywhere, Category=DefaultClasses, meta=(MetaClass="GameViewportClient", DisplayName="Game Viewport Client Class"))
-	FStringClassReference GameViewportClientClassName;
+	FSoftClassPath GameViewportClientClassName;
 
 	/** The class to use for local players. */
 	UPROPERTY()
@@ -676,7 +676,7 @@ public:
 
 	/** @todo document */
 	UPROPERTY(globalconfig, noclear, EditAnywhere, Category=DefaultClasses, meta=(MetaClass="LocalPlayer", DisplayName="Local Player Class"))
-	FStringClassReference LocalPlayerClassName;
+	FSoftClassPath LocalPlayerClassName;
 
 	/** The class for WorldSettings **/
 	UPROPERTY()
@@ -684,11 +684,11 @@ public:
 
 	/** @todo document */
 	UPROPERTY(globalconfig, noclear, EditAnywhere, Category=DefaultClasses, meta=(MetaClass="WorldSettings", DisplayName="World Settings Class"))
-	FStringClassReference WorldSettingsClassName;
+	FSoftClassPath WorldSettingsClassName;
 
 	/** @todo document */
 	UPROPERTY(globalconfig, noclear, meta=(MetaClass="NavigationSystem", DisplayName="Navigation System Class"))
-	FStringClassReference NavigationSystemClassName;
+	FSoftClassPath NavigationSystemClassName;
 
 	/** The class for NavigationSystem **/
 	UPROPERTY()
@@ -696,7 +696,7 @@ public:
 	
 	/** Name of behavior tree manager class */
 	UPROPERTY(globalconfig, noclear, meta=(MetaClass="AvoidanceManager", DisplayName="Avoidance Manager Class"))
-	FStringClassReference AvoidanceManagerClassName;
+	FSoftClassPath AvoidanceManagerClassName;
 	
 	/** The class for behavior tree manager **/
 	UPROPERTY()
@@ -708,17 +708,17 @@ public:
 
 	/** Name of PhysicsCollisionHandler class we should use by default. */
 	UPROPERTY(globalconfig, noclear, EditAnywhere, Category=DefaultClasses, meta=(MetaClass="PhysicsCollisionHandler", DisplayName="Physics Collision Handler Class"), AdvancedDisplay)
-	FStringClassReference PhysicsCollisionHandlerClassName;
+	FSoftClassPath PhysicsCollisionHandlerClassName;
 
 	UPROPERTY(globalconfig, noclear, meta=(MetaClass="GameUserSettings", DisplayName="Game User Settings Class"))
-	FStringClassReference GameUserSettingsClassName;
+	FSoftClassPath GameUserSettingsClassName;
 
 	UPROPERTY()
 	TSubclassOf<class UGameUserSettings> GameUserSettingsClass;
 
 	/** name of Controller class to be used as default AIController class for pawns */
 	UPROPERTY(globalconfig, noclear, meta = (MetaClass = "AI", DisplayName = "Default AIController class for all Pawns"))
-	FStringClassReference AIControllerClassName;
+	FSoftClassPath AIControllerClassName;
 
 	/** Global instance of the user game settings */
 	UPROPERTY()
@@ -730,15 +730,15 @@ public:
 
 	/** @todo document */
 	UPROPERTY(globalconfig, noclear, EditAnywhere, Category=DefaultClasses, meta=(MetaClass="LevelScriptActor", DisplayName="Level Script Actor Class"))
-	FStringClassReference LevelScriptActorClassName;
+	FSoftClassPath LevelScriptActorClassName;
 	
 	/** Name of the base class to use for new blueprints, configurable on a per-game basis */
 	UPROPERTY(globalconfig, noclear, EditAnywhere, Category=DefaultClasses, meta=(MetaClass="Object", DisplayName="Default Blueprint Base Class", AllowAbstract, BlueprintBaseOnly), AdvancedDisplay)
-	FStringClassReference DefaultBlueprintBaseClassName;
+	FSoftClassPath DefaultBlueprintBaseClassName;
 
 	/** Name of a singleton class to create at startup time, configurable per game */
 	UPROPERTY(globalconfig, noclear, EditAnywhere, Category=DefaultClasses, meta=(MetaClass="Object", DisplayName="Game Singleton Class"), AdvancedDisplay)
-	FStringClassReference GameSingletonClassName;
+	FSoftClassPath GameSingletonClassName;
 
 	/** A UObject spawned at initialization time to handle game-specific data */
 	UPROPERTY()
@@ -746,23 +746,19 @@ public:
 
 	/** Name of a singleton class to spawn as the AssetManager, configurable per game. If empty, it will not spawn one */
 	UPROPERTY(globalconfig, noclear, EditAnywhere, Category=DefaultClasses, meta=(MetaClass="Object", DisplayName="Asset Manager Class"), AdvancedDisplay)
-	FStringClassReference AssetManagerClassName;
+	FSoftClassPath AssetManagerClassName;
 
 	/** A UObject spawned at initialization time to handle game-specific data */
 	UPROPERTY()
 	class UAssetManager *AssetManager;
 
-	/** Path that levels for play on console will be saved to (relative to FPaths::ProjectSavedDir()) */
-	UPROPERTY(config)
-	FString PlayOnConsoleSaveDir;
-
 	/** A global default texture. */
 	UPROPERTY()
 	class UTexture2D* DefaultTexture;
 
 	/** @todo document */
 	UPROPERTY(globalconfig)
-	FStringAssetReference DefaultTextureName;
+	FSoftObjectPath DefaultTextureName;
 
 	/** A global default diffuse texture.*/
 	UPROPERTY()
@@ -770,7 +766,7 @@ public:
 
 	/** @todo document */
 	UPROPERTY(globalconfig)
-	FStringAssetReference DefaultDiffuseTextureName;
+	FSoftObjectPath DefaultDiffuseTextureName;
 
 	/** @todo document */
 	UPROPERTY()
@@ -778,7 +774,7 @@ public:
 
 	/** @todo document */
 	UPROPERTY(globalconfig)
-	FStringAssetReference DefaultBSPVertexTextureName;
+	FSoftObjectPath DefaultBSPVertexTextureName;
 
 	/** Texture used to get random image grain values for post processing */
 	UPROPERTY()
@@ -786,7 +782,7 @@ public:
 
 	/** @todo document */
 	UPROPERTY(globalconfig)
-	FStringAssetReference HighFrequencyNoiseTextureName;
+	FSoftObjectPath HighFrequencyNoiseTextureName;
 
 	/** Texture used to blur out of focus content, mimics the Bokeh shape of actual cameras */
 	UPROPERTY()
@@ -794,7 +790,7 @@ public:
 
 	/** @todo document */
 	UPROPERTY(globalconfig)
-	FStringAssetReference DefaultBokehTextureName;
+	FSoftObjectPath DefaultBokehTextureName;
 
 	/** Texture used to bloom when using FFT, mimics characteristic bloom produced in a camera from a signle bright source */
 	UPROPERTY()
@@ -802,7 +798,7 @@ public:
 
 	/** @todo document */
 	UPROPERTY(globalconfig)
-	FStringAssetReference DefaultBloomKernelTextureName;
+	FSoftObjectPath DefaultBloomKernelTextureName;
 
 	/** The material used to render wireframe meshes. */
 	UPROPERTY()
@@ -819,7 +815,7 @@ public:
 
 	/** @todo document */
 	UPROPERTY(globalconfig)
-	FStringAssetReference GeomMaterialName;
+	FSoftObjectPath GeomMaterialName;
 #endif
 
 	/** A material used to render debug meshes. */
@@ -828,7 +824,7 @@ public:
 
 	/** @todo document */
 	UPROPERTY(globalconfig)
-	FStringAssetReference DebugMeshMaterialName;
+	FSoftObjectPath DebugMeshMaterialName;
 
 	/** Material used for visualizing level membership in lit view port modes. */
 	UPROPERTY()
@@ -876,7 +872,7 @@ public:
 
 	/** @todo document */
 	UPROPERTY(globalconfig)
-	FStringAssetReference RemoveSurfaceMaterialName;
+	FSoftObjectPath RemoveSurfaceMaterialName;
 
 	/** Material that renders vertex color as emmissive. */
 	UPROPERTY()
@@ -933,7 +929,7 @@ public:
 
 	/** @todo document */
 	UPROPERTY(globalconfig)
-	FStringAssetReference BoneWeightMaterialName;
+	FSoftObjectPath BoneWeightMaterialName;
 
 	/** Materials used to render cloth properties on skeletal meshes */
 	UPROPERTY()
@@ -947,11 +943,11 @@ public:
 
 	/** Name of the material used to render cloth in the clothing tools */
 	UPROPERTY(globalconfig)
-	FStringAssetReference ClothPaintMaterialName;
+	FSoftObjectPath ClothPaintMaterialName;
 
 	/** Name of the material used to render cloth wireframe in the clothing tools */
 	UPROPERTY(globalconfig)
-	FStringAssetReference ClothPaintMaterialWireframeName;
+	FSoftObjectPath ClothPaintMaterialWireframeName;
 
 	/** A material used to render debug meshes. */
 	UPROPERTY()
@@ -960,7 +956,7 @@ public:
 
 	/** A material used to render debug opaque material. Used in various animation editor viewport features. */
 	UPROPERTY(globalconfig)
-	FStringAssetReference DebugEditorMaterialName;
+	FSoftObjectPath DebugEditorMaterialName;
 
 	/** Material used to render constraint limits */
 	UPROPERTY()
@@ -991,7 +987,7 @@ public:
 
 	/** @todo document */
 	UPROPERTY(globalconfig)
-	FStringAssetReference InvalidLightmapSettingsMaterialName;
+	FSoftObjectPath InvalidLightmapSettingsMaterialName;
 
 	/** Material that renders a message about preview shadows being used. */
 	UPROPERTY()
@@ -999,7 +995,7 @@ public:
 
 	/** @todo document */
 	UPROPERTY(globalconfig, EditAnywhere, Category=DefaultMaterials, meta=(AllowedClasses="Material", DisplayName="Preview Shadows Indicator Material"))
-	FStringAssetReference PreviewShadowsIndicatorMaterialName;
+	FSoftObjectPath PreviewShadowsIndicatorMaterialName;
 
 	/** Material that 'fakes' lighting, used for arrows, widgets. */
 	UPROPERTY()
@@ -1007,7 +1003,7 @@ public:
 
 	/** @todo document */
 	UPROPERTY(globalconfig)
-	FStringAssetReference ArrowMaterialName;
+	FSoftObjectPath ArrowMaterialName;
 
 	/** @todo document */
 	UPROPERTY(globalconfig)
@@ -1095,7 +1091,7 @@ public:
 
 	/** @todo document */
 	UPROPERTY(globalconfig)
-	FStringAssetReference EditorBrushMaterialName;
+	FSoftObjectPath EditorBrushMaterialName;
 #endif
 
 	/** PhysicalMaterial to use if none is defined for a particular object. */
@@ -1104,7 +1100,7 @@ public:
 
 	/** @todo document */
 	UPROPERTY(globalconfig)
-	FStringAssetReference DefaultPhysMaterialName;
+	FSoftObjectPath DefaultPhysMaterialName;
 
 	UPROPERTY(config)
 	TArray<FGameNameRedirect> ActiveGameNameRedirects;
@@ -1124,7 +1120,7 @@ public:
 
 	/** @todo document */
 	UPROPERTY(globalconfig)
-	FStringAssetReference PreIntegratedSkinBRDFTextureName;
+	FSoftObjectPath PreIntegratedSkinBRDFTextureName;
 
 	/** Texture used to do font rendering in shaders */
 	UPROPERTY()
@@ -1132,7 +1128,7 @@ public:
 
 	/** @todo document */
 	UPROPERTY(globalconfig)
-	FStringAssetReference MiniFontTextureName;
+	FSoftObjectPath MiniFontTextureName;
 
 	/** Texture used as a placeholder for terrain weight-maps to give the material the correct texture format. */
 	UPROPERTY()
@@ -1140,7 +1136,7 @@ public:
 
 	/** @todo document */
 	UPROPERTY(globalconfig)
-	FStringAssetReference WeightMapPlaceholderTextureName;
+	FSoftObjectPath WeightMapPlaceholderTextureName;
 
 	/** Texture used to display LightMapDensity */
 	UPROPERTY()
@@ -1148,7 +1144,7 @@ public:
 
 	/** @todo document */
 	UPROPERTY(globalconfig)
-	FStringAssetReference LightMapDensityTextureName;
+	FSoftObjectPath LightMapDensityTextureName;
 
 	// Variables.
 
@@ -2618,8 +2614,8 @@ public:
 	void DestroyNamedNetDriver(UWorld *InWorld, FName NetDriverName);
 	void DestroyNamedNetDriver(UPendingNetGame *PendingNetGame, FName NetDriverName);
 
-	virtual bool NetworkRemapPath( UNetDriver* Driver, FString &Str, bool reading=true) { return false; }
-	virtual bool NetworkRemapPath( UPendingNetGame *PendingNetGame, FString &Str, bool reading=true) { return false; }
+	virtual bool NetworkRemapPath( UNetDriver* Driver, FString &Str, bool bReading=true) { return false; }
+	virtual bool NetworkRemapPath( UPendingNetGame *PendingNetGame, FString &Str, bool bReading=true) { return false; }
 
 	virtual bool HandleOpenCommand( const TCHAR* Cmd, FOutputDevice& Ar, UWorld * InWorld );
 
diff --git a/Engine/Source/Runtime/Engine/Classes/Engine/GameViewportClient.h b/Engine/Source/Runtime/Engine/Classes/Engine/GameViewportClient.h
--- a/Engine/Source/Runtime/Engine/Classes/Engine/GameViewportClient.h
+++ b/Engine/Source/Runtime/Engine/Classes/Engine/GameViewportClient.h
@@ -16,6 +16,7 @@
 #include "Engine/TitleSafeZone.h"
 #include "Engine/GameViewportDelegates.h"
 #include "Engine/DebugDisplayProperty.h"
+#include "UObject/SoftObjectPath.h"
 
 #include "GameViewportClient.generated.h"
 
@@ -29,7 +30,6 @@ class UCanvas;
 class UGameInstance;
 class ULocalPlayer;
 class UNetDriver;
-struct FStringClassReference;
 class FHardwareCursor;
 
 /**
@@ -812,7 +812,7 @@ private:
 	void HandleViewportStatDisableAll(const bool bInAnyViewport);
 	
 	/** Adds a cursor to the set based on the enum and the class reference to it. */
-	void AddSoftwareCursor(EMouseCursor::Type Cursor, const FStringClassReference& CursorClass);
+	void AddSoftwareCursor(EMouseCursor::Type Cursor, const FSoftClassPath& CursorClass);
 
 private:
 	/** Slate window associated with this viewport client.  The same window may host more than one viewport client. */
