Copying //UE4/Dev-Framework to //UE4/Dev-Main (Source: //UE4/Dev-Framework @ 3628051)

diff --git a/Engine/Source/Editor/UnrealEd/Private/SourceCodeNavigation.cpp b/Engine/Source/Editor/UnrealEd/Private/SourceCodeNavigation.cpp
--- a/Engine/Source/Editor/UnrealEd/Private/SourceCodeNavigation.cpp
+++ b/Engine/Source/Editor/UnrealEd/Private/SourceCodeNavigation.cpp
@@ -1804,7 +1804,7 @@ bool FSourceCodeNavigation::FindClassHeaderPath( const UField *Field, FString &O
 		if(FSourceCodeNavigation::FindModulePath(*ModulePackageName + ModuleNameIdx + 1, ModuleBasePath))
 		{
 			// Get the metadata for the class path relative to the module base
-			FString ModuleRelativePath = ModulePackage->GetMetaData()->GetValue(Field, TEXT("ModuleRelativePath"));
+			const FString& ModuleRelativePath = ModulePackage->GetMetaData()->GetValue(Field, TEXT("ModuleRelativePath"));
 			if(ModuleRelativePath.Len() > 0)
 			{
 				OutClassHeaderPath = ModuleBasePath / ModuleRelativePath;
@@ -1830,7 +1830,7 @@ bool FSourceCodeNavigation::FindClassSourcePath( const UField *Field, FString &O
 		{
 			// Get the metadata for the class path relative to the module base
 			// Given this we can try and find the corresponding .cpp file
-			FString ModuleRelativePath = ModulePackage->GetMetaData()->GetValue(Field, TEXT("ModuleRelativePath"));
+			const FString& ModuleRelativePath = ModulePackage->GetMetaData()->GetValue(Field, TEXT("ModuleRelativePath"));
 			if(ModuleRelativePath.Len() > 0)
 			{
 				const FString PotentialCppLeafname = FPaths::GetBaseFilename(ModuleRelativePath) + TEXT(".cpp");
diff --git a/Engine/Source/Editor/UnrealEd/Public/Kismet2/ComponentEditorUtils.h b/Engine/Source/Editor/UnrealEd/Public/Kismet2/ComponentEditorUtils.h
--- a/Engine/Source/Editor/UnrealEd/Public/Kismet2/ComponentEditorUtils.h
+++ b/Engine/Source/Editor/UnrealEd/Public/Kismet2/ComponentEditorUtils.h
@@ -154,10 +154,35 @@ public:
 		check(InProperty != nullptr);
 		check(InSceneComponent != nullptr);
 
+		ensureMsgf(Cast<UBoolProperty>(InProperty) == nullptr, TEXT("ApplyDefaultValueChange cannot be safely called on a bool property with a non-bool value, becuase of bitfields"));
+
 		T* CurrentValue = PropertyOffset == INDEX_NONE ? InProperty->ContainerPtrToValuePtr<T>(InSceneComponent) : (T*)((uint8*)InSceneComponent + PropertyOffset);
-		if(CurrentValue != nullptr)
+		check(CurrentValue);
+
+		return ApplyDefaultValueChange(InSceneComponent, *CurrentValue, OldDefaultValue, NewDefaultValue);
+	}
+
+	// Bool specialization so it can properly handle bitfields
+	static bool ApplyDefaultValueChange(class USceneComponent* InSceneComponent, const class UProperty* InProperty, const bool& OldDefaultValue, const bool& NewDefaultValue, int32 PropertyOffset)
+	{
+		check(InProperty != nullptr);
+		check(InSceneComponent != nullptr);
+		
+		// Only bool properties can have bool values
+		const UBoolProperty* BoolProperty = Cast<UBoolProperty>(InProperty);
+		check(BoolProperty);
+
+		uint8* CurrentValue = PropertyOffset == INDEX_NONE ? InProperty->ContainerPtrToValuePtr<uint8>(InSceneComponent) : ((uint8*)InSceneComponent + PropertyOffset);
+		check(CurrentValue);
+
+		bool CurrentBool = BoolProperty->GetPropertyValue(CurrentValue);
+		if (ApplyDefaultValueChange(InSceneComponent, CurrentBool, OldDefaultValue, NewDefaultValue, false))
 		{
-			return ApplyDefaultValueChange(InSceneComponent, *CurrentValue, OldDefaultValue, NewDefaultValue);
+			BoolProperty->SetPropertyValue(CurrentValue, CurrentBool);
+
+			InSceneComponent->ReregisterComponent();
+
+			return true;
 		}
 
 		return false;
@@ -165,7 +190,7 @@ public:
 
 	// Given an instance of a template and a current value, propagates a default value change to the instance (only if applicable)
 	template<typename T>
-	static bool ApplyDefaultValueChange(class USceneComponent* InSceneComponent, T& CurrentValue, const T& OldDefaultValue, const T& NewDefaultValue)
+	static bool ApplyDefaultValueChange(class USceneComponent* InSceneComponent, T& CurrentValue, const T& OldDefaultValue, const T& NewDefaultValue, bool bReregisterComponent = true)
 	{
 		check(InSceneComponent != nullptr);
 
@@ -190,9 +215,12 @@ public:
 			// Modify the value
 			CurrentValue = NewDefaultValue;
 
-			// Re-register the component with the scene so that transforms are updated for display
-			InSceneComponent->ReregisterComponent();
-
+			if (bReregisterComponent && InSceneComponent->IsRegistered())
+			{
+				// Re-register the component with the scene so that transforms are updated for display
+				InSceneComponent->ReregisterComponent();
+			}
+			
 			return true;
 		}
 
diff --git a/Engine/Source/Editor/UnrealEd/Public/Kismet2/KismetEditorUtilities.h b/Engine/Source/Editor/UnrealEd/Public/Kismet2/KismetEditorUtilities.h
--- a/Engine/Source/Editor/UnrealEd/Public/Kismet2/KismetEditorUtilities.h
+++ b/Engine/Source/Editor/UnrealEd/Public/Kismet2/KismetEditorUtilities.h
@@ -42,6 +42,8 @@ enum class EBlueprintCompileOptions
 	SkeletonUpToDate = 0x8, 
 	/** Indicates this is a batch compile and that BroadcastCompiled and BroadcastBlueprintCompiled should be skipped */
 	BatchCompile = 0x10,
+	/** Skips saving blueprints even if save on compile is enabled */
+	SkipSave = 0x20,
 	/** Skips creating a reinstancer and running reinstancing routines - useful if calling code is performing reinstancing */
 	SkipReinstancing = 0x40,
 };
diff --git a/Engine/Source/Programs/AutomationTool/AutomationUtils/ProjectParams.cs b/Engine/Source/Programs/AutomationTool/AutomationUtils/ProjectParams.cs
--- a/Engine/Source/Programs/AutomationTool/AutomationUtils/ProjectParams.cs
+++ b/Engine/Source/Programs/AutomationTool/AutomationUtils/ProjectParams.cs
@@ -624,7 +624,7 @@ namespace AutomationTool
 				this.Stage = true;
 			}
 			this.StageDirectoryParam = ParseParamValueIfNotSpecified(Command, StageDirectoryParam, "stagingdirectory", String.Empty, true);
-			this.bCodeSign = GetParamValueIfNotSpecified(Command, CodeSign, CommandUtils.IsBuildMachine, "CodeSign");
+			this.bCodeSign = GetOptionalParamValueIfNotSpecified(Command, CodeSign, CommandUtils.IsBuildMachine, "CodeSign", "NoCodeSign").GetValueOrDefault();
 			this.bTreatNonShippingBinariesAsDebugFiles = GetParamValueIfNotSpecified(Command, TreatNonShippingBinariesAsDebugFiles, false, "TreatNonShippingBinariesAsDebugFiles");
 			this.Manifests = GetParamValueIfNotSpecified(Command, Manifests, this.Manifests, "manifests");
             this.CreateChunkInstall = GetParamValueIfNotSpecified(Command, CreateChunkInstall, this.CreateChunkInstall, "createchunkinstall");
diff --git a/Engine/Source/Programs/UnrealHeaderTool/Private/CodeGenerator.cpp b/Engine/Source/Programs/UnrealHeaderTool/Private/CodeGenerator.cpp
--- a/Engine/Source/Programs/UnrealHeaderTool/Private/CodeGenerator.cpp
+++ b/Engine/Source/Programs/UnrealHeaderTool/Private/CodeGenerator.cpp
@@ -845,7 +845,7 @@ FString FNativeClassHeaderGenerator::GetSingletonNameFuncAddr(UField* Item, bool
 
 FString FNativeClassHeaderGenerator::GetOverriddenName(const UField* Item)
 {
-	FString OverriddenName = Item->GetMetaData(TEXT("OverrideNativeName"));
+	const FString& OverriddenName = Item->GetMetaData(TEXT("OverrideNativeName"));
 	if (!OverriddenName.IsEmpty())
 	{
 		return OverriddenName.ReplaceCharWithEscapedChar();
@@ -870,7 +870,7 @@ FString FNativeClassHeaderGenerator::GetOverriddenPathName(const UField* Item)
 
 FString FNativeClassHeaderGenerator::GetOverriddenNameForLiteral(const UField* Item)
 {
-	FString OverriddenName = Item->GetMetaData(TEXT("OverrideNativeName"));
+	const FString& OverriddenName = Item->GetMetaData(TEXT("OverrideNativeName"));
 	if (!OverriddenName.IsEmpty())
 	{
 		return TEXT("TEXT(\"") + OverriddenName + TEXT("\")");
@@ -880,7 +880,7 @@ FString FNativeClassHeaderGenerator::GetOverriddenNameForLiteral(const UField* I
 
 FString FNativeClassHeaderGenerator::GetUTF8OverriddenNameForLiteral(const UField* Item)
 {
-	FString OverriddenName = Item->GetMetaData(TEXT("OverrideNativeName"));
+	const FString& OverriddenName = Item->GetMetaData(TEXT("OverrideNativeName"));
 	if (!OverriddenName.IsEmpty())
 	{
 		return CreateUTF8LiteralString(OverriddenName);
@@ -2060,7 +2060,8 @@ void FNativeClassHeaderGenerator::ExportNativeGeneratedInitCode(FOutputDevice& O
 
 	FString OverriddenClassName = *FNativeClassHeaderGenerator::GetOverriddenName(Class);
 
-	const FString InitSearchableValuesFunctionName = bIsDynamic ? Class->GetMetaData(TEXT("InitializeStaticSearchableValues")) : FString();
+	const FString EmptyString;
+	const FString& InitSearchableValuesFunctionName = bIsDynamic ? Class->GetMetaData(TEXT("InitializeStaticSearchableValues")) : EmptyString;
 	const FString InitSearchableValuesFunctionParam = InitSearchableValuesFunctionName.IsEmpty() ? FString(TEXT("nullptr")) :
 		FString::Printf(TEXT("&%s::%s"), ClassNameCPP, *InitSearchableValuesFunctionName);
 
diff --git a/Engine/Source/Programs/UnrealHeaderTool/Private/HeaderParser.cpp b/Engine/Source/Programs/UnrealHeaderTool/Private/HeaderParser.cpp
--- a/Engine/Source/Programs/UnrealHeaderTool/Private/HeaderParser.cpp
+++ b/Engine/Source/Programs/UnrealHeaderTool/Private/HeaderParser.cpp
@@ -2871,7 +2871,7 @@ void FHeaderParser::VerifyPropertyMarkups( UClass* TargetClass )
 
 			if (Prop->HasAnyPropertyFlags(CPF_BlueprintVisible))
 			{
-				const FString GetterFuncName = Prop->GetMetaData(TEXT("BlueprintGetter"));
+				const FString& GetterFuncName = Prop->GetMetaData(TEXT("BlueprintGetter"));
 				if (!GetterFuncName.IsEmpty())
 				{
 					if (UFunction* TargetFunc = FindTargetFunction(*GetterFuncName))
@@ -2887,7 +2887,7 @@ void FHeaderParser::VerifyPropertyMarkups( UClass* TargetClass )
 
 				if (!Prop->HasAnyPropertyFlags(CPF_BlueprintReadOnly))
 				{
-					const FString SetterFuncName = Prop->GetMetaData(TEXT("BlueprintSetter"));
+					const FString& SetterFuncName = Prop->GetMetaData(TEXT("BlueprintSetter"));
 					if (!SetterFuncName.IsEmpty())
 					{
 						if (UFunction* TargetFunc = FindTargetFunction(*SetterFuncName))
@@ -6959,6 +6959,7 @@ struct FExposeOnSpawnValidator
 		case CPT_Byte:
 		case CPT_Float:
 		case CPT_Bool:
+		case CPT_Bool8:
 		case CPT_ObjectReference:
 		case CPT_String:
 		case CPT_Text:
diff --git a/Engine/Source/Runtime/AIModule/Classes/BehaviorTree/BTDecorator.h b/Engine/Source/Runtime/AIModule/Classes/BehaviorTree/BTDecorator.h
--- a/Engine/Source/Runtime/AIModule/Classes/BehaviorTree/BTDecorator.h
+++ b/Engine/Source/Runtime/AIModule/Classes/BehaviorTree/BTDecorator.h
@@ -8,6 +8,15 @@
 
 class FBehaviorDecoratorDetails;
 
+enum class EBTDecoratorAbortRequest : uint8
+{
+	// request execution update when only result of condition changes and active branch of tree can potentially change too
+	ConditionResultChanged,
+
+	// request execution update every time as long as condition is still passing
+	ConditionPassing,
+};
+
 /** 
  * Decorators are supporting nodes placed on parent-child connection, that receive notification about execution flow and can be ticked
  *
@@ -106,6 +115,11 @@ protected:
 	/** calculates raw, core value of decorator's condition. Should not include calling IsInversed */
 	virtual bool CalculateRawConditionValue(UBehaviorTreeComponent& OwnerComp, uint8* NodeMemory) const;
 
+	/** more "flow aware" version of calling RequestExecution(this) on owning behavior tree component
+	 *  should be used in external events that may change result of CalculateRawConditionValue
+	 */
+	void ConditionalFlowAbort(UBehaviorTreeComponent& OwnerComp, EBTDecoratorAbortRequest RequestMode) const;
+
 	friend FBehaviorDecoratorDetails;
 
 	//----------------------------------------------------------------------//
diff --git a/Engine/Source/Runtime/AIModule/Classes/BehaviorTree/BTNode.h b/Engine/Source/Runtime/AIModule/Classes/BehaviorTree/BTNode.h
--- a/Engine/Source/Runtime/AIModule/Classes/BehaviorTree/BTNode.h
+++ b/Engine/Source/Runtime/AIModule/Classes/BehaviorTree/BTNode.h
@@ -148,6 +148,9 @@ class AIMODULE_API UBTNode : public UObject, public IGameplayTaskOwnerInterface
 
 	/** Get whether this node is using a blueprint for its logic */
 	virtual bool UsesBlueprint() const;
+
+	/** Called after creating new node in behavior tree editor, use for versioning */
+	virtual void OnNodeCreated() {}
 #endif
 
 	/** Gets called only for instanced nodes(bCreateNodeInstance == true). In practive overridden by BP-implemented BT nodes */
diff --git a/Engine/Source/Runtime/AIModule/Classes/BehaviorTree/Tasks/BTTask_MoveTo.h b/Engine/Source/Runtime/AIModule/Classes/BehaviorTree/Tasks/BTTask_MoveTo.h
--- a/Engine/Source/Runtime/AIModule/Classes/BehaviorTree/Tasks/BTTask_MoveTo.h
+++ b/Engine/Source/Runtime/AIModule/Classes/BehaviorTree/Tasks/BTTask_MoveTo.h
@@ -37,6 +37,7 @@ class AIMODULE_API UBTTask_MoveTo : public UBTTask_BlackboardBase
 {
 	GENERATED_UCLASS_BODY()
 
+	/** fixed distance added to threshold between AI and goal location in destination reach test */
 	UPROPERTY(config, Category = Node, EditAnywhere, meta=(ClampMin = "0.0", UIMin="0.0"))
 	float AcceptableRadius;
 
@@ -69,11 +70,21 @@ class AIMODULE_API UBTTask_MoveTo : public UBTTask_BlackboardBase
 	UPROPERTY(Category = Node, EditAnywhere, AdvancedDisplay)
 	uint32 bProjectGoalLocation : 1;
 
-	/** if set to true agent's radius will be added to AcceptableRadius for purposes of checking 
-	 *	if path's end point has been reached. Will result in AI stopping on contact with destination location*/
+	/** if set, radius of AI's capsule will be added to threshold between AI and goal location in destination reach test  */
 	UPROPERTY(Category = Node, EditAnywhere)
+	uint32 bReachTestIncludesAgentRadius : 1;
+	
+	/** if set, radius of goal's capsule will be added to threshold between AI and goal location in destination reach test  */
+	UPROPERTY(Category = Node, EditAnywhere)
+	uint32 bReachTestIncludesGoalRadius : 1;
+
+	/** DEPRECATED, please use combination of bReachTestIncludes*Radius instead */
+	UPROPERTY(Category = Node, VisibleInstanceOnly)
 	uint32 bStopOnOverlap : 1;
 
+	UPROPERTY()
+	uint32 bStopOnOverlapNeedsUpdate : 1;
+
 	/** if set, move will use pathfinding. Not exposed on purpose, please use BTTask_MoveDirectlyToward */
 	uint32 bUsePathfinding : 1;
 
@@ -85,6 +96,7 @@ class AIMODULE_API UBTTask_MoveTo : public UBTTask_BlackboardBase
 	virtual void OnTaskFinished(UBehaviorTreeComponent& OwnerComp, uint8* NodeMemory, EBTNodeResult::Type TaskResult) override;
 	virtual void TickTask(UBehaviorTreeComponent& OwnerComp, uint8* NodeMemory, float DeltaSeconds) override;
 	virtual uint16 GetInstanceMemorySize() const override;
+	virtual void PostLoad() override;
 
 	virtual void OnGameplayTaskDeactivated(UGameplayTask& Task) override;
 	virtual void OnMessage(UBehaviorTreeComponent& OwnerComp, uint8* NodeMemory, FName Message, int32 RequestID, bool bSuccess) override;
@@ -95,6 +107,7 @@ class AIMODULE_API UBTTask_MoveTo : public UBTTask_BlackboardBase
 
 #if WITH_EDITOR
 	virtual FName GetNodeIconName() const override;
+	virtual void OnNodeCreated() override;
 #endif // WITH_EDITOR
 
 protected:
diff --git a/Engine/Source/Runtime/AIModule/Private/BehaviorTree/BTDecorator.cpp b/Engine/Source/Runtime/AIModule/Private/BehaviorTree/BTDecorator.cpp
--- a/Engine/Source/Runtime/AIModule/Private/BehaviorTree/BTDecorator.cpp
+++ b/Engine/Source/Runtime/AIModule/Private/BehaviorTree/BTDecorator.cpp
@@ -2,6 +2,7 @@
 
 #include "BehaviorTree/BTDecorator.h"
 #include "BehaviorTree/BTCompositeNode.h"
+#include "VisualLogger/VisualLogger.h"
 
 UBTDecorator::UBTDecorator(const FObjectInitializer& ObjectInitializer) : Super(ObjectInitializer)
 {
@@ -81,6 +82,55 @@ void UBTDecorator::WrappedOnNodeProcessed(FBehaviorTreeSearchData& SearchData, E
 	}
 }
 
+void UBTDecorator::ConditionalFlowAbort(UBehaviorTreeComponent& OwnerComp, EBTDecoratorAbortRequest RequestMode) const
+{
+	if (FlowAbortMode == EBTFlowAbortMode::None)
+	{
+		return;
+	}
+
+	const int32 InstanceIdx = OwnerComp.FindInstanceContainingNode(GetParentNode());
+	if (InstanceIdx == INDEX_NONE)
+	{
+		return;
+	}
+
+	uint8* NodeMemory = OwnerComp.GetNodeMemory((UBTNode*)this, InstanceIdx);
+
+	const bool bIsExecutingBranch = OwnerComp.IsExecutingBranch(this, GetChildIndex());
+	const bool bPass = WrappedCanExecute(OwnerComp, NodeMemory);
+	const bool bAbortPending = OwnerComp.IsAbortPending();
+	const bool bAlwaysRequestWhenPassing = (RequestMode == EBTDecoratorAbortRequest::ConditionPassing);
+
+	const bool bLogRestart = (bIsExecutingBranch != bPass) || (bIsExecutingBranch && bPass && (bAlwaysRequestWhenPassing || bAbortPending));
+	UE_VLOG(&OwnerComp, LogBehaviorTree, Verbose, TEXT("%s, ConditionalFlowAbort(%s) pass:%d executingBranch:%d abortPending:%d => %s"),
+		*UBehaviorTreeTypes::DescribeNodeHelper(this),
+		bAlwaysRequestWhenPassing ? TEXT("always when passing") : TEXT("on change"),
+		bPass ? 1 : 0,
+		bIsExecutingBranch ? 1 : 0,
+		bAbortPending ? 1 : 0,
+		bLogRestart ? TEXT("restart") : TEXT("skip"));
+
+	if (bIsExecutingBranch != bPass)
+	{
+		OwnerComp.RequestExecution(this);
+	}
+	else if (!bIsExecutingBranch && !bPass && GetParentNode() && GetParentNode()->Children.IsValidIndex(GetChildIndex()))
+	{
+		// this condition here is to remove all active observers _BELOW_ this node
+		// because if this condition failed we no longer want to react to child-conditions
+		// value changes anyway since their nodes execution will be blocked by this condition 
+		// during tree search
+		const UBTCompositeNode* BranchRoot = GetParentNode()->Children[GetChildIndex()].ChildComposite;
+		OwnerComp.UnregisterAuxNodesInBranch(BranchRoot);
+	}
+	else if (bIsExecutingBranch && bPass && (bAlwaysRequestWhenPassing || bAbortPending))
+	{
+		// force result Aborted to restart from this decorator
+		OwnerComp.RequestExecution(GetParentNode(), InstanceIdx, this, GetChildIndex(), EBTNodeResult::Aborted);
+	}
+}
+
 FString UBTDecorator::GetStaticDescription() const
 {
 	FString FlowAbortDesc;
diff --git a/Engine/Source/Runtime/AIModule/Private/BehaviorTree/Decorators/BTDecorator_Blackboard.cpp b/Engine/Source/Runtime/AIModule/Private/BehaviorTree/Decorators/BTDecorator_Blackboard.cpp
--- a/Engine/Source/Runtime/AIModule/Private/BehaviorTree/Decorators/BTDecorator_Blackboard.cpp
+++ b/Engine/Source/Runtime/AIModule/Private/BehaviorTree/Decorators/BTDecorator_Blackboard.cpp
@@ -4,7 +4,6 @@
 #include "UObject/Package.h"
 #include "BehaviorTree/Blackboard/BlackboardKeyType_Enum.h"
 #include "BehaviorTree/Blackboard/BlackboardKeyType_NativeEnum.h"
-#include "VisualLogger/VisualLogger.h"
 #include "BehaviorTree/BTCompositeNode.h"
 #include "BehaviorTree/BlackboardComponent.h"
 
@@ -65,48 +64,12 @@ EBlackboardNotificationResult UBTDecorator_Blackboard::OnBlackboardKeyValueChang
 		return EBlackboardNotificationResult::RemoveObserver;
 	}
 
-	if (BlackboardKey.GetSelectedKeyID() == ChangedKeyID && GetFlowAbortMode() != EBTFlowAbortMode::None)
+	if (BlackboardKey.GetSelectedKeyID() == ChangedKeyID)
 	{
-		const bool bIsExecutingBranch = BehaviorComp->IsExecutingBranch(this, GetChildIndex());
-		const bool bPass = (EvaluateOnBlackboard(Blackboard) != IsInversed());
-		const bool bAbortPending = BehaviorComp->IsAbortPending();
-
-		UE_VLOG(BehaviorComp, LogBehaviorTree, Verbose, TEXT("%s, OnBlackboardChange[%s] pass:%d executingBranch:%d abortPending:%d => %s"),
-			*UBehaviorTreeTypes::DescribeNodeHelper(this),
-			*Blackboard.GetKeyName(ChangedKeyID).ToString(),
-			bPass ? 1 : 0,
-			bIsExecutingBranch ? 1 : 0,
-			bAbortPending ? 1 : 0,
-			(bAbortPending || (bIsExecutingBranch && !bPass) || (!bIsExecutingBranch && bPass)) ? TEXT("restart") : TEXT("skip"));
-
-		if ((bIsExecutingBranch && !bPass) || (!bIsExecutingBranch && bPass))
-		{
-			BehaviorComp->RequestExecution(this);
-		}
-		else if (!bIsExecutingBranch && !bPass && GetParentNode() && GetParentNode()->Children.IsValidIndex(GetChildIndex()))
-		{
-			// this condition here is to remove all active observers _BELOW_ this node
-			// because if this condition failed we no longer want to react to child-conditions
-			// value changes anyway since their nodes execution will be blocked by this condition 
-			// during tree search
-			const UBTCompositeNode* BranchRoot = GetParentNode()->Children[GetChildIndex()].ChildComposite;
-			BehaviorComp->UnregisterAuxNodesInBranch(BranchRoot);
-		}
-		else if (bIsExecutingBranch && bPass && (NotifyObserver == EBTBlackboardRestart::ValueChange || bAbortPending))
-		{
-			UE_VLOG(BehaviorComp, LogBehaviorTree, Verbose, TEXT("%s, OnBlackboardChange[%s] => restart"),
-				*UBehaviorTreeTypes::DescribeNodeHelper(this),
-				*Blackboard.GetKeyName(ChangedKeyID).ToString());
+		// can't simply use BehaviorComp->RequestExecution(this) here, we need to support condition/value change modes
 
-			// force result Aborted to restart from this decorator
-			// can't use helper function
-
-			const int32 InstanceIdx = BehaviorComp->FindInstanceContainingNode((UBTNode*)GetParentNode());
-			if (InstanceIdx != INDEX_NONE)
-			{
-				BehaviorComp->RequestExecution(GetParentNode(), InstanceIdx, this, GetChildIndex(), EBTNodeResult::Aborted);
-			}
-		}
+		const EBTDecoratorAbortRequest RequestMode = (NotifyObserver == EBTBlackboardRestart::ValueChange) ? EBTDecoratorAbortRequest::ConditionPassing : EBTDecoratorAbortRequest::ConditionResultChanged;
+		ConditionalFlowAbort(*BehaviorComp, RequestMode);
 	}
 
 	return EBlackboardNotificationResult::ContinueObserving;
diff --git a/Engine/Source/Runtime/AIModule/Private/BehaviorTree/Decorators/BTDecorator_Loop.cpp b/Engine/Source/Runtime/AIModule/Private/BehaviorTree/Decorators/BTDecorator_Loop.cpp
--- a/Engine/Source/Runtime/AIModule/Private/BehaviorTree/Decorators/BTDecorator_Loop.cpp
+++ b/Engine/Source/Runtime/AIModule/Private/BehaviorTree/Decorators/BTDecorator_Loop.cpp
@@ -37,7 +37,7 @@ void UBTDecorator_Loop::OnNodeActivation(FBehaviorTreeSearchData& SearchData)
 		// protect from truly infinite loop within single search
 		if (SearchData.SearchId != DecoratorMemory->SearchId)
 		{
-			if ((InfiniteLoopTimeoutTime < 0.f) || ((DecoratorMemory->TimeStarted + InfiniteLoopTimeoutTime) < GetWorld()->GetTimeSeconds()))
+			if ((InfiniteLoopTimeoutTime < 0.f) || ((DecoratorMemory->TimeStarted + InfiniteLoopTimeoutTime) > GetWorld()->GetTimeSeconds()))
 			{
 				bShouldLoop = true;
 			}
diff --git a/Engine/Source/Runtime/AIModule/Private/BehaviorTree/Tasks/BTTask_MoveTo.cpp b/Engine/Source/Runtime/AIModule/Private/BehaviorTree/Tasks/BTTask_MoveTo.cpp
--- a/Engine/Source/Runtime/AIModule/Private/BehaviorTree/Tasks/BTTask_MoveTo.cpp
+++ b/Engine/Source/Runtime/AIModule/Private/BehaviorTree/Tasks/BTTask_MoveTo.cpp
@@ -19,13 +19,15 @@ UBTTask_MoveTo::UBTTask_MoveTo(const FObjectInitializer& ObjectInitializer) : Su
 	bNotifyTaskFinished = true;
 
 	AcceptableRadius = GET_AI_CONFIG_VAR(AcceptanceRadius);
-	bStopOnOverlap = GET_AI_CONFIG_VAR(bFinishMoveOnGoalOverlap); 
+	bReachTestIncludesGoalRadius = bReachTestIncludesAgentRadius = bStopOnOverlap = GET_AI_CONFIG_VAR(bFinishMoveOnGoalOverlap);
 	bAllowStrafe = GET_AI_CONFIG_VAR(bAllowStrafing);
 	bAllowPartialPath = GET_AI_CONFIG_VAR(bAcceptPartialPaths);
 	bTrackMovingGoal = true;
 	bProjectGoalLocation = true;
 	bUsePathfinding = true;
 
+	bStopOnOverlapNeedsUpdate = true;
+
 	ObservedBlackboardValueTolerance = AcceptableRadius * 0.95f;
 
 	// accept only actors and vectors
@@ -83,7 +85,8 @@ EBTNodeResult::Type UBTTask_MoveTo::PerformMoveTask(UBehaviorTreeComponent& Owne
 		MoveReq.SetAllowPartialPath(bAllowPartialPath);
 		MoveReq.SetAcceptanceRadius(AcceptableRadius);
 		MoveReq.SetCanStrafe(bAllowStrafe);
-		MoveReq.SetReachTestIncludesAgentRadius(bStopOnOverlap);
+		MoveReq.SetReachTestIncludesAgentRadius(bReachTestIncludesAgentRadius);
+		MoveReq.SetReachTestIncludesGoalRadius(bReachTestIncludesGoalRadius);
 		MoveReq.SetProjectGoalLocation(bProjectGoalLocation);
 		MoveReq.SetUsePathfinding(bUsePathfinding);
 
@@ -389,6 +392,18 @@ uint16 UBTTask_MoveTo::GetInstanceMemorySize() const
 	return sizeof(FBTMoveToTaskMemory);
 }
 
+void UBTTask_MoveTo::PostLoad()
+{
+	Super::PostLoad();
+	
+	if (bStopOnOverlapNeedsUpdate)
+	{
+		bStopOnOverlapNeedsUpdate = false;
+		bReachTestIncludesAgentRadius = bStopOnOverlap;
+		bReachTestIncludesGoalRadius = false;
+	}
+}
+
 #if WITH_EDITOR
 
 FName UBTTask_MoveTo::GetNodeIconName() const
@@ -396,4 +411,9 @@ FName UBTTask_MoveTo::GetNodeIconName() const
 	return FName("BTEditor.Graph.BTNode.Task.MoveTo.Icon");
 }
 
+void UBTTask_MoveTo::OnNodeCreated()
+{
+	bStopOnOverlapNeedsUpdate = false;
+}
+
 #endif	// WITH_EDITOR
diff --git a/Engine/Source/Runtime/AIModule/Private/Navigation/PathFollowingComponent.cpp b/Engine/Source/Runtime/AIModule/Private/Navigation/PathFollowingComponent.cpp
--- a/Engine/Source/Runtime/AIModule/Private/Navigation/PathFollowingComponent.cpp
+++ b/Engine/Source/Runtime/AIModule/Private/Navigation/PathFollowingComponent.cpp
@@ -644,7 +644,7 @@ void UPathFollowingComponent::SetMovementComponent(UNavMovementComponent* MoveCo
 	if (MoveComp)
 	{
 		const FNavAgentProperties& NavAgentProps = MoveComp->GetNavAgentPropertiesRef();
-		MyDefaultAcceptanceRadius = NavAgentProps.AgentRadius;
+		MyDefaultAcceptanceRadius = NavAgentProps.AgentRadius * 0.1f;
 		MoveComp->PathFollowingComp = this;
 
 		UWorld* MyWorld = GetWorld();
@@ -1161,7 +1161,7 @@ bool UPathFollowingComponent::HasReachedInternal(const FVector& GoalLocation, fl
 	const FVector ToGoal = GoalLocation - AgentLocation;
 
 	const float Dist2DSq = ToGoal.SizeSquared2D();
-	const float UseRadius = FMath::Max(RadiusThreshold, GoalRadius + (AgentRadius * AgentRadiusMultiplier));
+	const float UseRadius = RadiusThreshold + GoalRadius + (AgentRadius * AgentRadiusMultiplier);
 	if (Dist2DSq > FMath::Square(UseRadius))
 	{
 		return false;
@@ -1217,19 +1217,12 @@ float UPathFollowingComponent::GetFinalAcceptanceRadius(const FNavigationPath& P
 
 		const float PathEndToGoalDistance = FVector::Dist(PathEndOverride ? *PathEndOverride : PathInstance.GetEndLocation(), OriginalGoalLocation);
 		const float Remaining = AcceptanceRadius - PathEndToGoalDistance;
-
-		if (Remaining <= MyDefaultAcceptanceRadius)
-		{
-			// goal is more than AcceptanceRadius from the end of the path
-			// so we just need to do regular path following, with default 
-			// acceptance radius
-			return  MyDefaultAcceptanceRadius;
-		}
-		else
-		{
-			return Remaining;
-		}
+		
+		// if goal is more than AcceptanceRadius from the end of the path
+		// we just need to do regular path following, with default radius value
+		return FMath::Max(Remaining, MyDefaultAcceptanceRadius);
 	}
+
 	return AcceptanceRadius;
 }
 
diff --git a/Engine/Source/Runtime/AssetRegistry/Private/AssetRegistry.cpp b/Engine/Source/Runtime/AssetRegistry/Private/AssetRegistry.cpp
--- a/Engine/Source/Runtime/AssetRegistry/Private/AssetRegistry.cpp
+++ b/Engine/Source/Runtime/AssetRegistry/Private/AssetRegistry.cpp
@@ -149,6 +149,8 @@ UAssetRegistryImpl::UAssetRegistryImpl(const FObjectInitializer& ObjectInitializ
 		}
 	}
 
+	
+
 	if (GConfig)
 	{
 		GConfig->GetBool(TEXT("AssetRegistry"), TEXT("bUpdateDiskCacheAfterLoad"), bUpdateDiskCacheAfterLoad, GEngineIni);
@@ -2024,7 +2026,7 @@ void UAssetRegistryImpl::AddFilesToSearch (const TArray<FString>& Files)
 
 #if WITH_EDITOR
 
-void UAssetRegistryImpl::OnDirectoryChanged (const TArray<FFileChangeData>& FileChanges)
+void UAssetRegistryImpl::OnDirectoryChanged(const TArray<FFileChangeData>& FileChanges)
 {
 	// Take local copy of FileChanges array as we wish to collapse pairs of 'Removed then Added' FileChangeData
 	// entries into a single 'Modified' entry.
@@ -2064,27 +2066,27 @@ void UAssetRegistryImpl::OnDirectoryChanged (const TArray<FFileChangeData>& File
 		const bool bIsValidPackageName = FPackageName::TryConvertFilenameToLongPackageName(File, LongPackageName);
 		const bool bIsValidPackage = bIsPackageFile && bIsValidPackageName;
 
-		if ( bIsValidPackage )
+		if (bIsValidPackage)
 		{
-			switch( FileChangesProcessed[FileIdx].Action )
+			switch (FileChangesProcessed[FileIdx].Action)
 			{
-				case FFileChangeData::FCA_Added:
-					// This is a package file that was created on disk. Mark it to be scanned for asset data.
-					NewFiles.AddUnique(File);
-					UE_LOG(LogAssetRegistry, Verbose, TEXT("File was added to content directory: %s"), *File);
-					break;
+			case FFileChangeData::FCA_Added:
+				// This is a package file that was created on disk. Mark it to be scanned for asset data.
+				NewFiles.AddUnique(File);
+				UE_LOG(LogAssetRegistry, Verbose, TEXT("File was added to content directory: %s"), *File);
+				break;
 
-				case FFileChangeData::FCA_Modified:
-					// This is a package file that changed on disk. Mark it to be scanned immediately for new or removed asset data.
-					ModifiedFiles.AddUnique(File);
-					UE_LOG(LogAssetRegistry, Verbose, TEXT("File changed in content directory: %s"), *File);
-					break;
+			case FFileChangeData::FCA_Modified:
+				// This is a package file that changed on disk. Mark it to be scanned immediately for new or removed asset data.
+				ModifiedFiles.AddUnique(File);
+				UE_LOG(LogAssetRegistry, Verbose, TEXT("File changed in content directory: %s"), *File);
+				break;
 
-				case FFileChangeData::FCA_Removed:
-					// This file was deleted. Remove all assets in the package from the registry.
-					RemovePackageData(*LongPackageName);
-					UE_LOG(LogAssetRegistry, Verbose, TEXT("File was removed from content directory: %s"), *File);
-					break;
+			case FFileChangeData::FCA_Removed:
+				// This file was deleted. Remove all assets in the package from the registry.
+				RemovePackageData(*LongPackageName);
+				UE_LOG(LogAssetRegistry, Verbose, TEXT("File was removed from content directory: %s"), *File);
+				break;
 			}
 		}
 	}
@@ -2094,48 +2096,7 @@ void UAssetRegistryImpl::OnDirectoryChanged (const TArray<FFileChangeData>& File
 		AddFilesToSearch(NewFiles);
 	}
 
-	if (ModifiedFiles.Num() > 0)
-	{
-		// Convert all the filenames to package names
-		TArray<FString> ModifiedPackageNames;
-		ModifiedPackageNames.Reserve(ModifiedFiles.Num());
-		for (const FString& File : ModifiedFiles)
-		{
-			ModifiedPackageNames.Add(FPackageName::FilenameToLongPackageName(File));
-		}
-
-		// Get the assets that are currently inside the package
-		TArray<TArray<FAssetData*>> ExistingFilesAssetData;
-		ExistingFilesAssetData.Reserve(ModifiedFiles.Num());
-		for (const FString& PackageName : ModifiedPackageNames)
-		{
-			TArray<FAssetData*>* PackageAssetsPtr = State.CachedAssetsByPackageName.Find(*PackageName);
-			if (PackageAssetsPtr && PackageAssetsPtr->Num() > 0)
-			{
-				ExistingFilesAssetData.Add(*PackageAssetsPtr);
-			}
-			else
-			{
-				ExistingFilesAssetData.AddDefaulted();
-			}
-		}
-
-		// Re-scan and update the asset registry with the new asset data
-		TArray<FName> FoundAssets;
-		ScanPathsAndFilesSynchronous(TArray<FString>(), ModifiedFiles, true, EAssetDataCacheMode::NoCache, &FoundAssets, nullptr);
-
-		// Remove any assets that are no longer present in the package
-		for (const TArray<FAssetData*>& OldPackageAssets : ExistingFilesAssetData)
-		{
-			for (FAssetData* OldPackageAsset : OldPackageAssets)
-			{
-				if (!FoundAssets.Contains(OldPackageAsset->ObjectPath))
-				{
-					RemoveAssetData(OldPackageAsset);
-				}
-			}
-		}
-	}
+	ScanModifiedAssetFiles(ModifiedFiles);
 }
 
 void UAssetRegistryImpl::OnAssetLoaded(UObject *AssetLoaded)
@@ -2234,6 +2195,51 @@ void UAssetRegistryImpl::UpdateRedirectCollector()
 
 #endif // WITH_EDITOR
 
+void UAssetRegistryImpl::ScanModifiedAssetFiles(const TArray<FString>& InFilePaths)
+{
+	if (InFilePaths.Num() > 0)
+	{
+		// Convert all the filenames to package names
+		TArray<FString> ModifiedPackageNames;
+		ModifiedPackageNames.Reserve(InFilePaths.Num());
+		for (const FString& File : InFilePaths)
+		{
+			ModifiedPackageNames.Add(FPackageName::FilenameToLongPackageName(File));
+		}
+
+		// Get the assets that are currently inside the package
+		TArray<TArray<FAssetData*>> ExistingFilesAssetData;
+		ExistingFilesAssetData.Reserve(InFilePaths.Num());
+		for (const FString& PackageName : ModifiedPackageNames)
+		{
+			TArray<FAssetData*>* PackageAssetsPtr = State.CachedAssetsByPackageName.Find(*PackageName);
+			if (PackageAssetsPtr && PackageAssetsPtr->Num() > 0)
+			{
+				ExistingFilesAssetData.Add(*PackageAssetsPtr);
+			}
+			else
+			{
+				ExistingFilesAssetData.AddDefaulted();
+			}
+		}
+
+		// Re-scan and update the asset registry with the new asset data
+		TArray<FName> FoundAssets;
+		ScanPathsAndFilesSynchronous(TArray<FString>(), InFilePaths, true, EAssetDataCacheMode::NoCache, &FoundAssets, nullptr);
+
+		// Remove any assets that are no longer present in the package
+		for (const TArray<FAssetData*>& OldPackageAssets : ExistingFilesAssetData)
+		{
+			for (FAssetData* OldPackageAsset : OldPackageAssets)
+			{
+				if (!FoundAssets.Contains(OldPackageAsset->ObjectPath))
+				{
+					RemoveAssetData(OldPackageAsset);
+				}
+			}
+		}
+	}
+}
 
 void UAssetRegistryImpl::OnContentPathMounted( const FString& InAssetPath, const FString& FileSystemPath )
 {
diff --git a/Engine/Source/Runtime/AssetRegistry/Private/AssetRegistry.h b/Engine/Source/Runtime/AssetRegistry/Private/AssetRegistry.h
--- a/Engine/Source/Runtime/AssetRegistry/Private/AssetRegistry.h
+++ b/Engine/Source/Runtime/AssetRegistry/Private/AssetRegistry.h
@@ -63,6 +63,7 @@ public:
 	virtual void ScanPathsSynchronous(const TArray<FString>& InPaths, bool bForceRescan = false) override;
 	virtual void ScanFilesSynchronous(const TArray<FString>& InFilePaths, bool bForceRescan = false) override;
 	virtual void PrioritizeSearchPath(const FString& PathToPrioritize) override;
+	virtual void ScanModifiedAssetFiles(const TArray<FString>& InFilePaths) override;
 	virtual void Serialize(FArchive& Ar) override;
 	virtual uint32 GetAllocatedSize(bool bLogDetailed = false) const override;
 	virtual void LoadPackageRegistryData(FArchive& Ar, TArray<FAssetData*>& Data) const override;
diff --git a/Engine/Source/Runtime/AssetRegistry/Public/IAssetRegistry.h b/Engine/Source/Runtime/AssetRegistry/Public/IAssetRegistry.h
--- a/Engine/Source/Runtime/AssetRegistry/Public/IAssetRegistry.h
+++ b/Engine/Source/Runtime/AssetRegistry/Public/IAssetRegistry.h
@@ -250,6 +250,9 @@ public:
 	/** If assets are currently being asynchronously scanned in the specified path, this will cause them to be scanned before other assets. */
 	virtual void PrioritizeSearchPath(const FString& PathToPrioritize) = 0;
 
+	/** Forces a rescan of specific filenames, call this when you need to refresh from disk */
+	virtual void ScanModifiedAssetFiles(const TArray<FString>& InFilePaths) = 0;
+
 	/** Event for when paths are added to the registry */
 	DECLARE_EVENT_OneParam( IAssetRegistry, FPathAddedEvent, const FString& /*Path*/ );
 	virtual FPathAddedEvent& OnPathAdded() = 0;
diff --git a/Engine/Source/Runtime/Core/Public/Math/Quat.h b/Engine/Source/Runtime/Core/Public/Math/Quat.h
--- a/Engine/Source/Runtime/Core/Public/Math/Quat.h
+++ b/Engine/Source/Runtime/Core/Public/Math/Quat.h
@@ -426,6 +426,16 @@ public:
 	 */
 	FString ToString() const;
 
+	/**
+	 * Initialize this FQuat from a FString. 
+	 * The string is expected to contain X=, Y=, Z=, W=, otherwise 
+	 * this FQuat will have indeterminate (invalid) values.
+	 *
+	 * @param InSourceString FString containing the quaternion values.
+	 * @return true if the FQuat was initialized; false otherwise.
+	 */
+	bool InitFromString(const FString& InSourceString);
+
 public:
 
 #if ENABLE_NAN_DIAGNOSTIC
@@ -710,6 +720,19 @@ FORCEINLINE FString FQuat::ToString() const
 	return FString::Printf(TEXT("X=%.9f Y=%.9f Z=%.9f W=%.9f"), X, Y, Z, W);
 }
 
+inline bool FQuat::InitFromString(const FString& InSourceString)
+{
+	X = Y = Z = 0.f;
+	W = 1.f;
+
+	const bool bSuccess
+		=  FParse::Value(*InSourceString, TEXT("X="), X)
+		&& FParse::Value(*InSourceString, TEXT("Y="), Y)
+		&& FParse::Value(*InSourceString, TEXT("Z="), Z)
+		&& FParse::Value(*InSourceString, TEXT("W="), W);
+	DiagnosticCheckNaN();
+	return bSuccess;
+}
 
 #ifdef IMPLEMENT_ASSIGNMENT_OPERATOR_MANUALLY
 FORCEINLINE FQuat& FQuat::operator=(const FQuat& Other)
@@ -811,9 +834,6 @@ FORCEINLINE FQuat FQuat::operator*(const FQuat& Q) const
 
 FORCEINLINE FQuat FQuat::operator*=(const FQuat& Q)
 {
-	/**
-	 * Now this uses VectorQuaternionMultiply that is optimized per platform.
-	 */
 	VectorRegister A = VectorLoadAligned(this);
 	VectorRegister B = VectorLoadAligned(&Q);
 	VectorRegister Result;
diff --git a/Engine/Source/Runtime/Core/Public/Math/TransformNonVectorized.h b/Engine/Source/Runtime/Core/Public/Math/TransformNonVectorized.h
--- a/Engine/Source/Runtime/Core/Public/Math/TransformNonVectorized.h
+++ b/Engine/Source/Runtime/Core/Public/Math/TransformNonVectorized.h
@@ -552,19 +552,31 @@ public:
 	FORCEINLINE FVector TransformPosition(const FVector& V) const;
 	FORCEINLINE FVector TransformPositionNoScale(const FVector& V) const;
 
-	/** Inverts the matrix and then transforms V - correctly handles scaling in this matrix. */
+	/** Inverts the transform and then transforms V - correctly handles scaling in this transform. */
 	FORCEINLINE FVector InverseTransformPosition(const FVector &V) const;
 	FORCEINLINE FVector InverseTransformPositionNoScale(const FVector &V) const;
 	FORCEINLINE FVector TransformVector(const FVector& V) const;
 	FORCEINLINE FVector TransformVectorNoScale(const FVector& V) const;
 
 	/**
-	*	Transform a direction vector by the inverse of this matrix - will not take into account translation part.
+	*	Transform a direction vector by the inverse of this transform - will not take into account translation part.
 	*	If you want to transform a surface normal (or plane) and correctly account for non-uniform scaling you should use TransformByUsingAdjointT with adjoint of matrix inverse.
 	*/
 	FORCEINLINE FVector InverseTransformVector(const FVector &V) const;
 	FORCEINLINE FVector InverseTransformVectorNoScale(const FVector &V) const;
 
+	/**
+	 * Transform a rotation.
+	 * For example if this is a LocalToWorld transform, TransformRotation(Q) would transform Q from local to world space.
+	 */
+	FORCEINLINE FQuat TransformRotation(const FQuat& Q) const;
+	
+	/**
+	* Inverse transform a rotation.
+	* For example if this is a LocalToWorld transform, InverseTransformRotation(Q) would transform Q from world to local space.
+	*/
+	FORCEINLINE FQuat InverseTransformRotation(const FQuat& Q) const;
+
 	FORCEINLINE FTransform GetScaled(float Scale) const;
 	FORCEINLINE FTransform GetScaled(FVector Scale) const;
 	FORCEINLINE FVector GetScaledAxis(EAxis::Type InAxis) const;
@@ -1425,6 +1437,15 @@ FORCEINLINE FVector FTransform::InverseTransformVectorNoScale(const FVector &V)
 	return (Rotation.UnrotateVector(V));
 }
 
+FORCEINLINE FQuat FTransform::TransformRotation(const FQuat& Q) const
+{
+	return GetRotation() * Q;
+}
+
+FORCEINLINE FQuat FTransform::InverseTransformRotation(const FQuat& Q) const
+{
+	return GetRotation().Inverse() * Q;
+}
 
 FORCEINLINE FTransform FTransform::operator*(const FTransform& Other) const
 {
diff --git a/Engine/Source/Runtime/Core/Public/Math/TransformVectorized.h b/Engine/Source/Runtime/Core/Public/Math/TransformVectorized.h
--- a/Engine/Source/Runtime/Core/Public/Math/TransformVectorized.h
+++ b/Engine/Source/Runtime/Core/Public/Math/TransformVectorized.h
@@ -560,13 +560,10 @@ public:
 	FORCEINLINE FVector		TransformPositionNoScale(const FVector& V) const;
 
 
-	/** Inverts the matrix and then transforms V - correctly handles scaling in this matrix. */
+	/** Inverts the transform and then transforms V - correctly handles scaling in this transform. */
 	FORCEINLINE FVector		InverseTransformPosition(const FVector &V) const;
-
 	FORCEINLINE FVector		InverseTransformPositionNoScale(const FVector &V) const;
-
 	FORCEINLINE FVector		TransformVector(const FVector& V) const;
-
 	FORCEINLINE FVector		TransformVectorNoScale(const FVector& V) const;
 
 	/** 
@@ -574,9 +571,20 @@ public:
 	 *	If you want to transform a surface normal (or plane) and correctly account for non-uniform scaling you should use TransformByUsingAdjointT with adjoint of matrix inverse.
 	 */
 	FORCEINLINE FVector InverseTransformVector(const FVector &V) const;
-
 	FORCEINLINE FVector InverseTransformVectorNoScale(const FVector &V) const;
 
+	/**
+	* Transform a rotation.
+	* For example if this is a LocalToWorld transform, TransformRotation(Q) would transform Q from local to world space.
+	*/
+	FORCEINLINE FQuat TransformRotation(const FQuat& Q) const;
+
+	/**
+	* Inverse transform a rotation.
+	* For example if this is a LocalToWorld transform, InverseTransformRotation(Q) would transform Q from world to local space.
+	*/
+	FORCEINLINE FQuat InverseTransformRotation(const FQuat& Q) const;
+
 	FORCEINLINE FTransform	GetScaled(float Scale) const;
 	FORCEINLINE FTransform	GetScaled(FVector Scale) const;
 	FORCEINLINE FVector		GetScaledAxis(EAxis::Type InAxis) const;
@@ -1768,6 +1776,16 @@ FORCEINLINE FVector FTransform::InverseTransformVectorNoScale(const FVector &V)
 	return Result;
 }
 
+FORCEINLINE FQuat FTransform::TransformRotation(const FQuat& Q) const
+{
+	return GetRotation() * Q;
+}
+
+FORCEINLINE FQuat FTransform::InverseTransformRotation(const FQuat& Q) const
+{
+	return GetRotation().Inverse() * Q;
+}
+
 FORCEINLINE FTransform FTransform::operator*(const FTransform& Other) const
 {
 	FTransform Output;
diff --git a/Engine/Source/Runtime/Core/Public/Stats/StatsData.h b/Engine/Source/Runtime/Core/Public/Stats/StatsData.h
--- a/Engine/Source/Runtime/Core/Public/Stats/StatsData.h
+++ b/Engine/Source/Runtime/Core/Public/Stats/StatsData.h
@@ -69,14 +69,6 @@ struct CORE_API FStatConstants
 	static const FStatNameAndInfo AdvanceFrame;
 };
 
-namespace Lex
-{
-	inline void FromString(FName& OutValue, const TCHAR* Buffer )
-	{
-		OutValue = FName( Buffer );
-	}
-}
-
 /** Parse a typed value into the specified out parameter.
  * 	Expects to find a FromString function that takes a reference to T. Defaults are provided in the Lex namespace.
  */
diff --git a/Engine/Source/Runtime/Core/Public/UObject/FrameworkObjectVersion.h b/Engine/Source/Runtime/Core/Public/UObject/FrameworkObjectVersion.h
--- a/Engine/Source/Runtime/Core/Public/UObject/FrameworkObjectVersion.h
+++ b/Engine/Source/Runtime/Core/Public/UObject/FrameworkObjectVersion.h
@@ -104,6 +104,9 @@ struct CORE_API FFrameworkObjectVersion
 		// in order to allow us to do all UFunction loading in a single pass (after classes and CDOs are created):
 		RemoveUField_Next,
 
+		// Fix User Defined structs so that all members are correct flagged blueprint visible 
+		UserDefinedStructsBlueprintVisible,
+
 		// -----<new versions can be added above this line>-------------------------------------------------
 		VersionPlusOne,
 		LatestVersion = VersionPlusOne - 1
diff --git a/Engine/Source/Runtime/Core/Public/UObject/NameTypes.h b/Engine/Source/Runtime/Core/Public/UObject/NameTypes.h
--- a/Engine/Source/Runtime/Core/Public/UObject/NameTypes.h
+++ b/Engine/Source/Runtime/Core/Public/UObject/NameTypes.h
@@ -88,7 +88,6 @@ namespace FNameDefs
 #endif
 }
 
-
 enum ELinkerNameTableConstructor    {ENAME_LinkerConstructor};
 
 /** Enumeration for finding name. */
@@ -1147,6 +1146,18 @@ inline uint32 GetTypeHash( const FName N )
 	return N.GetComparisonIndex() + N.GetNumber();
 }
 
+namespace Lex
+{
+	FORCEINLINE FString ToString(const FName& Name)
+	{
+		return Name.ToString();
+	}
+
+	FORCEINLINE void FromString(FName& Name, const TCHAR* Str)
+	{
+		Name = FName(Str);
+	}
+}
 
 FORCEINLINE FMinimalName NameToMinimalName(const FName& InName)
 {
diff --git a/Engine/Source/Runtime/CoreUObject/Private/Blueprint/BlueprintSupport.cpp b/Engine/Source/Runtime/CoreUObject/Private/Blueprint/BlueprintSupport.cpp
--- a/Engine/Source/Runtime/CoreUObject/Private/Blueprint/BlueprintSupport.cpp
+++ b/Engine/Source/Runtime/CoreUObject/Private/Blueprint/BlueprintSupport.cpp
@@ -395,7 +395,7 @@ TArray<UClass*> const& FScopedClassDependencyGather::GetCachedDependencies()
 // rather than littering the code with USE_DEFERRED_DEPENDENCY_CHECK_VERIFICATION_TESTS
 // checks, let's just define DEFERRED_DEPENDENCY_CHECK for the file
 #if USE_DEFERRED_DEPENDENCY_CHECK_VERIFICATION_TESTS
-	#define DEFERRED_DEPENDENCY_CHECK(CheckExpr) check(CheckExpr)
+	#define DEFERRED_DEPENDENCY_CHECK(CheckExpr) ensure(CheckExpr)
 #else  // USE_DEFERRED_DEPENDENCY_CHECK_VERIFICATION_TESTS
 	#define DEFERRED_DEPENDENCY_CHECK(CheckExpr)
 #endif // USE_DEFERRED_DEPENDENCY_CHECK_VERIFICATION_TESTS
@@ -858,13 +858,13 @@ bool FLinkerLoad::DeferExportCreation(const int32 Index)
 		return false;
 	}
 
-	if ((Export.Object != nullptr) || !Export.ClassIndex.IsImport())
+	if ((Export.Object != nullptr))
 	{
 		return false;
 	}
 
 	UClass* LoadClass = GetExportLoadClass(Index);
-	if (LoadClass == nullptr)
+	if (LoadClass == nullptr || LoadClass->HasAnyClassFlags(CLASS_Native))
 	{
 		return false;
 	}
@@ -873,7 +873,9 @@ bool FLinkerLoad::DeferExportCreation(const int32 Index)
 	bool const bIsPlaceholderClass = (AsPlaceholderClass != nullptr);
 
 	FLinkerLoad* ClassLinker = LoadClass->GetLinker();
-	if ( !bIsPlaceholderClass && ((ClassLinker == nullptr) || !ClassLinker->IsBlueprintFinalizationPending()) )
+	if ( !bIsPlaceholderClass 
+		&& ((ClassLinker == nullptr) || !ClassLinker->IsBlueprintFinalizationPending())
+		&& (!LoadClass->ClassDefaultObject || LoadClass->ClassDefaultObject->HasAnyFlags(RF_LoadCompleted) || !LoadClass->ClassDefaultObject->HasAnyFlags(RF_WasLoaded)) )
 	{
 		return false;
 	}
@@ -1545,6 +1547,9 @@ void FLinkerLoad::FinalizeBlueprint(UClass* LoadClass)
 		TArray<UObject*> ClassInstances;
 		GetObjectsOfClass(LoadClass, ClassInstances, /*bIncludeDerivedClasses =*/true);
 
+		// Filter out instances that are part of this package, they were handled in ResolveDeferredExports:
+		ClassInstances.RemoveAllSwap([LoadClass](UObject* Obj){return Obj->GetOutermost() == LoadClass->GetOutermost();});
+
 		for (UObject* ClassInst : ClassInstances)
 		{
 			// in the case that we do end up with instances, use this to find 
@@ -1604,6 +1609,8 @@ void FLinkerLoad::ResolveDeferredExports(UClass* LoadClass)
 
 	UObject* BlueprintCDO = DeferredCDOIndex != INDEX_NONE ? ExportMap[DeferredCDOIndex].Object : LoadClass->ClassDefaultObject;
 	DEFERRED_DEPENDENCY_CHECK(BlueprintCDO != nullptr);
+	
+	TArray<int32> DeferredTemplateObjects;
 
 	if (!FBlueprintSupport::IsDeferredExportCreationDisabled())
 	{
@@ -1641,7 +1648,11 @@ void FLinkerLoad::ResolveDeferredExports(UClass* LoadClass)
 			FObjectExport& Export = ExportMap[ExportIndex];
 			if (ULinkerPlaceholderExportObject* PlaceholderExport = Cast<ULinkerPlaceholderExportObject>(Export.Object))
 			{
-				DEFERRED_DEPENDENCY_CHECK(Export.ClassIndex.IsImport());
+				if(Export.ClassIndex.IsExport())
+				{
+					DeferredTemplateObjects.Push(ExportIndex);
+					continue;
+				}
 
 				UClass* ExportClass = GetExportLoadClass(ExportIndex);
 				// export class could be null... we create these placeholder 
@@ -1699,12 +1710,6 @@ void FLinkerLoad::ResolveDeferredExports(UClass* LoadClass)
 	// between now and when DeferredCDOIndex is cleared (they won't be resolved,
 	// so that is a problem!)
 	FResolvingExportTracker::Get().FlagFullExportResolvePassComplete(this);
-
-	for (TObjectIterator<ULinkerPlaceholderExportObject> PlaceholderIt; PlaceholderIt; ++PlaceholderIt)
-	{
-		ULinkerPlaceholderExportObject* PlaceholderObj = *PlaceholderIt;
-		DEFERRED_DEPENDENCY_CHECK((PlaceholderObj->GetOuter() != LinkerRoot) || PlaceholderObj->IsPendingKill());
-	}
 #endif // USE_DEFERRED_DEPENDENCY_CHECK_VERIFICATION_TESTS
 
 	// the ExportMap loop above could have recursed back into "finalization" for  
@@ -1752,6 +1757,16 @@ void FLinkerLoad::ResolveDeferredExports(UClass* LoadClass)
 			}
 		}
 
+		{
+			// Create any objects that (non CDO) objects that were deferred in this package:
+			TGuardValue<int32> ClearDeferredCDOToPreventDeferExportCreation(DeferredCDOIndex, INDEX_NONE);
+			for(int32 ExportIndex : DeferredTemplateObjects)
+			{
+				ExportMap[ExportIndex].Object = nullptr;
+				CreateExport(ExportIndex);
+			}
+		}
+
 		// sub-classes of this Blueprint could have had their CDO's 
 		// initialization deferred (this occurs when the sub-class CDO is 
 		// created before this super CDO has been fully serialized; we do this
diff --git a/Engine/Source/Runtime/CoreUObject/Private/Serialization/AsyncLoading.cpp b/Engine/Source/Runtime/CoreUObject/Private/Serialization/AsyncLoading.cpp
--- a/Engine/Source/Runtime/CoreUObject/Private/Serialization/AsyncLoading.cpp
+++ b/Engine/Source/Runtime/CoreUObject/Private/Serialization/AsyncLoading.cpp
@@ -243,7 +243,6 @@ static FORCEINLINE bool IsTimeLimitExceeded(double InTickStartTime, bool bUseTim
 		LastTestTime = CurrentTime;
 	}
 
-	FCoreDelegates::OnAsyncLoadingFlushUpdate.Broadcast();
 	return bTimeLimitExceeded;
 }
 FORCEINLINE bool FAsyncPackage::IsTimeLimitExceeded()
@@ -335,6 +334,11 @@ void FAsyncLoadingThread::InitializeAsyncThread()
 
 void FAsyncLoadingThread::CancelAsyncLoadingInternal()
 {
+	// Cancel is not thread safe because the loaded delegates expect to be called on the game thread
+	// EDL does not support this function, but for backward compatible reasons we are letting it run on the async loading thread
+	// If this is enabled for EDL it must be made properly thread safe
+	UE_CLOG(GEventDrivenLoaderEnabled && !IsInGameThread(), LogStreaming, Fatal, TEXT("CancelAsyncLoadingInternal is not thread safe! This must be fixed before being enabled for EDL"));
+
 	{
 		// Packages we haven't yet started processing.
 #if THREADSAFE_UOBJECTS
@@ -343,7 +347,11 @@ void FAsyncLoadingThread::CancelAsyncLoadingInternal()
 		const EAsyncLoadingResult::Type Result = EAsyncLoadingResult::Canceled;
 		for (FAsyncPackageDesc* PackageDesc : QueuedPackages)
 		{
-			PackageDesc->PackageLoadedDelegate.ExecuteIfBound(PackageDesc->Name, nullptr, Result);
+			if (PackageDesc->PackageLoadedDelegate.IsValid())
+			{
+				PackageDesc->PackageLoadedDelegate->ExecuteIfBound(PackageDesc->Name, nullptr, Result);
+			}
+
 			delete PackageDesc;
 		}
 		QueuedPackages.Reset();
@@ -409,14 +417,14 @@ void FAsyncLoadingThread::CancelAsyncLoadingInternal()
 	CancelLoadingEvent->Trigger();
 }
 
-void FAsyncLoadingThread::QueuePackage(const FAsyncPackageDesc& Package)
+void FAsyncLoadingThread::QueuePackage(FAsyncPackageDesc& Package)
 {
 	{
 #if THREADSAFE_UOBJECTS
 		FScopeLock QueueLock(&QueueCritical);
 #endif
 		QueuedPackagesCounter.Increment();
-		QueuedPackages.Add(new FAsyncPackageDesc(Package));
+		QueuedPackages.Add(new FAsyncPackageDesc(Package, MoveTemp(Package.PackageLoadedDelegate)));
 	}
 	NotifyAsyncLoadingStateHasMaybeChanged();
 
@@ -440,12 +448,12 @@ FAsyncPackage* FAsyncLoadingThread::FindExistingPackageAndAddCompletionCallback(
 	FAsyncPackage* Result = PackageList.FindRef(PackageRequest->Name);
 	if (Result)
 	{
-		if (PackageRequest->PackageLoadedDelegate.IsBound())
+		if (PackageRequest->PackageLoadedDelegate.IsValid())
 		{
 			const bool bInternalCallback = false;
-			Result->AddCompletionCallback(PackageRequest->PackageLoadedDelegate, bInternalCallback);
+			Result->AddCompletionCallback(MoveTemp(PackageRequest->PackageLoadedDelegate), bInternalCallback);
 		}
-			Result->AddRequestID(PackageRequest->RequestID);
+		Result->AddRequestID(PackageRequest->RequestID);
 		if (FlushTree)
 		{
 			Result->PopulateFlushTree(FlushTree);
@@ -515,10 +523,10 @@ void FAsyncLoadingThread::ProcessAsyncPackageRequest(FAsyncPackageDesc* InReques
 			FGCScopeGuard GCGuard;
 			Package = new FAsyncPackage(*InRequest);
 		}
-		if (InRequest->PackageLoadedDelegate.IsBound())
+		if (InRequest->PackageLoadedDelegate.IsValid())
 		{
 			const bool bInternalCallback = false;
-			Package->AddCompletionCallback(InRequest->PackageLoadedDelegate, bInternalCallback);
+			Package->AddCompletionCallback(MoveTemp(InRequest->PackageLoadedDelegate), bInternalCallback);
 		}
 		Package->SetDependencyRootPackage(InRootPackage);
 		if (FlushTree)
@@ -4748,6 +4756,9 @@ EAsyncPackageState::Type FAsyncLoadingThread::TickAsyncLoading(bool bUseTimeLimi
 				}
 			}
 		}
+
+		// Call update callback once per tick on the game thread
+		FCoreDelegates::OnAsyncLoadingFlushUpdate.Broadcast();
 	}
 
 	return Result;
@@ -5193,6 +5204,8 @@ FAsyncPackage::~FAsyncPackage()
 	{
 		for (FCompletionCallback& CompletionCallback : CompletionCallbacks)
 		{
+			checkSlow(CompletionCallback.bIsInternal || IsInGameThread());
+
 			if (!CompletionCallback.bCalled)
 			{
 				check(0);
@@ -5821,7 +5834,8 @@ void FAsyncPackage::AddImportDependency(const FName& PendingImport, FFlushTree*
 		!PackageToStream->bLoadHasFailed)
 	{
 		const bool bInternalCallback = true;
-		PackageToStream->AddCompletionCallback(FLoadPackageAsyncDelegate::CreateRaw(this, &FAsyncPackage::ImportFullyLoadedCallback), bInternalCallback);
+		TUniquePtr<FLoadPackageAsyncDelegate> InternalDelegate(new FLoadPackageAsyncDelegate(FLoadPackageAsyncDelegate::CreateRaw(this, &FAsyncPackage::ImportFullyLoadedCallback)));
+		PackageToStream->AddCompletionCallback(MoveTemp(InternalDelegate), bInternalCallback);
 		PackageToStream->DependencyRefCount.Increment();
 		PendingImportedPackages.Add(PackageToStream);
 		if (FlushTree)
@@ -5995,10 +6009,10 @@ void FAsyncPackage::ImportFullyLoadedCallback(const FName& InPackageName, UPacka
 		int32 Index = ContainsDependencyPackage(PendingImportedPackages, InPackageName);
 		if (Index != INDEX_NONE)
 		{
-		// Keep a reference to this package so that its linker doesn't go away too soon
-		ReferencedImports.Add(PendingImportedPackages[Index]);
-		PendingImportedPackages.RemoveAt(Index);
-	}
+			// Keep a reference to this package so that its linker doesn't go away too soon
+			ReferencedImports.Add(PendingImportedPackages[Index]);
+			PendingImportedPackages.RemoveAt(Index);
+		}
 	}
 }
 
@@ -6538,7 +6552,7 @@ void FAsyncPackage::CallCompletionCallbacks(bool bInternal, EAsyncLoadingResult:
 		if (CompletionCallback.bIsInternal == bInternal && !CompletionCallback.bCalled)
 		{
 			CompletionCallback.bCalled = true;
-			CompletionCallback.Callback.ExecuteIfBound(Desc.Name, LoadedPackage, LoadingResult);
+			CompletionCallback.Callback->ExecuteIfBound(Desc.Name, LoadedPackage, LoadingResult);
 		}
 	}
 }
@@ -6548,14 +6562,11 @@ void FAsyncPackage::Cancel()
 	UE_CLOG(GEventDrivenLoaderEnabled, LogStreaming, Fatal, TEXT("FAsyncPackage::Cancel is not supported with the new loader"));
 
 	// Call any completion callbacks specified.
+	bLoadHasFailed = true;
 	const EAsyncLoadingResult::Type Result = EAsyncLoadingResult::Canceled;
-	for (int32 CallbackIndex = 0; CallbackIndex < CompletionCallbacks.Num(); CallbackIndex++)
-	{
-		if (!CompletionCallbacks[CallbackIndex].bCalled)
-		{
-			CompletionCallbacks[CallbackIndex].Callback.ExecuteIfBound(Desc.Name, nullptr, Result);
-		}
-	}
+	CallCompletionCallbacks(true, Result);
+	CallCompletionCallbacks(false, Result);
+
 	{
 		// Clear load flags from any referenced objects
 		FScopeLock RereferncedObjectsLock(&ReferencedObjectsCritical);
@@ -6568,7 +6579,6 @@ void FAsyncPackage::Cancel()
 		EmptyReferencedObjects();
 	}
 
-	bLoadHasFailed = true;
 	if (LinkerRoot)
 	{
 		if (Linker)
@@ -6584,11 +6594,11 @@ void FAsyncPackage::Cancel()
 	PreLoadSortIndex = 0;
 }
 
-void FAsyncPackage::AddCompletionCallback(const FLoadPackageAsyncDelegate& Callback, bool bInternal)
+void FAsyncPackage::AddCompletionCallback(TUniquePtr<FLoadPackageAsyncDelegate>&& Callback, bool bInternal)
 {
 	// This is to ensure that there is no one trying to subscribe to a already loaded package
 	//check(!bLoadHasFinished && !bLoadHasFailed);
-	CompletionCallbacks.Add(FCompletionCallback(bInternal, Callback));
+	CompletionCallbacks.Emplace(bInternal, MoveTemp(Callback));
 }
 
 void FAsyncPackage::UpdateLoadPercentage()
@@ -6657,8 +6667,16 @@ int32 LoadPackageAsync(const FString& InName, const FGuid* InGuid /*= nullptr*/,
 	// this function, otherwise it would be added when the packages are being processed on the async thread).
 	const int32 RequestID = GPackageRequestID.Increment();
 	FAsyncLoadingThread::Get().AddPendingRequest(RequestID);
+
+	// Allocate delegate on Game Thread, it is not safe to copy delegates by value on other threads
+	TUniquePtr<FLoadPackageAsyncDelegate> CompletionDelegatePtr;
+	if (InCompletionDelegate.IsBound())
+	{
+		CompletionDelegatePtr.Reset(new FLoadPackageAsyncDelegate(InCompletionDelegate));
+	}
+
 	// Add new package request
-	FAsyncPackageDesc PackageDesc(RequestID, *PackageName, *PackageNameToLoad, InGuid ? *InGuid : FGuid(), InCompletionDelegate, InPackageFlags, InPIEInstanceID, InPackagePriority);
+	FAsyncPackageDesc PackageDesc(RequestID, *PackageName, *PackageNameToLoad, InGuid ? *InGuid : FGuid(), MoveTemp(CompletionDelegatePtr), InPackageFlags, InPIEInstanceID, InPackagePriority);
 	FAsyncLoadingThread::Get().QueuePackage(PackageDesc);
 
 	return RequestID;
diff --git a/Engine/Source/Runtime/CoreUObject/Private/Serialization/AsyncLoadingThread.h b/Engine/Source/Runtime/CoreUObject/Private/Serialization/AsyncLoadingThread.h
--- a/Engine/Source/Runtime/CoreUObject/Private/Serialization/AsyncLoadingThread.h
+++ b/Engine/Source/Runtime/CoreUObject/Private/Serialization/AsyncLoadingThread.h
@@ -392,7 +392,7 @@ public:
 	*
 	* @param Package package descriptor.
 	*/
-	void QueuePackage(const FAsyncPackageDesc& Package);
+	void QueuePackage(FAsyncPackageDesc& Package);
 
 	/**
 	* [GAME THREAD] Cancels streaming
diff --git a/Engine/Source/Runtime/CoreUObject/Private/UObject/Class.cpp b/Engine/Source/Runtime/CoreUObject/Private/UObject/Class.cpp
--- a/Engine/Source/Runtime/CoreUObject/Private/UObject/Class.cpp
+++ b/Engine/Source/Runtime/CoreUObject/Private/UObject/Class.cpp
@@ -3572,7 +3572,12 @@ void UClass::Serialize( FArchive& Ar )
 		{
 			bCooked = Ar.IsCooking();
 		}
-		Ar << bCooked;
+		bool bCookedAsBool = bCooked;
+		Ar << bCookedAsBool;
+		if (Ar.IsLoading())
+		{
+			bCooked = bCookedAsBool;
+		}
 	}
 
 	// Defaults.
@@ -3835,16 +3840,16 @@ bool UClass::HasProperty(UProperty* InProperty) const
 UClass::UClass(const FObjectInitializer& ObjectInitializer)
 :	UStruct( ObjectInitializer )
 ,	ClassUnique(0)
+,	bCooked(false)
 ,	ClassFlags(CLASS_None)
 ,	ClassCastFlags(0)
 ,	ClassWithin( UObject::StaticClass() )
 ,	ClassGeneratedBy(nullptr)
-,	bCooked(false)
 ,	ClassDefaultObject(nullptr)
 {
 	// If you add properties here, please update the other constructors and PurgeClass()
 
-	CppTypeInfo.Emplace(&DefaultCppClassTypeInfoStatic);
+	SetCppTypeInfoStatic(&DefaultCppClassTypeInfoStatic);
 }
 
 /**
@@ -3853,16 +3858,16 @@ UClass::UClass(const FObjectInitializer& ObjectInitializer)
 UClass::UClass(const FObjectInitializer& ObjectInitializer, UClass* InBaseClass)
 :	UStruct(ObjectInitializer, InBaseClass)
 ,	ClassUnique(0)
+,	bCooked(false)
 ,	ClassFlags(CLASS_None)
 ,	ClassCastFlags(0)
 ,	ClassWithin(UObject::StaticClass())
 ,	ClassGeneratedBy(nullptr)
-,	bCooked(false)
 ,	ClassDefaultObject(nullptr)
 {
 	// If you add properties here, please update the other constructors and PurgeClass()
 
-	CppTypeInfo.Emplace(&DefaultCppClassTypeInfoStatic);
+	SetCppTypeInfoStatic(&DefaultCppClassTypeInfoStatic);
 
 	UClass* ParentClass = GetSuperClass();
 	if (ParentClass)
@@ -3905,18 +3910,18 @@ UClass::UClass
 ,	ClassVTableHelperCtorCaller(InClassVTableHelperCtorCaller)
 ,	ClassAddReferencedObjects( InClassAddReferencedObjects )
 ,	ClassUnique				( 0 )
+,	bCooked					( false )
 ,	ClassFlags				( InClassFlags | CLASS_Native )
 ,	ClassCastFlags			( InClassCastFlags )
 ,	ClassWithin				( nullptr )
 ,	ClassGeneratedBy		( nullptr )
 ,	ClassConfigName			()
-,	bCooked					( false )
 ,	NetFields				()
 ,	ClassDefaultObject		( nullptr )
 {
 	// If you add properties here, please update the other constructors and PurgeClass()
 
-	CppTypeInfo.Emplace(&DefaultCppClassTypeInfoStatic);
+	SetCppTypeInfoStatic(&DefaultCppClassTypeInfoStatic);
 
 	// We store the pointer to the ConfigName in an FName temporarily, this cast is intentional
 	// as we expect the mis-typed data to get picked up in UClass::DeferredRegister. PVS-Studio
@@ -4135,11 +4140,12 @@ UFunction* UClass::FindFunctionByName(FName InName, EIncludeSuperFlag::Type Incl
 	UFunction* Result = FuncMap.FindRef(InName);
 	if (Result == nullptr && IncludeSuper == EIncludeSuperFlag::IncludeSuper)
 	{
-		if (Interfaces.Num())
+		UClass* SuperClass = GetSuperClass();
+		if (SuperClass || Interfaces.Num() > 0)
 		{
-			if (UFunction** InterfaceResult = InterfaceFuncMap.Find(InName))
+			if (UFunction** SuperResult = SuperFuncMap.Find(InName))
 			{
-				Result = *InterfaceResult;
+				Result = *SuperResult;
 			}
 			else
 			{
@@ -4152,23 +4158,12 @@ UFunction* UClass::FindFunctionByName(FName InName, EIncludeSuperFlag::Type Incl
 					}
 				}
 
-				InterfaceFuncMap.Add(InName, Result);
-			}
-		}
-
-		if (Result == nullptr)
-		{
-			if (UClass* SuperClass = GetSuperClass())
-			{
-				if (UFunction** ParentResult = ParentFuncMap.Find(InName))
-				{
-					Result = *ParentResult;
-				}
-				else
+				if (SuperClass && Result == nullptr)
 				{
 					Result = SuperClass->FindFunctionByName(InName);
-					ParentFuncMap.Add(InName, Result);
 				}
+
+				SuperFuncMap.Add(InName, Result);
 			}
 		}
 	}
diff --git a/Engine/Source/Runtime/CoreUObject/Private/UObject/LinkerLoad.cpp b/Engine/Source/Runtime/CoreUObject/Private/UObject/LinkerLoad.cpp
--- a/Engine/Source/Runtime/CoreUObject/Private/UObject/LinkerLoad.cpp
+++ b/Engine/Source/Runtime/CoreUObject/Private/UObject/LinkerLoad.cpp
@@ -2772,7 +2772,7 @@ UClass* FLinkerLoad::GetExportLoadClass(int32 Index)
 	UClass* ExportClass = (UClass*)IndexToObject(Export.ClassIndex);
 #if USE_DEFERRED_DEPENDENCY_CHECK_VERIFICATION_TESTS
 	check(!Export.ClassIndex.IsImport() || !(LoadFlags & LOAD_DeferDependencyLoads) || 
-		(ExportClass && ExportClass->HasAnyClassFlags(CLASS_Intrinsic)) || (Cast<ULinkerPlaceholderClass>(ExportClass) != nullptr));
+		(ExportClass && ExportClass->HasAnyClassFlags(CLASS_Native)) || (Cast<ULinkerPlaceholderClass>(ExportClass) != nullptr));
 #endif // USE_DEFERRED_DEPENDENCY_CHECK_VERIFICATION_TESTS
 
 	return ExportClass;
@@ -3800,10 +3800,10 @@ UObject* FLinkerLoad::CreateExport( int32 Index )
 				Export.Object->SetLinker(this, Index);
 
 				// If this object was allocated but never loaded (components created by a constructor) make sure it gets loaded
-				// Do this for all subobjects created in the native constructor.
+				// Do this for all subobjects created in the native constructor, and anything loaded in an editor build
 				FUObjectThreadContext::Get().ObjLoaded.AddUnique(Export.Object);
 				if (!Export.Object->HasAnyFlags(RF_LoadCompleted) &&
-					(Export.Object->HasAnyFlags(RF_DefaultSubObject) || (ThisParent && ThisParent->IsTemplate(RF_ClassDefaultObject))))
+					(!FPlatformProperties::RequiresCookedData() || Export.Object->HasAnyFlags(RF_DefaultSubObject) || (ThisParent && ThisParent->IsTemplate(RF_ClassDefaultObject))))
 				{
 					check(!GEventDrivenLoaderEnabled || !EVENT_DRIVEN_ASYNC_LOAD_ACTIVE_AT_RUNTIME);
 					Export.Object->SetFlags(RF_NeedLoad | RF_NeedPostLoad | RF_NeedPostLoadSubobjects | RF_WasLoaded);
diff --git a/Engine/Source/Runtime/CoreUObject/Private/UObject/LinkerPlaceholderBase.cpp b/Engine/Source/Runtime/CoreUObject/Private/UObject/LinkerPlaceholderBase.cpp
--- a/Engine/Source/Runtime/CoreUObject/Private/UObject/LinkerPlaceholderBase.cpp
+++ b/Engine/Source/Runtime/CoreUObject/Private/UObject/LinkerPlaceholderBase.cpp
@@ -36,8 +36,8 @@ public:
 	 * at the time of AddReferencingPropertyValue() we cannot know/record the 
 	 * address/index of array properties (as they may change during array 
 	 * re-allocation or compaction). So we must follow the property chain and 
-	 * check every UArrayProperty member for references to this (hence, the need
-	 * for this recursive function).
+	 * check every container (array, set, map) property member for references to 
+	 * this (hence, the need for this recursive function).
 	 * 
 	 * @param  PropertyChain    An ascending outer chain, where the property at index zero is the leaf (referencer) property.
 	 * @param  ChainIndex		An index into the PropertyChain that this call should start at and iterate DOWN to zero.
@@ -108,6 +108,59 @@ int32 FLinkerPlaceholderObjectImpl::ResolvePlaceholderValues(const TArray<const
 			// PropertyChain, no need to keep on here
 			break;
 		}
+		else if (const USetProperty* SetProperty = Cast<const USetProperty>(Property))
+		{
+#if USE_DEFERRED_DEPENDENCY_CHECK_VERIFICATION_TESTS
+			const UProperty* NextProperty = PropertyChain[PropertyIndex - 1];
+			check(NextProperty == SetProperty->ElementProp);
+#endif // USE_DEFERRED_DEPENDENCY_CHECK_VERIFICATION_TESTS
+
+			// because we can't know which set entry was set with a reference 
+			// to this object, we have to comb through them all
+			FScriptSetHelper SetHelper(SetProperty, ValueAddress);
+			int32 Num = SetHelper.Num();
+			for (int32 SetIndex = 0; Num; ++SetIndex)
+			{
+				if (SetHelper.IsValidIndex(SetIndex))
+				{
+					--Num;
+					uint8* ElementAddress = SetHelper.GetElementPtr(SetIndex);
+					ReplacementCount += ResolvePlaceholderValues(PropertyChain, PropertyIndex - 1, ElementAddress, OldValue, ReplacementValue);
+				}
+			}
+
+			// the above recursive call chewed through the rest of the
+			// PropertyChain, no need to keep on here
+			break;
+		}
+		else if (const UMapProperty* MapProperty = Cast<const UMapProperty>(Property))
+		{
+#if USE_DEFERRED_DEPENDENCY_CHECK_VERIFICATION_TESTS
+			const UProperty* NextProperty = PropertyChain[PropertyIndex - 1];
+			check(NextProperty == MapProperty->KeyProp);
+#endif // USE_DEFERRED_DEPENDENCY_CHECK_VERIFICATION_TESTS
+
+			// because we can't know which map entry was set with a reference 
+			// to this object, we have to comb through them all
+			FScriptMapHelper MapHelper(MapProperty, ValueAddress);
+			int32 Num = MapHelper.Num();
+			for (int32 MapIndex = 0; Num; ++MapIndex)
+			{
+				if (MapHelper.IsValidIndex(MapIndex))
+				{
+					--Num;
+					uint8* KeyAddress = MapHelper.GetKeyPtr(MapIndex);
+					ReplacementCount += ResolvePlaceholderValues(PropertyChain, PropertyIndex - 1, KeyAddress, OldValue, ReplacementValue);
+
+					uint8* MapValueAddress = MapHelper.GetValuePtr(MapIndex);
+					ReplacementCount += ResolvePlaceholderValues(PropertyChain, PropertyIndex - 1, MapValueAddress, OldValue, ReplacementValue);
+				}
+			}
+
+			// the above recursive call chewed through the rest of the
+			// PropertyChain, no need to keep on here
+			break;
+		}
 		else
 		{
 			const UProperty* NextProperty = PropertyChain[PropertyIndex - 1];
@@ -434,9 +487,9 @@ int32 FLinkerPlaceholderBase::ResolvePlaceholderPropertyValues(UObject* NewObjec
 			// there were none, then a property could have changed its value 
 			// after it was set to this
 			// 
-			// NOTE: this may seem it can be resolved by properties removing 
-			//       themselves from ReferencingProperties, but certain 
-			//       properties may be the inner of a UArrayProperty (meaning 
+			// NOTE: this may seem it can be resolved by properties removing themselves
+			//       from ReferencingProperties, but certain properties may be
+			//       the inner of a container (array, set, map) property (meaning 
 			//       there could be multiple references per property)... we'd 
 			//       have to inc/decrement a property ref-count to resolve that 
 			//       scenario
diff --git a/Engine/Source/Runtime/CoreUObject/Private/UObject/Obj.cpp b/Engine/Source/Runtime/CoreUObject/Private/UObject/Obj.cpp
--- a/Engine/Source/Runtime/CoreUObject/Private/UObject/Obj.cpp
+++ b/Engine/Source/Runtime/CoreUObject/Private/UObject/Obj.cpp
@@ -2010,7 +2010,7 @@ void UObject::LoadConfig( UClass* ConfigClass/*=NULL*/, const TCHAR* InFilename/
 
 #if WITH_EDITOR
 		static FName ConsoleVariableFName(TEXT("ConsoleVariable"));
-		FString CVarName = Property->GetMetaData(ConsoleVariableFName);
+		const FString& CVarName = Property->GetMetaData(ConsoleVariableFName);
 		if (!CVarName.IsEmpty())
 		{
 			Key = CVarName;
@@ -2207,7 +2207,7 @@ void UObject::SaveConfig( uint64 Flags, const TCHAR* InFilename, FConfigCacheIni
 
 #if WITH_EDITOR
 			static FName ConsoleVariableFName(TEXT("ConsoleVariable"));
-			FString CVarName = Property->GetMetaData(ConsoleVariableFName);
+			const FString& CVarName = Property->GetMetaData(ConsoleVariableFName);
 			if (!CVarName.IsEmpty())
 			{
 				Key = CVarName;
@@ -2396,7 +2396,7 @@ void UObject::UpdateSinglePropertyInConfigFile(const UProperty* InProperty, cons
 		
 #if WITH_EDITOR
 		static FName ConsoleVariableFName(TEXT("ConsoleVariable"));
-		FString CVarName = InProperty->GetMetaData(ConsoleVariableFName);
+		const FString& CVarName = InProperty->GetMetaData(ConsoleVariableFName);
 		if (!CVarName.IsEmpty())
 		{
 			PropertyKey = CVarName;
diff --git a/Engine/Source/Runtime/CoreUObject/Private/UObject/PropertyHelper.h b/Engine/Source/Runtime/CoreUObject/Private/UObject/PropertyHelper.h
--- a/Engine/Source/Runtime/CoreUObject/Private/UObject/PropertyHelper.h
+++ b/Engine/Source/Runtime/CoreUObject/Private/UObject/PropertyHelper.h
@@ -192,8 +192,8 @@ namespace DelegatePropertyTools
 			ObjName[i++] = *Buffer;
 		}
 		ObjName[i] = TCHAR('\0');
-		UClass* Cls = NULL;
-		UObject* Object = NULL;
+		UClass* Cls = nullptr;
+		UObject* Object = nullptr;
 		// Get function name
 		if (*Buffer == TCHAR('.'))
 		{
@@ -210,10 +210,10 @@ namespace DelegatePropertyTools
 			// if we're importing for a subobject, assume the parent object as the source
 			// otherwise, assume Parent for the source
 			// if neither are valid, then the importing fails
-			if (Parent == NULL)
+			if (Parent == nullptr)
 			{
 				ErrorText->Logf(ELogVerbosity::Warning, TEXT("Cannot import unqualified delegate name; no object to search"));
-				return NULL;
+				return nullptr;
 			}
 			// since we don't support nested components, we only need to check one level deep
 			else if (!Parent->HasAnyFlags(RF_ClassDefaultObject) && Parent->GetOuter()->HasAnyFlags(RF_ClassDefaultObject))
@@ -228,17 +228,28 @@ namespace DelegatePropertyTools
 			}
 			FCString::Strncpy(FuncName, ObjName, i + 1);
 		}
-		if (Cls == NULL)
+		if (Cls == nullptr)
 		{
 			Cls = FindObject<UClass>(ANY_PACKAGE,ObjName);
-			if (Cls != NULL)
+			if (Cls != nullptr)
 			{
 				Object = Cls->GetDefaultObject();
 			}
 			else
 			{
-				Object = StaticFindObject( UObject::StaticClass(), ANY_PACKAGE, ObjName );
-				if (Object != NULL)
+				// Check outer chain before checking all packages
+				UObject* OuterToCheck = Parent;
+				while (Object == nullptr && OuterToCheck)
+				{
+					Object = StaticFindObject(UObject::StaticClass(), OuterToCheck, ObjName);
+					OuterToCheck = OuterToCheck->GetOuter();
+				}
+				
+				if (Object == nullptr)
+				{
+					Object = StaticFindObject(UObject::StaticClass(), ANY_PACKAGE, ObjName);
+				}
+				if (Object != nullptr)
 				{
 					Cls = Object->GetClass();
 				}
@@ -246,10 +257,10 @@ namespace DelegatePropertyTools
 		}
 		UFunction *Func = FindField<UFunction>( Cls, FuncName );
 		// Check function params.
-		if( Func != NULL )
+		if(Func != nullptr)
 		{
 			// Find the delegate UFunction to check params
-			check(SignatureFunction != NULL && "Invalid delegate property");
+			check(SignatureFunction != nullptr && "Invalid delegate property");
 
 			// check return type and params
 			if(	Func->NumParms == SignatureFunction->NumParms )
@@ -260,7 +271,7 @@ namespace DelegatePropertyTools
 					if( It1->GetClass()!=It2->GetClass() || (It1->PropertyFlags&CPF_OutParm)!=(It2->PropertyFlags&CPF_OutParm) )
 					{
 						ErrorText->Logf(ELogVerbosity::Warning,TEXT("Function %s does not match param types with delegate signature %s"), *Func->GetName(), *SignatureFunction->GetName());
-						Func = NULL;
+						Func = nullptr;
 						break;
 					}
 				}
@@ -268,7 +279,7 @@ namespace DelegatePropertyTools
 			else
 			{
 				ErrorText->Logf(ELogVerbosity::Warning,TEXT("Function %s does not match number of params with delegate signature %s"),*Func->GetName(), *SignatureFunction->GetName());
-				Func = NULL;
+				Func = nullptr;
 			}
 		}
 		else 
@@ -276,11 +287,11 @@ namespace DelegatePropertyTools
 			ErrorText->Logf(ELogVerbosity::Warning,TEXT("Unable to find function %s in object %s for delegate (found class: %s)"),FuncName,ObjName,*GetNameSafe(Cls));
 		}
 
-		//UE_LOG(LogProperty, Log, TEXT("... importing delegate FunctionName:'%s'(%s)   Object:'%s'(%s)"),Func != NULL ? *Func->GetName() : TEXT("NULL"), FuncName, Object != NULL ? *Object->GetFullName() : TEXT("NULL"), ObjName);
+		//UE_LOG(LogProperty, Log, TEXT("... importing delegate FunctionName:'%s'(%s)   Object:'%s'(%s)"),Func != nullptr ? *Func->GetName() : TEXT("nullptr"), FuncName, Object != nullptr ? *Object->GetFullName() : TEXT("nullptr"), ObjName);
 	
-		Delegate.BindUFunction( Func ? Object : NULL, Func ? Func->GetFName() : NAME_None );
+		Delegate.BindUFunction( Func ? Object : nullptr, Func ? Func->GetFName() : NAME_None );
 
-		return ( Func != NULL && Object != NULL ) ? Buffer : NULL;
+		return ( Func != nullptr && Object != nullptr ) ? Buffer : nullptr;
 	}
 }
 
diff --git a/Engine/Source/Runtime/CoreUObject/Private/UObject/ScriptCore.cpp b/Engine/Source/Runtime/CoreUObject/Private/UObject/ScriptCore.cpp
--- a/Engine/Source/Runtime/CoreUObject/Private/UObject/ScriptCore.cpp
+++ b/Engine/Source/Runtime/CoreUObject/Private/UObject/ScriptCore.cpp
@@ -1081,14 +1081,11 @@ bool UObject::CallFunctionByNameWithArguments(const TCHAR* Str, FOutputDevice& A
 		// if ArgStr is empty but we have more params to read parse the function to see if these have defaults, if so set them
 		bool bFoundDefault = false;
 		bool bFailedImport = true;
+#if WITH_EDITOR
 		if (!FCString::Strcmp(*ArgStr, TEXT("")))
 		{
 			const FName DefaultPropertyKey(*(FString(TEXT("CPP_Default_")) + PropertyParam->GetName()));
-#if WITH_EDITOR
-			const FString PropertyDefaultValue = Function->GetMetaData(DefaultPropertyKey);
-#else
-			const FString PropertyDefaultValue = TEXT("");
-#endif
+			const FString& PropertyDefaultValue = Function->GetMetaData(DefaultPropertyKey);
 			if (!PropertyDefaultValue.IsEmpty()) 
 			{
 				bFoundDefault = true;
@@ -1097,6 +1094,7 @@ bool UObject::CallFunctionByNameWithArguments(const TCHAR* Str, FOutputDevice& A
 				bFailedImport = (Result == nullptr);
 			}
 		}
+#endif
 
 		if (!bFoundDefault)
 		{
diff --git a/Engine/Source/Runtime/CoreUObject/Public/Serialization/AsyncLoading.h b/Engine/Source/Runtime/CoreUObject/Public/Serialization/AsyncLoading.h
--- a/Engine/Source/Runtime/CoreUObject/Public/Serialization/AsyncLoading.h
+++ b/Engine/Source/Runtime/CoreUObject/Public/Serialization/AsyncLoading.h
@@ -404,7 +404,7 @@ struct FAsyncPackage : FGCObject
 		return Desc.NameToLoad;
 	}
 
-	void AddCompletionCallback(const FLoadPackageAsyncDelegate& Callback, bool bInternal);
+	void AddCompletionCallback(TUniquePtr<FLoadPackageAsyncDelegate>&& Callback, bool bInternal);
 
 	/** Gets the number of references to this package from other packages in the dependency tree. */
 	FORCEINLINE int32 GetDependencyRefCount() const
@@ -490,17 +490,17 @@ private:
 	{
 		bool bIsInternal;
 		bool bCalled;
-		FLoadPackageAsyncDelegate Callback;
+		TUniquePtr<FLoadPackageAsyncDelegate> Callback;
 
 		FCompletionCallback()
 			: bIsInternal(false)
 			, bCalled(false)
 		{
 		}
-		FCompletionCallback(bool bInInternal, FLoadPackageAsyncDelegate InCallback)
+		FCompletionCallback(bool bInInternal, TUniquePtr<FLoadPackageAsyncDelegate>&& InCallback)
 			: bIsInternal(bInInternal)
 			, bCalled(false)
-			, Callback(InCallback)
+			, Callback(MoveTemp(InCallback))
 		{
 		}
 	};
diff --git a/Engine/Source/Runtime/CoreUObject/Public/Serialization/AsyncPackage.h b/Engine/Source/Runtime/CoreUObject/Public/Serialization/AsyncPackage.h
--- a/Engine/Source/Runtime/CoreUObject/Public/Serialization/AsyncPackage.h
+++ b/Engine/Source/Runtime/CoreUObject/Public/Serialization/AsyncPackage.h
@@ -9,6 +9,7 @@
 #include "CoreMinimal.h"
 #include "UObject/UObjectGlobals.h"
 #include "Misc/Guid.h"
+#include "UniquePtr.h"
 
 struct FAsyncPackageDesc
 {
@@ -20,8 +21,8 @@ struct FAsyncPackageDesc
 	FName NameToLoad;
 	/** GUID of the package to load, or the zeroed invalid GUID for "don't care" */
 	FGuid Guid;
-	/** Delegate called on completion of loading */
-	FLoadPackageAsyncDelegate PackageLoadedDelegate;
+	/** Delegate called on completion of loading. This delegate can only be created and consumed on the game thread */
+	TUniquePtr<FLoadPackageAsyncDelegate> PackageLoadedDelegate;
 	/** The flags that should be applied to the package */
 	EPackageFlags PackageFlags;
 	/** Package loading priority. Higher number is higher priority. */
@@ -30,12 +31,12 @@ struct FAsyncPackageDesc
 	int32 PIEInstanceID;
 
 
-	FAsyncPackageDesc(int32 InRequestID, const FName& InName, FName InPackageToLoadFrom = NAME_None, const FGuid& InGuid = FGuid(), FLoadPackageAsyncDelegate InCompletionDelegate = FLoadPackageAsyncDelegate(), EPackageFlags InPackageFlags = PKG_None, int32 InPIEInstanceID = INDEX_NONE, TAsyncLoadPriority InPriority = 0)
+	FAsyncPackageDesc(int32 InRequestID, const FName& InName, FName InPackageToLoadFrom = NAME_None, const FGuid& InGuid = FGuid(), TUniquePtr<FLoadPackageAsyncDelegate>&& InCompletionDelegate = TUniquePtr<FLoadPackageAsyncDelegate>(), EPackageFlags InPackageFlags = PKG_None, int32 InPIEInstanceID = INDEX_NONE, TAsyncLoadPriority InPriority = 0)
 		: RequestID(InRequestID)
 		, Name(InName)
 		, NameToLoad(InPackageToLoadFrom)
 		, Guid(InGuid)
-		, PackageLoadedDelegate(InCompletionDelegate)
+		, PackageLoadedDelegate(MoveTemp(InCompletionDelegate))
 		, PackageFlags(InPackageFlags)
 		, Priority(InPriority)
 		, PIEInstanceID(InPIEInstanceID)
@@ -45,4 +46,30 @@ struct FAsyncPackageDesc
 			NameToLoad = Name;
 		}
 	}
+
+	/** This constructor does not modify the package loaded delegate as this is not safe outside the game thread */
+	FAsyncPackageDesc(const FAsyncPackageDesc& OldPackage)
+		: RequestID(OldPackage.RequestID)
+		, Name(OldPackage.Name)
+		, NameToLoad(OldPackage.NameToLoad)
+		, Guid(OldPackage.Guid)
+		, PackageFlags(OldPackage.PackageFlags)
+		, Priority(OldPackage.Priority)
+		, PIEInstanceID(OldPackage.PIEInstanceID)
+	{
+	}
+
+	/** This constructor will explicitly copy the package loaded delegate and invalidate the old one */
+	FAsyncPackageDesc(const FAsyncPackageDesc& OldPackage, TUniquePtr<FLoadPackageAsyncDelegate>&& InPackageLoadedDelegate)
+		: FAsyncPackageDesc(OldPackage)
+	{
+		PackageLoadedDelegate = MoveTemp(InPackageLoadedDelegate);
+	}
+
+#if DO_GUARD_SLOW
+	~FAsyncPackageDesc()
+	{
+		checkSlow(!PackageLoadedDelegate.IsValid() || IsInGameThread());
+	}
+#endif
 };
diff --git a/Engine/Source/Runtime/CoreUObject/Public/UObject/Class.h b/Engine/Source/Runtime/CoreUObject/Public/UObject/Class.h
--- a/Engine/Source/Runtime/CoreUObject/Public/UObject/Class.h
+++ b/Engine/Source/Runtime/CoreUObject/Public/UObject/Class.h
@@ -1882,7 +1882,7 @@ struct FCppClassTypeInfo : ICppClassTypeInfo
 	FCppClassTypeInfo& operator=(const FCppClassTypeInfo&) = delete;
 
 	// ICppClassTypeInfo implementation
-	bool IsAbstract() const
+	virtual bool IsAbstract() const override
 	{
 		return Info->bIsAbstract;
 	}
@@ -2032,7 +2032,10 @@ public:
 	ClassAddReferencedObjectsType ClassAddReferencedObjects;
 
 	/** Class pseudo-unique counter; used to accelerate unique instance name generation */
-	int32 ClassUnique;
+	uint32 ClassUnique:31;
+
+	/** Used to check if the class was cooked or not */
+	uint32 bCooked:1;
 
 	/** Class flags; See EClassFlags for more information */
 	EClassFlags ClassFlags;
@@ -2058,9 +2061,6 @@ public:
 	/** Which Name.ini file to load Config variables out of */
 	FName ClassConfigName;
 
-	/** Used to check if the class was cooked or not */
-	bool bCooked;
-
 	/** List of replication records */
 	TArray<FRepRecord> ClassReps;
 
@@ -2100,17 +2100,16 @@ public:
 	static void AssembleReferenceTokenStreams();
 
 private:
+#if WITH_EDITOR
 	/** Provides access to attributes of the underlying C++ class. Should never be unset. */
 	TOptional<FCppClassTypeInfo> CppTypeInfo;
+#endif
 
 	/** Map of all functions by name contained in this class */
 	TMap<FName, UFunction*> FuncMap;
 
-	/** A cache of all functions by name that exist in a parent context */
-	mutable TMap<FName, UFunction*> ParentFuncMap;
-
-	/** A cache of all functions by name that exist in an interface context */
-	mutable TMap<FName, UFunction*> InterfaceFuncMap;
+	/** A cache of all functions by name that exist in a parent (superclass or interface) context */
+	mutable TMap<FName, UFunction*> SuperFuncMap;
 
 public:
 	/**
@@ -2227,8 +2226,7 @@ public:
 	/** Clears the function name caches, in case things have changed */
 	void ClearFunctionMapsCaches()
 	{
-		ParentFuncMap.Empty();
-		InterfaceFuncMap.Empty();
+		SuperFuncMap.Empty();
 	}
 
 	/** Looks for a given function name */
@@ -2261,17 +2259,21 @@ public:
 	virtual void SerializeSuperStruct(FArchive& Ar) override;
 	// End of UStruct interface.
 
+#if WITH_EDITOR
 	/** Provides access to C++ type info. */
 	const ICppClassTypeInfo* GetCppTypeInfo() const
 	{
 		return CppTypeInfo ? &CppTypeInfo.GetValue() : nullptr;
 	}
+#endif
 
 	/** Sets C++ type information. Should not be NULL. */
 	void SetCppTypeInfoStatic(const FCppClassTypeInfoStatic* InCppTypeInfoStatic)
 	{
+#if WITH_EDITOR
 		check(InCppTypeInfoStatic);
 		CppTypeInfo.Emplace(InCppTypeInfoStatic);
+#endif
 	}
 	
 	/**
diff --git a/Engine/Source/Runtime/CoreUObject/Public/UObject/UnrealType.h b/Engine/Source/Runtime/CoreUObject/Public/UObject/UnrealType.h
--- a/Engine/Source/Runtime/CoreUObject/Public/UObject/UnrealType.h
+++ b/Engine/Source/Runtime/CoreUObject/Public/UObject/UnrealType.h
@@ -79,15 +79,15 @@ class COREUOBJECT_API UProperty : public UField
 	uint64		PropertyFlags;
 	uint16		RepIndex;
 
-	FName		RepNotifyFunc;
-
 private:
+	TEnumAsByte<ELifetimeCondition> BlueprintReplicationCondition;
+
 	// In memory variables (generated during Link()).
 	int32		Offset_Internal;
 
-	ELifetimeCondition BlueprintReplicationCondition;
-
 public:
+	FName		RepNotifyFunc;
+
 	/** In memory only: Linked list of properties from most-derived to base **/
 	UProperty*	PropertyLinkNext;
 	/** In memory only: Linked list of object reference properties from most-derived to base **/
diff --git a/Engine/Source/Runtime/Engine/Classes/Animation/AnimNotifies/AnimNotify_PlayParticleEffect.h b/Engine/Source/Runtime/Engine/Classes/Animation/AnimNotifies/AnimNotify_PlayParticleEffect.h
--- a/Engine/Source/Runtime/Engine/Classes/Animation/AnimNotifies/AnimNotify_PlayParticleEffect.h
+++ b/Engine/Source/Runtime/Engine/Classes/Animation/AnimNotifies/AnimNotify_PlayParticleEffect.h
@@ -45,6 +45,10 @@ public:
 	UPROPERTY(EditAnywhere, BlueprintReadWrite, Category="AnimNotify")
 	FRotator RotationOffset;
 
+	// Scale to spawn the particle system at
+	UPROPERTY(EditAnywhere, Category="AnimNotify")
+	FVector Scale;
+
 private:
 	// Cached version of the Rotation Offset already in Quat form
 	FQuat RotationOffsetQuat;
diff --git a/Engine/Source/Runtime/Engine/Classes/Components/ActorComponent.h b/Engine/Source/Runtime/Engine/Classes/Components/ActorComponent.h
--- a/Engine/Source/Runtime/Engine/Classes/Components/ActorComponent.h
+++ b/Engine/Source/Runtime/Engine/Classes/Components/ActorComponent.h
@@ -1,4 +1,4 @@
-// Copyright 1998-2017 Epic Games, Inc. All Rights Reserved.
+﻿// Copyright 1998-2017 Epic Games, Inc. All Rights Reserved.
 
 #pragma once
 
@@ -79,9 +79,6 @@ DECLARE_DYNAMIC_MULTICAST_DELEGATE_OneParam(FActorComponentDeactivateSignature,
 DECLARE_MULTICAST_DELEGATE_OneParam(FActorComponentGlobalCreatePhysicsSignature, UActorComponent*);
 DECLARE_MULTICAST_DELEGATE_OneParam(FActorComponentGlobalDestroyPhysicsSignature, UActorComponent*);
 
-DECLARE_MULTICAST_DELEGATE(FActorComponentInstanceCreatePhysicsSignature);
-DECLARE_MULTICAST_DELEGATE(FActorComponentInstanceDestroyPhysicsSignature);
-
 /**
  * ActorComponent is the base class for components that define reusable behavior that can be added to different types of Actors.
  * ActorComponents that have a transform are known as SceneComponents and those that can be rendered are PrimitiveComponents.
@@ -101,11 +98,6 @@ public:
 	/** Destroy component physics state global delegate.*/
 	static FActorComponentGlobalDestroyPhysicsSignature GlobalDestroyPhysicsDelegate;
 
-	/** Create component physics state delegate (for this specific instance).*/
-	FActorComponentInstanceCreatePhysicsSignature InstanceCreatePhysicsDelegate;
-	/** Destroy component physics state delegate (for this specific instance).*/
-	FActorComponentInstanceDestroyPhysicsSignature InstanceDestroyPhysicsDelegate;
-
 	/**
 	 * Default UObject constructor that takes an optional ObjectInitializer.
 	 */
@@ -131,133 +123,133 @@ protected:
 	/** 
 	 *  Indicates if this ActorComponent is currently registered with a scene. 
 	 */
-	uint32 bRegistered:1;
+	uint8 bRegistered:1;
 
 	/** If the render state is currently created for this component */
-	uint32 bRenderStateCreated:1;
+	uint8 bRenderStateCreated:1;
 
 	/** If the physics state is currently created for this component */
-	uint32 bPhysicsStateCreated:1;
+	uint8 bPhysicsStateCreated:1;
 
 	/** Is this component currently replicating? Should the network code consider it for replication? Owning Actor must be replicating first! */
 	UPROPERTY(EditDefaultsOnly, BlueprintReadOnly, Replicated, Category=ComponentReplication,meta=(DisplayName = "Component Replicates"))
-	uint32 bReplicates:1;
+	uint8 bReplicates:1;
 
 	/** Is this component safe to ID over the network by name?  */
 	UPROPERTY()
-	uint32 bNetAddressable:1;
+	uint8 bNetAddressable:1;
 
 private:
 	/** Is this component in need of its whole state being sent to the renderer? */
-	uint32 bRenderStateDirty:1;
+	uint8 bRenderStateDirty:1;
 
 	/** Is this component's transform in need of sending to the renderer? */
-	uint32 bRenderTransformDirty:1;
+	uint8 bRenderTransformDirty:1;
 
 	/** Is this component's dynamic data in need of sending to the renderer? */
-	uint32 bRenderDynamicDataDirty:1;
+	uint8 bRenderDynamicDataDirty:1;
 
 	/** Used to ensure that any subclass of UActorComponent that overrides PostRename calls up to the Super to make OwnedComponents arrays get updated correctly */
-	uint32 bRoutedPostRename:1;
+	uint8 bRoutedPostRename:1;
 
 public:
 
 	/** Does this component automatically register with its owner */
-	uint32 bAutoRegister:1;
+	uint8 bAutoRegister:1;
 
 protected:
-	uint32 bAllowReregistration:1;
+	uint8 bAllowReregistration:1;
 
 public:
 	/** Should this component be ticked in the editor */
-	uint32 bTickInEditor:1;
+	uint8 bTickInEditor:1;
 
 	/** If true, this component never needs a render update. */
-	uint32 bNeverNeedsRenderUpdate:1;
+	uint8 bNeverNeedsRenderUpdate:1;
 
 	/** Can we tick this concurrently on other threads? */
-	uint32 bAllowConcurrentTick:1;
+	uint8 bAllowConcurrentTick:1;
 
 	/** Can this component be destroyed (via K2_DestroyComponent) by any parent */
-	uint32 bAllowAnyoneToDestroyMe:1;
+	uint8 bAllowAnyoneToDestroyMe:1;
 
+#if WITH_EDITORONLY_DATA
 	/** True if this component was created by a construction script, and will be destroyed by DestroyConstructedComponents */
 	UPROPERTY()
-	uint32 bCreatedByConstructionScript_DEPRECATED:1;
+	uint8 bCreatedByConstructionScript_DEPRECATED:1;
 
 	/** True if this component was created as an instance component */
 	UPROPERTY()
-	uint32 bInstanceComponent_DEPRECATED:1;
+	uint8 bInstanceComponent_DEPRECATED:1;
+#endif
 
 	/** Whether the component is activated at creation or must be explicitly activated. */
 	UPROPERTY(EditAnywhere, BlueprintReadOnly, Category=Activation)
-	uint32 bAutoActivate:1;
+	uint8 bAutoActivate:1;
 
 	/** Whether the component is currently active. */
 	UPROPERTY(transient, ReplicatedUsing=OnRep_IsActive)
-	uint32 bIsActive:1;
+	uint8 bIsActive:1;
 
 	UPROPERTY(EditDefaultsOnly, Category="Variable")
-	uint32 bEditableWhenInherited:1;
+	uint8 bEditableWhenInherited:1;
 
 	/** Cached navigation relevancy flag for collision updates */
-	uint32 bNavigationRelevant : 1;
+	uint8 bNavigationRelevant : 1;
 
 protected:
 	/** Whether this component can potentially influence navigation */
 	UPROPERTY(EditAnywhere, Category = Collision, AdvancedDisplay)
-	uint32 bCanEverAffectNavigation : 1;
+	uint8 bCanEverAffectNavigation : 1;
 
 public:
 	/** If true, we call the virtual InitializeComponent */
-	uint32 bWantsInitializeComponent:1;
+	uint8 bWantsInitializeComponent:1;
 
 	/** If true, we call the virtual BeginPlay */
 	DEPRECATED(4.14, "bWantsBeginPlay was inconsistently enforced and is now unused. BeginPlay will now always be called for Actor Components.")
-	uint32 bWantsBeginPlay:1;
+	uint8 bWantsBeginPlay:1;
 
 	/** If true, the component will be excluded from non-editor builds */
 	UPROPERTY(EditAnywhere, BlueprintReadOnly, Category = Cooking)
-	uint32 bIsEditorOnly:1;
+	uint8 bIsEditorOnly:1;
 
 private:
 	/** Indicates that OnCreatedComponent has been called, but OnDestroyedComponent has not yet */
-	uint32 bHasBeenCreated:1;
+	uint8 bHasBeenCreated:1;
 
 	/** Indicates that InitializeComponent has been called, but UninitializeComponent has not yet */
-	uint32 bHasBeenInitialized:1;
+	uint8 bHasBeenInitialized:1;
 
 	/** Indicates that BeginPlay has been called, but EndPlay has not yet */
-	uint32 bHasBegunPlay:1;
+	uint8 bHasBegunPlay:1;
 
 	/** Indicates the the destruction process has begun for this component to avoid recursion */
-	uint32 bIsBeingDestroyed:1;
+	uint8 bIsBeingDestroyed:1;
 
 	/** Whether we've tried to register tick functions. Reset when they are unregistered. */
-	uint32 bTickFunctionsRegistered:1;
+	uint8 bTickFunctionsRegistered:1;
 
 #if WITH_EDITOR
 	/** During undo/redo it isn't safe to cache owner */
-	uint32 bCanUseCachedOwner:1;
+	uint8 bCanUseCachedOwner:1;
 #endif
 
+	/** True if this component was owned by a net startup actor during level load. */
+	uint8 bIsNetStartupComponent : 1;
+
 	/** Tracks whether the component has been added to one of the world's end of frame update lists */
-	uint32 MarkedForEndOfFrameUpdateState:2;
+	uint8 MarkedForEndOfFrameUpdateState:2;
 	friend struct FMarkComponentEndOfFrameUpdateState;
 
 	friend class FActorComponentInstanceData;
 	friend class FActorComponentDetails;
 
-	/** True if this component was owned by a net startup actor during level load. */
-	bool bIsNetStartupComponent;
-
 public:
 	UPROPERTY()
 	EComponentCreationMethod CreationMethod;
 
 private:
-	mutable AActor* OwnerPrivate;
-
 	UPROPERTY()
 	TArray<FSimpleMemberReference> UCSModifiedProperties;
 
@@ -436,6 +428,8 @@ public:
 
 private:
 
+	mutable AActor* OwnerPrivate;
+
 	/** 
 	 * Pointer to the world that this component is currently registered with. 
 	 * This is only non-NULL when the component is registered.
diff --git a/Engine/Source/Runtime/Engine/Classes/Components/ForceFeedbackComponent.h b/Engine/Source/Runtime/Engine/Classes/Components/ForceFeedbackComponent.h
--- a/Engine/Source/Runtime/Engine/Classes/Components/ForceFeedbackComponent.h
+++ b/Engine/Source/Runtime/Engine/Classes/Components/ForceFeedbackComponent.h
@@ -1,4 +1,4 @@
-// Copyright 1998-2017 Epic Games, Inc. All Rights Reserved.
+﻿// Copyright 1998-2017 Epic Games, Inc. All Rights Reserved.
 
 
 #pragma once
@@ -78,6 +78,10 @@ public:
 	UPROPERTY(EditAnywhere, BlueprintReadWrite, Category=ForceFeedback)
 	uint8 bLooping:1;
 
+	/** Should the playback of the forcefeedback pattern ignore time dilation and use the app's delta time */
+	UPROPERTY(EditAnywhere, BlueprintReadWrite, Category=ForceFeedback)
+	uint8 bIgnoreTimeDilation:1;
+
 	/** Should the Attenuation Settings asset be used (false) or should the properties set directly on the component be used for attenuation properties */
 	UPROPERTY(EditAnywhere, BlueprintReadWrite, Category=Attenuation)
 	uint8 bOverrideAttenuation:1;
diff --git a/Engine/Source/Runtime/Engine/Classes/Components/MeshComponent.h b/Engine/Source/Runtime/Engine/Classes/Components/MeshComponent.h
--- a/Engine/Source/Runtime/Engine/Classes/Components/MeshComponent.h
+++ b/Engine/Source/Runtime/Engine/Classes/Components/MeshComponent.h
@@ -8,6 +8,7 @@
 #include "Engine/TextureStreamingTypes.h"
 #include "Components/PrimitiveComponent.h"
 #include "Materials/MaterialInterface.h"
+#include "SortedMap.h"
 #include "MeshComponent.generated.h"
 
 /**
@@ -114,8 +115,8 @@ public:
 	 */
 	float GetScalarParameterDefaultValue(const FName ParameterName)
 	{
-		float* Value = ScalarParameterDefaultValues.Find(ParameterName);
-		return (Value) ? *Value : 0.f;
+		FMaterialParameterCache* ParameterCache = MaterialParameterCache.Find(ParameterName);
+		return (ParameterCache ? ParameterCache->ScalarParameterDefaultValue : 0.f);
 	}
 protected:
 	/** Retrieves all the (scalar/vector-)parameters from within the used materials on the SkeletalMesh, and stores material index vs parameter names */
@@ -124,14 +125,20 @@ protected:
 	/** Mark cache parameters map as dirty, cache will be rebuild once SetScalar/SetVector functions are called */
 	void MarkCachedMaterialParameterNameIndicesDirty();
 	
-	/** Map containing material indices for the retrieved scalar material parameter names */
-	TMap<FName, TArray<int32>> ScalarParameterMaterialIndices;
-	/** Map containing material default parameter for the scalar parameter names 
-	 * We only cache the first one as we can't trace back from [name, index] 
-	 * This data is used for animation system to set default back to it*/
-	TMap<FName, float> ScalarParameterDefaultValues;
-	/** Map containing material indices for the retrieved vector material parameter names */
-	TMap<FName, TArray<int32>> VectorParameterMaterialIndices;
+	/** Struct containing information about a given parameter name */
+	struct FMaterialParameterCache
+	{
+		/** Material indices for the retrieved scalar material parameter names */
+		TArray<int32> ScalarParameterMaterialIndices;
+		/** Material indices for the retrieved vector material parameter names */
+		TArray<int32> VectorParameterMaterialIndices;
+		/** Material default parameter for the scalar parameter
+		 * We only cache the last one as we can't trace back from [name, index] 
+		 * This data is used for animation system to set default back to it*/
+		float ScalarParameterDefaultValue = 0.f;
+	};
+
+	TSortedMap<FName, FMaterialParameterCache> MaterialParameterCache;
 
 	/** Flag whether or not the cached material parameter indices map is dirty (defaults to true, and is set from SetMaterial/Set(Skeletal)Mesh */
 	uint32 bCachedMaterialParameterIndicesAreDirty : 1;
