Copy of CL #3631888

diff --git a/Engine/Source/Runtime/CoreUObject/Private/UObject/LinkerLoad.cpp b/Engine/Source/Runtime/CoreUObject/Private/UObject/LinkerLoad.cpp
--- a/Engine/Source/Runtime/CoreUObject/Private/UObject/LinkerLoad.cpp
+++ b/Engine/Source/Runtime/CoreUObject/Private/UObject/LinkerLoad.cpp
@@ -3800,13 +3800,21 @@ UObject* FLinkerLoad::CreateExport( int32 Index )
 				Export.Object->SetLinker(this, Index);
 
 				// If this object was allocated but never loaded (components created by a constructor) make sure it gets loaded
-				// Do this for all subobjects created in the native constructor, and anything loaded in an editor build
+				// Don't do this for any packages that have previously fully loaded as they may have in memory changes
 				FUObjectThreadContext::Get().ObjLoaded.AddUnique(Export.Object);
-				if (!Export.Object->HasAnyFlags(RF_LoadCompleted) &&
-					(!FPlatformProperties::RequiresCookedData() || Export.Object->HasAnyFlags(RF_DefaultSubObject) || (ThisParent && ThisParent->IsTemplate(RF_ClassDefaultObject))))
+				if (!Export.Object->HasAnyFlags(RF_LoadCompleted) && !LinkerRoot->IsFullyLoaded())
 				{
 					check(!GEventDrivenLoaderEnabled || !EVENT_DRIVEN_ASYNC_LOAD_ACTIVE_AT_RUNTIME);
-					Export.Object->SetFlags(RF_NeedLoad | RF_NeedPostLoad | RF_NeedPostLoadSubobjects | RF_WasLoaded);
+
+					if (Export.Object->HasAnyFlags(RF_ClassDefaultObject))
+					{
+						// Class default objects cannot have PostLoadSubobjects called on them
+						Export.Object->SetFlags(RF_NeedLoad | RF_NeedPostLoad | RF_WasLoaded);
+					}
+					else
+					{
+						Export.Object->SetFlags(RF_NeedLoad | RF_NeedPostLoad | RF_NeedPostLoadSubobjects | RF_WasLoaded);
+					}
 				}
 			}
 			return Export.Object;
