[NO NON-PII COMMIT MESSAGE]

diff --git a/Engine/Source/Runtime/CoreUObject/Private/UObject/Package.cpp b/Engine/Source/Runtime/CoreUObject/Private/UObject/Package.cpp
--- a/Engine/Source/Runtime/CoreUObject/Private/UObject/Package.cpp
+++ b/Engine/Source/Runtime/CoreUObject/Private/UObject/Package.cpp
@@ -188,10 +188,10 @@ bool UPackage::IsFullyLoaded() const
 	if( !bHasBeenFullyLoaded && !HasAnyInternalFlags(EInternalObjectFlags::AsyncLoading) )
 	{
 		FString DummyFilename;
-		// Try to find matching package in package file cache.
+		// Try to find matching package in package file cache. We use the source package name here as it may be loaded into a temporary package
 		if (	!GetConvertedDynamicPackageNameToTypeName().Contains(GetFName()) &&
 				(
-					!FPackageName::DoesPackageExist( *GetName(), NULL, &DummyFilename ) || 
+					!FPackageName::DoesPackageExist( *FileName.ToString(), NULL, &DummyFilename ) || 
 					(GIsEditor && IFileManager::Get().FileSize(*DummyFilename) < 0) 
 				)
 			)
