CIS fix for prior change, some packages have no source package name and it warns if you try to look it up

diff --git a/Engine/Source/Runtime/CoreUObject/Private/UObject/Package.cpp b/Engine/Source/Runtime/CoreUObject/Private/UObject/Package.cpp
--- a/Engine/Source/Runtime/CoreUObject/Private/UObject/Package.cpp
+++ b/Engine/Source/Runtime/CoreUObject/Private/UObject/Package.cpp
@@ -188,10 +188,11 @@ bool UPackage::IsFullyLoaded() const
 	if( !bHasBeenFullyLoaded && !HasAnyInternalFlags(EInternalObjectFlags::AsyncLoading) )
 	{
 		FString DummyFilename;
+		FString SourcePackageName = FileName != NAME_None ? FileName.ToString() : GetName();
 		// Try to find matching package in package file cache. We use the source package name here as it may be loaded into a temporary package
 		if (	!GetConvertedDynamicPackageNameToTypeName().Contains(GetFName()) &&
 				(
-					!FPackageName::DoesPackageExist( *FileName.ToString(), NULL, &DummyFilename ) || 
+					!FPackageName::DoesPackageExist(*SourcePackageName, NULL, &DummyFilename ) ||
 					(GIsEditor && IFileManager::Get().FileSize(*DummyFilename) < 0) 
 				)
 			)
