Fix a 4.18->Main merge issue that caused Blueprint nativization to fail.

- Mirrored from //UE4/Dev-Framework (3740428).

diff --git a/Engine/Source/Runtime/CoreUObject/Private/UObject/SavePackage.cpp b/Engine/Source/Runtime/CoreUObject/Private/UObject/SavePackage.cpp
--- a/Engine/Source/Runtime/CoreUObject/Private/UObject/SavePackage.cpp
+++ b/Engine/Source/Runtime/CoreUObject/Private/UObject/SavePackage.cpp
@@ -3704,16 +3704,13 @@ FSavePackageResultStruct UPackage::Save(UPackage* InOuter, UObject* Base, EObjec
 									}
 								}
 
-								UE_LOG(LogSavePackage, Display, TEXT("Package %s contains assets, that were converted into native code. Package will not be saved."), *InOuter->GetName());
+								UE_LOG(LogSavePackage, Verbose, TEXT("Package %s contains assets that are being converted to native code."), *InOuter->GetName());
 								return ESavePackageResult::ReplaceCompletely;
 							}
 							else if (ReplacmentResult == EReplacementResult::GenerateStub)
 							{
 								bRequestStub = true;
 							}
-
-							UE_LOG(LogSavePackage, Verbose, TEXT("Package %s contains assets that are being converted to native code."), *InOuter->GetName());
-							return ESavePackageResult::ReplaceCompletely;
 						}
 					}
 #endif
