Copying //UE4/Dev-Framework to //UE4/Dev-Main (Source: //UE4/Dev-Framework @ 3779049)

diff --git a/Engine/Source/Runtime/Core/Private/Misc/AssertionMacros.cpp b/Engine/Source/Runtime/Core/Private/Misc/AssertionMacros.cpp
--- a/Engine/Source/Runtime/Core/Private/Misc/AssertionMacros.cpp
+++ b/Engine/Source/Runtime/Core/Private/Misc/AssertionMacros.cpp
@@ -81,32 +81,14 @@ namespace
 	#define CALLSTACK_IGNOREDEPTH 2
 #endif // PLATFORM_LINUX
 
-void InternalPrintScriptCallstack(bool bEmptyWhenDone)
-{
-#if DO_BLUEPRINT_GUARD
-	// Walk the script stack, if any
-	FBlueprintExceptionTracker& BlueprintExceptionTracker = FBlueprintExceptionTracker::Get();
-	if( BlueprintExceptionTracker.ScriptStack.Num() > 0 )
-	{
-		FString ScriptStack = TEXT( "\n\nScript Stack:\n" );
-		for (int32 FrameIdx = BlueprintExceptionTracker.ScriptStack.Num() - 1; FrameIdx >= 0; --FrameIdx)
-		{
-			ScriptStack += BlueprintExceptionTracker.ScriptStack[FrameIdx].GetStackDescription() + TEXT( "\n" );
-		}
-
-		UE_LOG( LogOutputDevice, Warning, TEXT( "%s" ), *ScriptStack );
-
-		if (bEmptyWhenDone)
-		{
-			BlueprintExceptionTracker.ScriptStack.Empty();
-		}
-	}
-#endif
-}
+CORE_API void (*GPrintScriptCallStackFn)() = nullptr;
 
 void PrintScriptCallstack()
 {
-	InternalPrintScriptCallstack(false);
+	if(GPrintScriptCallStackFn)
+	{
+		GPrintScriptCallStackFn();
+	}
 }
 
 /**
@@ -235,7 +217,7 @@ VARARG_BODY(void, FDebug::LogAssertFailedMessage, const TCHAR*, VARARG_EXTRA(con
 	if( !GIsCriticalError )
 	{
 		// Print out the blueprint callstack
-		InternalPrintScriptCallstack(true);
+		PrintScriptCallstack();
 
 		TCHAR DescriptionString[4096];
 		GET_VARARGS( DescriptionString, ARRAY_COUNT( DescriptionString ), ARRAY_COUNT( DescriptionString ) - 1, Fmt, Fmt );
@@ -275,7 +257,7 @@ void FDebug::EnsureFailed(const ANSICHAR* Expr, const ANSICHAR* File, int32 Line
 	FPlatformAtomics::InterlockedIncrement(&ActiveEnsureCount);
 
 	// Print out the blueprint callstack
-	InternalPrintScriptCallstack(false);
+	PrintScriptCallstack();
 
 	// Print initial debug message for this error
 	TCHAR ErrorString[MAX_SPRINTF];
diff --git a/Engine/Source/Runtime/Core/Private/Misc/CoreGlobals.cpp b/Engine/Source/Runtime/Core/Private/Misc/CoreGlobals.cpp
--- a/Engine/Source/Runtime/Core/Private/Misc/CoreGlobals.cpp
+++ b/Engine/Source/Runtime/Core/Private/Misc/CoreGlobals.cpp
@@ -318,7 +318,6 @@ DEFINE_STAT(STAT_TextureMemory);
 DEFINE_STAT(STAT_MemoryPhysXTotalAllocationSize);
 DEFINE_STAT(STAT_MemoryICUTotalAllocationSize);
 DEFINE_STAT(STAT_MemoryICUDataFileAllocationSize);
-DEFINE_STAT(STAT_AnimationMemory);
 DEFINE_STAT(STAT_PrecomputedVisibilityMemory);
 DEFINE_STAT(STAT_SkeletalMeshVertexMemory);
 DEFINE_STAT(STAT_SkeletalMeshIndexMemory);
diff --git a/Engine/Source/Runtime/Core/Public/Misc/CoreDelegates.h b/Engine/Source/Runtime/Core/Public/Misc/CoreDelegates.h
--- a/Engine/Source/Runtime/Core/Public/Misc/CoreDelegates.h
+++ b/Engine/Source/Runtime/Core/Public/Misc/CoreDelegates.h
@@ -99,6 +99,7 @@ public:
 	DECLARE_MULTICAST_DELEGATE(FOnAsyncLoadingFlushUpdate);
 	static FOnAsyncLoadingFlushUpdate OnAsyncLoadingFlushUpdate;
 
+	// Callback on the game thread when an async load is started. This goes off before the packages has finished loading
 	DECLARE_MULTICAST_DELEGATE_OneParam(FOnAsyncLoadPackage, const FString&);
 	static FOnAsyncLoadPackage OnAsyncLoadPackage;
 
diff --git a/Engine/Source/Runtime/Core/Public/Misc/CoreMisc.h b/Engine/Source/Runtime/Core/Public/Misc/CoreMisc.h
--- a/Engine/Source/Runtime/Core/Public/Misc/CoreMisc.h
+++ b/Engine/Source/Runtime/Core/Public/Misc/CoreMisc.h
@@ -179,29 +179,6 @@ struct CORE_API FUrlConfig
 
 bool CORE_API StringHasBadDashes(const TCHAR* Str);
 
-/**
- * Helper for script stack traces
- */
-struct FScriptTraceStackNode
-{
-	FName	Scope;
-	FName	FunctionName;
-
-	FScriptTraceStackNode(FName InScope, FName InFunctionName)
-		: Scope(InScope)
-		, FunctionName(InFunctionName)
-	{
-	}
-
-	FString GetStackDescription() const
-	{
-		return Scope.ToString() + TEXT(".") + FunctionName.ToString();
-	}
-
-private:
-	FScriptTraceStackNode() {};
-};
-
 /** Helper structure for boolean values in config */
 struct CORE_API FBoolConfigValueHelper
 {
@@ -225,6 +202,8 @@ public:
 #endif
 
 #if DO_BLUEPRINT_GUARD
+struct FFrame;
+
 /** 
  * Helper struct for dealing with Blueprint exceptions 
  */
@@ -252,8 +231,8 @@ public:
 	// Script entry point tracking
 	int32 ScriptEntryTag;
 
-	// Stack names from the VM to be unrolled when we assert
-	TArray<FScriptTraceStackNode> ScriptStack;
+	// Stack pointers from the VM to be unrolled when we assert
+	TArray<const FFrame*> ScriptStack;
 };
 
 #endif // DO_BLUEPRINT_GUARD
diff --git a/Engine/Source/Runtime/Core/Public/Misc/CoreStats.h b/Engine/Source/Runtime/Core/Public/Misc/CoreStats.h
--- a/Engine/Source/Runtime/Core/Public/Misc/CoreStats.h
+++ b/Engine/Source/Runtime/Core/Public/Misc/CoreStats.h
@@ -10,7 +10,6 @@
 
 /** Memory stats */
 DECLARE_MEMORY_STAT_EXTERN(TEXT("Audio Memory Used"),STAT_AudioMemory,STATGROUP_Memory, CORE_API);
-DECLARE_MEMORY_STAT_EXTERN(TEXT("Animation Memory"),STAT_AnimationMemory,STATGROUP_Memory, CORE_API);
 DECLARE_MEMORY_STAT_EXTERN(TEXT("Precomputed Visibility Memory"),STAT_PrecomputedVisibilityMemory,STATGROUP_Memory, CORE_API);
 DECLARE_MEMORY_STAT_EXTERN(TEXT("SkeletalMesh Vertex Memory"),STAT_SkeletalMeshVertexMemory,STATGROUP_Memory, CORE_API);
 DECLARE_MEMORY_STAT_EXTERN(TEXT("SkeletalMesh Index Memory"),STAT_SkeletalMeshIndexMemory,STATGROUP_Memory, CORE_API);
diff --git a/Engine/Source/Runtime/Core/Public/ProfilingDebugging/ResourceSize.h b/Engine/Source/Runtime/Core/Public/ProfilingDebugging/ResourceSize.h
--- a/Engine/Source/Runtime/Core/Public/ProfilingDebugging/ResourceSize.h
+++ b/Engine/Source/Runtime/Core/Public/ProfilingDebugging/ResourceSize.h
@@ -5,15 +5,15 @@
 #include "CoreTypes.h"
 #include "Misc/AssertionMacros.h"
 
-/** Indicates which resource size should be returned. */
+/** Indicate what types of resources should be included for calculating used memory */
 namespace EResourceSizeMode
 {
 	enum Type
 	{
-		/** Only exclusive resource size */
+		/** Only include memory used by non-UObject resources that are directly owned by this UObject. This is used to show memory actually used at runtime */
 		Exclusive,
-		/** Resource size of the object and all of its references */
-		Inclusive,
+		/** Include exclusive resources and UObject serialized memory for this and all child UObjects, but not memory for external referenced assets or editor only members. This is used in the editor to estimate maximum required memory */
+		EstimatedTotal,
 	};
 };
 
diff --git a/Engine/Source/Runtime/CoreUObject/Private/Misc/PackageName.cpp b/Engine/Source/Runtime/CoreUObject/Private/Misc/PackageName.cpp
--- a/Engine/Source/Runtime/CoreUObject/Private/Misc/PackageName.cpp
+++ b/Engine/Source/Runtime/CoreUObject/Private/Misc/PackageName.cpp
@@ -1306,12 +1306,16 @@ bool FPackageName::IsScriptPackage(const FString& InPackageName)
 	return InPackageName.StartsWith(FLongPackagePathsSingleton::Get().ScriptRootPath);
 }
 
-// Are we a package that resides in memory and not on disk
 bool FPackageName::IsMemoryPackage(const FString& InPackageName)
 {
 	return InPackageName.StartsWith(FLongPackagePathsSingleton::Get().MemoryRootPath);
 }
 
+bool FPackageName::IsTempPackage(const FString& InPackageName)
+{
+	return InPackageName.StartsWith(FLongPackagePathsSingleton::Get().TempRootPath);
+}
+
 bool FPackageName::IsLocalizedPackage(const FString& InPackageName)
 {
 	// Minimum valid package name length is "/A/L10N"
diff --git a/Engine/Source/Runtime/CoreUObject/Private/Serialization/AsyncLoading.cpp b/Engine/Source/Runtime/CoreUObject/Private/Serialization/AsyncLoading.cpp
--- a/Engine/Source/Runtime/CoreUObject/Private/Serialization/AsyncLoading.cpp
+++ b/Engine/Source/Runtime/CoreUObject/Private/Serialization/AsyncLoading.cpp
@@ -4630,6 +4630,24 @@ EAsyncPackageState::Type FAsyncLoadingThread::ProcessLoadedPackages(bool bUseTim
 				const EAsyncLoadingResult::Type LoadingResult = Package->HasLoadFailed() ? EAsyncLoadingResult::Failed : EAsyncLoadingResult::Succeeded;
 				Package->CallCompletionCallbacks(bInternalCallbacks, LoadingResult);
 
+#if WITH_EDITOR
+				// In the editor we need to find any assets and add them to list for later callback
+				UPackage* LoadedPackage = Package->GetLoadedPackage();
+
+				if (LoadedPackage)
+				{
+					TArray<UObject*> TopLevelObjects;
+					GetObjectsWithOuter(LoadedPackage, TopLevelObjects, false);
+
+					for (UObject* TopLevelObject : TopLevelObjects)
+					{
+						if (TopLevelObject->IsAsset())
+						{
+							LoadedAssets.Add(TopLevelObject);
+						}
+					}
+				}
+#endif
 				// We don't need the package anymore
 				PackagesToDelete.AddUnique(Package);
 				Package->MarkRequestIDsAsComplete();
@@ -4680,6 +4698,19 @@ EAsyncPackageState::Type FAsyncLoadingThread::ProcessLoadedPackages(bool bUseTim
 		FBlueprintSupport::FlushReinstancingQueue();
 #endif
 
+#if WITH_EDITOR
+		// In editor builds, call the asset load callback. This happens in both editor and standalone to match EndLoad
+		for (const FWeakObjectPtr& WeakAsset : LoadedAssets)
+		{
+			// It may have been unloaded/marked pending kill since being added, ignore those cases
+			if (UObject* LoadedAsset = WeakAsset.Get())
+			{
+				FCoreUObjectDelegates::OnAssetLoaded.Broadcast(LoadedAsset);
+			}
+		}
+		LoadedAssets.Reset();
+#endif
+
 		// We're not done until all packages have been deleted
 		Result = PackagesToDelete.Num() ? EAsyncPackageState::PendingImports : EAsyncPackageState::Complete;
 	}
@@ -6570,6 +6601,12 @@ void FAsyncPackage::CallCompletionCallbacks(bool bInternal, EAsyncLoadingResult:
 	}
 }
 
+UPackage* FAsyncPackage::GetLoadedPackage()
+{
+	UPackage* LoadedPackage = (!bLoadHasFailed) ? LinkerRoot : nullptr;
+	return LoadedPackage;
+}
+
 void FAsyncPackage::Cancel()
 {
 	UE_CLOG(GEventDrivenLoaderEnabled, LogStreaming, Fatal, TEXT("FAsyncPackage::Cancel is not supported with the new loader"));
diff --git a/Engine/Source/Runtime/CoreUObject/Private/Serialization/AsyncLoadingThread.h b/Engine/Source/Runtime/CoreUObject/Private/Serialization/AsyncLoadingThread.h
--- a/Engine/Source/Runtime/CoreUObject/Private/Serialization/AsyncLoadingThread.h
+++ b/Engine/Source/Runtime/CoreUObject/Private/Serialization/AsyncLoadingThread.h
@@ -179,6 +179,9 @@ class FAsyncLoadingThread : public FRunnable
 	TArray<FAsyncPackage*> LoadedPackagesToProcess;
 	TArray<FAsyncPackage*> PackagesToDelete;
 	TMap<FName, FAsyncPackage*> LoadedPackagesToProcessNameLookup;
+#if WITH_EDITOR
+	TArray<FWeakObjectPtr> LoadedAssets;
+#endif
 #if THREADSAFE_UOBJECTS
 	/** [ASYNC/GAME THREAD] Critical section for LoadedPackagesToProcess list. 
 	 * Note this is only required for looking up existing packages on the async loading thread 
diff --git a/Engine/Source/Runtime/CoreUObject/Private/UObject/LinkerLoad.cpp b/Engine/Source/Runtime/CoreUObject/Private/UObject/LinkerLoad.cpp
--- a/Engine/Source/Runtime/CoreUObject/Private/UObject/LinkerLoad.cpp
+++ b/Engine/Source/Runtime/CoreUObject/Private/UObject/LinkerLoad.cpp
@@ -3521,6 +3521,10 @@ UObject* FLinkerLoad::CreateExport( int32 Index )
 		UClass* LoadClass = GetExportLoadClass(Index);
 		if( !LoadClass && !Export.ClassIndex.IsNull() ) // Hack to load packages with classes which do not exist.
 		{
+			Export.bExportLoadFailed = true;
+
+			FString OuterName = Export.OuterIndex.IsNull() ? LinkerRoot->GetFullName() : GetFullImpExpName(Export.OuterIndex);
+			UE_CLOG(Export.ObjectFlags & EObjectFlags::RF_Public, LogLinker, Warning, TEXT("Unable to load %s with outer %s because its class does not exist"), *Export.ObjectName.ToString(), *OuterName);
 			return NULL;
 		}
 
diff --git a/Engine/Source/Runtime/CoreUObject/Private/UObject/LinkerSave.cpp b/Engine/Source/Runtime/CoreUObject/Private/UObject/LinkerSave.cpp
--- a/Engine/Source/Runtime/CoreUObject/Private/UObject/LinkerSave.cpp
+++ b/Engine/Source/Runtime/CoreUObject/Private/UObject/LinkerSave.cpp
@@ -266,23 +266,4 @@ FArchive& FLinkerSave::operator<<(FLazyObjectPtr& LazyObjectPtr)
 	FUniqueObjectGuid ID;
 	ID = LazyObjectPtr.GetUniqueID();
 	return *this << ID;
-}
-
-FArchive& FLinkerSave::operator<<(FSoftObjectPtr& SoftObjectPtr)
-{
-	FSoftObjectPath ID;
-	UObject *Object = SoftObjectPtr.Get();
-
-	if (Object)
-	{
-		// Use object in case name has changed. 
-		ID = FSoftObjectPath(Object);
-	}
-	else
-	{
-		ID = SoftObjectPtr.GetUniqueID();
-	}
-
-	ID.Serialize(*this);
-	return *this;
-}
+}
\ No newline at end of file
diff --git a/Engine/Source/Runtime/CoreUObject/Private/UObject/Obj.cpp b/Engine/Source/Runtime/CoreUObject/Private/UObject/Obj.cpp
--- a/Engine/Source/Runtime/CoreUObject/Private/UObject/Obj.cpp
+++ b/Engine/Source/Runtime/CoreUObject/Private/UObject/Obj.cpp
@@ -1642,13 +1642,6 @@ const FName FPrimaryAssetId::PrimaryAssetNameTag(TEXT("PrimaryAssetName"));
 
 void UObject::GetAssetRegistryTags(TArray<FAssetRegistryTag>& OutTags) const
 {
-	// Add ResourceSize if non-zero. GetResourceSize is not const because many override implementations end up calling Serialize on this pointers.
-	const SIZE_T ResourceSize = const_cast<UObject*>(this)->GetResourceSizeBytes(EResourceSizeMode::Exclusive);
-	if ( ResourceSize > 0 || ( !GetClass()->HasAnyClassFlags(CLASS_CompiledFromBlueprint) && HasAnyFlags(RF_ClassDefaultObject) ) )
-	{
-		OutTags.Add( FAssetRegistryTag("ResourceSize", FString::Printf(TEXT("%d"), (ResourceSize + 512) / 1024), FAssetRegistryTag::TT_Numerical) );
-	}
-
 	// Add primary asset info if valid
 	FPrimaryAssetId PrimaryAssetId = GetPrimaryAssetId();
 	if (PrimaryAssetId.IsValid())
@@ -1669,15 +1662,38 @@ const FName& UObject::SourceFileTagName()
 #if WITH_EDITOR
 void UObject::GetAssetRegistryTagMetadata(TMap<FName, FAssetRegistryTagMetadata>& OutMetadata) const
 {
-	OutMetadata.Add("ResourceSize",
+	OutMetadata.Add(FPrimaryAssetId::PrimaryAssetTypeTag,
 		FAssetRegistryTagMetadata()
-			.SetDisplayName(NSLOCTEXT("UObject", "Size", "Size"))
-			.SetSuffix(NSLOCTEXT("UObject", "KilobytesSuffix", "Kb"))
-			.SetTooltip(NSLOCTEXT("UObject", "SizeTooltip", "The size of the asset in kilobytes"))
-		);
+		.SetDisplayName(NSLOCTEXT("UObject", "PrimaryAssetType", "Primary Asset Type"))
+		.SetTooltip(NSLOCTEXT("UObject", "PrimaryAssetTypeTooltip", "Type registered with the Asset Manager system"))
+	);
+
+	OutMetadata.Add(FPrimaryAssetId::PrimaryAssetNameTag,
+		FAssetRegistryTagMetadata()
+		.SetDisplayName(NSLOCTEXT("UObject", "PrimaryAssetName", "Primary Asset Name"))
+		.SetTooltip(NSLOCTEXT("UObject", "PrimaryAssetNameTooltip", "Logical name registered with the Asset Manager system"))
+	);
 }
 #endif
 
+void UObject::GetResourceSizeEx(FResourceSizeEx& CumulativeResourceSize)
+{
+	if (CumulativeResourceSize.GetResourceSizeMode() == EResourceSizeMode::EstimatedTotal)
+	{
+		// Include this object's serialize size, and recursively call on direct subobjects
+		FArchiveCountMem MemoryCount(this);
+		CumulativeResourceSize.AddDedicatedSystemMemoryBytes(MemoryCount.GetMax());
+
+		TArray<UObject*> SubObjects;
+		GetObjectsWithOuter(this, SubObjects, false);
+
+		for (UObject* SubObject : SubObjects)
+		{
+			SubObject->GetResourceSizeEx(CumulativeResourceSize);
+		}
+	}
+}
+
 bool UObject::IsAsset() const
 {
 	// Assets are not transient or CDOs. They must be public.
@@ -3741,7 +3757,7 @@ bool StaticExec( UWorld* InWorld, const TCHAR* Cmd, FOutputDevice& Ar )
 								FArchiveCountMem Count(Object);
 
 								// Get the 'old-style' resource size and the truer resource size
-								const SIZE_T ResourceSize = It->GetResourceSizeBytes(EResourceSizeMode::Inclusive);
+								const SIZE_T ResourceSize = It->GetResourceSizeBytes(EResourceSizeMode::EstimatedTotal);
 								const SIZE_T TrueResourceSize = It->GetResourceSizeBytes(EResourceSizeMode::Exclusive);
 								
 								if (bListClass)
@@ -4264,6 +4280,17 @@ void UObject::PreDestroyFromReplication()
 
 }
 
+#if WITH_EDITOR
+/*-----------------------------------------------------------------------------
+	Data Validation.
+-----------------------------------------------------------------------------*/
+
+EDataValidationResult UObject::IsDataValid(TArray<FText>& ValidationErrors)
+{
+	return EDataValidationResult::NotValidated;
+}
+
+#endif // WITH_EDITOR
 /** IsNameStableForNetworking means an object can be referred to its path name (relative to outer) over the network */
 bool UObject::IsNameStableForNetworking() const
 {
diff --git a/Engine/Source/Runtime/CoreUObject/Private/UObject/ObjectMemoryAnalyzer.cpp b/Engine/Source/Runtime/CoreUObject/Private/UObject/ObjectMemoryAnalyzer.cpp
--- a/Engine/Source/Runtime/CoreUObject/Private/UObject/ObjectMemoryAnalyzer.cpp
+++ b/Engine/Source/Runtime/CoreUObject/Private/UObject/ObjectMemoryAnalyzer.cpp
@@ -192,7 +192,7 @@ SIZE_T FObjectMemoryAnalyzer::CalculateSizeRecursive(UObject* Object)
 
 	SIZE_T InclusiveSize = Annotation.ExclusiveMemoryUsage;
 	Annotation.ExclusiveResourceSize = Object->GetResourceSizeBytes(EResourceSizeMode::Exclusive);
-	Annotation.InclusiveResourceSize = Object->GetResourceSizeBytes(EResourceSizeMode::Inclusive);
+	Annotation.InclusiveResourceSize = Object->GetResourceSizeBytes(EResourceSizeMode::EstimatedTotal);
 
 	TArray<UObject*> ReferencedObjects;
 	GetReferencedObjects(Object, ReferencedObjects);
diff --git a/Engine/Source/Runtime/CoreUObject/Private/UObject/ScriptCore.cpp b/Engine/Source/Runtime/CoreUObject/Private/UObject/ScriptCore.cpp
--- a/Engine/Source/Runtime/CoreUObject/Private/UObject/ScriptCore.cpp
+++ b/Engine/Source/Runtime/CoreUObject/Private/UObject/ScriptCore.cpp
@@ -175,6 +175,25 @@ void FBlueprintCoreDelegates::SetScriptMaximumLoopIterations( const int32 Maximu
 	}
 }
 
+void PrintScriptCallStackImpl()
+{
+#if DO_BLUEPRINT_GUARD
+	FBlueprintExceptionTracker& BlueprintExceptionTracker = FBlueprintExceptionTracker::Get();
+	if( BlueprintExceptionTracker.ScriptStack.Num() > 0 )
+	{
+		FString ScriptStack = TEXT( "\n\nScript Stack:\n" );
+		for (int32 FrameIdx = BlueprintExceptionTracker.ScriptStack.Num() - 1; FrameIdx >= 0; --FrameIdx)
+		{
+			ScriptStack += BlueprintExceptionTracker.ScriptStack[FrameIdx]->GetStackDescription() + TEXT( "\n" );
+		}
+
+		UE_LOG( LogOutputDevice, Warning, TEXT( "%s" ), *ScriptStack );
+	}
+#endif
+}
+
+extern CORE_API void (*GPrintScriptCallStackFn)();
+
 //////////////////////////////////////////////////////////////////////////
 // FEditorScriptExecutionGuard
 FEditorScriptExecutionGuard::FEditorScriptExecutionGuard()
@@ -185,6 +204,8 @@ FEditorScriptExecutionGuard::FEditorScriptExecutionGuard()
 	if( GIsEditor && !FApp::IsGame() )
 	{
 		GInitRunaway();
+		
+		GPrintScriptCallStackFn = &PrintScriptCallStackImpl;
 	}
 }
 
@@ -343,7 +364,7 @@ FString FFrame::GetScriptCallstack()
 	{
 		for (int32 i = BlueprintExceptionTracker.ScriptStack.Num() - 1; i >= 0; --i)
 		{
-			ScriptStack += TEXT("\t") + BlueprintExceptionTracker.ScriptStack[i].GetStackDescription() + TEXT("\n");
+			ScriptStack += TEXT("\t") + BlueprintExceptionTracker.ScriptStack[i]->GetStackDescription() + TEXT("\n");
 		}
 	}
 	else
@@ -357,6 +378,11 @@ FString FFrame::GetScriptCallstack()
 	return ScriptStack;
 }
 
+FString FFrame::GetStackDescription() const
+{
+	return Node->GetOuter()->GetName() + TEXT(".") + Node->GetName();
+}
+
 //
 // Error or warning handler.
 //
diff --git a/Engine/Source/Runtime/CoreUObject/Private/UObject/UObjectGlobals.cpp b/Engine/Source/Runtime/CoreUObject/Private/UObject/UObjectGlobals.cpp
--- a/Engine/Source/Runtime/CoreUObject/Private/UObject/UObjectGlobals.cpp
+++ b/Engine/Source/Runtime/CoreUObject/Private/UObject/UObjectGlobals.cpp
@@ -100,25 +100,19 @@ namespace LoadPackageStats
 FCoreUObjectDelegates::FRegisterHotReloadAddedClassesDelegate FCoreUObjectDelegates::RegisterHotReloadAddedClassesDelegate;
 FCoreUObjectDelegates::FRegisterClassForHotReloadReinstancingDelegate FCoreUObjectDelegates::RegisterClassForHotReloadReinstancingDelegate;
 FCoreUObjectDelegates::FReinstanceHotReloadedClassesDelegate FCoreUObjectDelegates::ReinstanceHotReloadedClassesDelegate;
-// Delegates used by SavePackage()
 FCoreUObjectDelegates::FIsPackageOKToSaveDelegate FCoreUObjectDelegates::IsPackageOKToSaveDelegate;
-FCoreUObjectDelegates::FAutoPackageBackupDelegate FCoreUObjectDelegates::AutoPackageBackupDelegate;
-
 FCoreUObjectDelegates::FOnPackageReloaded FCoreUObjectDelegates::OnPackageReloaded;
 FCoreUObjectDelegates::FNetworkFileRequestPackageReload FCoreUObjectDelegates::NetworkFileRequestPackageReload;
-
+#if WITH_EDITOR
+FCoreUObjectDelegates::FAutoPackageBackupDelegate FCoreUObjectDelegates::AutoPackageBackupDelegate;
 FCoreUObjectDelegates::FOnPreObjectPropertyChanged FCoreUObjectDelegates::OnPreObjectPropertyChanged;
 FCoreUObjectDelegates::FOnObjectPropertyChanged FCoreUObjectDelegates::OnObjectPropertyChanged;
-
-#if WITH_EDITOR
-// Set of objects modified this frame
 TSet<UObject*> FCoreUObjectDelegates::ObjectsModifiedThisFrame;
 FCoreUObjectDelegates::FOnObjectModified FCoreUObjectDelegates::OnObjectModified;
 FCoreUObjectDelegates::FOnAssetLoaded FCoreUObjectDelegates::OnAssetLoaded;
 FCoreUObjectDelegates::FOnObjectSaved FCoreUObjectDelegates::OnObjectSaved;
 #endif // WITH_EDITOR
 
-
 FSimpleMulticastDelegate& FCoreUObjectDelegates::GetPreGarbageCollectDelegate()
 {
 	static FSimpleMulticastDelegate Delegate;
@@ -140,7 +134,6 @@ FSimpleMulticastDelegate FCoreUObjectDelegates::PostGarbageCollectConditionalBeg
 FCoreUObjectDelegates::FPreLoadMapDelegate FCoreUObjectDelegates::PreLoadMap;
 FCoreUObjectDelegates::FPostLoadMapDelegate FCoreUObjectDelegates::PostLoadMapWithWorld;
 PRAGMA_DISABLE_DEPRECATION_WARNINGS
-FSimpleMulticastDelegate FCoreUObjectDelegates::PostLoadMap;
 FCoreUObjectDelegates::FSoftObjectPathLoaded FCoreUObjectDelegates::StringAssetReferenceLoaded;
 FCoreUObjectDelegates::FSoftObjectPathSaving FCoreUObjectDelegates::StringAssetReferenceSaving;
 FCoreUObjectDelegates::FOnRedirectorFollowed FCoreUObjectDelegates::RedirectorFollowed;
@@ -1520,6 +1513,7 @@ void EndLoad()
 	FScopedSlowTask SlowTask(0, NSLOCTEXT("Core", "PerformingPostLoad", "Performing post-load..."), ShouldReportProgress());
 
 	int32 NumObjectsLoaded = 0, NumObjectsFound = 0;
+	TSet<UObject*> AssetsLoaded;
 #endif
 
 	while (--ThreadContext.ObjBeginLoadCount == 0 && (ThreadContext.ObjLoaded.Num() || ThreadContext.ImportCount || ThreadContext.ForcedExportCount))
@@ -1638,17 +1632,14 @@ void EndLoad()
 			}
 
 #if WITH_EDITOR
-			// Send global notification for each object that was loaded.
-			// Useful for updating UI such as ContentBrowser's loaded status.
+			// Schedule asset loaded callbacks for later
+			for( int32 CurObjIndex=0; CurObjIndex<ObjLoaded.Num(); CurObjIndex++ )
 			{
-				for( int32 CurObjIndex=0; CurObjIndex<ObjLoaded.Num(); CurObjIndex++ )
+				UObject* Obj = ObjLoaded[CurObjIndex];
+				check(Obj);
+				if ( Obj->IsAsset() )
 				{
-					UObject* Obj = ObjLoaded[CurObjIndex];
-					check(Obj);
-					if ( Obj->IsAsset() )
-					{
-						FCoreUObjectDelegates::OnAssetLoaded.Broadcast(Obj);
-					}
+					AssetsLoaded.Add(Obj);
 				}
 			}
 #endif	// WITH_EDITOR
@@ -1709,6 +1700,16 @@ void EndLoad()
 
 	// Loaded new objects, so allow reaccessing asset ptrs
 	FSoftObjectPath::InvalidateTag();
+
+#if WITH_EDITOR
+	// Now call asset loaded callbacks for anything that was loaded. We do this at the very end so any nested objects will load properly
+	// Useful for updating UI such as ContentBrowser's loaded status.
+	for (UObject* LoadedAsset : AssetsLoaded)
+	{
+		check(LoadedAsset);
+		FCoreUObjectDelegates::OnAssetLoaded.Broadcast(LoadedAsset);
+	}
+#endif	// WITH_EDITOR
 }
 
 /*-----------------------------------------------------------------------------
diff --git a/Engine/Source/Runtime/CoreUObject/Public/Misc/PackageName.h b/Engine/Source/Runtime/CoreUObject/Public/Misc/PackageName.h
--- a/Engine/Source/Runtime/CoreUObject/Public/Misc/PackageName.h
+++ b/Engine/Source/Runtime/CoreUObject/Public/Misc/PackageName.h
@@ -380,12 +380,19 @@ public:
 	static bool IsScriptPackage(const FString& InPackageName);
 
 	/**
-	 * Checks the root of hte package's path to see if it's a memory package
-	 *		This should be set for packages that reside in memory and not on disk, we treat them similar to a script package
+	 * Checks the root of the package's path to see if it's a memory package
+	 * This should be set for packages that reside in memory and not on disk, we treat them similar to a script package
 	 * @return true if the root of the patch matches the memory path
 	 */
 	static bool IsMemoryPackage(const FString& InPackageName);
 
+	/**
+	 * Checks the root of the package's path to see if it is a temp package
+	 * Temp packages are sometimes saved to disk, and sometimes only exist in memory. They are never in source control
+	 * @return true if the root of the patch matches the temp path
+	 */
+	static bool IsTempPackage(const FString& InPackageName);
+
 	/**
 	 * Checks the root of the package's path to see if it is a localized package
 	 * @return true if the root of the path matches any localized root path
diff --git a/Engine/Source/Runtime/CoreUObject/Public/Serialization/AsyncLoading.h b/Engine/Source/Runtime/CoreUObject/Public/Serialization/AsyncLoading.h
--- a/Engine/Source/Runtime/CoreUObject/Public/Serialization/AsyncLoading.h
+++ b/Engine/Source/Runtime/CoreUObject/Public/Serialization/AsyncLoading.h
@@ -484,6 +484,9 @@ struct FAsyncPackage : FGCObject
 	/** Removes all objects from the list and clears async loading flags */
 	void EmptyReferencedObjects();
 
+	/** Returns the UPackage wrapped by this, if it is valid */
+	UPackage* GetLoadedPackage();
+
 private:	
 
 	struct FCompletionCallback
diff --git a/Engine/Source/Runtime/CoreUObject/Public/UObject/LinkerSave.h b/Engine/Source/Runtime/CoreUObject/Public/UObject/LinkerSave.h
--- a/Engine/Source/Runtime/CoreUObject/Public/UObject/LinkerSave.h
+++ b/Engine/Source/Runtime/CoreUObject/Public/UObject/LinkerSave.h
@@ -80,7 +80,6 @@ public:
 	FArchive& operator<<( FName& InName );
 	FArchive& operator<<( UObject*& Obj );
 	FArchive& operator<<( FLazyObjectPtr& LazyObjectPtr );
-	FArchive& operator<<( FSoftObjectPtr& SoftObjectPtr );
 
 #if WITH_EDITOR
 	// proxy for debugdata
diff --git a/Engine/Source/Runtime/CoreUObject/Public/UObject/Object.h b/Engine/Source/Runtime/CoreUObject/Public/UObject/Object.h
--- a/Engine/Source/Runtime/CoreUObject/Public/UObject/Object.h
+++ b/Engine/Source/Runtime/CoreUObject/Public/UObject/Object.h
@@ -493,35 +493,19 @@ public:
 	}
 
 	/**
-	 * Returns the size of the object/ resource for display to artists/ LDs in the Editor. The
-	 * default behavior is to return 0 which indicates that the resource shouldn't
-	 * display its size.
-	 *
-	 * @param	Mode	Indicates which resource size should be returned
-	 * @return	Size of resource as to be displayed to artists/ LDs in the Editor.
-	 */
-	DEPRECATED(4.14, "GetResourceSize is deprecated. Please use GetResourceSizeEx or GetResourceSizeBytes instead.")
-	virtual SIZE_T GetResourceSize(EResourceSizeMode::Type Mode)
-	{
-		return GetResourceSizeBytes(Mode);
-	}
-
-	/**
-	 * Get the size of the object/resource for display to artists/LDs in the Editor.
+	 * Get the size of the object/resource for use in memory tools or to display to artists/LDs in the Editor
 	 * This is the extended version which separates up the used memory into different memory regions (the actual definition of which may be platform specific).
 	 *
 	 * @param	CumulativeResourceSize	Struct used to count up the cumulative size of the resource as to be displayed to artists/LDs in the Editor.
 	 */
-	virtual void GetResourceSizeEx(FResourceSizeEx& CumulativeResourceSize)
-	{
-	}
+	virtual void GetResourceSizeEx(FResourceSizeEx& CumulativeResourceSize);
 
 	/**
-	 * Get the size of the object/resource for display to artists/LDs in the Editor.
+	 * Get the size of the object/resource for use in memory tools or to display to artists/LDs in the Editor
 	 * This is the simple version which just returns the total number of bytes used by this object.
 	 *
 	 * @param	Mode					Indicates which resource size should be returned.
-	 * @return The cumulative size of the resource as to be displayed to artists/LDs in the Editor.
+	 * @return The cumulative size of this object in memory
 	 */
 	SIZE_T GetResourceSizeBytes(EResourceSizeMode::Type Mode)
 	{
@@ -731,6 +715,13 @@ public:
 	/** Called right before being marked for destruction due to network replication */
 	virtual void PreDestroyFromReplication();
 
+#if WITH_EDITOR
+	/** 
+	 * @return		Returns Valid if this object has data validation rules set up for it and the data for this object is valid. Returns Invalid if it does not pass the rules. Returns NotValidated if no rules are set for this object.
+	 */
+	virtual EDataValidationResult IsDataValid(TArray<FText>& ValidationErrors);
+#endif // WITH_EDITOR
+
 	/**
  	 *******************************************************
 	 * Non virtual functions, not intended to be overridden
diff --git a/Engine/Source/Runtime/CoreUObject/Public/UObject/Stack.h b/Engine/Source/Runtime/CoreUObject/Public/UObject/Stack.h
--- a/Engine/Source/Runtime/CoreUObject/Public/UObject/Stack.h
+++ b/Engine/Source/Runtime/CoreUObject/Public/UObject/Stack.h
@@ -186,6 +186,11 @@ public:
 	* This will return the StackTrace of the all script frames currently active
 	**/
 	COREUOBJECT_API static FString GetScriptCallstack();
+
+	/** 
+	 * This will return a string of the form "ScopeName.FunctionName" associated with this stack frame:
+	 */
+	COREUOBJECT_API FString GetStackDescription() const;
 };
 
 
@@ -207,8 +212,7 @@ inline FFrame::FFrame( UObject* InObject, UFunction* InNode, void* InLocals, FFr
 	, bArrayContextFailed(false)
 {
 #if DO_BLUEPRINT_GUARD
-	FScriptTraceStackNode StackNode(InNode->GetOuter()->GetFName(), InNode->GetFName());
-	FBlueprintExceptionTracker::Get().ScriptStack.Push(StackNode);
+	FBlueprintExceptionTracker::Get().ScriptStack.Push(this);
 #endif
 }
 
diff --git a/Engine/Source/Runtime/CoreUObject/Public/UObject/UObjectGlobals.h b/Engine/Source/Runtime/CoreUObject/Public/UObject/UObjectGlobals.h
--- a/Engine/Source/Runtime/CoreUObject/Public/UObject/UObjectGlobals.h
+++ b/Engine/Source/Runtime/CoreUObject/Public/UObject/UObjectGlobals.h
@@ -1755,11 +1755,6 @@ protected:
 	bool			bShouldIgnoreTransient;
 };
 
-
-/** Delegate types for source control package saving checks and adding package to default changelist */
-DECLARE_DELEGATE_RetVal_TwoParams( bool, FCheckForAutoAddDelegate, UPackage*, const FString& );
-DECLARE_DELEGATE_OneParam( FAddPackageToDefaultChangelistDelegate, const TCHAR* );
-
 /** Defined in PackageReload.h */
 enum class EPackageReloadPhase : uint8;
 class FPackageReloadedEvent;
@@ -1771,40 +1766,23 @@ class FGarbageCollectionTracer;
  */
 struct COREUOBJECT_API FCoreUObjectDelegates
 {
-	// Callback for object property modifications
+#if WITH_EDITOR
+	/** Callback for object property modifications, called by UObject::PostEditChangeProperty with a single property event */
 	DECLARE_MULTICAST_DELEGATE_TwoParams(FOnObjectPropertyChanged, UObject*, struct FPropertyChangedEvent&);
-
-	// Called when a property is changed
 	static FOnObjectPropertyChanged OnObjectPropertyChanged;
 
-	// Callback for PreEditChange
+	/** Callback for object property modifications, called by UObject::PreEditChange with a full property chain */
 	DECLARE_MULTICAST_DELEGATE_TwoParams(FOnPreObjectPropertyChanged, UObject*, const class FEditPropertyChain&);
-
-	// Called before a property is changed
 	static FOnPreObjectPropertyChanged OnPreObjectPropertyChanged;
 
-	/** Delegate type for making auto backup of package */
-	DECLARE_DELEGATE_RetVal_OneParam(bool, FAutoPackageBackupDelegate, const UPackage&);
-
-	/** Called by ReloadPackage during package reloading. It will be called several times for different phases of fix-up to allow custom code to handle updating objects as needed */
-	DECLARE_MULTICAST_DELEGATE_TwoParams(FOnPackageReloaded, EPackageReloadPhase, FPackageReloadedEvent*);
-	static FOnPackageReloaded OnPackageReloaded;
-
-	/** Called when a package reload request is received from a network file server */
-	DECLARE_DELEGATE_OneParam(FNetworkFileRequestPackageReload, const TArray<FString>& /*PackageNames*/);
-	static FNetworkFileRequestPackageReload NetworkFileRequestPackageReload;
-
-#if WITH_EDITOR
-	// Callback for all object modifications
+	/** Called when an object is registered for change with UObject::Modify. This gets called in both the editor and standalone game editor builds, for every object modified */
 	DECLARE_MULTICAST_DELEGATE_OneParam(FOnObjectModified, UObject*);
-
-	// Called when any object is modified at all 
 	static FOnObjectModified OnObjectModified;
 
-	// Set of objects modified this frame, to prevent multiple triggerings of the OnObjectModified delegate.
+	/** Set of objects modified this frame, to prevent multiple triggerings of the OnObjectModified delegate */
 	static TSet<UObject*> ObjectsModifiedThisFrame;
 
-	// Broadcast OnObjectModified if the broadcast hasn't ocurred for this object in this frame
+	/** Broadcast OnObjectModified if the broadcast hasn't occurred for this object in this frame */
 	static void BroadcastOnObjectModified(UObject* Object)
 	{
 		if (OnObjectModified.IsBound() && !ObjectsModifiedThisFrame.Contains(Object))
@@ -1814,29 +1792,29 @@ struct COREUOBJECT_API FCoreUObjectDelegates
 		}
 	}
 
-	// Callback for when an asset is loaded (Editor)
-	DECLARE_MULTICAST_DELEGATE_OneParam(FOnAssetLoaded, UObject*);
-
-	// Called when an asset is loaded
-	static FOnAssetLoaded OnAssetLoaded;
-
-	// Callback for when an asset is saved (Editor)
+	/** Callback for when an asset is saved. This is called from UObject::PreSave before it is actually written to disk, for every object saved */
 	DECLARE_MULTICAST_DELEGATE_OneParam(FOnObjectSaved, UObject*);
-
-	// Called when an asset is saved
 	static FOnObjectSaved OnObjectSaved;
 
-#endif	//WITH_EDITOR
-
-
+	/** Callback for when an asset is loaded. This gets called in both the editor and standalone game editor builds, but only for objects that return true for IsAsset() */
+	DECLARE_MULTICAST_DELEGATE_OneParam(FOnAssetLoaded, UObject*);
+	static FOnAssetLoaded OnAssetLoaded;
 
 	/** Delegate used by SavePackage() to create the package backup */
+	DECLARE_DELEGATE_RetVal_OneParam(bool, FAutoPackageBackupDelegate, const UPackage&);
 	static FAutoPackageBackupDelegate AutoPackageBackupDelegate;
+#endif // WITH_EDITOR
 
-	/** Delegate type for saving check */
-	DECLARE_DELEGATE_RetVal_ThreeParams(bool, FIsPackageOKToSaveDelegate, UPackage*, const FString&, FOutputDevice*);
+	/** Called by ReloadPackage during package reloading. It will be called several times for different phases of fix-up to allow custom code to handle updating objects as needed */
+	DECLARE_MULTICAST_DELEGATE_TwoParams(FOnPackageReloaded, EPackageReloadPhase, FPackageReloadedEvent*);
+	static FOnPackageReloaded OnPackageReloaded;
+
+	/** Called when a package reload request is received from a network file server */
+	DECLARE_DELEGATE_OneParam(FNetworkFileRequestPackageReload, const TArray<FString>& /*PackageNames*/);
+	static FNetworkFileRequestPackageReload NetworkFileRequestPackageReload;
 
 	/** Delegate used by SavePackage() to check whether a package should be saved */
+	DECLARE_DELEGATE_RetVal_ThreeParams(bool, FIsPackageOKToSaveDelegate, UPackage*, const FString&, FOutputDevice*);
 	static FIsPackageOKToSaveDelegate IsPackageOKToSaveDelegate;
 
 	/** Delegate for registering hot-reloaded classes that have been added  */
@@ -1851,52 +1829,47 @@ struct COREUOBJECT_API FCoreUObjectDelegates
 	DECLARE_MULTICAST_DELEGATE(FReinstanceHotReloadedClassesDelegate);
 	static FReinstanceHotReloadedClassesDelegate ReinstanceHotReloadedClassesDelegate;
 
-	// Sent at the very beginning of LoadMap
+	/** Sent at the very beginning of LoadMap */
 	DECLARE_MULTICAST_DELEGATE_OneParam(FPreLoadMapDelegate, const FString& /* MapName */);
 	static FPreLoadMapDelegate PreLoadMap;
 
-	// Sent at the _successful_ end of LoadMap
+	/** Sent at the _successful_ end of LoadMap */
 	DECLARE_MULTICAST_DELEGATE_OneParam(FPostLoadMapDelegate, UWorld* /* LoadedWorld */);
 	static FPostLoadMapDelegate PostLoadMapWithWorld;
 
-	DEPRECATED(4.16, "Use PostLoadMapWithWorld instead.")
-	static FSimpleMulticastDelegate PostLoadMap;
-
-	// Sent at the _successful_ end of LoadMap
+	/** Sent when a network replay has started */
 	static FSimpleMulticastDelegate PostDemoPlay;
 
-	// Called before garbage collection
+	/** Called before garbage collection */
 	static FSimpleMulticastDelegate& GetPreGarbageCollectDelegate();
 
-	// Delegate type for reachability analysis external roots callback. First parameter is FGarbageCollectionTracer to use for tracing, second is flags with which objects should be kept alive regardless, third is whether to force single threading
+	/** Delegate type for reachability analysis external roots callback. First parameter is FGarbageCollectionTracer to use for tracing, second is flags with which objects should be kept alive regardless, third is whether to force single threading */
 	DECLARE_MULTICAST_DELEGATE_ThreeParams(FTraceExternalRootsForReachabilityAnalysisDelegate, FGarbageCollectionTracer&, EObjectFlags, bool);
 
-	// Called as last phase of reachability analysis. Allow external systems to add UObject roots *after* first reachability pass has been done
+	/** Called as last phase of reachability analysis. Allow external systems to add UObject roots *after* first reachability pass has been done */
 	static FTraceExternalRootsForReachabilityAnalysisDelegate TraceExternalRootsForReachabilityAnalysis;
 
-	// Called after reachability analysis, before any purging
+	/** Called after reachability analysis, before any purging */
 	static FSimpleMulticastDelegate PostReachabilityAnalysis;
 
-	// Called after garbage collection
+	/** Called after garbage collection */
 	static FSimpleMulticastDelegate& GetPostGarbageCollect();
 
-	// Called before ConditionalBeginDestroy phase of garbage collection
+	/** Called before ConditionalBeginDestroy phase of garbage collection */
 	static FSimpleMulticastDelegate PreGarbageCollectConditionalBeginDestroy;
 
-	// Called after ConditionalBeginDestroy phase of garbage collection
+	/** Called after ConditionalBeginDestroy phase of garbage collection */
 	static FSimpleMulticastDelegate PostGarbageCollectConditionalBeginDestroy;
 
-	/** delegate type for querying whether a loaded object should replace an already existing one */
-	DECLARE_DELEGATE_RetVal_OneParam(bool, FOnLoadObjectsOnTop, const FString&);
-
 	/** Queries whether an object should be loaded on top ( replace ) an already existing one */
+	DECLARE_DELEGATE_RetVal_OneParam(bool, FOnLoadObjectsOnTop, const FString&);
 	static FOnLoadObjectsOnTop ShouldLoadOnTop;
 
-	/** called when path to world root is changed */
+	/** Called when path to world root is changed */
 	DECLARE_MULTICAST_DELEGATE_OneParam(FPackageCreatedForLoad, class UPackage*);
 	static FPackageCreatedForLoad PackageCreatedForLoad;
 
-	/** Called when trying to figure out if a UObject is a primary asset, if it doesn't know */
+	/** Called when trying to figure out if a UObject is a primary asset, if it doesn't implement GetPrimaryAssetId itself */
 	DECLARE_DELEGATE_RetVal_OneParam(FPrimaryAssetId, FGetPrimaryAssetIdForObject, const UObject*);
 	static FGetPrimaryAssetIdForObject GetPrimaryAssetIdForObject;
 
@@ -2441,3 +2414,11 @@ namespace UE4CodeGen_Private
 #else
 	#define IF_WITH_EDITORONLY_DATA(x, y) y
 #endif
+
+enum class EDataValidationResult : uint8
+{
+	Invalid,
+	Valid,
+
+	NotValidated
+};
diff --git a/Engine/Source/Runtime/Engine/Classes/Animation/AnimSequence.h b/Engine/Source/Runtime/Engine/Classes/Animation/AnimSequence.h
--- a/Engine/Source/Runtime/Engine/Classes/Animation/AnimSequence.h
+++ b/Engine/Source/Runtime/Engine/Classes/Animation/AnimSequence.h
@@ -574,7 +574,6 @@ public:
 	virtual void PostDuplicate(bool bDuplicateForPIE) override;
 #endif // WITH_EDITOR
 	virtual void BeginDestroy() override;
-	virtual void GetResourceSizeEx(FResourceSizeEx& CumulativeResourceSize) override;
 	virtual void GetAssetRegistryTags(TArray<FAssetRegistryTag>& OutTags) const override;
 	//~ End UObject Interface
 
diff --git a/Engine/Source/Runtime/Engine/Classes/Animation/AnimSet.h b/Engine/Source/Runtime/Engine/Classes/Animation/AnimSet.h
--- a/Engine/Source/Runtime/Engine/Classes/Animation/AnimSet.h
+++ b/Engine/Source/Runtime/Engine/Classes/Animation/AnimSet.h
@@ -118,7 +118,6 @@ public:
 
 	//~ Begin UObject Interface
 	virtual void PostLoad() override;
-	virtual void GetResourceSizeEx(FResourceSizeEx& CumulativeResourceSize) override;
 	//~ End UObject Interface
 	
 	//~ Begin UAnimSet Interface
diff --git a/Engine/Source/Runtime/Engine/Classes/Animation/MorphTarget.h b/Engine/Source/Runtime/Engine/Classes/Animation/MorphTarget.h
--- a/Engine/Source/Runtime/Engine/Classes/Animation/MorphTarget.h
+++ b/Engine/Source/Runtime/Engine/Classes/Animation/MorphTarget.h
@@ -64,10 +64,7 @@ struct FMorphTargetLODModel
 	int32 NumBaseMeshVerts;
 	
 	/** Get Resource Size */
-	DEPRECATED(4.14, "GetResourceSize is deprecated. Please use GetResourceSizeEx or GetResourceSizeBytes instead.")
-	SIZE_T GetResourceSize() const;
 	void GetResourceSizeEx(FResourceSizeEx& CumulativeResourceSize) const;
-	SIZE_T GetResourceSizeBytes() const;
 
 	FMorphTargetLODModel()
 		: NumBaseMeshVerts(0)
diff --git a/Engine/Source/Runtime/Engine/Classes/EdGraph/EdGraphPin.h b/Engine/Source/Runtime/Engine/Classes/EdGraph/EdGraphPin.h
--- a/Engine/Source/Runtime/Engine/Classes/EdGraph/EdGraphPin.h
+++ b/Engine/Source/Runtime/Engine/Classes/EdGraph/EdGraphPin.h
@@ -1,4 +1,4 @@
-﻿// Copyright 1998-2017 Epic Games, Inc. All Rights Reserved.
+// Copyright 1998-2017 Epic Games, Inc. All Rights Reserved.
 
 #pragma once
 
@@ -481,9 +481,12 @@ public:
 #endif // WITH_EDITORONLY_DATA
 	}
 
-	/** Get the current DefaultObject path name, or DefaultValue if its null */
+	/** Returns an internal string representation of the string/object/text default value, that is compatible with AutogeneratedDefaultValue and ImportText */
 	ENGINE_API FString GetDefaultAsString() const;
 
+	/** Returns a human readable FText representation of the string/object/text default value */
+	ENGINE_API FText GetDefaultAsText() const;
+
 	/** Returns true if the current default value matches the autogenerated default value */
 	ENGINE_API bool DoesDefaultValueMatchAutogenerated() const;
 
diff --git a/Engine/Source/Runtime/Engine/Classes/Engine/Blueprint.h b/Engine/Source/Runtime/Engine/Classes/Engine/Blueprint.h
--- a/Engine/Source/Runtime/Engine/Classes/Engine/Blueprint.h
+++ b/Engine/Source/Runtime/Engine/Classes/Engine/Blueprint.h
@@ -743,8 +743,15 @@ public:
 	/** Collect blueprints that depend on this blueprint. */
 	virtual void GatherDependencies(TSet<TWeakObjectPtr<UBlueprint>>& InDependencies) const;
 
+	/** Checks all nodes in all graphs to see if they should be replaced by other nodes */
 	virtual void ReplaceDeprecatedNodes();
 
+	/** Clears out any editor data regarding a blueprint class, this can be called when you want to unload a blueprint */
+	virtual void ClearEditorReferences();
+
+	/** Returns Valid if this object has data validation rules set up for it and the data for this object is valid. Returns Invalid if it does not pass the rules. Returns NotValidated if no rules are set for this object. */
+	virtual EDataValidationResult IsDataValid(TArray<FText>& ValidationErrors) override;
+
 #endif	//#if WITH_EDITOR
 
 	//~ Begin UObject Interface
diff --git a/Engine/Source/Runtime/Engine/Classes/Engine/BlueprintGeneratedClass.h b/Engine/Source/Runtime/Engine/Classes/Engine/BlueprintGeneratedClass.h
--- a/Engine/Source/Runtime/Engine/Classes/Engine/BlueprintGeneratedClass.h
+++ b/Engine/Source/Runtime/Engine/Classes/Engine/BlueprintGeneratedClass.h
@@ -58,12 +58,6 @@ public:
 	// Reverse map from code offset to source node
 	TMap< int32, TWeakObjectPtr<UEdGraphNode> > LineNumberToSourceNodeMap;
 
-	// Reverse map from code offset to macro source node
-	TMap< int32, TWeakObjectPtr<UEdGraphNode> > LineNumberToMacroSourceNodeMap;
-
-	// Reverse map from code offset to macro instance node(s)
-	TMultiMap< int32, TWeakObjectPtr<UEdGraphNode> > LineNumberToMacroInstanceNodeMap;
-
 	// Reverse map from code offset to source pin
 	TMap< int32, FEdGraphPinReference > LineNumberToSourcePinMap;
 
@@ -260,36 +254,6 @@ public:
 		return nullptr;
 	}
 
-	// Finds the macro source node associated with the code location Function+CodeOffset, or nullptr if there isn't one
-	UEdGraphNode* FindMacroSourceNodeFromCodeLocation(UFunction* Function, int32 CodeOffset) const
-	{
-		if (const FDebuggingInfoForSingleFunction* pFuncInfo = PerFunctionLineNumbers.Find(Function))
-		{
-			return pFuncInfo->LineNumberToMacroSourceNodeMap.FindRef(CodeOffset).Get();
-		}
-
-		return nullptr;
-	}
-
-	// Finds the macro instance node(s) associated with the code location Function+CodeOffset. The returned set can be empty.
-	void FindMacroInstanceNodesFromCodeLocation(UFunction* Function, int32 CodeOffset, TArray<UEdGraphNode*>& MacroInstanceNodes) const
-	{
-		MacroInstanceNodes.Empty();
-		if (const FDebuggingInfoForSingleFunction* pFuncInfo = PerFunctionLineNumbers.Find(Function))
-		{
-			TArray<TWeakObjectPtr<UEdGraphNode>> MacroInstanceNodePtrs;
-			pFuncInfo->LineNumberToMacroInstanceNodeMap.MultiFind(CodeOffset, MacroInstanceNodePtrs);
-			for(auto MacroInstanceNodePtrIt = MacroInstanceNodePtrs.CreateConstIterator(); MacroInstanceNodePtrIt; ++MacroInstanceNodePtrIt)
-			{
-				TWeakObjectPtr<UEdGraphNode> MacroInstanceNodePtr = *MacroInstanceNodePtrIt;
-				if(UEdGraphNode* MacroInstanceNode = MacroInstanceNodePtr.Get())
-				{
-					MacroInstanceNodes.AddUnique(MacroInstanceNode);
-				}
-			}
-		}
-	}
-
 	// Finds the source pin associated with the code location Function+CodeOffset, or nullptr if there isn't one
 	UEdGraphPin* FindSourcePinFromCodeLocation(UFunction* Function, int32 CodeOffset) const
 	{
@@ -343,13 +307,6 @@ public:
 					OutNodeToCodeAssociations.Add(CodeLocation.Key);
 				}
 			}
-			for (auto CodeLocation : pFuncInfo->LineNumberToMacroSourceNodeMap)
-			{
-				if (CodeLocation.Value == SourceNode)
-				{
-					OutNodeToCodeAssociations.Add(CodeLocation.Key);
-				}
-			}
 		}
 	}
 
@@ -416,33 +373,23 @@ public:
 	}
 
 	// Adds a debug record for a source node and destination in the bytecode of a specified function
-	void RegisterNodeToCodeAssociation(UEdGraphNode* TrueSourceNode, UEdGraphNode* MacroSourceNode, const TArray<TWeakObjectPtr<UEdGraphNode>>& MacroInstanceNodes, UFunction* InFunction, int32 CodeOffset, bool bBreakpointSite)
+	void RegisterNodeToCodeAssociation(UEdGraphNode* SourceNode, UFunction* InFunction, int32 CodeOffset, bool bBreakpointSite)
 	{
 		//@TODO: Nasty expansion behavior during compile time
 		if (bBreakpointSite)
 		{
-			new (DebugNodeLineNumbers) FNodeToCodeAssociation(TrueSourceNode, InFunction, CodeOffset);
-			DebugNodeIndexLookup.Add(TrueSourceNode, DebugNodeLineNumbers.Num() - 1);
+			new (DebugNodeLineNumbers) FNodeToCodeAssociation(SourceNode, InFunction, CodeOffset);
+			DebugNodeIndexLookup.Add(SourceNode, DebugNodeLineNumbers.Num() - 1);
 		}
 
 		FDebuggingInfoForSingleFunction& PerFuncInfo = PerFunctionLineNumbers.FindOrAdd(InFunction);
-		PerFuncInfo.LineNumberToSourceNodeMap.Add(CodeOffset, TrueSourceNode);
-
-		if (MacroSourceNode)
-		{
-			PerFuncInfo.LineNumberToMacroSourceNodeMap.Add(CodeOffset, MacroSourceNode);
-		}
-
-		for(auto MacroInstanceNodePtrIt = MacroInstanceNodes.CreateConstIterator(); MacroInstanceNodePtrIt; ++MacroInstanceNodePtrIt)
-		{
-			PerFuncInfo.LineNumberToMacroInstanceNodeMap.Add(CodeOffset, *MacroInstanceNodePtrIt);
-		}
+		PerFuncInfo.LineNumberToSourceNodeMap.Add(CodeOffset, SourceNode);
 	}
 
-	void RegisterPureNodeScriptCodeRange(UEdGraphNode* TrueSourceNode, UFunction* InFunction, FInt32Range InPureNodeScriptCodeRange)
+	void RegisterPureNodeScriptCodeRange(UEdGraphNode* SourceNode, UFunction* InFunction, FInt32Range InPureNodeScriptCodeRange)
 	{
 		FDebuggingInfoForSingleFunction& PerFuncInfo = PerFunctionLineNumbers.FindOrAdd(InFunction);
-		PerFuncInfo.PureNodeScriptCodeRangeMap.Add(TrueSourceNode, InPureNodeScriptCodeRange);
+		PerFuncInfo.PureNodeScriptCodeRangeMap.Add(SourceNode, InPureNodeScriptCodeRange);
 	}
 
 	void RegisterPinToCodeAssociation(UEdGraphPin const* SourcePin, UFunction* InFunction, int32 CodeOffset)
diff --git a/Engine/Source/Runtime/Engine/Classes/Engine/SkeletalMesh.h b/Engine/Source/Runtime/Engine/Classes/Engine/SkeletalMesh.h
--- a/Engine/Source/Runtime/Engine/Classes/Engine/SkeletalMesh.h
+++ b/Engine/Source/Runtime/Engine/Classes/Engine/SkeletalMesh.h
@@ -425,11 +425,6 @@ struct FClothingAssetData_Legacy
 #endif// #if WITH_APEX_CLOTHING
 	// serialization
 	friend FArchive& operator<<(FArchive& Ar, FClothingAssetData_Legacy& A);
-	// get resource size
-	DEPRECATED(4.14, "GetResourceSize is deprecated. Please use GetResourceSizeEx or GetResourceSizeBytes instead.")
-	SIZE_T GetResourceSize() const;
-	void GetResourceSizeEx(FResourceSizeEx& CumulativeResourceSize) const;
-	SIZE_T GetResourceSizeBytes() const;
 };
 
 //~ Begin Material Interface for USkeletalMesh - contains a material and a shadow casting flag
diff --git a/Engine/Source/Runtime/Engine/Classes/GameFramework/Actor.h b/Engine/Source/Runtime/Engine/Classes/GameFramework/Actor.h
--- a/Engine/Source/Runtime/Engine/Classes/GameFramework/Actor.h
+++ b/Engine/Source/Runtime/Engine/Classes/GameFramework/Actor.h
@@ -288,7 +288,8 @@ private:
 	UPROPERTY()
 	uint8 bActorEnableCollision:1;
 
-	/** Set when actor is about to be deleted. */
+	/** Set when actor is about to be deleted. Needs to be uproperty so included in transactions. */
+	UPROPERTY(Transient, DuplicateTransient)
 	uint8 bActorIsBeingDestroyed:1;
 
 	enum class EActorBeginPlayState : uint8
@@ -1933,6 +1934,11 @@ public:
 	 */
 	virtual bool CheckStillInWorld();
 
+#if WITH_EDITOR
+	/** Returns Valid if this object has data validation rules set up for it and the data for this object is valid. Returns Invalid if it does not pass the rules. Returns NotValidated if no rules are set for this object. */
+	virtual EDataValidationResult IsDataValid(TArray<FText>& ValidationErrors) override;
+#endif // WITH_EDITOR
+
 
 	//~=============================================================================
 	// Actor overlap tracking
diff --git a/Engine/Source/Runtime/Engine/Classes/PhysicsEngine/BodySetup.h b/Engine/Source/Runtime/Engine/Classes/PhysicsEngine/BodySetup.h
--- a/Engine/Source/Runtime/Engine/Classes/PhysicsEngine/BodySetup.h
+++ b/Engine/Source/Runtime/Engine/Classes/PhysicsEngine/BodySetup.h
@@ -61,10 +61,7 @@ struct FBodySetupUVInfo
 	}
 
 	/** Get resource size of UV info */
-	DEPRECATED(4.14, "GetResourceSize is deprecated. Please use GetResourceSizeEx or GetResourceSizeBytes instead.")
-	SIZE_T GetResourceSize() const;
 	void GetResourceSizeEx(FResourceSizeEx& CumulativeResourceSize) const;
-	SIZE_T GetResourceSizeBytes() const;
 
 	void FillFromTriMesh(const FTriMeshCollisionData& TriMeshCollisionData);
 };
diff --git a/Engine/Source/Runtime/Engine/Classes/PhysicsEngine/BodySetupEnums.h b/Engine/Source/Runtime/Engine/Classes/PhysicsEngine/BodySetupEnums.h
--- a/Engine/Source/Runtime/Engine/Classes/PhysicsEngine/BodySetupEnums.h
+++ b/Engine/Source/Runtime/Engine/Classes/PhysicsEngine/BodySetupEnums.h
@@ -41,3 +41,12 @@ namespace EBodyCollisionResponse
 		//BodyCollision_Custom UMETA(DisplayName="Custom")
 	};
 }
+
+
+/** Helpers to convert enum to string */
+namespace Lex
+{
+	ENGINE_API const TCHAR* ToString(ECollisionTraceFlag Enum);
+	ENGINE_API const TCHAR* ToString(EPhysicsType Enum);
+	ENGINE_API const TCHAR* ToString(EBodyCollisionResponse::Type Enum);
+}
diff --git a/Engine/Source/Runtime/Engine/Private/Actor.cpp b/Engine/Source/Runtime/Engine/Private/Actor.cpp
--- a/Engine/Source/Runtime/Engine/Private/Actor.cpp
+++ b/Engine/Source/Runtime/Engine/Private/Actor.cpp
@@ -169,15 +169,6 @@ bool AActor::CheckActorComponents()
 			UE_LOG(LogCheckComponents, Error, TEXT("Component is a template but I am not. Me = %s, Component = %s"), *this->GetFullName(), *Inner->GetFullName());
 			bResult = false;
 		}
-		UObject* Archetype = Inner->GetArchetype();
-		if (Archetype != Inner->GetClass()->GetDefaultObject())
-		{
-			if (Archetype != GetClass()->GetDefaultSubobjectByName(Inner->GetFName()))
-			{
-				UE_LOG(LogCheckComponents, Error, TEXT("Component archetype is not the CDO nor a default subobject of my class. Me = %s, Component = %s, Archetype = %s"), *this->GetFullName(), *Inner->GetFullName(), *Archetype->GetFullName());
-				bResult = false;
-			}
-		}
 	}
 	for (int32 Index = 0; Index < BlueprintCreatedComponents.Num(); Index++)
 	{
@@ -3142,6 +3133,7 @@ void AActor::BeginPlay()
 		{
 			Component->RegisterAllComponentTickFunctions(true);
 			Component->BeginPlay();
+			ensureMsgf(Component->HasBegunPlay(), TEXT("Failed to route BeginPlay (%s)"), *Component->GetFullName());
 		}
 		else
 		{
diff --git a/Engine/Source/Runtime/Engine/Private/ActorEditor.cpp b/Engine/Source/Runtime/Engine/Private/ActorEditor.cpp
--- a/Engine/Source/Runtime/Engine/Private/ActorEditor.cpp
+++ b/Engine/Source/Runtime/Engine/Private/ActorEditor.cpp
@@ -768,8 +768,8 @@ void AActor::CheckForDeprecated()
 			->AddToken(FTextToken::Create(FText::Format(LOCTEXT( "MapCheck_Message_ActorIsObselete_Deprecated", "{ActorName} : Obsolete and must be removed! (Class is deprecated)" ), Arguments) ))
 			->AddToken(FMapErrorToken::Create(FMapErrors::ActorIsObselete));
 	}
-
-	if ( GetClass()->HasAnyClassFlags(CLASS_Abstract) )
+	// don't check to see if this is an abstract class if this is the CDO
+	if ( !(GetFlags() & RF_ClassDefaultObject) && GetClass()->HasAnyClassFlags(CLASS_Abstract) )
 	{
 		FFormatNamedArguments Arguments;
 		Arguments.Add(TEXT("ActorName"), FText::FromString(GetName()));
@@ -782,24 +782,10 @@ void AActor::CheckForDeprecated()
 
 void AActor::CheckForErrors()
 {
-	if ( GetClass()->HasAnyClassFlags(CLASS_Deprecated) )
+	int32 OldNumWarnings = FMessageLog("MapCheck").NumMessages(EMessageSeverity::Warning);
+	CheckForDeprecated();
+	if (OldNumWarnings < FMessageLog("MapCheck").NumMessages(EMessageSeverity::Warning))
 	{
-		FFormatNamedArguments Arguments;
-		Arguments.Add(TEXT("ActorName"), FText::FromString(GetName()));
-		FMessageLog("MapCheck").Warning()
-			->AddToken(FUObjectToken::Create(this))
-			->AddToken(FTextToken::Create(FText::Format(LOCTEXT( "MapCheck_Message_ActorIsObselete_Deprecated", "{ActorName} : Obsolete and must be removed! (Class is deprecated)" ), Arguments) ))
-			->AddToken(FMapErrorToken::Create(FMapErrors::ActorIsObselete));
-		return;
-	}
-	if ( GetClass()->HasAnyClassFlags(CLASS_Abstract) )
-	{
-		FFormatNamedArguments Arguments;
-		Arguments.Add(TEXT("ActorName"), FText::FromString(GetName()));
-		FMessageLog("MapCheck").Warning()
-			->AddToken(FUObjectToken::Create(this))
-			->AddToken(FTextToken::Create(FText::Format(LOCTEXT( "MapCheck_Message_ActorIsObselete_Abstract", "{ActorName} : Obsolete and must be removed! (Class is abstract)" ), Arguments) ))
-			->AddToken(FMapErrorToken::Create(FMapErrors::ActorIsObselete));
 		return;
 	}
 
@@ -865,6 +851,25 @@ void AActor::SetLODParent(UPrimitiveComponent* InLODParent, float InParentDrawDi
 		Component->SetLODParentPrimitive(InLODParent);
 	}
 }
+
+EDataValidationResult AActor::IsDataValid(TArray<FText>& ValidationErrors)
+{
+	bool bSuccess = CheckDefaultSubobjectsInternal();
+
+	int32 OldNumMapWarningsAndErrors = FMessageLog("MapCheck").NumMessages(EMessageSeverity::Warning);
+	CheckForErrors();
+	int32 NewNumMapWarningsAndErrors = FMessageLog("MapCheck").NumMessages(EMessageSeverity::Warning);
+	if (NewNumMapWarningsAndErrors != OldNumMapWarningsAndErrors)
+	{
+		FFormatNamedArguments Arguments;
+		Arguments.Add(TEXT("ActorName"), FText::FromString(GetName()));
+		FText ErrorMsg = FText::Format(LOCTEXT("IsDataValid_Failed_CheckForErrors", "{ActorName} is not valid. See the MapCheck log messages for details."), Arguments);
+		ValidationErrors.Add(ErrorMsg);
+		bSuccess = false;
+	}
+
+	return bSuccess ? EDataValidationResult::Valid : EDataValidationResult::Invalid;
+}
 #undef LOCTEXT_NAMESPACE
 
 #endif // WITH_EDITOR
diff --git a/Engine/Source/Runtime/Engine/Private/Animation/AnimSequence.cpp b/Engine/Source/Runtime/Engine/Private/Animation/AnimSequence.cpp
--- a/Engine/Source/Runtime/Engine/Private/Animation/AnimSequence.cpp
+++ b/Engine/Source/Runtime/Engine/Private/Animation/AnimSequence.cpp
@@ -149,20 +149,6 @@ void UAnimSequence::PostInitProperties()
 	Super::PostInitProperties();
 }
 
-void UAnimSequence::GetResourceSizeEx(FResourceSizeEx& CumulativeResourceSize)
-{
-	Super::GetResourceSizeEx(CumulativeResourceSize);
-
-	if (CumulativeResourceSize.GetResourceSizeMode() == EResourceSizeMode::Exclusive)
-	{
-		// All of the sequence data is serialized and will be counted as part of the direct object size rather than as a resource
-	}
-	else
-	{
-		CumulativeResourceSize.AddDedicatedSystemMemoryBytes((CompressedTrackOffsets.Num() == 0) ? GetApproxRawSize() : GetApproxCompressedSize());
-	}
-}
-
 void UAnimSequence::GetAssetRegistryTags(TArray<FAssetRegistryTag>& OutTags) const
 {
 #if WITH_EDITORONLY_DATA
@@ -515,12 +501,6 @@ void UAnimSequence::PostLoad()
 	bWasCompressedWithoutTranslations = false; //@todoanim: @fixmelh : AnimRotationOnly - GetAnimSet()->bAnimRotationOnly;
 #endif // WITH_EDITORONLY_DATA
 
-	if( IsRunningGame() )
-	{
-		// this probably will not show newly created animations in PIE but will show them in the game once they have been saved off
-		INC_DWORD_STAT_BY( STAT_AnimationMemory, GetResourceSizeBytes(EResourceSizeMode::Exclusive) );
-	}
-
 	{
 		LOG_SCOPE_VERBOSITY_OVERRIDE(LogAnimation, ELogVerbosity::Warning);
  		// convert animnotifies
@@ -692,11 +672,6 @@ void UAnimSequence::BeginDestroy()
 	// clear any active codec links
 	RotationCodec = NULL;
 	TranslationCodec = NULL;
-
-	if( IsRunningGame() )
-	{
-		DEC_DWORD_STAT_BY( STAT_AnimationMemory, GetResourceSizeBytes(EResourceSizeMode::Exclusive) );
-	}
 }
 
 #if WITH_EDITOR
@@ -5587,7 +5562,7 @@ void GatherAnimSequenceStats(FOutputDevice& Ar)
 			NumTransTracksWithOneKey, NumRotTracksWithOneKey, NumScaleTracksWithOneKey,
 			TotalNumTransKeys, TotalNumRotKeys, TotalNumScaleKeys,
 			*FAnimationUtils::GetAnimationKeyFormatString(static_cast<AnimationKeyFormat>(Seq->KeyEncodingFormat)),
-			(int32)Seq->GetResourceSizeBytes(EResourceSizeMode::Exclusive) );
+			(int32)Seq->GetResourceSizeBytes(EResourceSizeMode::EstimatedTotal) );
 	}
 	Ar.Logf( TEXT("======================================================================") );
 	Ar.Logf( TEXT("Total Num Tracks: %i trans, %i rot, %i scale, %i trans1, %i rot1, %i scale1"), GlobalNumTransTracks, GlobalNumRotTracks, GlobalNumScaleTracks, GlobalNumTransTracksWithOneKey, GlobalNumRotTracksWithOneKey, GlobalNumScaleTracksWithOneKey  );
diff --git a/Engine/Source/Runtime/Engine/Private/Animation/AnimSet.cpp b/Engine/Source/Runtime/Engine/Private/Animation/AnimSet.cpp
--- a/Engine/Source/Runtime/Engine/Private/Animation/AnimSet.cpp
+++ b/Engine/Source/Runtime/Engine/Private/Animation/AnimSet.cpp
@@ -127,31 +127,6 @@ int32 UAnimSet::GetMeshLinkupIndex(USkeletalMesh* SkelMesh)
 	return (*IndexPtr);
 }
 
-
-void UAnimSet::GetResourceSizeEx(FResourceSizeEx& CumulativeResourceSize)
-{
-	Super::GetResourceSizeEx(CumulativeResourceSize);
-
-	if (CumulativeResourceSize.GetResourceSizeMode() == EResourceSizeMode::Exclusive)
-	{
-		// This object only references others, it doesn't have any real resource bytes
-	}
-	else
-	{
-#if WITH_EDITORONLY_DATA
-		for( int32 i=0; i<Sequences.Num(); i++ )
-		{
-			UAnimSequence* AnimSeq = Sequences[i];
-			if( AnimSeq )
-			{
-				AnimSeq->GetResourceSizeEx(CumulativeResourceSize);
-			}			
-		}
-#endif	//#if WITH_EDITORONLY_DATA
-	}
-}
-
-
 void UAnimSet::ResetAnimSet()
 {
 #if WITH_EDITORONLY_DATA
diff --git a/Engine/Source/Runtime/Engine/Private/Animation/AnimationUtils.cpp b/Engine/Source/Runtime/Engine/Private/Animation/AnimationUtils.cpp
--- a/Engine/Source/Runtime/Engine/Private/Animation/AnimationUtils.cpp
+++ b/Engine/Source/Runtime/Engine/Private/Animation/AnimationUtils.cpp
@@ -1318,21 +1318,6 @@ void FAnimationUtils::CompressAnimSequenceExplicit(
 
 						TotalNoWinnerRounds++;
 					}
-
-					// Update the memory stats
-#if STATS
-					if( IsRunningGame() )
-					{
-						if (SizeDecrease > 0)
-						{
-							DEC_DWORD_STAT_BY( STAT_AnimationMemory, SizeDecrease );
-						}
-						else
-						{
-							INC_DWORD_STAT_BY( STAT_AnimationMemory, -SizeDecrease );
-						}
-					}
-#endif
 				}
 
 				// Make sure we got that right.
diff --git a/Engine/Source/Runtime/Engine/Private/AssetManager.cpp b/Engine/Source/Runtime/Engine/Private/AssetManager.cpp
--- a/Engine/Source/Runtime/Engine/Private/AssetManager.cpp
+++ b/Engine/Source/Runtime/Engine/Private/AssetManager.cpp
@@ -544,7 +544,8 @@ void UAssetManager::UpdateCachedAssetData(const FPrimaryAssetId& PrimaryAssetId,
 		if (OldData && OldData->AssetPtr.ToSoftObjectPath() != NewAssetPath)
 		{
 			UE_LOG(LogAssetManager, Warning, TEXT("Found Duplicate PrimaryAssetID %s, this must be resolved before saving. Path %s is replacing path %s"), *PrimaryAssetId.ToString(), *OldData->AssetPtr.ToString(), *NewAssetPath.ToString());
-			if (!bAllowDuplicates)
+			// Don't ensure for editor only types, this will not cause an actual game problem
+			if (!bAllowDuplicates && !TypeData.Info.bIsEditorOnly)
 			{
 				ensureMsgf(!OldData, TEXT("Found Duplicate PrimaryAssetID %s! Path %s is replacing path %s"), *PrimaryAssetId.ToString(), *OldData->AssetPtr.ToString(), *NewAssetPath.ToString());
 			}
diff --git a/Engine/Source/Runtime/Engine/Private/AudioDevice.cpp b/Engine/Source/Runtime/Engine/Private/AudioDevice.cpp
--- a/Engine/Source/Runtime/Engine/Private/AudioDevice.cpp
+++ b/Engine/Source/Runtime/Engine/Private/AudioDevice.cpp
@@ -4787,7 +4787,7 @@ void FAudioDevice::DumpActiveSounds() const
 			for (const TPair<UPTRINT, FWaveInstance*>& WaveInstancePair : ActiveSound->WaveInstances)
 			{
 				const FWaveInstance* WaveInstance = WaveInstancePair.Value;
-				UE_LOG(LogAudio, Display, TEXT("   %s (%.3g) (%d) - %.3g"), *WaveInstance->GetName(), WaveInstance->WaveData->GetDuration(), WaveInstance->WaveData->GetResourceSizeBytes(EResourceSizeMode::Inclusive), WaveInstance->GetActualVolume());
+				UE_LOG(LogAudio, Display, TEXT("   %s (%.3g) (%d) - %.3g"), *WaveInstance->GetName(), WaveInstance->WaveData->GetDuration(), WaveInstance->WaveData->GetResourceSizeBytes(EResourceSizeMode::EstimatedTotal), WaveInstance->GetActualVolume());
 			}
 		}
 	}
diff --git a/Engine/Source/Runtime/Engine/Private/Blueprint.cpp b/Engine/Source/Runtime/Engine/Private/Blueprint.cpp
--- a/Engine/Source/Runtime/Engine/Private/Blueprint.cpp
+++ b/Engine/Source/Runtime/Engine/Private/Blueprint.cpp
@@ -1767,6 +1767,11 @@ void UBlueprint::ReplaceDeprecatedNodes()
 	}
 }
 
+void UBlueprint::ClearEditorReferences()
+{
+	FKismetEditorUtilities::OnBlueprintUnloaded.Broadcast(this);
+}
+
 UInheritableComponentHandler* UBlueprint::GetInheritableComponentHandler(bool bCreateIfNecessary)
 {
 	static const FBoolConfigValueHelper EnableInheritableComponents(TEXT("Kismet"), TEXT("bEnableInheritableComponents"), GEngineIni);
@@ -1784,9 +1789,12 @@ UInheritableComponentHandler* UBlueprint::GetInheritableComponentHandler(bool bC
 	return InheritableComponentHandler;
 }
 
-#endif
 
-#if WITH_EDITOR
+EDataValidationResult UBlueprint::IsDataValid(TArray<FText>& ValidationErrors)
+{
+	return GeneratedClass ? GeneratedClass->GetDefaultObject()->IsDataValid(ValidationErrors) : EDataValidationResult::Invalid;
+}
+
 FName UBlueprint::GetFunctionNameFromClassByGuid(const UClass* InClass, const FGuid FunctionGuid)
 {
 	return FBlueprintEditorUtils::GetFunctionNameFromClassByGuid(InClass, FunctionGuid);
diff --git a/Engine/Source/Runtime/Engine/Private/Camera/CameraAnim.cpp b/Engine/Source/Runtime/Engine/Private/Camera/CameraAnim.cpp
--- a/Engine/Source/Runtime/Engine/Private/Camera/CameraAnim.cpp
+++ b/Engine/Source/Runtime/Engine/Private/Camera/CameraAnim.cpp
@@ -167,24 +167,3 @@ void UCameraAnim::CalcLocalAABB()
 		}
 	}
 }
-
-void UCameraAnim::GetResourceSizeEx(FResourceSizeEx& CumulativeResourceSize)
-{
-	Super::GetResourceSizeEx(CumulativeResourceSize);
-
-	if (CumulativeResourceSize.GetResourceSizeMode() == EResourceSizeMode::Inclusive && CameraInterpGroup)
-	{
-		// find move track
-		UInterpTrackMove *MoveTrack = NULL;
-		for (int32 TrackIdx = 0; TrackIdx < CameraInterpGroup->InterpTracks.Num(); ++TrackIdx)
-		{
-			// somehow movement track's not calculated when you just used serialize, so I'm adding it here. 
-			MoveTrack = Cast<UInterpTrackMove>(CameraInterpGroup->InterpTracks[TrackIdx]);
-			if (MoveTrack)
-			{
-				FArchiveCountMem CountBytesSize(MoveTrack);
-				CumulativeResourceSize.AddDedicatedSystemMemoryBytes(CountBytesSize.GetNum());
-			}
-		}
-	}
-}
diff --git a/Engine/Source/Runtime/Engine/Private/Components/ActorComponent.cpp b/Engine/Source/Runtime/Engine/Private/Components/ActorComponent.cpp
--- a/Engine/Source/Runtime/Engine/Private/Components/ActorComponent.cpp
+++ b/Engine/Source/Runtime/Engine/Private/Components/ActorComponent.cpp
@@ -1001,6 +1001,7 @@ void UActorComponent::RegisterComponentWithWorld(UWorld* InWorld)
 			if (!bHasBegunPlay)
 			{
 				BeginPlay();
+				ensureMsgf(bHasBegunPlay, TEXT("Failed to route BeginPlay (%s)"), *GetFullName());
 			}
 		}
 	}
diff --git a/Engine/Source/Runtime/Engine/Private/Components/BrushComponent.cpp b/Engine/Source/Runtime/Engine/Private/Components/BrushComponent.cpp
--- a/Engine/Source/Runtime/Engine/Private/Components/BrushComponent.cpp
+++ b/Engine/Source/Runtime/Engine/Private/Components/BrushComponent.cpp
@@ -528,17 +528,6 @@ void UBrushComponent::PostLoad()
 #endif
 }
 
-void UBrushComponent::GetResourceSizeEx(FResourceSizeEx& CumulativeResourceSize)
-{
-	Super::GetResourceSizeEx(CumulativeResourceSize);
-
-	// Count the bodysetup we own as well for 'inclusive' stats
-	if((CumulativeResourceSize.GetResourceSizeMode() == EResourceSizeMode::Inclusive) && (BrushBodySetup != NULL))
-	{
-		BrushBodySetup->GetResourceSizeEx(CumulativeResourceSize);
-	}
-}
-
 uint8 UBrushComponent::GetStaticDepthPriorityGroup() const
 {
 	ABrush* BrushOwner = Cast<ABrush>(GetOwner());
diff --git a/Engine/Source/Runtime/Engine/Private/Components/ModelComponent.cpp b/Engine/Source/Runtime/Engine/Private/Components/ModelComponent.cpp
--- a/Engine/Source/Runtime/Engine/Private/Components/ModelComponent.cpp
+++ b/Engine/Source/Runtime/Engine/Private/Components/ModelComponent.cpp
@@ -320,18 +320,6 @@ void UModelComponent::PostEditUndo()
 }
 #endif // WITH_EDITOR
 
-void UModelComponent::GetResourceSizeEx(FResourceSizeEx& CumulativeResourceSize)
-{
-	Super::GetResourceSizeEx(CumulativeResourceSize);
-
-	// Count the bodysetup we own as well for 'inclusive' stats
-	if((CumulativeResourceSize.GetResourceSizeMode() == EResourceSizeMode::Inclusive) && (ModelBodySetup != NULL))
-	{
-		ModelBodySetup->GetResourceSizeEx(CumulativeResourceSize);
-	}
-}
-
-
 bool UModelComponent::IsNameStableForNetworking() const
 {
 	// UModelComponent is always persistent for the duration of a game session, and so can be considered to have a stable name
diff --git a/Engine/Source/Runtime/Engine/Private/EdGraph/EdGraphPin.cpp b/Engine/Source/Runtime/Engine/Private/EdGraph/EdGraphPin.cpp
--- a/Engine/Source/Runtime/Engine/Private/EdGraph/EdGraphPin.cpp
+++ b/Engine/Source/Runtime/Engine/Private/EdGraph/EdGraphPin.cpp
@@ -1,4 +1,4 @@
-﻿// Copyright 1998-2017 Epic Games, Inc. All Rights Reserved.
+// Copyright 1998-2017 Epic Games, Inc. All Rights Reserved.
 
 #include "EdGraph/EdGraphPin.h"
 #include "UObject/BlueprintsObjectVersion.h"
@@ -603,7 +603,9 @@ FString UEdGraphPin::GetDefaultAsString() const
 	}
 	else if(!DefaultTextValue.IsEmpty())
 	{
-		return DefaultTextValue.ToString();
+		FString TextAsString;
+		FTextStringHelper::WriteToString(TextAsString, DefaultTextValue);
+		return TextAsString;
 	}
 	else
 	{
@@ -611,6 +613,23 @@ FString UEdGraphPin::GetDefaultAsString() const
 	}
 }
 
+FText UEdGraphPin::GetDefaultAsText() const
+{
+	if (DefaultObject)
+	{
+		return FText::AsCultureInvariant(DefaultObject->GetPathName());
+	}
+	else if (!DefaultTextValue.IsEmpty())
+	{
+		return DefaultTextValue;
+	}
+	else if (!DefaultValue.IsEmpty())
+	{
+		return FText::AsCultureInvariant(DefaultValue);
+	}
+	return FText::GetEmpty();
+}
+
 bool UEdGraphPin::DoesDefaultValueMatchAutogenerated() const
 {
 	if (const UEdGraphSchema* Schema = GetSchema())
diff --git a/Engine/Source/Runtime/Engine/Private/EngineUtils.cpp b/Engine/Source/Runtime/Engine/Private/EngineUtils.cpp
--- a/Engine/Source/Runtime/Engine/Private/EngineUtils.cpp
+++ b/Engine/Source/Runtime/Engine/Private/EngineUtils.cpp
@@ -123,7 +123,7 @@ bool FContentComparisonHelper::CompareClasses(const FString& InBaseClassName, co
 								int32 NewIndex = AssetList->AddZeroed();
 								FContentComparisonAssetInfo& Info = (*AssetList)[NewIndex];
 								Info.AssetName = Object->GetFullName();
-								Info.ResourceSize = Object->GetResourceSizeBytes(EResourceSizeMode::Inclusive);
+								Info.ResourceSize = Object->GetResourceSizeBytes(EResourceSizeMode::EstimatedTotal);
 							}
 						}
 					}
diff --git a/Engine/Source/Runtime/Engine/Private/Font.cpp b/Engine/Source/Runtime/Engine/Private/Font.cpp
--- a/Engine/Source/Runtime/Engine/Private/Font.cpp
+++ b/Engine/Source/Runtime/Engine/Private/Font.cpp
@@ -358,55 +358,32 @@ void UFont::GetResourceSizeEx(FResourceSizeEx& CumulativeResourceSize)
 {
 	Super::GetResourceSizeEx(CumulativeResourceSize);
 
-	switch(FontCacheType)
-	{
-	case EFontCacheType::Offline:
-		{
-			for (UTexture2D* Texture : Textures)
-			{
-				if (Texture)
-				{
-					Texture->GetResourceSizeEx(CumulativeResourceSize);
-				}
-			}
-		}
-		break;
+	// Nested Textures/FontFaces get included with the general subobject inclusion for EstimatedTotal
+	// We do need to handle nested lazy load data
 
-	case EFontCacheType::Runtime:
+	if (FontCacheType == EFontCacheType::Runtime && CumulativeResourceSize.GetResourceSizeMode() == EResourceSizeMode::EstimatedTotal)
+	{
+		auto GetTypefaceResourceSize = [&CumulativeResourceSize](const FTypeface& Typeface)
 		{
-			auto GetTypefaceResourceSize = [&CumulativeResourceSize](const FTypeface& Typeface)
+			for (const FTypefaceEntry& TypefaceEntry : Typeface.Fonts)
 			{
-				for (const FTypefaceEntry& TypefaceEntry : Typeface.Fonts)
+				const UFontFace* FontFace = Cast<const UFontFace>(TypefaceEntry.Font.GetFontFaceAsset());
+				if (!FontFace && TypefaceEntry.Font.GetLoadingPolicy() == EFontLoadingPolicy::LazyLoad)
 				{
-					const UFontFace* FontFace = Cast<const UFontFace>(TypefaceEntry.Font.GetFontFaceAsset());
-					if (FontFace)
+					const int64 FileSize = IFileManager::Get().FileSize(*TypefaceEntry.Font.GetFontFilename());
+					if (FileSize > 0)
 					{
-						const_cast<UFontFace*>(FontFace)->GetResourceSizeEx(CumulativeResourceSize);
+						CumulativeResourceSize.AddDedicatedSystemMemoryBytes(FileSize);
 					}
-					else if (TypefaceEntry.Font.GetLoadingPolicy() == EFontLoadingPolicy::LazyLoad)
-					{
-						const int64 FileSize = IFileManager::Get().FileSize(*TypefaceEntry.Font.GetFontFilename());
-						if (FileSize > 0)
-						{
-							CumulativeResourceSize.AddDedicatedSystemMemoryBytes(FileSize);
-						}
-					}
-				}
-			};
-
-			if (CumulativeResourceSize.GetResourceSizeMode() == EResourceSizeMode::Inclusive)
-			{
-				// Sum the contained font data sizes
-				GetTypefaceResourceSize(CompositeFont.DefaultTypeface);
-				for (const FCompositeSubFont& SubTypeface : CompositeFont.SubTypefaces)
-				{
-					GetTypefaceResourceSize(SubTypeface.Typeface);
 				}
 			}
-		}
-		break;
+		};
 
-	default:
-		break;
+		// Sum the contained font data sizes
+		GetTypefaceResourceSize(CompositeFont.DefaultTypeface);
+		for (const FCompositeSubFont& SubTypeface : CompositeFont.SubTypefaces)
+		{
+			GetTypefaceResourceSize(SubTypeface.Typeface);
+		}
 	}
 }
diff --git a/Engine/Source/Runtime/Engine/Private/InstancedStaticMesh.cpp b/Engine/Source/Runtime/Engine/Private/InstancedStaticMesh.cpp
--- a/Engine/Source/Runtime/Engine/Private/InstancedStaticMesh.cpp
+++ b/Engine/Source/Runtime/Engine/Private/InstancedStaticMesh.cpp
@@ -2121,10 +2121,6 @@ void UInstancedStaticMeshComponent::GetResourceSizeEx(FResourceSizeEx& Cumulativ
 	}
 	CumulativeResourceSize.AddDedicatedSystemMemoryBytes(InstanceReorderTable.GetAllocatedSize());
 	CumulativeResourceSize.AddDedicatedSystemMemoryBytes(RemovedInstances.GetAllocatedSize());
-
-#if WITH_EDITOR
-	CumulativeResourceSize.AddDedicatedSystemMemoryBytes(SelectedInstances.GetAllocatedSize());
-#endif
 }
 
 void UInstancedStaticMeshComponent::BeginDestroy()
diff --git a/Engine/Source/Runtime/Engine/Private/LevelStreaming.cpp b/Engine/Source/Runtime/Engine/Private/LevelStreaming.cpp
--- a/Engine/Source/Runtime/Engine/Private/LevelStreaming.cpp
+++ b/Engine/Source/Runtime/Engine/Private/LevelStreaming.cpp
@@ -1111,7 +1111,7 @@ ULevelStreamingKismet* ULevelStreamingKismet::LoadLevelInstanceBySoftObjectPtr(U
 		return nullptr;
 	}
 
-	return LoadLevelInstance_Internal(World, Level.ToString(), Location, Rotation, bOutSuccess);
+	return LoadLevelInstance_Internal(World, Level.GetLongPackageName(), Location, Rotation, bOutSuccess);
 }
 
 ULevelStreamingKismet* ULevelStreamingKismet::LoadLevelInstance_Internal(UWorld* World, const FString& LongPackageName, const FVector Location, const FRotator Rotation, bool& bOutSuccess)
diff --git a/Engine/Source/Runtime/Engine/Private/Materials/Material.cpp b/Engine/Source/Runtime/Engine/Private/Materials/Material.cpp
--- a/Engine/Source/Runtime/Engine/Private/Materials/Material.cpp
+++ b/Engine/Source/Runtime/Engine/Private/Materials/Material.cpp
@@ -4448,32 +4448,15 @@ void UMaterial::GetResourceSizeEx(FResourceSizeEx& CumulativeResourceSize)
 		}
 	}
 
-	if (CumulativeResourceSize.GetResourceSizeMode() == EResourceSizeMode::Inclusive)
+	for (int32 QualityLevelIndex = 0; QualityLevelIndex < EMaterialQualityLevel::Num; QualityLevelIndex++)
 	{
-		for (int32 QualityLevelIndex = 0; QualityLevelIndex < EMaterialQualityLevel::Num; QualityLevelIndex++)
+		for (int32 FeatureLevelIndex = 0; FeatureLevelIndex < ERHIFeatureLevel::Num; FeatureLevelIndex++)
 		{
-			for (int32 FeatureLevelIndex = 0; FeatureLevelIndex < ERHIFeatureLevel::Num; FeatureLevelIndex++)
+			if (FMaterialResource* CurrentResource = MaterialResources[QualityLevelIndex][FeatureLevelIndex])
 			{
-				FMaterialResource* CurrentResource = MaterialResources[QualityLevelIndex][FeatureLevelIndex];
 				CurrentResource->GetResourceSizeEx(CumulativeResourceSize);
 			}
 		}
-
-		TArray<UTexture*> TheReferencedTextures;
-		for ( int32 ExpressionIndex= 0 ; ExpressionIndex < Expressions.Num() ; ++ExpressionIndex )
-		{
-			UMaterialExpressionTextureSample* TextureSample = Cast<UMaterialExpressionTextureSample>( Expressions[ExpressionIndex] );
-			if ( TextureSample && TextureSample->Texture )
-			{
-				UTexture* Texture						= TextureSample->Texture;
-				const bool bTextureAlreadyConsidered	= TheReferencedTextures.Contains( Texture );
-				if ( !bTextureAlreadyConsidered )
-				{
-					TheReferencedTextures.Add( Texture );
-					Texture->GetResourceSizeEx(CumulativeResourceSize);
-				}
-			}
-		}
 	}
 }
 
diff --git a/Engine/Source/Runtime/Engine/Private/Materials/MaterialInstance.cpp b/Engine/Source/Runtime/Engine/Private/Materials/MaterialInstance.cpp
--- a/Engine/Source/Runtime/Engine/Private/Materials/MaterialInstance.cpp
+++ b/Engine/Source/Runtime/Engine/Private/Materials/MaterialInstance.cpp
@@ -3483,15 +3483,12 @@ void UMaterialInstance::GetResourceSizeEx(FResourceSizeEx& CumulativeResourceSiz
 
 	if (bHasStaticPermutationResource)
 	{
-		if (CumulativeResourceSize.GetResourceSizeMode() == EResourceSizeMode::Inclusive)
+		for (int32 QualityLevelIndex = 0; QualityLevelIndex < EMaterialQualityLevel::Num; QualityLevelIndex++)
 		{
-			for (int32 QualityLevelIndex = 0; QualityLevelIndex < EMaterialQualityLevel::Num; QualityLevelIndex++)
+			for (int32 FeatureLevelIndex = 0; FeatureLevelIndex < ERHIFeatureLevel::Num; FeatureLevelIndex++)
 			{
-				for (int32 FeatureLevelIndex = 0; FeatureLevelIndex < ERHIFeatureLevel::Num; FeatureLevelIndex++)
-				{
-					FMaterialResource* CurrentResource = StaticPermutationMaterialResources[QualityLevelIndex][FeatureLevelIndex];
-					CurrentResource->GetResourceSizeEx(CumulativeResourceSize);
-				}
+				FMaterialResource* CurrentResource = StaticPermutationMaterialResources[QualityLevelIndex][FeatureLevelIndex];
+				CurrentResource->GetResourceSizeEx(CumulativeResourceSize);
 			}
 		}
 	}
diff --git a/Engine/Source/Runtime/Engine/Private/Materials/MaterialShared.cpp b/Engine/Source/Runtime/Engine/Private/Materials/MaterialShared.cpp
--- a/Engine/Source/Runtime/Engine/Private/Materials/MaterialShared.cpp
+++ b/Engine/Source/Runtime/Engine/Private/Materials/MaterialShared.cpp
@@ -1287,60 +1287,55 @@ void FMaterialResource::GetRepresentativeShaderTypesAndDescriptions(TMap<FName,
 	}
 }
 
-SIZE_T FMaterialResource::GetResourceSizeInclusive()
-{
-	FResourceSizeEx ResSize = FResourceSizeEx(EResourceSizeMode::Inclusive);
-	GetResourceSizeEx(ResSize);
-	return ResSize.GetTotalMemoryBytes();
-}
-
 void FMaterialResource::GetResourceSizeEx(FResourceSizeEx& CumulativeResourceSize)
 {
-	if (CumulativeResourceSize.GetResourceSizeMode() == EResourceSizeMode::Inclusive)
-	{
-		TSet<const FMaterialShaderMap*> UniqueShaderMaps;
-		TMap<FShaderId, FShader*> UniqueShaders;
-		TArray<FShaderPipeline*> ShaderPipelines;
-		TSet<FShaderResourceId> UniqueShaderResourceIds;
+	TSet<const FMaterialShaderMap*> UniqueShaderMaps;
+	TMap<FShaderId, FShader*> UniqueShaders;
+	TArray<FShaderPipeline*> ShaderPipelines;
+	TSet<FShaderResourceId> UniqueShaderResourceIds;
 
-		CumulativeResourceSize.AddDedicatedSystemMemoryBytes(sizeof(FMaterialResource));
-		UniqueShaderMaps.Add(GetGameThreadShaderMap());
+	CumulativeResourceSize.AddDedicatedSystemMemoryBytes(sizeof(FMaterialResource));
+	UniqueShaderMaps.Add(GetGameThreadShaderMap());
 
-		for (TSet<const FMaterialShaderMap*>::TConstIterator It(UniqueShaderMaps); It; ++It)
+	for (TSet<const FMaterialShaderMap*>::TConstIterator It(UniqueShaderMaps); It; ++It)
+	{
+		const FMaterialShaderMap* MaterialShaderMap = *It;
+		if (MaterialShaderMap)
 		{
-			const FMaterialShaderMap* MaterialShaderMap = *It;
-			if (MaterialShaderMap)
+			CumulativeResourceSize.AddDedicatedSystemMemoryBytes(MaterialShaderMap->GetSizeBytes());
+
+			// Shaders are shared, so only count them in total mode
+			if (CumulativeResourceSize.GetResourceSizeMode() == EResourceSizeMode::EstimatedTotal)
 			{
-				CumulativeResourceSize.AddDedicatedSystemMemoryBytes(MaterialShaderMap->GetSizeBytes());
 				MaterialShaderMap->GetShaderList(UniqueShaders);
 				MaterialShaderMap->GetShaderPipelineList(ShaderPipelines);
 			}
 		}
+	}
 
-		for (auto& KeyValue : UniqueShaders)
+	for (auto& KeyValue : UniqueShaders)
+	{
+		FShader* Shader = KeyValue.Value;
+		if (Shader)
 		{
-			auto* Shader = KeyValue.Value;
-			if (Shader)
-			{
-				CumulativeResourceSize.AddDedicatedSystemMemoryBytes(AddShaderSize(Shader, UniqueShaderResourceIds));
-			}
+			CumulativeResourceSize.AddDedicatedSystemMemoryBytes(AddShaderSize(Shader, UniqueShaderResourceIds));
 		}
+	}
 
-		for (FShaderPipeline* Pipeline : ShaderPipelines)
+	for (FShaderPipeline* Pipeline : ShaderPipelines)
+	{
+		if (Pipeline)
 		{
-			if (Pipeline)
+			for (FShader* Shader : Pipeline->GetShaders())
 			{
-				for (FShader* Shader : Pipeline->GetShaders())
+				if (Shader)
 				{
-					if (Shader)
-					{
-						CumulativeResourceSize.AddDedicatedSystemMemoryBytes(AddShaderSize(Shader, UniqueShaderResourceIds));
-					}
+					CumulativeResourceSize.AddDedicatedSystemMemoryBytes(AddShaderSize(Shader, UniqueShaderResourceIds));
 				}
-				CumulativeResourceSize.AddDedicatedSystemMemoryBytes(Pipeline->GetSizeBytes());
 			}
+			CumulativeResourceSize.AddDedicatedSystemMemoryBytes(Pipeline->GetSizeBytes());
 		}
-	}
+	}	
 }
 
 /**
diff --git a/Engine/Source/Runtime/Engine/Private/MorphMesh.cpp b/Engine/Source/Runtime/Engine/Private/MorphMesh.cpp
--- a/Engine/Source/Runtime/Engine/Private/MorphMesh.cpp
+++ b/Engine/Source/Runtime/Engine/Private/MorphMesh.cpp
@@ -41,21 +41,7 @@ void UMorphTarget::GetResourceSizeEx(FResourceSizeEx& CumulativeResourceSize)
 	}
 }
 
-
-//////////////////////////////////////////////////////////////////////
-SIZE_T FMorphTargetLODModel::GetResourceSize() const
-{
-	return GetResourceSizeBytes();
-}
-
 void FMorphTargetLODModel::GetResourceSizeEx(FResourceSizeEx& CumulativeResourceSize) const
 {
 	CumulativeResourceSize.AddUnknownMemoryBytes(Vertices.GetAllocatedSize() + sizeof(int32));
 }
-
-SIZE_T FMorphTargetLODModel::GetResourceSizeBytes() const
-{
-	FResourceSizeEx ResSize;
-	GetResourceSizeEx(ResSize);
-	return ResSize.GetTotalMemoryBytes();
-}
diff --git a/Engine/Source/Runtime/Engine/Private/Particles/ParticleBeam2EmitterInstance.cpp b/Engine/Source/Runtime/Engine/Private/Particles/ParticleBeam2EmitterInstance.cpp
--- a/Engine/Source/Runtime/Engine/Private/Particles/ParticleBeam2EmitterInstance.cpp
+++ b/Engine/Source/Runtime/Engine/Private/Particles/ParticleBeam2EmitterInstance.cpp
@@ -1251,7 +1251,7 @@ void FParticleBeam2EmitterInstance::GetAllocatedSize(int32& OutNum, int32& OutMa
  */
 void FParticleBeam2EmitterInstance::GetResourceSizeEx(FResourceSizeEx& CumulativeResourceSize)
 {
-	if (CumulativeResourceSize.GetResourceSizeMode() == EResourceSizeMode::Inclusive || (Component && Component->SceneProxy))
+	if (CumulativeResourceSize.GetResourceSizeMode() == EResourceSizeMode::EstimatedTotal || (Component && Component->SceneProxy))
 	{
 		int32 MaxActiveParticleDataSize = (ParticleData != NULL) ? (MaxActiveParticles * ParticleStride) : 0;
 		int32 MaxActiveParticleIndexSize = 0;
diff --git a/Engine/Source/Runtime/Engine/Private/Particles/ParticleEmitterInstances.cpp b/Engine/Source/Runtime/Engine/Private/Particles/ParticleEmitterInstances.cpp
--- a/Engine/Source/Runtime/Engine/Private/Particles/ParticleEmitterInstances.cpp
+++ b/Engine/Source/Runtime/Engine/Private/Particles/ParticleEmitterInstances.cpp
@@ -2975,7 +2975,7 @@ void FParticleSpriteEmitterInstance::GetAllocatedSize(int32& OutNum, int32& OutM
  */
 void FParticleSpriteEmitterInstance::GetResourceSizeEx(FResourceSizeEx& CumulativeResourceSize)
 {
-	if (CumulativeResourceSize.GetResourceSizeMode() == EResourceSizeMode::Inclusive || (Component && Component->SceneProxy))
+	if (CumulativeResourceSize.GetResourceSizeMode() == EResourceSizeMode::EstimatedTotal || (Component && Component->SceneProxy))
 	{
 		int32 MaxActiveParticleDataSize = (ParticleData != NULL) ? (MaxActiveParticles * ParticleStride) : 0;
 		int32 MaxActiveParticleIndexSize = (ParticleIndices != NULL) ? (MaxActiveParticles * sizeof(uint16)) : 0;
@@ -3619,7 +3619,7 @@ void FParticleMeshEmitterInstance::GetAllocatedSize(int32& OutNum, int32& OutMax
  */
 void FParticleMeshEmitterInstance::GetResourceSizeEx(FResourceSizeEx& CumulativeResourceSize)
 {
-	if (CumulativeResourceSize.GetResourceSizeMode() == EResourceSizeMode::Inclusive || (Component && Component->SceneProxy))
+	if (CumulativeResourceSize.GetResourceSizeMode() == EResourceSizeMode::EstimatedTotal || (Component && Component->SceneProxy))
 	{
 		int32 MaxActiveParticleDataSize = (ParticleData != NULL) ? (MaxActiveParticles * ParticleStride) : 0;
 		int32 MaxActiveParticleIndexSize = (ParticleIndices != NULL) ? (MaxActiveParticles * sizeof(uint16)) : 0;
diff --git a/Engine/Source/Runtime/Engine/Private/Particles/ParticleTrail2EmitterInstance.cpp b/Engine/Source/Runtime/Engine/Private/Particles/ParticleTrail2EmitterInstance.cpp
--- a/Engine/Source/Runtime/Engine/Private/Particles/ParticleTrail2EmitterInstance.cpp
+++ b/Engine/Source/Runtime/Engine/Private/Particles/ParticleTrail2EmitterInstance.cpp
@@ -2691,7 +2691,7 @@ void FParticleRibbonEmitterInstance::GetAllocatedSize(int32& OutNum, int32& OutM
  */
 void FParticleRibbonEmitterInstance::GetResourceSizeEx(FResourceSizeEx& CumulativeResourceSize)
 {
-	if (CumulativeResourceSize.GetResourceSizeMode() == EResourceSizeMode::Inclusive || (Component && Component->SceneProxy))
+	if (CumulativeResourceSize.GetResourceSizeMode() == EResourceSizeMode::EstimatedTotal || (Component && Component->SceneProxy))
 	{
 		int32 MaxActiveParticleDataSize = (ParticleData != NULL) ? (MaxActiveParticles * ParticleStride) : 0;
 		int32 MaxActiveParticleIndexSize = (ParticleIndices != NULL) ? (MaxActiveParticles * sizeof(uint16)) : 0;
@@ -3987,7 +3987,7 @@ void FParticleAnimTrailEmitterInstance::GetAllocatedSize(int32& OutNum, int32& O
  */
 void FParticleAnimTrailEmitterInstance::GetResourceSizeEx(FResourceSizeEx& CumulativeResourceSize)
 {
-	if (CumulativeResourceSize.GetResourceSizeMode() == EResourceSizeMode::Inclusive || (Component && Component->SceneProxy))
+	if (CumulativeResourceSize.GetResourceSizeMode() == EResourceSizeMode::EstimatedTotal || (Component && Component->SceneProxy))
 	{
 		int32 MaxActiveParticleDataSize = (ParticleData != NULL) ? (MaxActiveParticles * ParticleStride) : 0;
 		int32 MaxActiveParticleIndexSize = (ParticleIndices != NULL) ? (MaxActiveParticles * sizeof(uint16)) : 0;
diff --git a/Engine/Source/Runtime/Engine/Private/PhysicsEngine/BodySetup.cpp b/Engine/Source/Runtime/Engine/Private/PhysicsEngine/BodySetup.cpp
--- a/Engine/Source/Runtime/Engine/Private/PhysicsEngine/BodySetup.cpp
+++ b/Engine/Source/Runtime/Engine/Private/PhysicsEngine/BodySetup.cpp
@@ -40,6 +40,40 @@
 #include "ProfilingDebugging/CookStats.h"
 #include "AnimPhysObjectVersion.h"
 
+/** Helper for enum output... */
+#ifndef CASE_ENUM_TO_TEXT
+#define CASE_ENUM_TO_TEXT(txt) case txt: return TEXT(#txt);
+#endif
+
+namespace Lex
+{
+	const TCHAR* ToString(ECollisionTraceFlag Enum)
+	{
+		switch (Enum)
+		{
+			FOREACH_ENUM_ECOLLISIONTRACEFLAG(CASE_ENUM_TO_TEXT)
+		}
+		return TEXT("<Unknown ECollisionTraceFlag>");
+	}
+
+	const TCHAR* ToString(EPhysicsType Enum)
+	{
+		switch (Enum)
+		{
+			FOREACH_ENUM_EPHYSICSTYPE(CASE_ENUM_TO_TEXT)
+		}
+		return TEXT("<Unknown EPhysicsType>");
+	}
+
+	const TCHAR* ToString(EBodyCollisionResponse::Type Enum)
+	{
+		switch (Enum)
+		{
+			FOREACH_ENUM_EBODYCOLLISIONRESPONSE(CASE_ENUM_TO_TEXT)
+		}
+		return TEXT("<Unknown EBodyCollisionResponse>");
+	}
+}
 
 
 FCookBodySetupInfo::FCookBodySetupInfo() :
@@ -120,11 +154,6 @@ static TAutoConsoleVariable<float> CVarMaxContactOffset(
 	ECVF_Default);
 
 
-SIZE_T FBodySetupUVInfo::GetResourceSize() const
-{
-	return GetResourceSizeBytes();
-}
-
 void FBodySetupUVInfo::GetResourceSizeEx(FResourceSizeEx& CumulativeResourceSize) const
 {
 	CumulativeResourceSize.AddDedicatedSystemMemoryBytes(IndexBuffer.GetAllocatedSize());
@@ -138,14 +167,6 @@ void FBodySetupUVInfo::GetResourceSizeEx(FResourceSizeEx& CumulativeResourceSize
 	CumulativeResourceSize.AddDedicatedSystemMemoryBytes(VertUVs.GetAllocatedSize());
 }
 
-SIZE_T FBodySetupUVInfo::GetResourceSizeBytes() const
-{
-	FResourceSizeEx ResSize;
-	GetResourceSizeEx(ResSize);
-	return ResSize.GetTotalMemoryBytes();
-}
-
-
 DEFINE_LOG_CATEGORY(LogPhysics);
 UBodySetup::UBodySetup(const FObjectInitializer& ObjectInitializer)
 	: Super(ObjectInitializer)
diff --git a/Engine/Source/Runtime/Engine/Private/PhysicsEngine/PhysicsAsset.cpp b/Engine/Source/Runtime/Engine/Private/PhysicsEngine/PhysicsAsset.cpp
--- a/Engine/Source/Runtime/Engine/Private/PhysicsEngine/PhysicsAsset.cpp
+++ b/Engine/Source/Runtime/Engine/Private/PhysicsEngine/PhysicsAsset.cpp
@@ -827,15 +827,9 @@ void UPhysicsAsset::GetResourceSizeEx(FResourceSizeEx& CumulativeResourceSize)
 {
 	Super::GetResourceSizeEx(CumulativeResourceSize);
 
-	for (const auto& SingleBody : SkeletalBodySetups)
-	{
-		SingleBody->GetResourceSizeEx(CumulativeResourceSize);
-	}
-
+	// Nested body setups are handled by default implementation
 	CumulativeResourceSize.AddDedicatedSystemMemoryBytes(BodySetupIndexMap.GetAllocatedSize());
 	CumulativeResourceSize.AddDedicatedSystemMemoryBytes(CollisionDisableTable.GetAllocatedSize());
-
-	// @todo implement inclusive mode
 }
 
 #undef LOCTEXT_NAMESPACE
diff --git a/Engine/Source/Runtime/Engine/Private/SkeletalMesh.cpp b/Engine/Source/Runtime/Engine/Private/SkeletalMesh.cpp
--- a/Engine/Source/Runtime/Engine/Private/SkeletalMesh.cpp
+++ b/Engine/Source/Runtime/Engine/Private/SkeletalMesh.cpp
@@ -160,44 +160,6 @@ static bool SaveApexClothingAssetToBlob(const apex::ClothingAsset *InAsset, TArr
 
 #endif//#if WITH_APEX_CLOTHING
 
-////////////////////////////////////////////////
-SIZE_T FClothingAssetData_Legacy::GetResourceSize() const
-{
-	return GetResourceSizeBytes();
-}
-
-void FClothingAssetData_Legacy::GetResourceSizeEx(FResourceSizeEx& CumulativeResourceSize) const
-{
-#if WITH_APEX_CLOTHING
-	if (ApexClothingAsset)
-	{
-		physx::PxU32 LODLevel = ApexClothingAsset->getNumGraphicalLodLevels();
-		for(physx::PxU32 LODId=0; LODId < LODLevel; ++LODId)
-		{
-			if(const apex::RenderMeshAsset* RenderAsset = ApexClothingAsset->getRenderMeshAsset(LODId))
-			{
-				apex::RenderMeshAssetStats AssetStats;
-				RenderAsset->getStats(AssetStats);
-
-				CumulativeResourceSize.AddUnknownMemoryBytes(AssetStats.totalBytes);
-			}
-		}
-	}
-#endif // #if WITH_APEX_CLOTHING
-}
-
-SIZE_T FClothingAssetData_Legacy::GetResourceSizeBytes() const
-{
-	FResourceSizeEx ResSize;
-	GetResourceSizeEx(ResSize);
-	return ResSize.GetTotalMemoryBytes();
-}
-
-
-
-
-
-
 
 /*-----------------------------------------------------------------------------
 FGPUSkinVertexBase
@@ -723,65 +685,15 @@ const FMeshUVChannelInfo* USkeletalMesh::GetUVChannelData(int32 MaterialIndex) c
 void USkeletalMesh::GetResourceSizeEx(FResourceSizeEx& CumulativeResourceSize)
 {
 	Super::GetResourceSizeEx(CumulativeResourceSize);
-
-#if WITH_EDITORONLY_DATA
-	if (ImportedModel.IsValid())
-	{
-		ImportedModel->GetResourceSizeEx(CumulativeResourceSize);
-	}
-#endif // WITH_EDITOR
+	// Default implementation handles subobjects
 
 	if (SkeletalMeshRenderData.IsValid())
 	{
 		SkeletalMeshRenderData->GetResourceSizeEx(CumulativeResourceSize);
 	}
 
-	if (CumulativeResourceSize.GetResourceSizeMode() == EResourceSizeMode::Inclusive)
-	{
-		for (const auto& MorphTarget : MorphTargets)
-		{
-			MorphTarget->GetResourceSizeEx(CumulativeResourceSize);
-		}
-
-		for(const FClothingAssetData_Legacy& LegacyAsset : ClothingAssets_DEPRECATED)
-		{
-			LegacyAsset.GetResourceSizeEx(CumulativeResourceSize);
-		}
-
-		for(UClothingAssetBase* ClothingAsset : MeshClothingAssets)
-		{
-			ClothingAsset->GetResourceSizeEx(CumulativeResourceSize);
-		}
-
-		TSet<UMaterialInterface*> UniqueMaterials;
-		for(int32 MaterialIndex = 0; MaterialIndex < Materials.Num(); ++MaterialIndex)
-		{
-			UMaterialInterface* Material = Materials[MaterialIndex].MaterialInterface;
-			bool bAlreadyCounted = false;
-			UniqueMaterials.Add(Material, &bAlreadyCounted);
-			if(!bAlreadyCounted && Material)
-			{
-				Material->GetResourceSizeEx(CumulativeResourceSize);
-			}
-		}
-
-#if WITH_EDITORONLY_DATA
-		CumulativeResourceSize.AddDedicatedSystemMemoryBytes(RetargetBasePose.GetAllocatedSize());
-#endif
-
-		CumulativeResourceSize.AddDedicatedSystemMemoryBytes(RefBasesInvMatrix.GetAllocatedSize());
-		CumulativeResourceSize.AddDedicatedSystemMemoryBytes(RefSkeleton.GetDataSize());
-
-		if (BodySetup)
-		{
-			BodySetup->GetResourceSizeEx(CumulativeResourceSize);
-		}
-
-		if (PhysicsAsset)
-		{
-			PhysicsAsset->GetResourceSizeEx(CumulativeResourceSize);
-		}
-	}
+	CumulativeResourceSize.AddDedicatedSystemMemoryBytes(RefBasesInvMatrix.GetAllocatedSize());
+	CumulativeResourceSize.AddDedicatedSystemMemoryBytes(RefSkeleton.GetDataSize());
 }
 
 /**
diff --git a/Engine/Source/Runtime/Engine/Private/SkeletalRenderCPUSkin.h b/Engine/Source/Runtime/Engine/Private/SkeletalRenderCPUSkin.h
--- a/Engine/Source/Runtime/Engine/Private/SkeletalRenderCPUSkin.h
+++ b/Engine/Source/Runtime/Engine/Private/SkeletalRenderCPUSkin.h
@@ -71,15 +71,7 @@ public:
 	/** a weight factor to blend between simulated positions and skinned positions */
 	float ClothBlendWeight;
 
-	/**
-	* Returns the size of memory allocated by render data
-	*/
-	DEPRECATED(4.14, "GetResourceSize is deprecated. Please use GetResourceSizeEx or GetResourceSizeBytes instead.")
-	SIZE_T GetResourceSize()
-	{
-		return GetResourceSizeBytes();
-	}
-
+	/** Returns the size of memory allocated by render data */
 	virtual void GetResourceSizeEx(FResourceSizeEx& CumulativeResourceSize)
  	{
 		CumulativeResourceSize.AddDedicatedSystemMemoryBytes(sizeof(*this));
@@ -88,13 +80,6 @@ public:
 		CumulativeResourceSize.AddDedicatedSystemMemoryBytes(ActiveMorphTargets.GetAllocatedSize());
  	}
 
-	SIZE_T GetResourceSizeBytes()
-	{
-		FResourceSizeEx ResSize;
-		GetResourceSizeEx(ResSize);
-		return ResSize.GetTotalMemoryBytes();
-	}
-
 	/** Update Simulated Positions & Normals from Clothing actor */
 	bool UpdateClothSimulationData(USkinnedMeshComponent* InMeshComponent);
 };
@@ -210,23 +195,10 @@ private:
 		/**
 	 	 * Get Resource Size : return the size of Resource this allocates
 	 	 */
-		DEPRECATED(4.14, "GetResourceSize is deprecated. Please use GetResourceSizeEx or GetResourceSizeBytes instead.")
-		SIZE_T GetResourceSize()
-		{
-			return GetResourceSizeBytes();
-		}
-
 		void GetResourceSizeEx(FResourceSizeEx& CumulativeResourceSize)
 		{
 			CumulativeResourceSize.AddUnknownMemoryBytes(StaticMeshVertexBuffer.GetResourceSize() + PositionVertexBuffer.GetStride() * PositionVertexBuffer.GetNumVertices());
 		}
-
-		SIZE_T GetResourceSizeBytes()
-		{
-			FResourceSizeEx ResSize;
-			GetResourceSizeEx(ResSize);
-			return ResSize.GetTotalMemoryBytes();
-		}
 	};
 
 	/** Render data for each LOD */
diff --git a/Engine/Source/Runtime/Engine/Private/SkeletalRenderGPUSkin.h b/Engine/Source/Runtime/Engine/Private/SkeletalRenderGPUSkin.h
--- a/Engine/Source/Runtime/Engine/Private/SkeletalRenderGPUSkin.h
+++ b/Engine/Source/Runtime/Engine/Private/SkeletalRenderGPUSkin.h
@@ -96,15 +96,7 @@ public:
 	*/
 	bool ActiveMorphTargetsEqual(const TArray<FActiveMorphTarget>& CompareActiveMorphTargets, const TArray<float>& CompareMorphTargetWeights);
 	
-	/**
-	* Returns the size of memory allocated by render data
-	*/
-	DEPRECATED(4.14, "GetResourceSize is deprecated. Please use GetResourceSizeEx or GetResourceSizeBytes instead.")
-	SIZE_T GetResourceSize()
-	{
-		return GetResourceSizeBytes();
-	}
-
+	/** Returns the size of memory allocated by render data */
 	virtual void GetResourceSizeEx(FResourceSizeEx& CumulativeResourceSize)
 	{
 		CumulativeResourceSize.AddDedicatedSystemMemoryBytes(sizeof(*this));
@@ -113,13 +105,6 @@ public:
 		CumulativeResourceSize.AddUnknownMemoryBytes(ActiveMorphTargets.GetAllocatedSize());
 	}
 
-	SIZE_T GetResourceSizeBytes()
-	{
-		FResourceSizeEx ResSize;
-		GetResourceSizeEx(ResSize);
-		return ResSize.GetTotalMemoryBytes();
-	}
-
 	/** Update Simulated Positions & Normals from APEX Clothing actor */
 	bool UpdateClothSimulationData(USkinnedMeshComponent* InMeshComponent);
 };
@@ -458,25 +443,12 @@ private:
 		/**
 		 * @return memory in bytes of size of the resources for this LOD
 		 */
-		DEPRECATED(4.14, "GetResourceSize is deprecated. Please use GetResourceSizeEx or GetResourceSizeBytes instead.")
-		SIZE_T GetResourceSize()
-		{
-			return GetResourceSizeBytes();
-		}
-
 		void GetResourceSizeEx(FResourceSizeEx& CumulativeResourceSize)
 		{
 			CumulativeResourceSize.AddUnknownMemoryBytes(MorphVertexBuffer.GetResourceSize());
 			CumulativeResourceSize.AddUnknownMemoryBytes(GPUSkinVertexFactories.GetResourceSize());
 		}
 
-		SIZE_T GetResourceSizeBytes()
-		{
-			FResourceSizeEx ResSize;
-			GetResourceSizeEx(ResSize);
-			return ResSize.GetTotalMemoryBytes();
-		}
-
 		FSkeletalMeshRenderData* SkelMeshRenderData;
 		// index into FSkeletalMeshRenderData::LODRenderData[]
 		int32 LODIndex;
diff --git a/Engine/Source/Runtime/Engine/Private/SoundCue.cpp b/Engine/Source/Runtime/Engine/Private/SoundCue.cpp
--- a/Engine/Source/Runtime/Engine/Private/SoundCue.cpp
+++ b/Engine/Source/Runtime/Engine/Private/SoundCue.cpp
@@ -310,31 +310,6 @@ FString USoundCue::GetDesc()
 	return Description;
 }
 
-void USoundCue::GetResourceSizeEx(FResourceSizeEx& CumulativeResourceSize)
-{
-	Super::GetResourceSizeEx(CumulativeResourceSize);
-
-	if( CumulativeResourceSize.GetResourceSizeMode() == EResourceSizeMode::Exclusive )
-	{
-		// Nothing to add
-	}
-	else
-	{
-		// Sum up the size of referenced waves
-		TArray<USoundNodeWavePlayer*> WavePlayers;
-		RecursiveFindNode<USoundNodeWavePlayer>( FirstNode, WavePlayers );
-
-		for( int32 WaveIndex = 0; WaveIndex < WavePlayers.Num(); ++WaveIndex )
-		{
-			USoundWave* SoundWave = WavePlayers[WaveIndex]->GetSoundWave();
-			if (SoundWave)
-			{
-				SoundWave->GetResourceSizeEx(CumulativeResourceSize);
-			}
-		}
-	}
-}
-
 int32 USoundCue::GetResourceSizeForFormat(FName Format)
 {
 	TArray<USoundNodeWavePlayer*> WavePlayers;
