Add trivial hash function for FPrimaryAssetType, needed for sample in main

diff --git a/Engine/Source/Runtime/CoreUObject/Public/UObject/PrimaryAssetId.h b/Engine/Source/Runtime/CoreUObject/Public/UObject/PrimaryAssetId.h
--- a/Engine/Source/Runtime/CoreUObject/Public/UObject/PrimaryAssetId.h
+++ b/Engine/Source/Runtime/CoreUObject/Public/UObject/PrimaryAssetId.h
@@ -60,6 +60,11 @@ struct FPrimaryAssetType
 	bool ImportTextItem(const TCHAR*& Buffer, int32 PortFlags, UObject* Parent, FOutputDevice* ErrorText);
 	bool SerializeFromMismatchedTag(struct FPropertyTag const& Tag, FArchive& Ar);
 
+	friend inline uint32 GetTypeHash(const FPrimaryAssetType& Key)
+	{
+		return GetTypeHash(Key.Name);
+	}
+
 private:
 	friend COREUOBJECT_API UScriptStruct* Z_Construct_UScriptStruct_FPrimaryAssetType();
 
