Disabling more false positives from PVS studio.

diff --git a/Engine/Source/Editor/Matinee/Private/MatineeDraw.cpp b/Engine/Source/Editor/Matinee/Private/MatineeDraw.cpp
--- a/Engine/Source/Editor/Matinee/Private/MatineeDraw.cpp
+++ b/Engine/Source/Editor/Matinee/Private/MatineeDraw.cpp
@@ -1878,7 +1878,7 @@ void FMatineeViewportClient::Draw(FViewport* InViewport, FCanvas* Canvas)
 		{
 			// If we're parented then we're only visible if our parent group is not collapsed
 			check( CurParentGroup != NULL );
-			if( CurParentGroup->bCollapsed )
+			if( CurParentGroup->bCollapsed ) //-V522
 			{
 				// Parent group is collapsed, so we should not be rendered
 				bIsGroupVisible = false;
diff --git a/Engine/Source/Editor/UnrealEd/Private/Fbx/FbxSkeletalMeshImport.cpp b/Engine/Source/Editor/UnrealEd/Private/Fbx/FbxSkeletalMeshImport.cpp
--- a/Engine/Source/Editor/UnrealEd/Private/Fbx/FbxSkeletalMeshImport.cpp
+++ b/Engine/Source/Editor/UnrealEd/Private/Fbx/FbxSkeletalMeshImport.cpp
@@ -3234,7 +3234,7 @@ bool UnFbx::FFbxImporter::FillSkelMeshImporterFromFbx( FSkeletalMeshImportData&
 		delete[] UVMappingMode;
 	}
 	
-	return true;
+	return true; //-V773
 }
 
 void UnFbx::FFbxImporter::InsertNewLODToBaseSkeletalMesh(USkeletalMesh* InSkeletalMesh, USkeletalMesh* BaseSkeletalMesh, int32 DesiredLOD, UFbxSkeletalMeshImportData* TemplateImportData)
diff --git a/Engine/Source/Runtime/Core/Private/HAL/UnrealMemory.cpp b/Engine/Source/Runtime/Core/Private/HAL/UnrealMemory.cpp
--- a/Engine/Source/Runtime/Core/Private/HAL/UnrealMemory.cpp
+++ b/Engine/Source/Runtime/Core/Private/HAL/UnrealMemory.cpp
@@ -279,7 +279,7 @@ static void FMallocBinnedOverrunTest()
 		[Pointer, ArraySize]()
 	{
 		//FPlatformProcess::Sleep(.01f);
-		Pointer[ArraySize / 2] = 0xcc;
+		Pointer[ArraySize / 2] = 0xcc; //-V774
 	},
 		TStatId()
 		);
diff --git a/Engine/Source/Runtime/CoreUObject/Private/UObject/LinkerLoad.cpp b/Engine/Source/Runtime/CoreUObject/Private/UObject/LinkerLoad.cpp
--- a/Engine/Source/Runtime/CoreUObject/Private/UObject/LinkerLoad.cpp
+++ b/Engine/Source/Runtime/CoreUObject/Private/UObject/LinkerLoad.cpp
@@ -2197,7 +2197,7 @@ FLinkerLoad::EVerifyResult FLinkerLoad::VerifyImport(int32 ImportIndex)
 		VerifyImportInner(ImportIndex, WarningAppend);
 
 		// if the redirector wasn't found, then it truly doesn't exist
-		if (Import.SourceIndex == INDEX_NONE)
+		if (Import.SourceIndex == INDEX_NONE) //-V547
 		{
 			Result = VERIFY_Failed;
 		}
@@ -2262,7 +2262,7 @@ FLinkerLoad::EVerifyResult FLinkerLoad::VerifyImport(int32 ImportIndex)
 		Import.ClassPackage = OriginalImport.ClassPackage;
 
 		// if nothing above failed, then we are good to go
-		if (Result != VERIFY_Failed)
+		if (Result != VERIFY_Failed) //-V547
 		{
 			// we update the runtime information (SourceIndex, SourceLinker) to point to the object the redirector pointed to
 			Import.SourceIndex = Import.XObject->GetLinkerIndex();
diff --git a/Engine/Source/Runtime/Engine/Private/HierarchicalInstancedStaticMesh.cpp b/Engine/Source/Runtime/Engine/Private/HierarchicalInstancedStaticMesh.cpp
--- a/Engine/Source/Runtime/Engine/Private/HierarchicalInstancedStaticMesh.cpp
+++ b/Engine/Source/Runtime/Engine/Private/HierarchicalInstancedStaticMesh.cpp
@@ -1583,7 +1583,7 @@ void FHierarchicalStaticMeshSceneProxy::GetDynamicMeshElements(const TArray<cons
 					SCOPE_CYCLE_COUNTER(STAT_FoliageTraversalTime);
 
 					// validate that the bounding box is layed out correctly in memory
-					check((const FVector4*)&ClusterTree[0].BoundMin + 1 == (const FVector4*)&ClusterTree[0].BoundMax);
+					check((const FVector4*)&ClusterTree[0].BoundMin + 1 == (const FVector4*)&ClusterTree[0].BoundMax); //-V594
 					//check(UPTRINT(&ClusterTree[0].BoundMin) % 16 == 0);
 					//check(UPTRINT(&ClusterTree[0].BoundMax) % 16 == 0);
 
diff --git a/Engine/Source/Runtime/Engine/Private/Particles/ParticleSystemRender.cpp b/Engine/Source/Runtime/Engine/Private/Particles/ParticleSystemRender.cpp
--- a/Engine/Source/Runtime/Engine/Private/Particles/ParticleSystemRender.cpp
+++ b/Engine/Source/Runtime/Engine/Private/Particles/ParticleSystemRender.cpp
@@ -5975,7 +5975,7 @@ int32 FDynamicRibbonEmitterData::FillVertexData(struct FAsyncBufferFillData& Dat
 				FLinearColor CurrColor = PackingParticle->Color;
 
 				FVector PrevPosition = PrevParticle->Location; //-V522
-				FVector PrevTangent = PrevTrailPayload->Tangent;
+				FVector PrevTangent = PrevTrailPayload->Tangent; //-V522
 				FVector PrevUp = PrevWorkingUp;
 				FLinearColor PrevColor = PrevParticle->Color;
 				float PrevSize = PrevParticle->Size.X * Source.Scale.X;
