Copying //UE4/Dev-Framework to //UE4/Dev-Main (Source: //UE4/Dev-Framework @ 3967517)

diff --git a/Engine/Source/Runtime/AIModule/Classes/Navigation/NavLinkProxy.h b/Engine/Source/Runtime/AIModule/Classes/Navigation/NavLinkProxy.h
new file mode 100644
--- /dev/null
+++ b/Engine/Source/Runtime/AIModule/Classes/Navigation/NavLinkProxy.h
@@ -0,0 +1,125 @@
+// Copyright 1998-2018 Epic Games, Inc. All Rights Reserved.
+
+#pragma once
+
+#include "CoreMinimal.h"
+#include "UObject/ObjectMacros.h"
+#include "Templates/SubclassOf.h"
+#include "GameFramework/Actor.h"
+#include "EngineDefines.h"
+#include "AI/Navigation/NavRelevantInterface.h"
+#include "AI/Navigation/NavLinkDefinition.h"
+#include "NavLinkHostInterface.h"
+#include "NavLinkProxy.generated.h"
+
+class UBillboardComponent;
+class UNavLinkCustomComponent;
+class UNavLinkRenderingComponent;
+class UPathFollowingComponent;
+struct FNavigationRelevantData;
+
+DECLARE_DYNAMIC_MULTICAST_DELEGATE_TwoParams( FSmartLinkReachedSignature, AActor*, MovingActor, const FVector&, DestinationPoint );
+
+UCLASS(Blueprintable, autoCollapseCategories=(SmartLink, Actor), hideCategories=(Input))
+class AIMODULE_API ANavLinkProxy : public AActor, public INavLinkHostInterface, public INavRelevantInterface
+{
+	GENERATED_UCLASS_BODY()
+
+	/** Navigation links (point to point) added to navigation data */
+	UPROPERTY(EditAnywhere, Category=SimpleLink)
+	TArray<FNavigationLink> PointLinks;
+	
+	/** Navigation links (segment to segment) added to navigation data
+	*	@todo hidden from use until we fix segment links. Not really working now*/
+	UPROPERTY()
+	TArray<FNavigationSegmentLink> SegmentLinks;
+
+private:
+	/** Smart link: can affect path following */
+	UPROPERTY(VisibleAnywhere, Category=SmartLink)
+	UNavLinkCustomComponent* SmartLinkComp;
+public:
+
+	/** Smart link: toggle relevancy */
+	UPROPERTY(EditAnywhere, Category=SmartLink)
+	bool bSmartLinkIsRelevant;
+
+#if WITH_EDITORONLY_DATA
+private:
+	/** Editor Preview */
+	UPROPERTY()
+	UNavLinkRenderingComponent* EdRenderComp;
+
+	UPROPERTY()
+	UBillboardComponent* SpriteComponent;
+public:
+#endif // WITH_EDITORONLY_DATA
+
+	// BEGIN INavRelevantInterface
+	virtual void GetNavigationData(FNavigationRelevantData& Data) const override;
+	virtual FBox GetNavigationBounds() const override;
+	virtual bool IsNavigationRelevant() const override;
+	// END INavRelevantInterface
+
+	// BEGIN INavLinkHostInterface
+	virtual bool GetNavigationLinksClasses(TArray<TSubclassOf<UNavLinkDefinition> >& OutClasses) const override;
+	virtual bool GetNavigationLinksArray(TArray<FNavigationLink>& OutLink, TArray<FNavigationSegmentLink>& OutSegments) const override;
+	// END INavLinkHostInterface
+
+#if WITH_EDITOR
+	virtual void PostEditChangeProperty(FPropertyChangedEvent& PropertyChangedEvent) override;
+	virtual void PostEditUndo() override;
+	virtual void PostEditImport() override;
+#endif // WITH_EDITOR
+	virtual void PostRegisterAllComponents() override;
+	virtual void PostLoad() override;
+
+#if ENABLE_VISUAL_LOG
+protected:
+	virtual void BeginPlay() override;
+
+public:
+#endif // ENABLE_VISUAL_LOG
+
+	virtual FBox GetComponentsBoundingBox(bool bNonColliding = false) const override;
+
+	//////////////////////////////////////////////////////////////////////////
+	// Blueprint interface for smart links
+
+	/** called when agent reaches smart link during path following, use ResumePathFollowing() to give control back */
+	UFUNCTION(BlueprintImplementableEvent)
+	void ReceiveSmartLinkReached(AActor* Agent, const FVector& Destination);
+
+	/** resume normal path following */
+	UFUNCTION(BlueprintCallable, Category="AI|Navigation")
+	void ResumePathFollowing(AActor* Agent);
+
+	/** check if smart link is enabled */
+	UFUNCTION(BlueprintCallable, Category="AI|Navigation")
+	bool IsSmartLinkEnabled() const;
+
+	/** change state of smart link */
+	UFUNCTION(BlueprintCallable, Category="AI|Navigation")
+	void SetSmartLinkEnabled(bool bEnabled);
+
+	/** check if any agent is moving through smart link right now */
+	UFUNCTION(BlueprintCallable, Category="AI|Navigation")
+	bool HasMovingAgents() const;
+
+protected:
+
+	UPROPERTY(BlueprintAssignable)
+	FSmartLinkReachedSignature OnSmartLinkReached;
+
+	void NotifySmartLinkReached(UNavLinkCustomComponent* LinkComp, UObject* PathingAgent, const FVector& DestPoint);
+
+public:
+	/** Returns SmartLinkComp subobject **/
+	UNavLinkCustomComponent* GetSmartLinkComp() const { return SmartLinkComp; }
+#if WITH_EDITORONLY_DATA
+	/** Returns EdRenderComp subobject **/
+	UNavLinkRenderingComponent* GetEdRenderComp() const { return EdRenderComp; }
+	/** Returns SpriteComponent subobject **/
+	UBillboardComponent* GetSpriteComponent() const { return SpriteComponent; }
+#endif
+};
diff --git a/Engine/Source/Runtime/AIModule/Classes/Navigation/PathFollowingComponent.h b/Engine/Source/Runtime/AIModule/Classes/Navigation/PathFollowingComponent.h
--- a/Engine/Source/Runtime/AIModule/Classes/Navigation/PathFollowingComponent.h
+++ b/Engine/Source/Runtime/AIModule/Classes/Navigation/PathFollowingComponent.h
@@ -9,16 +9,18 @@
 #include "Components/ActorComponent.h"
 #include "EngineDefines.h"
 #include "AI/Navigation/NavigationTypes.h"
+#include "NavigationData.h"
 #include "AITypes.h"
 #include "AIResourceInterface.h"
-#include "AI/Navigation/NavigationData.h"
 #include "GameFramework/NavMovementComponent.h"
+#include "AI/Navigation/PathFollowingAgentInterface.h"
 #include "PathFollowingComponent.generated.h"
 
 class Error;
 class FDebugDisplayInfo;
 class INavLinkCustomInterface;
 class UCanvas;
+class ANavigationData;
 
 AIMODULE_API DECLARE_LOG_CATEGORY_EXTERN(LogPathFollowing, Warning, All);
 
@@ -211,7 +213,7 @@ enum class EPathFollowingReachMode : uint8
 };
 
 UCLASS(config=Engine)
-class AIMODULE_API UPathFollowingComponent : public UActorComponent, public IAIResourceInterface
+class AIMODULE_API UPathFollowingComponent : public UActorComponent, public IAIResourceInterface, public IPathFollowingAgentInterface
 {
 	GENERATED_UCLASS_BODY()
 
@@ -227,6 +229,7 @@ class AIMODULE_API UPathFollowingComponent : public UActorComponent, public IAIR
 	FMoveComplete OnRequestFinished;
 
 	//~ Begin UActorComponent Interface
+	virtual void OnRegister() override;
 	virtual void TickComponent(float DeltaTime, enum ELevelTick TickType, FActorComponentTickFunction *ThisTickFunction) override;
 	//~ End UActorComponent Interface
 
@@ -377,14 +380,12 @@ class AIMODULE_API UPathFollowingComponent : public UActorComponent, public IAIR
 	UFUNCTION()
 	virtual void OnActorBump(AActor* SelfActor, AActor* OtherActor, FVector NormalImpulse, const FHitResult& Hit);
 
-	/** Called when movement is blocked by a collision with another actor.  */
-	virtual void OnMoveBlockedBy(const FHitResult& BlockingImpact) {}
-
-	/** Called when falling movement starts. */
-	virtual void OnStartedFalling();
-
-	/** Called when falling movement ends. */
-	virtual void OnLanded() {}
+	// IPathFollowingAgentInterface begin
+	virtual void OnUnableToMove(const UObject& Instigator) override;		
+	//virtual void OnMoveBlockedBy(const FHitResult& BlockingImpact) {}
+	virtual void OnStartedFalling() override;
+	virtual void OnLanded() override {}
+	// IPathFollowingAgentInterface end
 
 	/** Check if path following can be activated */
 	virtual bool IsPathFollowingAllowed() const;
@@ -596,6 +597,10 @@ protected:
 	/** reset path following data */
 	virtual void Reset();
 
+	/** Called if owning Controller possesses new pawn or ends up pawn-less. 
+	 *	Doesn't get called if owner is not an AContoller */
+	virtual void OnNewPawn(APawn* NewPawn);
+
 	/** should verify if agent if still on path ater movement has been resumed? */
 	virtual bool ShouldCheckPathOnResume() const;
 
diff --git a/Engine/Source/Runtime/AIModule/Classes/Navigation/PathFollowingManager.h b/Engine/Source/Runtime/AIModule/Classes/Navigation/PathFollowingManager.h
new file mode 100644
--- /dev/null
+++ b/Engine/Source/Runtime/AIModule/Classes/Navigation/PathFollowingManager.h
@@ -0,0 +1,21 @@
+// Copyright 1998-2018 Epic Games, Inc. All Rights Reserved.
+
+#pragma once
+
+#include "UObject/Object.h"
+#include "AI/NavigationSystemBase.h"
+#include "PathFollowingManager.generated.h"
+
+class AController;
+
+
+UCLASS()
+class AIMODULE_API UPathFollowingManager : public UObject, public IPathFollowingManagerInterface
+{
+	GENERATED_BODY()
+public:
+	UPathFollowingManager(const FObjectInitializer& ObjectInitializer = FObjectInitializer::Get());
+
+	static void StopMovement(const AController& Controller);
+	static bool IsFollowingAPath(const AController& Controller);
+};
diff --git a/Engine/Source/Runtime/AIModule/Classes/Perception/AISenseConfig_Blueprint.h b/Engine/Source/Runtime/AIModule/Classes/Perception/AISenseConfig_Blueprint.h
--- a/Engine/Source/Runtime/AIModule/Classes/Perception/AISenseConfig_Blueprint.h
+++ b/Engine/Source/Runtime/AIModule/Classes/Perception/AISenseConfig_Blueprint.h
@@ -11,7 +11,7 @@
 
 class UAISense_Blueprint;
 
-UCLASS(Blueprintable, Abstract)
+UCLASS(Blueprintable, Abstract, hidedropdown)
 class AIMODULE_API UAISenseConfig_Blueprint : public UAISenseConfig
 {
 	GENERATED_UCLASS_BODY()
diff --git a/Engine/Source/Runtime/AIModule/Classes/Perception/AISense_Blueprint.h b/Engine/Source/Runtime/AIModule/Classes/Perception/AISense_Blueprint.h
--- a/Engine/Source/Runtime/AIModule/Classes/Perception/AISense_Blueprint.h
+++ b/Engine/Source/Runtime/AIModule/Classes/Perception/AISense_Blueprint.h
@@ -14,7 +14,7 @@ class UAIPerceptionComponent;
 class UAISenseEvent;
 class UUserDefinedStruct;
 
-UCLASS(ClassGroup = AI, Abstract, Blueprintable)
+UCLASS(ClassGroup = AI, Abstract, Blueprintable, hidedropdown)
 class AIMODULE_API UAISense_Blueprint : public UAISense
 {
 	GENERATED_BODY()
diff --git a/Engine/Source/Runtime/AIModule/Private/AIController.cpp b/Engine/Source/Runtime/AIModule/Private/AIController.cpp
--- a/Engine/Source/Runtime/AIModule/Private/AIController.cpp
+++ b/Engine/Source/Runtime/AIModule/Private/AIController.cpp
@@ -2,7 +2,7 @@
 
 #include "AIController.h"
 #include "CollisionQueryParams.h"
-#include "AI/Navigation/NavigationSystem.h"
+#include "NavigationSystem.h"
 #include "Components/CapsuleComponent.h"
 #include "GameFramework/PhysicsVolume.h"
 #include "Actions/PawnActionsComponent.h"
@@ -476,12 +476,8 @@ void AAIController::Possess(APawn* InPawn)
 		return;
 	}
 
-	// no point in doing navigation setup if pawn has no movement component
-	const UPawnMovementComponent* MovementComp = InPawn->GetMovementComponent();
-	if (MovementComp != NULL)
-	{
-		UpdateNavigationComponents();
-	}
+	// not calling UpdateNavigationComponents() anymore. The PathFollowingComponent 
+	// is now observing newly possessed pawns (via OnNewPawn)
 
 	if (PathFollowingComponent)
 	{
@@ -562,11 +558,6 @@ void AAIController::SetPawn(APawn* InPawn)
 	}
 }
 
-void AAIController::InitNavigationControl(UPathFollowingComponent*& PathFollowingComp)
-{
-	PathFollowingComp = PathFollowingComponent;
-}
-
 EPathFollowingRequestResult::Type AAIController::MoveToActor(AActor* Goal, float AcceptanceRadius, bool bStopOnOverlap, bool bUsePathfinding, bool bCanStrafe, TSubclassOf<UNavigationQueryFilter> FilterClass, bool bAllowPartialPaths)
 {
 	// abort active movement to keep only one request running
@@ -647,7 +638,7 @@ FPathFollowingRequestResult AAIController::MoveTo(const FAIMoveRequest& MoveRequ
 		// fail if projection to navigation is required but it failed
 		if (bCanRequestMove && MoveRequest.IsProjectingGoal())
 		{
-			UNavigationSystem* NavSys = UNavigationSystem::GetCurrent(GetWorld());
+			UNavigationSystemV1* NavSys = FNavigationSystem::GetCurrent<UNavigationSystemV1>(GetWorld());
 			const FNavAgentProperties& AgentProps = GetNavAgentPropertiesRef();
 			FNavLocation ProjectedLocation;
 
@@ -753,7 +744,7 @@ bool AAIController::BuildPathfindingQuery(const FAIMoveRequest& MoveRequest, FPa
 {
 	bool bResult = false;
 
-	UNavigationSystem* NavSys = UNavigationSystem::GetCurrent(GetWorld());
+	UNavigationSystemV1* NavSys = FNavigationSystem::GetCurrent<UNavigationSystemV1>(GetWorld());
 	const ANavigationData* NavData = (NavSys == nullptr) ? nullptr :
 		MoveRequest.IsUsingPathfinding() ? NavSys->GetNavDataForProps(GetNavAgentPropertiesRef()) :
 		NavSys->GetAbstractNavData();
@@ -798,7 +789,7 @@ void AAIController::FindPathForMoveRequest(const FAIMoveRequest& MoveRequest, FP
 {
 	SCOPE_CYCLE_COUNTER(STAT_AI_Overall);
 
-	UNavigationSystem* NavSys = UNavigationSystem::GetCurrent(GetWorld());
+	UNavigationSystemV1* NavSys = FNavigationSystem::GetCurrent<UNavigationSystemV1>(GetWorld());
 	if (NavSys)
 	{
 		FPathFindingResult PathResult = NavSys->FindPathSync(Query);
diff --git a/Engine/Source/Runtime/AIModule/Private/AIModule.cpp b/Engine/Source/Runtime/AIModule/Private/AIModule.cpp
--- a/Engine/Source/Runtime/AIModule/Private/AIModule.cpp
+++ b/Engine/Source/Runtime/AIModule/Private/AIModule.cpp
@@ -50,6 +50,7 @@ IMPLEMENT_MODULE(FAIModule, AIModule)
 
 void FAIModule::StartupModule()
 { 
+	FModuleManager::LoadModulePtr< IModuleInterface >("NavigationSystem");
 #if WITH_EDITOR 
 	FModuleManager::LoadModulePtr< IModuleInterface >("AITestSuite");
 
diff --git a/Engine/Source/Runtime/AIModule/Private/BehaviorTree/BehaviorTreeComponent.cpp b/Engine/Source/Runtime/AIModule/Private/BehaviorTree/BehaviorTreeComponent.cpp
--- a/Engine/Source/Runtime/AIModule/Private/BehaviorTree/BehaviorTreeComponent.cpp
+++ b/Engine/Source/Runtime/AIModule/Private/BehaviorTree/BehaviorTreeComponent.cpp
@@ -1009,7 +1009,7 @@ void UBehaviorTreeComponent::ApplySearchUpdates(const TArray<FBehaviorTreeSearch
 		{
 			// special case: service node at root of top most subtree - don't remove/re-add them when tree is in looping mode
 			// don't bother with decorators parent == root means that they are on child branches
-			if (bLoopExecution && UpdateInfo.AuxNode->GetParentNode() == InstanceStack[0].RootNode &&
+			if (bLoopExecution && UpdateInfo.AuxNode->GetMyNode() == InstanceStack[0].RootNode &&
 				UpdateInfo.AuxNode->IsA(UBTService::StaticClass()))
 			{
 				if (UpdateInfo.Mode == EBTNodeUpdateMode::Remove ||
@@ -1907,6 +1907,9 @@ bool UBehaviorTreeComponent::PushInstance(UBehaviorTree& TreeAsset)
 			UBTService* ServiceNode = RootNode->Services[ServiceIndex];
 			uint8* NodeMemory = (uint8*)ServiceNode->GetNodeMemory<uint8>(InstanceStack[ActiveInstanceIdx]);
 
+			// send initial on search start events in case someone is using them for init logic
+			ServiceNode->NotifyParentActivation(SearchData);
+
 			InstanceStack[ActiveInstanceIdx].ActiveAuxNodes.Add(ServiceNode);
 			ServiceNode->WrappedOnBecomeRelevant(*this, NodeMemory);
 		}
diff --git a/Engine/Source/Runtime/AIModule/Private/BehaviorTree/Decorators/BTDecorator_DoesPathExist.cpp b/Engine/Source/Runtime/AIModule/Private/BehaviorTree/Decorators/BTDecorator_DoesPathExist.cpp
--- a/Engine/Source/Runtime/AIModule/Private/BehaviorTree/Decorators/BTDecorator_DoesPathExist.cpp
+++ b/Engine/Source/Runtime/AIModule/Private/BehaviorTree/Decorators/BTDecorator_DoesPathExist.cpp
@@ -3,8 +3,8 @@
 #include "BehaviorTree/Decorators/BTDecorator_DoesPathExist.h"
 #include "UObject/Package.h"
 #include "GameFramework/Actor.h"
-#include "AI/Navigation/NavigationSystem.h"
-#include "AI/Navigation/RecastNavMesh.h"
+#include "NavigationSystem.h"
+#include "NavMesh/RecastNavMesh.h"
 #include "BehaviorTree/BlackboardComponent.h"
 #include "AIController.h"
 
@@ -60,7 +60,7 @@ bool UBTDecorator_DoesPathExist::CalculateRawConditionValue(UBehaviorTreeCompone
 	
 	bool bHasPath = false;
 
-	const UNavigationSystem* NavSys = UNavigationSystem::GetCurrent(OwnerComp.GetWorld());
+	const UNavigationSystemV1* NavSys = FNavigationSystem::GetCurrent<UNavigationSystemV1>(OwnerComp.GetWorld());
 	if (NavSys && bHasPointA && bHasPointB)
 	{
 		const AAIController* AIOwner = OwnerComp.GetAIOwner();
diff --git a/Engine/Source/Runtime/AIModule/Private/Blueprint/AIBlueprintHelperLibrary.cpp b/Engine/Source/Runtime/AIModule/Private/Blueprint/AIBlueprintHelperLibrary.cpp
--- a/Engine/Source/Runtime/AIModule/Private/Blueprint/AIBlueprintHelperLibrary.cpp
+++ b/Engine/Source/Runtime/AIModule/Private/Blueprint/AIBlueprintHelperLibrary.cpp
@@ -13,10 +13,15 @@
 #include "BehaviorTree/BlackboardComponent.h"
 #include "Blueprint/AIAsyncTaskBlueprintProxy.h"
 #include "Animation/AnimInstance.h"
-#include "AI/Navigation/NavigationPath.h"
+#include "NavigationPath.h"
+#include "NavigationData.h"
+#include "NavigationSystem.h"
+#include "Logging/MessageLog.h"
 
 DEFINE_LOG_CATEGORY_STATIC(LogAIBlueprint, Warning, All);
 
+#define LOCTEXT_NAMESPACE "AIBlueprintHelperLibrary"
+
 //----------------------------------------------------------------------//
 // UAIAsyncTaskBlueprintProxy
 //----------------------------------------------------------------------//
@@ -338,3 +343,162 @@ UNavigationPath* UAIBlueprintHelperLibrary::GetCurrentPath(AController* Controll
 	return ResultPath;
 }
 
+
+namespace
+{
+	UPathFollowingComponent* InitNavigationControl(AController& Controller)
+	{
+		AAIController* AsAIController = Cast<AAIController>(&Controller);
+		UPathFollowingComponent* PathFollowingComp = nullptr;
+
+		if (AsAIController)
+		{
+			PathFollowingComp = AsAIController->GetPathFollowingComponent();
+		}
+		else
+		{
+			PathFollowingComp = Controller.FindComponentByClass<UPathFollowingComponent>();
+			if (PathFollowingComp == nullptr)
+			{
+				PathFollowingComp = NewObject<UPathFollowingComponent>(&Controller);
+				PathFollowingComp->RegisterComponentWithWorld(Controller.GetWorld());
+				PathFollowingComp->Initialize();
+			}
+		}
+
+		return PathFollowingComp;
+	}
+}
+
+void UAIBlueprintHelperLibrary::SimpleMoveToActor(AController* Controller, const AActor* Goal)
+{
+	UNavigationSystemV1* NavSys = Controller ? FNavigationSystem::GetCurrent<UNavigationSystemV1>(Controller->GetWorld()) : nullptr;
+	if (NavSys == nullptr || Goal == nullptr || Controller == nullptr || Controller->GetPawn() == nullptr)
+	{
+		UE_LOG(LogNavigation, Warning, TEXT("UNavigationSystemV1::SimpleMoveToActor called for NavSys:%s Controller:%s controlling Pawn:%s with goal actor %s (if any of these is None then there's your problem"),
+			*GetNameSafe(NavSys), *GetNameSafe(Controller), Controller ? *GetNameSafe(Controller->GetPawn()) : TEXT("NULL"), *GetNameSafe(Goal));
+		return;
+	}
+
+	UPathFollowingComponent* PFollowComp = InitNavigationControl(*Controller);
+
+	if (PFollowComp == nullptr)
+	{
+		FMessageLog("PIE").Warning(FText::Format(
+			LOCTEXT("SimpleMoveErrorNoComp", "SimpleMove failed for {0}: missing components"),
+			FText::FromName(Controller->GetFName())
+		));
+		return;
+	}
+
+	if (!PFollowComp->IsPathFollowingAllowed())
+	{
+		FMessageLog("PIE").Warning(FText::Format(
+			LOCTEXT("SimpleMoveErrorMovement", "SimpleMove failed for {0}: movement not allowed"),
+			FText::FromName(Controller->GetFName())
+		));
+		return;
+	}
+
+	const bool bAlreadyAtGoal = PFollowComp->HasReached(*Goal, EPathFollowingReachMode::OverlapAgentAndGoal);
+
+	// script source, keep only one move request at time
+	if (PFollowComp->GetStatus() != EPathFollowingStatus::Idle)
+	{
+		PFollowComp->AbortMove(*NavSys, FPathFollowingResultFlags::ForcedScript | FPathFollowingResultFlags::NewRequest
+			, FAIRequestID::AnyRequest, bAlreadyAtGoal ? EPathFollowingVelocityMode::Reset : EPathFollowingVelocityMode::Keep);
+	}
+
+	if (bAlreadyAtGoal)
+	{
+		PFollowComp->RequestMoveWithImmediateFinish(EPathFollowingResult::Success);
+	}
+	else
+	{
+		const ANavigationData* NavData = NavSys->GetNavDataForProps(Controller->GetNavAgentPropertiesRef());
+		if (NavData)
+		{
+			FPathFindingQuery Query(Controller, *NavData, Controller->GetNavAgentLocation(), Goal->GetActorLocation());
+			FPathFindingResult Result = NavSys->FindPathSync(Query);
+			if (Result.IsSuccessful())
+			{
+				Result.Path->SetGoalActorObservation(*Goal, 100.0f);
+				PFollowComp->RequestMove(FAIMoveRequest(Goal), Result.Path);
+			}
+			else if (PFollowComp->GetStatus() != EPathFollowingStatus::Idle)
+			{
+				PFollowComp->RequestMoveWithImmediateFinish(EPathFollowingResult::Invalid);
+			}
+		}
+	}
+}
+
+void UAIBlueprintHelperLibrary::SimpleMoveToLocation(AController* Controller, const FVector& GoalLocation)
+{
+	UNavigationSystemV1* NavSys = Controller ? FNavigationSystem::GetCurrent<UNavigationSystemV1>(Controller->GetWorld()) : nullptr;
+	if (NavSys == nullptr || Controller == nullptr || Controller->GetPawn() == nullptr)
+	{
+		UE_LOG(LogNavigation, Warning, TEXT("UNavigationSystemV1::SimpleMoveToActor called for NavSys:%s Controller:%s controlling Pawn:%s (if any of these is None then there's your problem"),
+			*GetNameSafe(NavSys), *GetNameSafe(Controller), Controller ? *GetNameSafe(Controller->GetPawn()) : TEXT("NULL"));
+		return;
+	}
+
+	UPathFollowingComponent* PFollowComp = InitNavigationControl(*Controller);
+
+	if (PFollowComp == nullptr)
+	{
+		FMessageLog("PIE").Warning(FText::Format(
+			LOCTEXT("SimpleMoveErrorNoComp", "SimpleMove failed for {0}: missing components"),
+			FText::FromName(Controller->GetFName())
+		));
+		return;
+	}
+
+	if (!PFollowComp->IsPathFollowingAllowed())
+	{
+		FMessageLog("PIE").Warning(FText::Format(
+			LOCTEXT("SimpleMoveErrorMovement", "SimpleMove failed for {0}: movement not allowed"),
+			FText::FromName(Controller->GetFName())
+		));
+		return;
+	}
+
+	const bool bAlreadyAtGoal = PFollowComp->HasReached(GoalLocation, EPathFollowingReachMode::OverlapAgent);
+
+	// script source, keep only one move request at time
+	if (PFollowComp->GetStatus() != EPathFollowingStatus::Idle)
+	{
+		PFollowComp->AbortMove(*NavSys, FPathFollowingResultFlags::ForcedScript | FPathFollowingResultFlags::NewRequest
+			, FAIRequestID::AnyRequest, bAlreadyAtGoal ? EPathFollowingVelocityMode::Reset : EPathFollowingVelocityMode::Keep);
+	}
+
+	// script source, keep only one move request at time
+	if (PFollowComp->GetStatus() != EPathFollowingStatus::Idle)
+	{
+		PFollowComp->AbortMove(*NavSys, FPathFollowingResultFlags::ForcedScript | FPathFollowingResultFlags::NewRequest);
+	}
+
+	if (bAlreadyAtGoal)
+	{
+		PFollowComp->RequestMoveWithImmediateFinish(EPathFollowingResult::Success);
+	}
+	else
+	{
+		const ANavigationData* NavData = NavSys->GetNavDataForProps(Controller->GetNavAgentPropertiesRef());
+		if (NavData)
+		{
+			FPathFindingQuery Query(Controller, *NavData, Controller->GetNavAgentLocation(), GoalLocation);
+			FPathFindingResult Result = NavSys->FindPathSync(Query);
+			if (Result.IsSuccessful())
+			{
+				PFollowComp->RequestMove(FAIMoveRequest(GoalLocation), Result.Path);
+			}
+			else if (PFollowComp->GetStatus() != EPathFollowingStatus::Idle)
+			{
+				PFollowComp->RequestMoveWithImmediateFinish(EPathFollowingResult::Invalid);
+			}
+		}
+	}
+}
+
+#undef LOCTEXT_NAMESPACE
\ No newline at end of file
diff --git a/Engine/Source/Runtime/AIModule/Private/DataProviders/AIDataProvider.cpp b/Engine/Source/Runtime/AIModule/Private/DataProviders/AIDataProvider.cpp
--- a/Engine/Source/Runtime/AIModule/Private/DataProviders/AIDataProvider.cpp
+++ b/Engine/Source/Runtime/AIModule/Private/DataProviders/AIDataProvider.cpp
@@ -29,12 +29,6 @@ void FAIDataProviderValue::GetMatchingProperties(TArray<FName>& MatchingProperti
 	}
 }
 
-template<typename T>
-T* FAIDataProviderValue::GetRawValuePtr() const
-{
-	return CachedProperty ? CachedProperty->ContainerPtrToValuePtr<T>(DataBinding) : nullptr;
-}
-
 void FAIDataProviderValue::BindData(const UObject* Owner, int32 RequestId) const
 {
 	if (DataBinding && ensure(Owner))
diff --git a/Engine/Source/Runtime/AIModule/Private/EnvironmentQuery/EnvQueryInstance.cpp b/Engine/Source/Runtime/AIModule/Private/EnvironmentQuery/EnvQueryInstance.cpp
--- a/Engine/Source/Runtime/AIModule/Private/EnvironmentQuery/EnvQueryInstance.cpp
+++ b/Engine/Source/Runtime/AIModule/Private/EnvironmentQuery/EnvQueryInstance.cpp
@@ -14,6 +14,7 @@
 #include "EnvironmentQuery/Contexts/EnvQueryContext_Item.h"
 #include "EnvironmentQuery/Generators/EnvQueryGenerator_Composite.h"
 #include "AISystem.h"
+#include "AITypes.h"
 
 #if UE_BUILD_SHIPPING
 #define eqs_ensure ensure
@@ -121,11 +122,16 @@ bool FEnvQueryInstance::PrepareContext(UClass* ContextClass, FEnvQueryContextDat
 		if (CachedData == NULL)
 		{
 			UEnvQueryManager* QueryManager = UEnvQueryManager::GetCurrent(World);
-			UEnvQueryContext* ContextOb = QueryManager->PrepareLocalContext(ContextClass);
+			const UEnvQueryContext* ContextOb = QueryManager->PrepareLocalContext(ContextClass);
 
+			check(ContextOb);
 			ContextOb->ProvideContext(*this, ContextData);
 
-			DEC_MEMORY_STAT_BY(STAT_AI_EQS_InstanceMemory, GetContextAllocatedSize());
+#if STATS
+			// this contraption has been added to nail down a location of a rare bug
+			const uint32 AllocatedSize = GetContextAllocatedSize();
+			DEC_MEMORY_STAT_BY(STAT_AI_EQS_InstanceMemory, AllocatedSize);
+#endif // STATS
 
 			ContextCache.Add(ContextClass, ContextData);
 
@@ -198,8 +204,12 @@ bool FEnvQueryInstance::PrepareContext(UClass* Context, TArray<FVector>& Data)
 		Data.SetNumUninitialized(ContextData.NumValues);
 		for (int32 ValueIndex = 0; ValueIndex < ContextData.NumValues; ValueIndex++)
 		{
-			Data[ValueIndex] = DefTypeOb->GetItemLocation(ContextRawData);
-			ContextRawData += DefTypeValueSize;
+			const FVector ItemLocation = DefTypeOb->GetItemLocation(ContextRawData);
+			if (FAISystem::IsValidLocation(ItemLocation))
+			{
+				Data[ValueIndex] = ItemLocation;
+				ContextRawData += DefTypeValueSize;
+			}
 		}
 	}
 
@@ -225,8 +235,12 @@ bool FEnvQueryInstance::PrepareContext(UClass* Context, TArray<FRotator>& Data)
 		Data.SetNumUninitialized(ContextData.NumValues);
 		for (int32 ValueIndex = 0; ValueIndex < ContextData.NumValues; ValueIndex++)
 		{
-			Data[ValueIndex] = DefTypeOb->GetItemRotation(ContextRawData);
-			ContextRawData += DefTypeValueSize;
+			const FRotator ItemRotation = DefTypeOb->GetItemRotation(ContextRawData);
+			if (FAISystem::IsValidRotation(ItemRotation))
+			{
+				Data[ValueIndex] = ItemRotation;
+				ContextRawData += DefTypeValueSize;
+			}
 		}
 	}
 
diff --git a/Engine/Source/Runtime/AIModule/Private/EnvironmentQuery/EnvQueryTraceHelpers.cpp b/Engine/Source/Runtime/AIModule/Private/EnvironmentQuery/EnvQueryTraceHelpers.cpp
--- a/Engine/Source/Runtime/AIModule/Private/EnvironmentQuery/EnvQueryTraceHelpers.cpp
+++ b/Engine/Source/Runtime/AIModule/Private/EnvironmentQuery/EnvQueryTraceHelpers.cpp
@@ -1,7 +1,7 @@
 // Copyright 1998-2018 Epic Games, Inc. All Rights Reserved.
 
 #include "EnvironmentQuery/EnvQueryTraceHelpers.h"
-#include "AI/Navigation/NavigationData.h"
+#include "NavigationData.h"
 
 template<>
 void FEQSHelpers::FBatchTrace::DoSingleSourceMultiDestinations<EEnvTraceShape::Line>(const FVector& Source, TArray<FNavLocation>& Points)
diff --git a/Engine/Source/Runtime/AIModule/Private/EnvironmentQuery/EnvQueryTypes.cpp b/Engine/Source/Runtime/AIModule/Private/EnvironmentQuery/EnvQueryTypes.cpp
--- a/Engine/Source/Runtime/AIModule/Private/EnvironmentQuery/EnvQueryTypes.cpp
+++ b/Engine/Source/Runtime/AIModule/Private/EnvironmentQuery/EnvQueryTypes.cpp
@@ -3,8 +3,8 @@
 #include "EnvironmentQuery/EnvQueryTypes.h"
 #include "UObject/Package.h"
 #include "AI/Navigation/NavAgentInterface.h"
-#include "AI/Navigation/NavigationSystem.h"
-#include "AI/Navigation/RecastNavMesh.h"
+#include "NavigationSystem.h"
+#include "NavMesh/RecastNavMesh.h"
 #include "EnvironmentQuery/Items/EnvQueryItemType_VectorBase.h"
 #include "EnvironmentQuery/Items/EnvQueryItemType_ActorBase.h"
 #include "BehaviorTree/BTNode.h"
@@ -268,7 +268,7 @@ namespace FEQSHelpers
 {
 	const ANavigationData* FindNavigationDataForQuery(FEnvQueryInstance& QueryInstance)
 	{
-		const UNavigationSystem* NavSys = QueryInstance.World->GetNavigationSystem();
+		const UNavigationSystemV1* NavSys = FNavigationSystem::GetCurrent<UNavigationSystemV1>(QueryInstance.World);
 		
 		if (NavSys == nullptr)
 		{
@@ -280,10 +280,10 @@ namespace FEQSHelpers
 		if (NavAgent)
 		{
 			const FNavAgentProperties& NavAgentProps = NavAgent->GetNavAgentPropertiesRef();
-			return NavSys->GetNavDataForProps(NavAgentProps);
+			return Cast<const ANavigationData>(NavSys->GetNavDataForProps(NavAgentProps));
 		}
 
-		return NavSys->GetMainNavData();
+		return Cast<const ANavigationData>(NavSys->GetDefaultNavDataInstance());
 	}
 }
 
diff --git a/Engine/Source/Runtime/AIModule/Private/EnvironmentQuery/Generators/EnvQueryGenerator_PathingGrid.cpp b/Engine/Source/Runtime/AIModule/Private/EnvironmentQuery/Generators/EnvQueryGenerator_PathingGrid.cpp
--- a/Engine/Source/Runtime/AIModule/Private/EnvironmentQuery/Generators/EnvQueryGenerator_PathingGrid.cpp
+++ b/Engine/Source/Runtime/AIModule/Private/EnvironmentQuery/Generators/EnvQueryGenerator_PathingGrid.cpp
@@ -1,7 +1,7 @@
 // Copyright 1998-2018 Epic Games, Inc. All Rights Reserved.
 
 #include "EnvironmentQuery/Generators/EnvQueryGenerator_PathingGrid.h"
-#include "AI/Navigation/RecastNavMesh.h"
+#include "NavMesh/RecastNavMesh.h"
 
 #define LOCTEXT_NAMESPACE "EnvQueryGenerator"
 
diff --git a/Engine/Source/Runtime/AIModule/Private/EnvironmentQuery/Items/EnvQueryItemType_Actor.cpp b/Engine/Source/Runtime/AIModule/Private/EnvironmentQuery/Items/EnvQueryItemType_Actor.cpp
--- a/Engine/Source/Runtime/AIModule/Private/EnvironmentQuery/Items/EnvQueryItemType_Actor.cpp
+++ b/Engine/Source/Runtime/AIModule/Private/EnvironmentQuery/Items/EnvQueryItemType_Actor.cpp
@@ -3,6 +3,7 @@
 #include "EnvironmentQuery/Items/EnvQueryItemType_Actor.h"
 #include "UObject/WeakObjectPtr.h"
 #include "GameFramework/Actor.h"
+#include "AITypes.h"
 #include "EnvironmentQuery/EnvQueryTypes.h"
 
 UEnvQueryItemType_Actor::UEnvQueryItemType_Actor(const FObjectInitializer& ObjectInitializer) : Super(ObjectInitializer)
@@ -25,13 +26,13 @@ void UEnvQueryItemType_Actor::SetValue(uint8* RawData, const FWeakObjectPtr& Val
 FVector UEnvQueryItemType_Actor::GetItemLocation(const uint8* RawData) const
 {
 	AActor* MyActor = UEnvQueryItemType_Actor::GetValue(RawData);
-	return MyActor ? MyActor->GetActorLocation() : FVector::ZeroVector;
+	return MyActor ? MyActor->GetActorLocation() : FAISystem::InvalidLocation;
 }
 
 FRotator UEnvQueryItemType_Actor::GetItemRotation(const uint8* RawData) const
 {
 	AActor* MyActor = UEnvQueryItemType_Actor::GetValue(RawData);
-	return MyActor ? MyActor->GetActorRotation() : FRotator::ZeroRotator;
+	return MyActor ? MyActor->GetActorRotation() : FAISystem::InvalidRotation;
 }
 
 AActor* UEnvQueryItemType_Actor::GetActor(const uint8* RawData) const
diff --git a/Engine/Source/Runtime/AIModule/Private/EnvironmentQuery/Tests/EnvQueryTest_Pathfinding.cpp b/Engine/Source/Runtime/AIModule/Private/EnvironmentQuery/Tests/EnvQueryTest_Pathfinding.cpp
--- a/Engine/Source/Runtime/AIModule/Private/EnvironmentQuery/Tests/EnvQueryTest_Pathfinding.cpp
+++ b/Engine/Source/Runtime/AIModule/Private/EnvironmentQuery/Tests/EnvQueryTest_Pathfinding.cpp
@@ -3,8 +3,8 @@
 #include "EnvironmentQuery/Tests/EnvQueryTest_Pathfinding.h"
 #include "AI/Navigation/NavAgentInterface.h"
 #include "Engine/World.h"
-#include "AI/Navigation/NavigationData.h"
-#include "AI/Navigation/NavigationSystem.h"
+#include "NavigationData.h"
+#include "NavigationSystem.h"
 #include "EnvironmentQuery/Items/EnvQueryItemType_VectorBase.h"
 #include "EnvironmentQuery/Contexts/EnvQueryContext_Querier.h"
 
@@ -39,7 +39,7 @@ void UEnvQueryTest_Pathfinding::RunTest(FEnvQueryInstance& QueryInstance) const
 	float MinThresholdValue = FloatValueMin.GetValue();
 	float MaxThresholdValue = FloatValueMax.GetValue();
 
-	UNavigationSystem* NavSys = QueryInstance.World->GetNavigationSystem();
+	UNavigationSystemV1* NavSys = FNavigationSystem::GetCurrent<UNavigationSystemV1>(QueryInstance.World);
 	if (NavSys == nullptr || QueryOwner == nullptr)
 	{
 		return;
@@ -166,7 +166,7 @@ void UEnvQueryTest_Pathfinding::PostLoad()
 	SetWorkOnFloatValues(TestMode != EEnvTestPathfinding::PathExist);
 }
 
-bool UEnvQueryTest_Pathfinding::TestPathFrom(const FVector& ItemPos, const FVector& ContextPos, EPathFindingMode::Type Mode, const ANavigationData& NavData, UNavigationSystem& NavSys, FSharedConstNavQueryFilter NavFilter, const UObject* PathOwner) const
+bool UEnvQueryTest_Pathfinding::TestPathFrom(const FVector& ItemPos, const FVector& ContextPos, EPathFindingMode::Type Mode, const ANavigationData& NavData, UNavigationSystemV1& NavSys, FSharedConstNavQueryFilter NavFilter, const UObject* PathOwner) const
 {
 	FPathFindingQuery Query(PathOwner, NavData, ItemPos, ContextPos, NavFilter);
 	Query.SetAllowPartialPaths(false);
@@ -175,7 +175,7 @@ bool UEnvQueryTest_Pathfinding::TestPathFrom(const FVector& ItemPos, const FVect
 	return bPathExists;
 }
 
-bool UEnvQueryTest_Pathfinding::TestPathTo(const FVector& ItemPos, const FVector& ContextPos, EPathFindingMode::Type Mode, const ANavigationData& NavData, UNavigationSystem& NavSys, FSharedConstNavQueryFilter NavFilter, const UObject* PathOwner) const
+bool UEnvQueryTest_Pathfinding::TestPathTo(const FVector& ItemPos, const FVector& ContextPos, EPathFindingMode::Type Mode, const ANavigationData& NavData, UNavigationSystemV1& NavSys, FSharedConstNavQueryFilter NavFilter, const UObject* PathOwner) const
 {
 	FPathFindingQuery Query(PathOwner, NavData, ContextPos, ItemPos, NavFilter);
 	Query.SetAllowPartialPaths(false);
@@ -184,7 +184,7 @@ bool UEnvQueryTest_Pathfinding::TestPathTo(const FVector& ItemPos, const FVector
 	return bPathExists;
 }
 
-float UEnvQueryTest_Pathfinding::FindPathCostFrom(const FVector& ItemPos, const FVector& ContextPos, EPathFindingMode::Type Mode, const ANavigationData& NavData, UNavigationSystem& NavSys, FSharedConstNavQueryFilter NavFilter, const UObject* PathOwner) const
+float UEnvQueryTest_Pathfinding::FindPathCostFrom(const FVector& ItemPos, const FVector& ContextPos, EPathFindingMode::Type Mode, const ANavigationData& NavData, UNavigationSystemV1& NavSys, FSharedConstNavQueryFilter NavFilter, const UObject* PathOwner) const
 {
 	FPathFindingQuery Query(PathOwner, NavData, ItemPos, ContextPos, NavFilter);
 	Query.SetAllowPartialPaths(false);
@@ -193,7 +193,7 @@ float UEnvQueryTest_Pathfinding::FindPathCostFrom(const FVector& ItemPos, const
 	return (Result.IsSuccessful()) ? Result.Path->GetCost() : BIG_NUMBER;
 }
 
-float UEnvQueryTest_Pathfinding::FindPathCostTo(const FVector& ItemPos, const FVector& ContextPos, EPathFindingMode::Type Mode, const ANavigationData& NavData, UNavigationSystem& NavSys, FSharedConstNavQueryFilter NavFilter, const UObject* PathOwner) const
+float UEnvQueryTest_Pathfinding::FindPathCostTo(const FVector& ItemPos, const FVector& ContextPos, EPathFindingMode::Type Mode, const ANavigationData& NavData, UNavigationSystemV1& NavSys, FSharedConstNavQueryFilter NavFilter, const UObject* PathOwner) const
 {
 	FPathFindingQuery Query(PathOwner, NavData, ContextPos, ItemPos, NavFilter);
 	Query.SetAllowPartialPaths(false);
@@ -202,7 +202,7 @@ float UEnvQueryTest_Pathfinding::FindPathCostTo(const FVector& ItemPos, const FV
 	return (Result.IsSuccessful()) ? Result.Path->GetCost() : BIG_NUMBER;
 }
 
-float UEnvQueryTest_Pathfinding::FindPathLengthFrom(const FVector& ItemPos, const FVector& ContextPos, EPathFindingMode::Type Mode, const ANavigationData& NavData, UNavigationSystem& NavSys, FSharedConstNavQueryFilter NavFilter, const UObject* PathOwner) const
+float UEnvQueryTest_Pathfinding::FindPathLengthFrom(const FVector& ItemPos, const FVector& ContextPos, EPathFindingMode::Type Mode, const ANavigationData& NavData, UNavigationSystemV1& NavSys, FSharedConstNavQueryFilter NavFilter, const UObject* PathOwner) const
 {
 	FPathFindingQuery Query(PathOwner, NavData, ItemPos, ContextPos, NavFilter);
 	Query.SetAllowPartialPaths(false);
@@ -211,7 +211,7 @@ float UEnvQueryTest_Pathfinding::FindPathLengthFrom(const FVector& ItemPos, cons
 	return (Result.IsSuccessful()) ? Result.Path->GetLength() : BIG_NUMBER;
 }
 
-float UEnvQueryTest_Pathfinding::FindPathLengthTo(const FVector& ItemPos, const FVector& ContextPos, EPathFindingMode::Type Mode, const ANavigationData& NavData, UNavigationSystem& NavSys, FSharedConstNavQueryFilter NavFilter, const UObject* PathOwner) const
+float UEnvQueryTest_Pathfinding::FindPathLengthTo(const FVector& ItemPos, const FVector& ContextPos, EPathFindingMode::Type Mode, const ANavigationData& NavData, UNavigationSystemV1& NavSys, FSharedConstNavQueryFilter NavFilter, const UObject* PathOwner) const
 {
 	FPathFindingQuery Query(PathOwner, NavData, ContextPos, ItemPos, NavFilter);
 	Query.SetAllowPartialPaths(false);
@@ -220,7 +220,7 @@ float UEnvQueryTest_Pathfinding::FindPathLengthTo(const FVector& ItemPos, const
 	return (Result.IsSuccessful()) ? Result.Path->GetLength() : BIG_NUMBER;
 }
 
-ANavigationData* UEnvQueryTest_Pathfinding::FindNavigationData(UNavigationSystem& NavSys, UObject* Owner) const
+ANavigationData* UEnvQueryTest_Pathfinding::FindNavigationData(UNavigationSystemV1& NavSys, UObject* Owner) const
 {
 	INavAgentInterface* NavAgent = Cast<INavAgentInterface>(Owner);
 	if (NavAgent)
@@ -228,7 +228,7 @@ ANavigationData* UEnvQueryTest_Pathfinding::FindNavigationData(UNavigationSystem
 		return NavSys.GetNavDataForProps(NavAgent->GetNavAgentPropertiesRef());
 	}
 
-	return NavSys.GetMainNavData(FNavigationSystem::DontCreate);
+	return NavSys.GetDefaultNavDataInstance(FNavigationSystem::DontCreate);
 }
 
 #undef LOCTEXT_NAMESPACE
diff --git a/Engine/Source/Runtime/AIModule/Private/EnvironmentQuery/Tests/EnvQueryTest_PathfindingBatch.cpp b/Engine/Source/Runtime/AIModule/Private/EnvironmentQuery/Tests/EnvQueryTest_PathfindingBatch.cpp
--- a/Engine/Source/Runtime/AIModule/Private/EnvironmentQuery/Tests/EnvQueryTest_PathfindingBatch.cpp
+++ b/Engine/Source/Runtime/AIModule/Private/EnvironmentQuery/Tests/EnvQueryTest_PathfindingBatch.cpp
@@ -2,8 +2,9 @@
 
 #include "EnvironmentQuery/Tests/EnvQueryTest_PathfindingBatch.h"
 #include "Engine/World.h"
-#include "AI/Navigation/RecastNavMesh.h"
-#include "AI/Navigation/RecastQueryFilter.h"
+#include "NavigationSystem.h"
+#include "NavMesh/RecastNavMesh.h"
+#include "NavMesh/RecastQueryFilter.h"
 
 #define LOCTEXT_NAMESPACE "EnvQueryGenerator"
 
@@ -67,7 +68,7 @@ void UEnvQueryTest_PathfindingBatch::RunTest(FEnvQueryInstance& QueryInstance) c
 	float MaxThresholdValue = FloatValueMax.GetValue();
 	float RangeMultiplierValue = ScanRangeMultiplier.GetValue();
 
-	UNavigationSystem* NavSys = QueryInstance.World->GetNavigationSystem();
+	UNavigationSystemV1* NavSys = FNavigationSystem::GetCurrent<UNavigationSystemV1>(QueryInstance.World);
 	if (NavSys == nullptr || QueryOwner == nullptr)
 	{
 		return;
diff --git a/Engine/Source/Runtime/AIModule/Private/EnvironmentQuery/Tests/EnvQueryTest_Project.cpp b/Engine/Source/Runtime/AIModule/Private/EnvironmentQuery/Tests/EnvQueryTest_Project.cpp
--- a/Engine/Source/Runtime/AIModule/Private/EnvironmentQuery/Tests/EnvQueryTest_Project.cpp
+++ b/Engine/Source/Runtime/AIModule/Private/EnvironmentQuery/Tests/EnvQueryTest_Project.cpp
@@ -2,7 +2,7 @@
 
 #include "EnvironmentQuery/Tests/EnvQueryTest_Project.h"
 #include "AI/Navigation/NavigationTypes.h"
-#include "AI/Navigation/NavigationData.h"
+#include "NavigationData.h"
 #include "EnvironmentQuery/Items/EnvQueryItemType_VectorBase.h"
 #include "EnvironmentQuery/Items/EnvQueryItemType_Point.h"
 #include "EnvironmentQuery/EnvQueryTraceHelpers.h"
diff --git a/Engine/Source/Runtime/AIModule/Private/GameplayDebugger/GameplayDebuggerCategory_AI.cpp b/Engine/Source/Runtime/AIModule/Private/GameplayDebugger/GameplayDebuggerCategory_AI.cpp
--- a/Engine/Source/Runtime/AIModule/Private/GameplayDebugger/GameplayDebuggerCategory_AI.cpp
+++ b/Engine/Source/Runtime/AIModule/Private/GameplayDebugger/GameplayDebuggerCategory_AI.cpp
@@ -8,11 +8,11 @@
 #include "ShowFlags.h"
 #include "PrimitiveViewRelevance.h"
 #include "Components/PrimitiveComponent.h"
-#include "AI/Navigation/NavigationSystem.h"
+#include "NavigationSystem.h"
 #include "GameFramework/PlayerController.h"
 #include "Engine/Engine.h"
 #include "EngineGlobals.h"
-#include "AI/Navigation/RecastNavMesh.h"
+#include "NavMesh/RecastNavMesh.h"
 #include "Engine/Canvas.h"
 #include "AIController.h"
 #include "BehaviorTree/BehaviorTreeComponent.h"
@@ -178,7 +178,7 @@ void FGameplayDebuggerCategory_AI::CollectData(APlayerController* OwnerPC, AActo
 
 		DataPack.MontageInfo = MyChar ? GetNameSafe(MyChar->GetCurrentMontage()) : FString();
 
-		UNavigationSystem* NavSys = UNavigationSystem::GetCurrent(MyPawn->GetWorld());
+		UNavigationSystemV1* NavSys = FNavigationSystem::GetCurrent<UNavigationSystemV1>(MyPawn->GetWorld());
 		const ANavigationData* NavData = MyController && NavSys ? NavSys->GetNavDataForProps(MyController->GetNavAgentPropertiesRef()) : nullptr;
 		DataPack.NavDataInfo = NavData ? NavData->GetConfig().Name.ToString() : FString();
 
diff --git a/Engine/Source/Runtime/AIModule/Private/GameplayDebugger/GameplayDebuggerCategory_Navmesh.cpp b/Engine/Source/Runtime/AIModule/Private/GameplayDebugger/GameplayDebuggerCategory_Navmesh.cpp
--- a/Engine/Source/Runtime/AIModule/Private/GameplayDebugger/GameplayDebuggerCategory_Navmesh.cpp
+++ b/Engine/Source/Runtime/AIModule/Private/GameplayDebugger/GameplayDebuggerCategory_Navmesh.cpp
@@ -4,9 +4,9 @@
 
 #if WITH_GAMEPLAY_DEBUGGER
 
-#include "AI/Navigation/NavigationSystem.h"
+#include "NavigationSystem.h"
 #include "GameFramework/PlayerController.h"
-#include "AI/Navigation/RecastNavMesh.h"
+#include "NavMesh/RecastNavMesh.h"
 
 FGameplayDebuggerCategory_Navmesh::FGameplayDebuggerCategory_Navmesh()
 {
@@ -25,18 +25,18 @@ TSharedRef<FGameplayDebuggerCategory> FGameplayDebuggerCategory_Navmesh::MakeIns
 void FGameplayDebuggerCategory_Navmesh::CollectData(APlayerController* OwnerPC, AActor* DebugActor)
 {
 #if WITH_RECAST
-	ARecastNavMesh* NavData = nullptr;
+	const ARecastNavMesh* NavData = nullptr;
 
 	APawn* PlayerPawn = OwnerPC ? OwnerPC->GetPawnOrSpectator() : nullptr;
 	APawn* DebugActorAsPawn = Cast<APawn>(DebugActor);
 	APawn* DestPawn = DebugActorAsPawn ? DebugActorAsPawn : PlayerPawn;
 	if (OwnerPC && DestPawn)
 	{
-		UNavigationSystem* NavSys = UNavigationSystem::GetCurrent(OwnerPC->GetWorld());
+		UNavigationSystemV1* NavSys = FNavigationSystem::GetCurrent<UNavigationSystemV1>(OwnerPC->GetWorld());
 		if (NavSys) 
 		{
 			const FNavAgentProperties& NavAgentProperties = DestPawn->GetNavAgentPropertiesRef();
-			NavData = Cast<ARecastNavMesh>(NavSys->GetNavDataForProps(NavAgentProperties));
+			NavData = Cast<const ARecastNavMesh>(NavSys->GetNavDataForProps(NavAgentProperties));
 		}
 	}
 
diff --git a/Engine/Source/Runtime/AIModule/Private/GameplayDebugger/GameplayDebuggerCategory_Navmesh.h b/Engine/Source/Runtime/AIModule/Private/GameplayDebugger/GameplayDebuggerCategory_Navmesh.h
--- a/Engine/Source/Runtime/AIModule/Private/GameplayDebugger/GameplayDebuggerCategory_Navmesh.h
+++ b/Engine/Source/Runtime/AIModule/Private/GameplayDebugger/GameplayDebuggerCategory_Navmesh.h
@@ -6,7 +6,7 @@
 
 #include "CoreMinimal.h"
 #include "GameplayDebuggerCategory.h"
-#include "AI/Navigation/NavMeshRenderingComponent.h"
+#include "NavMesh/NavMeshRenderingComponent.h"
 
 class APlayerController;
 
diff --git a/Engine/Source/Runtime/AIModule/Private/Navigation/CrowdFollowingComponent.cpp b/Engine/Source/Runtime/AIModule/Private/Navigation/CrowdFollowingComponent.cpp
--- a/Engine/Source/Runtime/AIModule/Private/Navigation/CrowdFollowingComponent.cpp
+++ b/Engine/Source/Runtime/AIModule/Private/Navigation/CrowdFollowingComponent.cpp
@@ -1,14 +1,14 @@
 // Copyright 1998-2018 Epic Games, Inc. All Rights Reserved.
 
 #include "Navigation/CrowdFollowingComponent.h"
-#include "AI/Navigation/NavigationSystem.h"
-#include "AI/Navigation/RecastNavMesh.h"
+#include "NavigationSystem.h"
+#include "NavMesh/RecastNavMesh.h"
 #include "VisualLogger/VisualLoggerTypes.h"
 #include "VisualLogger/VisualLogger.h"
 #include "AIModuleLog.h"
 #include "Navigation/CrowdManager.h"
-#include "AI/Navigation/NavAreas/NavArea.h"
-#include "AI/Navigation/AbstractNavData.h"
+#include "NavAreas/NavArea.h"
+#include "AbstractNavData.h"
 #include "AIConfig.h"
 #include "Navigation/MetaNavMeshPath.h"
 #include "GameFramework/CharacterMovementComponent.h"
@@ -474,7 +474,7 @@ void UCrowdFollowingComponent::Initialize()
 		if (!bRegisteredWithCrowdSimulation)
 		{
 			// crowd manager might not be created yet if this component was spawned during level's begin play (possessing a pawn placed in level)
-			UNavigationSystem* NavSys = UNavigationSystem::GetCurrent(GetWorld());
+			UNavigationSystemV1* NavSys = FNavigationSystem::GetCurrent<UNavigationSystemV1>(GetWorld());
 			if (NavSys && !NavSys->IsInitialized())
 			{
 				NavSys->OnNavigationInitDone.AddUObject(this, &UCrowdFollowingComponent::OnNavigationInitDone);
@@ -734,14 +734,14 @@ void LogPathPartHelper(AActor* LogOwner, FNavMeshPath* NavMeshPath, int32 StartI
 			CenterPt /= Verts.Num();
 
 			const UNavArea* DefArea = AreaClass ? ((UClass*)AreaClass)->GetDefaultObject<UNavArea>() : NULL;
-			const FColor PolygonColor = AreaClass != UNavigationSystem::GetDefaultWalkableArea() ? (DefArea ? DefArea->DrawColor : NavMesh->GetConfig().Color) : FColorList::LightSteelBlue;
+			const FColor PolygonColor = AreaClass != UNavigationSystemV1::GetDefaultWalkableArea() ? (DefArea ? DefArea->DrawColor : NavMesh->GetConfig().Color) : FColorList::LightSteelBlue;
 
 			CorridorPoly.SetColor(PolygonColor.WithAlpha(100));
 			CorridorPoly.Points.Reset();
 			CorridorPoly.Points.Append(Verts);
 			Snapshot->ElementsToDraw.Add(CorridorPoly);
 
-			if (AreaClass && AreaClass != UNavigationSystem::GetDefaultWalkableArea())
+			if (AreaClass && AreaClass != UNavigationSystemV1::GetDefaultWalkableArea())
 			{
 				FVisualLogShapeElement AreaMarkElem(EVisualLoggerShapeElement::Segment);
 				AreaMarkElem.SetColor(FColorList::Orange.WithAlpha(100));
diff --git a/Engine/Source/Runtime/AIModule/Private/Navigation/CrowdManager.cpp b/Engine/Source/Runtime/AIModule/Private/Navigation/CrowdManager.cpp
--- a/Engine/Source/Runtime/AIModule/Private/Navigation/CrowdManager.cpp
+++ b/Engine/Source/Runtime/AIModule/Private/Navigation/CrowdManager.cpp
@@ -4,17 +4,17 @@
 #include "GameFramework/Actor.h"
 #include "GameFramework/Pawn.h"
 #include "GameFramework/Controller.h"
-#include "AI/Navigation/NavigationSystem.h"
-#include "AI/Navigation/RecastNavMesh.h"
+#include "NavigationSystem.h"
+#include "NavMesh/RecastNavMesh.h"
 #include "VisualLogger/VisualLogger.h"
 #include "AIModuleLog.h"
 
 #if WITH_RECAST
-#include "AI/Navigation/RecastHelpers.h"
+#include "NavMesh/RecastHelpers.h"
 #include "DetourCrowd/DetourObstacleAvoidance.h"
 #include "DetourCrowd/DetourCrowd.h"
 #include "Detour/DetourNavMesh.h"
-#include "AI/Navigation/RecastQueryFilter.h"
+#include "NavMesh/RecastQueryFilter.h"
 #endif
 
 #include "Navigation/CrowdFollowingComponent.h"
@@ -691,7 +691,7 @@ void UCrowdManager::AddAgent(const ICrowdAgentInterface* Agent, FCrowdAgentData&
 	const UCrowdFollowingComponent* CrowdComponent = Cast<const UCrowdFollowingComponent>(Agent);
 	if (CrowdComponent)
 	{
-		UNavigationSystem* NavSys = Cast<UNavigationSystem>(GetOuter());
+		UNavigationSystemV1* NavSys = Cast<UNavigationSystemV1>(GetOuter());
 		MyLinkFilter = MakeShareable(new FRecastSpeciaLinkFilter(NavSys, CrowdComponent->GetOuter()));
 	}
 
@@ -805,7 +805,7 @@ void UCrowdManager::ApplyVelocity(UCrowdFollowingComponent* AgentComponent, int3
 
 void UCrowdManager::UpdateAgentPaths()
 {
-	UNavigationSystem* NavSys = Cast<UNavigationSystem>(GetOuter());
+	UNavigationSystemV1* NavSys = Cast<UNavigationSystemV1>(GetOuter());
 	ARecastNavMesh* RecastNavData = Cast<ARecastNavMesh>(MyNavData);
 	if (RecastNavData == NULL)
 	{
@@ -1255,7 +1255,7 @@ void UCrowdManager::UpdateNavData()
 {
 	if (MyNavData == NULL)
 	{
-		UNavigationSystem* NavSys = Cast<UNavigationSystem>(GetOuter());
+		UNavigationSystemV1* NavSys = Cast<UNavigationSystemV1>(GetOuter());
 		if (NavSys)
 		{
 			for (int32 Idx = 0; Idx < NavSys->NavDataSet.Num(); Idx++)
@@ -1372,18 +1372,18 @@ void UCrowdManager::PostMovePointUpdate()
 
 UWorld* UCrowdManager::GetWorld() const
 {
-	UNavigationSystem* NavSys = Cast<UNavigationSystem>(GetOuter());
+	UNavigationSystemV1* NavSys = Cast<UNavigationSystemV1>(GetOuter());
 	return NavSys ? NavSys->GetWorld() : NULL;
 }
 
 UCrowdManager* UCrowdManager::GetCurrent(UObject* WorldContextObject)
 {
-	UNavigationSystem* NavSys = UNavigationSystem::GetCurrent(WorldContextObject);
+	UNavigationSystemV1* NavSys = FNavigationSystem::GetCurrent<UNavigationSystemV1>(WorldContextObject);
 	return NavSys ? Cast<UCrowdManager>(NavSys->GetCrowdManager()) : NULL;
 }
 
 UCrowdManager* UCrowdManager::GetCurrent(UWorld* World)
 {
-	UNavigationSystem* NavSys = UNavigationSystem::GetCurrent(World);
+	UNavigationSystemV1* NavSys = FNavigationSystem::GetCurrent<UNavigationSystemV1>(World);
 	return NavSys ? Cast<UCrowdManager>(NavSys->GetCrowdManager()) : NULL;
 }
diff --git a/Engine/Source/Runtime/AIModule/Private/Navigation/MetaNavMeshPath.cpp b/Engine/Source/Runtime/AIModule/Private/Navigation/MetaNavMeshPath.cpp
--- a/Engine/Source/Runtime/AIModule/Private/Navigation/MetaNavMeshPath.cpp
+++ b/Engine/Source/Runtime/AIModule/Private/Navigation/MetaNavMeshPath.cpp
@@ -2,10 +2,11 @@
 
 #include "Navigation/MetaNavMeshPath.h"
 #include "GameFramework/Controller.h"
-#include "AI/Navigation/NavigationSystem.h"
+#include "NavigationSystem.h"
 #include "VisualLogger/VisualLoggerTypes.h"
 #include "VisualLogger/VisualLogger.h"
-#include "AI/Navigation/NavAreas/NavArea.h"
+#include "NavAreas/NavArea.h"
+#include "NavigationSystemTypes.h"
 #include "DrawDebugHelpers.h"
 
 const FNavPathType FMetaNavMeshPath::Type(&FMetaNavMeshPath::Super::Type);
@@ -35,8 +36,8 @@ FMetaNavMeshPath::FMetaNavMeshPath(const TArray<FMetaPathWayPoint>& InWaypoints,
 FMetaNavMeshPath::FMetaNavMeshPath(const TArray<FMetaPathWayPoint>& InWaypoints, const AController& Owner)
 	: FMetaNavMeshPath()
 {
-	UNavigationSystem* NavSys = UNavigationSystem::GetCurrent(Owner.GetWorld());
-	const ANavigationData* NavData = NavSys ? NavSys->GetNavDataForProps(Owner.GetNavAgentPropertiesRef()) : nullptr;
+	UNavigationSystemV1* NavSys = FNavigationSystem::GetCurrent<UNavigationSystemV1>(Owner.GetWorld());
+	const ANavigationData* NavData = NavSys ? Cast<const ANavigationData>(NavSys->GetNavDataForProps(Owner.GetNavAgentPropertiesRef())) : (const ANavigationData*)nullptr;
 
 	if (ensure(NavData))
 	{
@@ -57,8 +58,8 @@ FMetaNavMeshPath::FMetaNavMeshPath(const TArray<FVector>& InWaypoints, const ANa
 
 FMetaNavMeshPath::FMetaNavMeshPath(const TArray<FVector>& InWaypoints, const AController& Owner) : FMetaNavMeshPath()
 {
-	UNavigationSystem* NavSys = UNavigationSystem::GetCurrent(Owner.GetWorld());
-	const ANavigationData* NavData = NavSys ? NavSys->GetNavDataForProps(Owner.GetNavAgentPropertiesRef()) : nullptr;
+	UNavigationSystemV1* NavSys = FNavigationSystem::GetCurrent<UNavigationSystemV1>(Owner.GetWorld());
+	const ANavigationData* NavData = Cast<const ANavigationData>(NavSys ? NavSys->GetNavDataForProps(Owner.GetNavAgentPropertiesRef()) : nullptr);
 
 	if (ensure(NavData))
 	{
@@ -221,7 +222,7 @@ float FMetaNavMeshPath::GetLengthFromPosition(FVector SegmentStart, uint32 NextP
 float FMetaNavMeshPath::GetCostFromIndex(int32 PathPointIndex) const
 {
 	// return approximation of full path * default cost, there's not enough data to give accurate value
-	const UNavArea* DefaultAreaOb = UNavigationSystem::GetDefaultWalkableArea().GetDefaultObject();
+	const UNavArea* DefaultAreaOb = UNavigationSystemV1::GetDefaultWalkableArea().GetDefaultObject();
 	const float DefaultAreaCost = DefaultAreaOb ? DefaultAreaOb->DefaultCost : 1.0f;
 	return ApproximateLength * DefaultAreaCost;
 }
diff --git a/Engine/Source/Runtime/AIModule/Private/Navigation/NavLocalGridData.cpp b/Engine/Source/Runtime/AIModule/Private/Navigation/NavLocalGridData.cpp
--- a/Engine/Source/Runtime/AIModule/Private/Navigation/NavLocalGridData.cpp
+++ b/Engine/Source/Runtime/AIModule/Private/Navigation/NavLocalGridData.cpp
@@ -3,7 +3,7 @@
 #include "Navigation/NavLocalGridData.h"
 #include "Navigation/NavLocalGridManager.h"
 #include "GraphAStar.h"
-#include "AI/Navigation/NavigationPath.h"
+#include "NavigationPath.h"
 
 namespace FGridMathHelpers
 {
diff --git a/Engine/Source/Runtime/AIModule/Private/Navigation/NavLocalGridManager.cpp b/Engine/Source/Runtime/AIModule/Private/Navigation/NavLocalGridManager.cpp
--- a/Engine/Source/Runtime/AIModule/Private/Navigation/NavLocalGridManager.cpp
+++ b/Engine/Source/Runtime/AIModule/Private/Navigation/NavLocalGridManager.cpp
@@ -2,7 +2,7 @@
 
 #include "Navigation/NavLocalGridManager.h"
 #include "AISystem.h"
-#include "AI/Navigation/NavigationSystem.h"
+#include "NavigationSystem.h"
 #include "Engine/Engine.h"
 
 float UNavLocalGridManager::GridCellSize = 50.0f;
@@ -199,8 +199,8 @@ void UNavLocalGridManager::ProjectGrids(const TArray<int32>& GridIndices)
 {
 	DECLARE_SCOPE_CYCLE_COUNTER(TEXT("NavGrid: Project"), STAT_GridProject, STATGROUP_AI);
 
-	const UNavigationSystem* NavSys = UNavigationSystem::GetCurrent(GetWorld());
-	const ANavigationData* NavData = NavSys ? NavSys->GetMainNavData() : nullptr;
+	const UNavigationSystemV1* NavSys = FNavigationSystem::GetCurrent<UNavigationSystemV1>(GetWorld());
+	const ANavigationData* NavData = NavSys ? NavSys->GetDefaultNavDataInstance() : nullptr;
 	
 	if (NavData)
 	{
diff --git a/Engine/Source/Runtime/AIModule/Private/Navigation/PathFollowingComponent.cpp b/Engine/Source/Runtime/AIModule/Private/Navigation/PathFollowingComponent.cpp
--- a/Engine/Source/Runtime/AIModule/Private/Navigation/PathFollowingComponent.cpp
+++ b/Engine/Source/Runtime/AIModule/Private/Navigation/PathFollowingComponent.cpp
@@ -5,16 +5,16 @@
 #include "TimerManager.h"
 #include "GameFramework/Pawn.h"
 #include "GameFramework/Controller.h"
-#include "AI/Navigation/NavigationSystem.h"
-#include "AI/Navigation/RecastNavMesh.h"
+#include "NavigationSystem.h"
+#include "NavMesh/RecastNavMesh.h"
 #include "AISystem.h"
 #include "BrainComponent.h"
 #include "Engine/Canvas.h"
 #include "AIController.h"
 #include "VisualLogger/VisualLoggerTypes.h"
 #include "VisualLogger/VisualLogger.h"
-#include "AI/Navigation/AbstractNavData.h"
-#include "AI/Navigation/NavLinkCustomInterface.h"
+#include "AbstractNavData.h"
+#include "NavLinkCustomInterface.h"
 #include "Navigation/MetaNavMeshPath.h"
 #include "AIConfig.h"
 
@@ -618,6 +618,22 @@ void UPathFollowingComponent::UpdateCachedComponents()
 	UpdateMovementComponent(/*bForce=*/true);
 }
 
+void UPathFollowingComponent::OnNewPawn(APawn* NewPawn)
+{
+	UpdateCachedComponents();
+}
+
+void UPathFollowingComponent::OnRegister()
+{
+	Super::OnRegister();
+	AController* ControllerOwner = Cast<AController>(GetOwner());
+
+	if (ControllerOwner)
+	{
+		ControllerOwner->GetOnNewPawnNotifier().AddUObject(this, &UPathFollowingComponent::OnNewPawn);
+	}
+}
+
 void UPathFollowingComponent::TickComponent(float DeltaTime, enum ELevelTick TickType, FActorComponentTickFunction *ThisTickFunction)
 {
 	Super::TickComponent(DeltaTime, TickType, ThisTickFunction);
@@ -644,21 +660,22 @@ void UPathFollowingComponent::SetMovementComponent(UNavMovementComponent* MoveCo
 	{
 		const FNavAgentProperties& NavAgentProps = MoveComp->GetNavAgentPropertiesRef();
 		MyDefaultAcceptanceRadius = NavAgentProps.AgentRadius * 0.1f;
-		MoveComp->PathFollowingComp = this;
+		MoveComp->SetPathFollowingAgent(Cast<IPathFollowingAgentInterface>(this));
 
-		UWorld* MyWorld = GetWorld();
-		if (MyWorld && MyWorld->GetNavigationSystem())
+		UNavigationSystemV1* NavSys = FNavigationSystem::GetCurrent<UNavigationSystemV1>(GetWorld());
+		
+		if (NavSys)
 		{
-			MyNavData = MyWorld->GetNavigationSystem()->GetNavDataForProps(NavAgentProps);
+			MyNavData = NavSys->GetNavDataForProps(NavAgentProps);
 			if (MyNavData == nullptr)
 			{
-				if (MyWorld->GetNavigationSystem()->IsInitialized() == false)
+				if (NavSys->IsInitialized() == false)
 				{
-					MyWorld->GetNavigationSystem()->OnNavigationInitDone.AddUObject(this, &UPathFollowingComponent::OnNavigationInitDone);
+					NavSys->OnNavigationInitDone.AddUObject(this, &UPathFollowingComponent::OnNavigationInitDone);
 				}
 				else
 				{
-					MyWorld->GetNavigationSystem()->OnNavDataRegisteredEvent.AddUniqueDynamic(this, &UPathFollowingComponent::OnNavDataRegistered);
+					NavSys->OnNavDataRegisteredEvent.AddUniqueDynamic(this, &UPathFollowingComponent::OnNavDataRegistered);
 				}
 			}
 		}
@@ -670,10 +687,11 @@ void UPathFollowingComponent::OnNavigationInitDone()
 	UWorld* MyWorld = GetWorld();
 	if (MovementComp && MyWorld)
 	{
-		check(MyWorld->GetNavigationSystem());
+		UNavigationSystemV1* NavSys = FNavigationSystem::GetCurrent<UNavigationSystemV1>(MyWorld);
+		check(NavSys);
 		const FNavAgentProperties& NavAgentProps = MovementComp->GetNavAgentPropertiesRef();
-		MyNavData = MyWorld->GetNavigationSystem()->GetNavDataForProps(NavAgentProps);
-		MyWorld->GetNavigationSystem()->OnNavigationInitDone.RemoveAll(this);
+		MyNavData = NavSys->GetNavDataForProps(NavAgentProps);
+		NavSys->OnNavigationInitDone.RemoveAll(this);
 	}
 }
 
@@ -686,8 +704,9 @@ void UPathFollowingComponent::OnNavDataRegistered(ANavigationData* NavData)
 		{
 			MyNavData = NavData;
 			UWorld* MyWorld = GetWorld();
-			check(MyWorld && MyWorld->GetNavigationSystem());
-			MyWorld->GetNavigationSystem()->OnNavDataRegisteredEvent.RemoveAll(this);
+			UNavigationSystemV1* NavSys = FNavigationSystem::GetCurrent<UNavigationSystemV1>(MyWorld);
+			check(NavSys);
+			NavSys->OnNavDataRegisteredEvent.RemoveAll(this);
 		}
 	}
 }
@@ -848,7 +867,7 @@ void UPathFollowingComponent::SetMoveSegment(int32 SegmentStartIndex)
 		// handle moving through custom nav links
 		if (PathPt0.CustomLinkId)
 		{
-			UNavigationSystem* NavSys = UNavigationSystem::GetCurrent(GetWorld());
+			UNavigationSystemV1* NavSys = FNavigationSystem::GetCurrent<UNavigationSystemV1>(GetWorld());
 			INavLinkCustomInterface* CustomNavLink = NavSys->GetCustomLink(PathPt0.CustomLinkId);
 			StartUsingCustomLink(CustomNavLink, SegmentEnd);
 		}
@@ -1559,8 +1578,8 @@ FNavLocation UPathFollowingComponent::GetCurrentNavLocation() const
 		return FNavLocation();
 	}
 
-	UWorld* MyWorld = GetWorld();
-	if (MyWorld == nullptr || MyWorld->GetNavigationSystem() == nullptr)
+	UNavigationSystemV1* NavSys = FNavigationSystem::GetCurrent<UNavigationSystemV1>(GetWorld());
+	if (NavSys == nullptr)
 	{
 		return FNavLocation();
 	}
@@ -1574,7 +1593,7 @@ FNavLocation UPathFollowingComponent::GetCurrentNavLocation() const
 		const AActor* OwnerActor = MovementComp->GetOwner();
 		const FVector OwnerExtent = OwnerActor ? OwnerActor->GetSimpleCollisionCylinderExtent() : FVector::ZeroVector;
 
-		MyWorld->GetNavigationSystem()->ProjectPointToNavigation(OwnerLoc, CurrentNavLocation, OwnerExtent, MyNavData);
+		NavSys->ProjectPointToNavigation(OwnerLoc, CurrentNavLocation, OwnerExtent, MyNavData);
 	}
 
 	return CurrentNavLocation;
@@ -1754,6 +1773,11 @@ bool UPathFollowingComponent::IsPathFollowingAllowed() const
 	return MovementComp && MovementComp->CanStartPathFollowing();
 }
 
+void UPathFollowingComponent::OnUnableToMove(const UObject& Instigator)
+{
+	AbortMove(Instigator, FPathFollowingResultFlags::MovementStop);
+}
+
 void UPathFollowingComponent::OnStartedFalling()
 {
 	bWalkingNavLinkStart = false;
diff --git a/Engine/Source/Runtime/AIModule/Private/Perception/AISense_Sight.cpp b/Engine/Source/Runtime/AIModule/Private/Perception/AISense_Sight.cpp
--- a/Engine/Source/Runtime/AIModule/Private/Perception/AISense_Sight.cpp
+++ b/Engine/Source/Runtime/AIModule/Private/Perception/AISense_Sight.cpp
@@ -633,10 +633,11 @@ void UAISense_Sight::RemoveAllQueriesToTarget(const FAISightTarget::FTargetId& T
 void UAISense_Sight::OnListenerForgetsActor(const FPerceptionListener& Listener, AActor& ActorToForget)
 {
 	const uint32 ListenerId = Listener.GetListenerID();
+	const uint32 TargetId = ActorToForget.GetUniqueID();
 	
 	for (FAISightQuery& SightQuery : SightQueryQueue)
 	{
-		if (SightQuery.ObserverId == ListenerId)
+		if (SightQuery.ObserverId == ListenerId && SightQuery.TargetId == TargetId)
 		{
 			// assuming one query per observer-target pair
 			SightQuery.ForgetPreviousResult();
diff --git a/Engine/Source/Runtime/AssetRegistry/Private/AssetDataGatherer.cpp b/Engine/Source/Runtime/AssetRegistry/Private/AssetDataGatherer.cpp
--- a/Engine/Source/Runtime/AssetRegistry/Private/AssetDataGatherer.cpp
+++ b/Engine/Source/Runtime/AssetRegistry/Private/AssetDataGatherer.cpp
@@ -386,7 +386,7 @@ FAssetDataGatherer::FAssetDataGatherer(const TArray<FString>& InPaths, const TAr
 	, bFinishedInitialDiscovery( false )
 	, Thread(nullptr)
 {
-	bGatherDependsData = (GIsEditor && !FParse::Param( FCommandLine::Get(), TEXT("NoDependsGathering") )) || FParse::Param(FCommandLine::Get(),TEXT("ForceDependsGathering")) ;
+	bGatherDependsData = (GIsEditor && !FParse::Param( FCommandLine::Get(), TEXT("NoDependsGathering") )) || FParse::Param(FCommandLine::Get(),TEXT("ForceDependsGathering"));
 
 	if (FParse::Param(FCommandLine::Get(), TEXT("NoAssetRegistryCache")) || FParse::Param(FCommandLine::Get(), TEXT("multiprocess")))
 	{
@@ -397,7 +397,7 @@ FAssetDataGatherer::FAssetDataGatherer(const TArray<FString>& InPaths, const TAr
 		if (AssetDataCacheMode == EAssetDataCacheMode::UseMonolithicCache)
 		{
 			bLoadAndSaveCache = true;
-			CacheFilename = FPaths::ProjectIntermediateDir() / TEXT("CachedAssetRegistry.bin");
+			CacheFilename = FPaths::ProjectIntermediateDir() / (bGatherDependsData ? TEXT("CachedAssetRegistry.bin") : TEXT("CachedAssetRegistryNoDeps.bin"));
 		}
 		else if (InPaths.Num() > 0)
 		{
@@ -409,7 +409,7 @@ FAssetDataGatherer::FAssetDataGatherer(const TArray<FString>& InPaths, const TAr
 			}
 
 			bLoadAndSaveCache = true;
-			CacheFilename = FPaths::ProjectIntermediateDir() / TEXT("AssetRegistryCache") / FString::Printf(TEXT("%08x.bin"), CacheHash);
+			CacheFilename = FPaths::ProjectIntermediateDir() / TEXT("AssetRegistryCache") / FString::Printf(TEXT("%08x%s.bin"), CacheHash, bGatherDependsData ? TEXT("") : TEXT("NoDeps"));
 		}
 	}
 
diff --git a/Engine/Source/Runtime/Core/Public/UObject/UnrealNames.inl b/Engine/Source/Runtime/Core/Public/UObject/UnrealNames.inl
--- a/Engine/Source/Runtime/Core/Public/UObject/UnrealNames.inl
+++ b/Engine/Source/Runtime/Core/Public/UObject/UnrealNames.inl
@@ -129,6 +129,7 @@ REGISTER_NAME(250,BSP)
 REGISTER_NAME(251,EditorSettings)
 REGISTER_NAME(252,AudioThread)
 REGISTER_NAME(253,ID)
+REGISTER_NAME(254,UserDefinedEnum)
 
 // Online
 REGISTER_NAME(280,DGram)
diff --git a/Engine/Source/Runtime/CoreUObject/Private/Blueprint/BlueprintSupport.cpp b/Engine/Source/Runtime/CoreUObject/Private/Blueprint/BlueprintSupport.cpp
--- a/Engine/Source/Runtime/CoreUObject/Private/Blueprint/BlueprintSupport.cpp
+++ b/Engine/Source/Runtime/CoreUObject/Private/Blueprint/BlueprintSupport.cpp
@@ -1860,8 +1860,23 @@ void FLinkerLoad::ResolveDeferredExports(UClass* LoadClass)
 			TGuardValue<int32> ClearDeferredCDOToPreventDeferExportCreation(DeferredCDOIndex, INDEX_NONE);
 			for(int32 ExportIndex : DeferredTemplateObjects)
 			{
-				ExportMap[ExportIndex].Object = nullptr;
-				CreateExport(ExportIndex);
+				FObjectExport& Export = ExportMap[ExportIndex];
+				ULinkerPlaceholderExportObject* PlaceholderExport = Cast<ULinkerPlaceholderExportObject>(Export.Object);
+				if (ensure(PlaceholderExport))
+				{
+					// replace the placeholder with the proper object instance
+					Export.Object = nullptr;
+					UObject* ExportObj = CreateExport(ExportIndex);
+
+					PlaceholderExport->ResolveAllPlaceholderReferences(ExportObj);
+					ResolvedDeferredSubobjects(PlaceholderExport);
+
+					PlaceholderExport->MarkPendingKill();
+					if (ExportObj != nullptr)
+					{
+						Preload(ExportObj);
+					}
+				}
 			}
 		}
 
diff --git a/Engine/Source/Runtime/CoreUObject/Private/UObject/CoreRedirects.cpp b/Engine/Source/Runtime/CoreUObject/Private/UObject/CoreRedirects.cpp
--- a/Engine/Source/Runtime/CoreUObject/Private/UObject/CoreRedirects.cpp
+++ b/Engine/Source/Runtime/CoreUObject/Private/UObject/CoreRedirects.cpp
@@ -445,7 +445,10 @@ bool FCoreRedirects::RedirectNameAndValues(ECoreRedirectFlags Type, const FCoreR
 				{
 					if (*FoundValueRedirect)
 					{
-						UE_LOG(LogLinker, Error, TEXT("RedirectNameAndValues(%s) found multiple conflicting value redirects, %s and %s!"), *OldObjectName.ToString(), *(*FoundValueRedirect)->OldName.ToString(), *Redirect->OldName.ToString());
+						if (!LegacyCompareEqual((*FoundValueRedirect)->ValueChanges, Redirect->ValueChanges))
+						{
+							UE_LOG(LogLinker, Error, TEXT("RedirectNameAndValues(%s) found multiple conflicting value redirects, %s and %s!"), *OldObjectName.ToString(), *(*FoundValueRedirect)->OldName.ToString(), *Redirect->OldName.ToString());
+						}
 					}
 					else
 					{
@@ -1009,6 +1012,11 @@ ECoreRedirectFlags FCoreRedirects::GetFlagsForTypeName(FName PackageName, FName
 		return ECoreRedirectFlags::Type_Class;
 	}
 
+	if (TypeName == NAME_UserDefinedEnum)
+	{
+		return ECoreRedirectFlags::Type_Enum;
+	}
+
 	return ECoreRedirectFlags::Type_Object;
 }
 
@@ -1059,9 +1067,9 @@ static void RegisterNativeRedirects40(TArray<FCoreRedirect>& Redirects)
 	CLASS_REDIRECT("MovementComp_Character", "CharacterMovementComponent");
 	CLASS_REDIRECT("MovementComp_Projectile", "ProjectileMovementComponent");
 	CLASS_REDIRECT("MovementComp_Rotating", "RotatingMovementComponent");
-	CLASS_REDIRECT("NavAreaDefault", "/Script/Engine.NavArea_Default");
-	CLASS_REDIRECT("NavAreaDefinition", "/Script/Engine.NavArea");
-	CLASS_REDIRECT("NavAreaNull", "/Script/Engine.NavArea_Null");
+	CLASS_REDIRECT("NavAreaDefault", "/Script/NavigationSystem.NavArea_Default");
+	CLASS_REDIRECT("NavAreaDefinition", "/Script/NavigationSystem.NavArea");
+	CLASS_REDIRECT("NavAreaNull", "/Script/NavigationSystem.NavArea_Null");
 	CLASS_REDIRECT("PhysicsActor", "StaticMeshActor");
 	CLASS_REDIRECT("PhysicsBSJointActor", "PhysicsConstraintActor");
 	CLASS_REDIRECT("PhysicsHingeActor", "PhysicsConstraintActor");
@@ -1098,7 +1106,7 @@ static void RegisterNativeRedirects40(TArray<FCoreRedirect>& Redirects)
 	CLASS_REDIRECT("VimGeneratedClass", "AnimBlueprintGeneratedClass");
 	CLASS_REDIRECT("VimInstance", "AnimInstance");
 	CLASS_REDIRECT("WorldInfo", "WorldSettings");
-	CLASS_REDIRECT_INSTANCES("NavAreaMeta", "/Script/Engine.NavArea_Default");
+	CLASS_REDIRECT_INSTANCES("NavAreaMeta", "/Script/NavigationSystem.NavArea_Default");
 
 	STRUCT_REDIRECT("VimDebugData", "AnimBlueprintDebugData");
 
@@ -1388,7 +1396,7 @@ static void RegisterNativeRedirects46(TArray<FCoreRedirect>& Redirects)
 	CLASS_REDIRECT("EmitterSpawnable", "Emitter");
 	CLASS_REDIRECT("SlateWidgetStyleAsset", "/Script/SlateCore.SlateWidgetStyleAsset");
 	CLASS_REDIRECT("SlateWidgetStyleContainerBase", "/Script/SlateCore.SlateWidgetStyleContainerBase");
-	CLASS_REDIRECT("SmartNavLinkComponent", "/Script/Engine.NavLinkCustomComponent");
+	CLASS_REDIRECT("SmartNavLinkComponent", "/Script/NavigationSystem.NavLinkCustomComponent");
 	CLASS_REDIRECT("WidgetBlueprint", "/Script/UMGEditor.WidgetBlueprint");
 
 	PROPERTY_REDIRECT("AnimNotify.Received_Notify.AnimSeq", "Animation");
diff --git a/Engine/Source/Runtime/CoreUObject/Private/UObject/LinkerLoad.cpp b/Engine/Source/Runtime/CoreUObject/Private/UObject/LinkerLoad.cpp
--- a/Engine/Source/Runtime/CoreUObject/Private/UObject/LinkerLoad.cpp
+++ b/Engine/Source/Runtime/CoreUObject/Private/UObject/LinkerLoad.cpp
@@ -1182,7 +1182,7 @@ FLinkerLoad::ELinkerStatus FLinkerLoad::SerializePackageFileSummary()
 				else if (SerializedCustomVersion.Version > LatestVersion->Version)
 				{
 					// Loading a package with a newer custom version than the current one.
-					UE_LOG(LogLinker, Error, TEXT("Package %s was saved with a newer custom version than the current. Tag %s  PackageVersion %d  MaxExpected %d"), *Filename, *SerializedCustomVersion.Key.ToString(), SerializedCustomVersion.Version, LatestVersion->Version);
+					UE_LOG(LogLinker, Error, TEXT("Package %s was saved with a newer custom version than the current. Tag %s Name '%s' PackageVersion %d  MaxExpected %d"), *Filename, *SerializedCustomVersion.Key.ToString(), *LatestVersion->GetFriendlyName().ToString(), SerializedCustomVersion.Version, LatestVersion->Version);
 					return LINKER_Failed;
 				}
 				else if (SerializedCustomVersion.Version != LatestVersion->Version)
@@ -2227,31 +2227,44 @@ FLinkerLoad::EVerifyResult FLinkerLoad::VerifyImport(int32 ImportIndex)
 				UObject* DestObject = Redir->DestinationObject;
 
 				// check to make sure the destination obj was loaded,
-				if ( DestObject == NULL )
+				if (DestObject == nullptr)
 				{
 					Result = VERIFY_Failed;
 				}
-				// check that in fact it was the type we thought it should be
-				else if ( DestObject->GetClass()->GetFName() != OriginalImport.ClassName
-
-					// if the destination object is a CDO, allow class changes
-					&&	!DestObject->HasAnyFlags(RF_ClassDefaultObject) )
+				else
 				{
-					Result = VERIFY_Failed;
-					// if the destination is a ObjectRedirector you've most likely made a nasty circular loop
-					if( Redir->DestinationObject->GetClass() == UObjectRedirector::StaticClass() )
+					// Blueprint CDOs are always allowed to change class, otherwise we need to do a name check for all parent classes
+					bool bIsValidClass = DestObject->HasAnyFlags(RF_ClassDefaultObject);
+					UClass* CheckClass = DestObject->GetClass();
+
+					while (!bIsValidClass && CheckClass)
 					{
-						WarningAppend += LOCTEXT("LoadWarningSuffix_circularredirection", " [circular redirection]").ToString();
+						if (CheckClass->GetFName() == OriginalImport.ClassName)
+						{
+							bIsValidClass = true;
+							break;
+						}
+						CheckClass = CheckClass->GetSuperClass();
 					}
-				}
-				else
-				{
-					Result = VERIFY_Redirected;
 
-					// now, fake our Import to be what the redirector pointed to
-					Import.XObject = Redir->DestinationObject;
-					FUObjectThreadContext::Get().ImportCount++;
-					FLinkerManager::Get().AddLoaderWithNewImports(this);
+					if (!bIsValidClass)
+					{
+						Result = VERIFY_Failed;
+						// if the destination is a ObjectRedirector you've most likely made a nasty circular loop
+						if (Redir->DestinationObject->GetClass() == UObjectRedirector::StaticClass())
+						{
+							WarningAppend += LOCTEXT("LoadWarningSuffix_circularredirection", " [circular redirection]").ToString();
+						}
+					}
+					else
+					{
+						Result = VERIFY_Redirected;
+
+						// now, fake our Import to be what the redirector pointed to
+						Import.XObject = Redir->DestinationObject;
+						FUObjectThreadContext::Get().ImportCount++;
+						FLinkerManager::Get().AddLoaderWithNewImports(this);
+					}
 				}
 			}
 		}
diff --git a/Engine/Source/Runtime/CoreUObject/Private/UObject/SavePackage.cpp b/Engine/Source/Runtime/CoreUObject/Private/UObject/SavePackage.cpp
--- a/Engine/Source/Runtime/CoreUObject/Private/UObject/SavePackage.cpp
+++ b/Engine/Source/Runtime/CoreUObject/Private/UObject/SavePackage.cpp
@@ -2519,6 +2519,11 @@ public:
 		}
 	}
 
+	/** Do nothing when serializing soft references, this is required because the presave on soft references can fix redirectors, which is unsafe at this point */
+	virtual FArchive& operator<<(FLazyObjectPtr& Value) override { return *this; }
+	virtual FArchive& operator<<(FSoftObjectPtr& Value) override { return *this; }
+	virtual FArchive& operator<<(FSoftObjectPath& Value) override { return *this; }
+
 private:
 
 	/**
diff --git a/Engine/Source/Runtime/CoreUObject/Private/UObject/SoftObjectPath.cpp b/Engine/Source/Runtime/CoreUObject/Private/UObject/SoftObjectPath.cpp
--- a/Engine/Source/Runtime/CoreUObject/Private/UObject/SoftObjectPath.cpp
+++ b/Engine/Source/Runtime/CoreUObject/Private/UObject/SoftObjectPath.cpp
@@ -7,6 +7,7 @@
 #include "Misc/PackageName.h"
 #include "UObject/LinkerLoad.h"
 #include "UObject/UObjectThreadContext.h"
+#include "UObject/CoreRedirects.h"
 #include "Misc/RedirectCollector.h"
 
 FSoftObjectPath::FSoftObjectPath(const UObject* InObject)
@@ -48,19 +49,19 @@ void FSoftObjectPath::SetPath(FString Path)
 		{
 			// Possibly an ExportText path. Trim the ClassName.
 			Path = FPackageName::ExportTextPathToObjectPath(Path);
-	}
+		}
 
 		int32 ColonIndex = INDEX_NONE;
 
 		if (Path.FindChar(':', ColonIndex))
-	{
+		{
 			// Has a subobject, split on that then create a name from the temporary path
 			SubPathString = Path.Mid(ColonIndex + 1);
 			Path.RemoveAt(ColonIndex, Path.Len() - ColonIndex);
 			AssetPathName = *Path;
-	}
-	else
-	{
+		}
+		else
+		{
 			// No Subobject
 			AssetPathName = *Path;
 			SubPathString.Empty();
@@ -71,6 +72,11 @@ void FSoftObjectPath::SetPath(FString Path)
 bool FSoftObjectPath::PreSavePath()
 {
 #if WITH_EDITOR
+	if (IsNull())
+	{
+		return false;
+	}
+
 	FName FoundRedirection = GRedirectCollector.GetAssetPathRedirection(AssetPathName);
 
 	if (FoundRedirection != NAME_None)
@@ -78,6 +84,11 @@ bool FSoftObjectPath::PreSavePath()
 		AssetPathName = FoundRedirection;
 		return true;
 	}
+
+	if (FixupCoreRedirects())
+	{
+		return true;
+	}
 #endif // WITH_EDITOR
 	return false;
 }
@@ -153,11 +164,11 @@ void FSoftObjectPath::SerializePath(FArchive& Ar)
 			PostLoadPath();
 		}
 		if (Ar.GetPortFlags()&PPF_DuplicateForPIE)
-	{
-		// Remap unique ID if necessary
-		// only for fixing up cross-level references, inter-level references handled in FDuplicateDataReader
-		FixupForPIE();
-	}
+		{
+			// Remap unique ID if necessary
+			// only for fixing up cross-level references, inter-level references handled in FDuplicateDataReader
+			FixupForPIE();
+		}
 	}
 #endif // WITH_EDITOR
 }
@@ -181,7 +192,7 @@ bool FSoftObjectPath::ExportTextItem(FString& ValueStr, FSoftObjectPath const& D
 		return false;
 	}
 
-	if (IsValid())
+	if (!IsNull())
 	{
 		// Fixup any redirectors
 		FSoftObjectPath Temp = *this;
@@ -232,8 +243,20 @@ bool FSoftObjectPath::ImportTextItem(const TCHAR*& Buffer, int32 PortFlags, UObj
 
 	SetPath(MoveTemp(ImportedPath));
 
-	// Consider this a load, so Config string references get cooked
-	PostLoadPath();
+#if WITH_EDITOR
+	if (Parent && IsEditorOnlyObject(Parent))
+	{
+		// We're probably reading config for an editor only object, we need to mark this reference as editor only
+		FSoftObjectPathSerializationScope SerializationScope(NAME_None, NAME_None, ESoftObjectPathCollectType::EditorOnlyCollect, ESoftObjectPathSerializeType::AlwaysSerialize);
+
+		PostLoadPath();
+	}
+	else
+#endif
+	{
+		// Consider this a load, so Config string references get cooked
+		PostLoadPath();
+	}
 
 	return true;
 }
@@ -259,7 +282,7 @@ bool SerializeFromMismatchedTagTemplate(FString& Output, const FPropertyTag& Tag
 			Output = ObjPtr->GetPathName();
 		}
 		else
-	{
+		{
 			Output = FString();
 		}
 		return true;
@@ -270,8 +293,8 @@ bool SerializeFromMismatchedTagTemplate(FString& Output, const FPropertyTag& Tag
 		Ar << String;
 
 		Output = String;
-	return true;
-}
+		return true;
+	}
 	return false;
 }
 
@@ -300,9 +323,22 @@ UObject* FSoftObjectPath::TryLoad() const
 {
 	UObject* LoadedObject = nullptr;
 
-	if ( IsValid() )
+	if (!IsNull())
 	{
 		LoadedObject = LoadObject<UObject>(nullptr, *ToString());
+
+#if WITH_EDITOR
+		// Look at core redirects if we didn't find the object
+		if (!LoadedObject)
+		{
+			FSoftObjectPath FixupObjectPath = *this;
+			if (FixupObjectPath.FixupCoreRedirects())
+			{
+				LoadedObject = LoadObject<UObject>(nullptr, *FixupObjectPath.ToString());
+			}
+		}
+#endif
+
 		while (UObjectRedirector* Redirector = Cast<UObjectRedirector>(LoadedObject))
 		{
 			LoadedObject = Redirector->DestinationObject;
@@ -316,7 +352,7 @@ UObject* FSoftObjectPath::ResolveObject() const
 {
 	// Don't try to resolve if we're saving a package because StaticFindObject can't be used here
 	// and we usually don't want to force references to weak pointers while saving.
-	if (!IsValid() || GIsSavingPackage)
+	if (IsNull() || GIsSavingPackage)
 	{
 		return nullptr;
 	}
@@ -327,9 +363,7 @@ UObject* FSoftObjectPath::ResolveObject() const
 	{
 		// If we are in PIE and this hasn't already been fixed up, we need to fixup at resolution time. We cannot modify the path as it may be somewhere like a blueprint CDO
 		FSoftObjectPath FixupObjectPath = *this;
-		FixupObjectPath.FixupForPIE();
-
-		if (FixupObjectPath.AssetPathName != AssetPathName)
+		if (FixupObjectPath.FixupForPIE())
 		{
 			PathString = FixupObjectPath.ToString();
 		}
@@ -337,6 +371,20 @@ UObject* FSoftObjectPath::ResolveObject() const
 #endif
 
 	UObject* FoundObject = FindObject<UObject>(nullptr, *PathString);
+
+#if WITH_EDITOR
+	// Look at core redirects if we didn't find the object
+	if (!FoundObject)
+	{
+		FSoftObjectPath FixupObjectPath = *this;
+		if (FixupObjectPath.FixupCoreRedirects())
+		{
+			PathString = FixupObjectPath.ToString();
+			FoundObject = FindObject<UObject>(nullptr, *PathString);
+		}
+	}
+#endif
+
 	while (UObjectRedirector* Redirector = Cast<UObjectRedirector>(FoundObject))
 	{
 		FoundObject = Redirector->DestinationObject;
@@ -361,10 +409,10 @@ void FSoftObjectPath::ClearPIEPackageNames()
 	PIEPackageNames.Empty();
 }
 
-void FSoftObjectPath::FixupForPIE()
+bool FSoftObjectPath::FixupForPIE()
 {
 #if WITH_EDITOR
-	if (GPlayInEditorID != INDEX_NONE && IsValid())
+	if (GPlayInEditorID != INDEX_NONE && !IsNull())
 	{
 		FString Path = ToString();
 
@@ -380,13 +428,36 @@ void FSoftObjectPath::FixupForPIE()
 
 			// Duplicate if this an already registered PIE package or this looks like a level subobject reference
 			if (bIsChildOfLevel || PIEPackageNames.Contains(PIEPackage))
-				{
-					// Need to prepend PIE prefix, as we're in PIE and this refers to an object in a PIE package
-					SetPath(MoveTemp(PIEPath));
+			{
+				// Need to prepend PIE prefix, as we're in PIE and this refers to an object in a PIE package
+				SetPath(MoveTemp(PIEPath));
+
+				return true;
 			}
 		}
 	}
 #endif
+	return false;
+}
+
+bool FSoftObjectPath::FixupCoreRedirects()
+{
+	FCoreRedirectObjectName OldName = FCoreRedirectObjectName(ToString());
+	FCoreRedirectObjectName NewName = FCoreRedirects::GetRedirectedName(ECoreRedirectFlags::Type_Object, OldName);
+
+	// This also might be a class
+	if (OldName == NewName)
+	{
+		NewName = FCoreRedirects::GetRedirectedName(ECoreRedirectFlags::Type_Class, OldName);
+	}
+
+	if (OldName != NewName)
+	{
+		SetPath(NewName.ToString());
+		return true;
+	}
+
+	return false;
 }
 
 bool FSoftClassPath::SerializeFromMismatchedTag(struct FPropertyTag const& Tag, FArchive& Ar)
@@ -406,15 +477,15 @@ bool FSoftClassPath::SerializeFromMismatchedTag(struct FPropertyTag const& Tag,
 	{
 		SetPath(MoveTemp(Path));
 		PostLoadPath();
-				}
+	}
 
 	return bReturn;
-			}
+}
 
 UClass* FSoftClassPath::ResolveClass() const
 {
 	return Cast<UClass>(ResolveObject());
-		}
+}
 
 FSoftClassPath FSoftClassPath::GetOrCreateIDForClass(const UClass *InClass)
 {
@@ -482,7 +553,7 @@ bool FSoftObjectPathThreadContext::GetSerializationOptions(FName& OutPackageName
 			if (bEditorOnly && CurrentCollectType == ESoftObjectPathCollectType::AlwaysCollect)
 			{
 				CurrentCollectType = ESoftObjectPathCollectType::EditorOnlyCollect;
-	}
+			}
 
 			bFoundAnything = true;
 		}
diff --git a/Engine/Source/Runtime/CoreUObject/Private/UObject/UObjectArchetype.cpp b/Engine/Source/Runtime/CoreUObject/Private/UObject/UObjectArchetype.cpp
--- a/Engine/Source/Runtime/CoreUObject/Private/UObject/UObjectArchetype.cpp
+++ b/Engine/Source/Runtime/CoreUObject/Private/UObject/UObjectArchetype.cpp
@@ -98,7 +98,7 @@ UObject* UObject::GetArchetypeFromRequiredInfo(UClass* Class, UObject* Outer, FN
 #if WITH_EDITOR
 	// While compiling we just want to use whatever is in the object hierarchy,
 	// as some instances within the hierarchy may also be compiling:
-	bUseUpToDateClass = Class->GetAuthoritativeClass() == Class && (!GCompilingBlueprint || GIsReinstancing);
+	bUseUpToDateClass = GIsReinstancing && Class->GetAuthoritativeClass() == Class;
 #endif
 	return GetArchetypeFromRequiredInfoImpl(Class, Outer, Name, ObjectFlags, bUseUpToDateClass);
 }
diff --git a/Engine/Source/Runtime/CoreUObject/Public/Misc/RedirectCollector.h b/Engine/Source/Runtime/CoreUObject/Public/Misc/RedirectCollector.h
--- a/Engine/Source/Runtime/CoreUObject/Public/Misc/RedirectCollector.h
+++ b/Engine/Source/Runtime/CoreUObject/Public/Misc/RedirectCollector.h
@@ -29,7 +29,8 @@ private:
 		bool operator==(const FPackagePropertyPair& Other) const
 		{
 			return Package == Other.Package &&
-				Property == Other.Property;
+				Property == Other.Property &&
+				bReferencedByEditorOnlyProperty == Other.bReferencedByEditorOnlyProperty;
 		}
 
 		const FName& GetCachedPackageName() const
diff --git a/Engine/Source/Runtime/CoreUObject/Public/UObject/Class.h b/Engine/Source/Runtime/CoreUObject/Public/UObject/Class.h
--- a/Engine/Source/Runtime/CoreUObject/Public/UObject/Class.h
+++ b/Engine/Source/Runtime/CoreUObject/Public/UObject/Class.h
@@ -2063,7 +2063,7 @@ class COREUOBJECT_API UClass : public UStruct
 #endif
 {
 	DECLARE_CASTED_CLASS_INTRINSIC_NO_CTOR(UClass, UStruct, 0, TEXT("/Script/CoreUObject"), CASTCLASS_UClass, NO_API)
-	DECLARE_WITHIN(UPackage)
+	DECLARE_WITHIN_UPACKAGE()
 
 public:
 #if UCLASS_FAST_ISA_IMPL == UCLASS_ISA_INDEXTREE
@@ -2697,7 +2697,7 @@ protected:
 class COREUOBJECT_API UDynamicClass : public UClass
 {
 	DECLARE_CASTED_CLASS_INTRINSIC_NO_CTOR(UDynamicClass, UClass, 0, TEXT("/Script/CoreUObject"), CASTCLASS_None, NO_API)
-	DECLARE_WITHIN(UPackage)
+	DECLARE_WITHIN_UPACKAGE()
 
 public:
 
diff --git a/Engine/Source/Runtime/CoreUObject/Public/UObject/ObjectMacros.h b/Engine/Source/Runtime/CoreUObject/Public/UObject/ObjectMacros.h
--- a/Engine/Source/Runtime/CoreUObject/Public/UObject/ObjectMacros.h
+++ b/Engine/Source/Runtime/CoreUObject/Public/UObject/ObjectMacros.h
@@ -1484,10 +1484,16 @@ public: \
 	DECLARE_CASTED_CLASS_INTRINSIC_WITH_API( TClass, TSuperClass, TStaticFlags, TPackage, TStaticCastFlags, NO_API) \
 
 // Declare that objects of class being defined reside within objects of the specified class.
-#define DECLARE_WITHIN( TWithinClass ) \
+#define DECLARE_WITHIN_INTERNAL( TWithinClass, bCanUseOnCDO ) \
 	/** The required type of this object's outer ({{ typedef-type }}) */ \
 	typedef class TWithinClass WithinClass;  \
-	TWithinClass* GetOuter##TWithinClass() const { return (TWithinClass*)GetOuter(); }
+	TWithinClass* GetOuter##TWithinClass() const { return (ensure(bCanUseOnCDO || !HasAnyFlags(RF_ClassDefaultObject)) ? (TWithinClass*)GetOuter() : nullptr); }
+
+#define DECLARE_WITHIN( TWithinClass ) \
+	DECLARE_WITHIN_INTERNAL( TWithinClass, false )
+
+#define DECLARE_WITHIN_UPACKAGE() \
+	DECLARE_WITHIN_INTERNAL( UPackage, true )
 
 // Register a class at startup time.
 #define IMPLEMENT_CLASS(TClass, TClassCrc) \
diff --git a/Engine/Source/Runtime/CoreUObject/Public/UObject/SoftObjectPath.h b/Engine/Source/Runtime/CoreUObject/Public/UObject/SoftObjectPath.h
--- a/Engine/Source/Runtime/CoreUObject/Public/UObject/SoftObjectPath.h
+++ b/Engine/Source/Runtime/CoreUObject/Public/UObject/SoftObjectPath.h
@@ -152,8 +152,11 @@ struct COREUOBJECT_API FSoftObjectPath
 	/** Handles when a string asset reference has been loaded, call if loading with a method that skips SerializePath. This does not modify path but might call callbacks */
 	void PostLoadPath() const;
 
-	/** Fixes up this SoftObjectPath to add or remove the PIE prefix depending on what is currently active */
-	void FixupForPIE();
+	/** Fixes up this SoftObjectPath to add or remove the PIE prefix depending on what is currently active, returns true if it was modified */
+	bool FixupForPIE();
+
+	/** Fixes soft object path for CoreRedirects to handle renamed native objects, returns true if it was modified */
+	bool FixupCoreRedirects();
 
 	FORCEINLINE friend uint32 GetTypeHash(FSoftObjectPath const& This)
 	{
diff --git a/Engine/Source/Runtime/Engine/Classes/AI/Navigation/NavAgentInterface.h b/Engine/Source/Runtime/Engine/Classes/AI/Navigation/NavAgentInterface.h
--- a/Engine/Source/Runtime/Engine/Classes/AI/Navigation/NavAgentInterface.h
+++ b/Engine/Source/Runtime/Engine/Classes/AI/Navigation/NavAgentInterface.h
@@ -9,6 +9,7 @@
 #include "NavAgentInterface.generated.h"
 
 class AActor;
+class IPathFollowingAgentInterface;
 
 UINTERFACE(MinimalAPI, meta=(CannotImplementInterfaceInBlueprint))
 class UNavAgentInterface : public UInterface
@@ -49,6 +50,9 @@ class INavAgentInterface
 	/** Checks if the agent is actively following a navigation path */
 	virtual bool IsFollowingAPath() const { return false; }
 
+	/** Retrieves PathFollowingAgent for this NavAgent */
+	virtual IPathFollowingAgentInterface* GetPathFollowingAgent() const { return nullptr; }
+
 	//----------------------------------------------------------------------//
 	// DEPRECATED
 	//----------------------------------------------------------------------//
diff --git a/Engine/Source/Runtime/Engine/Classes/AI/Navigation/NavAreaBase.h b/Engine/Source/Runtime/Engine/Classes/AI/Navigation/NavAreaBase.h
new file mode 100644
--- /dev/null
+++ b/Engine/Source/Runtime/Engine/Classes/AI/Navigation/NavAreaBase.h
@@ -0,0 +1,46 @@
+// Copyright 1998-2018 Epic Games, Inc. All Rights Reserved.
+
+#pragma once
+
+#include "UObject/Object.h"
+#include "AI/Navigation/NavigationTypes.h"
+#include "NavAreaBase.generated.h"
+
+
+// a stub class. Actual implementation in NavigationSystem module.
+UCLASS(DefaultToInstanced, abstract, Config = Engine)
+class ENGINE_API UNavAreaBase : public UObject
+{
+	GENERATED_BODY()
+
+protected:
+	uint8 bIsMetaArea : 1;
+
+public:
+	UNavAreaBase(const FObjectInitializer& ObjectInitializer = FObjectInitializer::Get());
+
+	// leftover from NavigationSystem extraction from the Engine code
+	virtual bool IsLowArea() const { return false; }
+	virtual bool IsMetaArea() const { return (bIsMetaArea != 0); }
+
+	/**
+	*	Picks an navigation area class that should be used for Actor when
+	*	queried by NavAgent. 
+	*/
+	static TSubclassOf<UNavAreaBase> PickAreaClassForAgent(TSubclassOf<UNavAreaBase> AreaClass, const AActor& Actor, const FNavAgentProperties& NavAgent)
+	{
+		const UNavAreaBase* CDO = AreaClass.Get() ? AreaClass->GetDefaultObject<const UNavAreaBase>() : nullptr;
+		return CDO && CDO->IsMetaArea()
+			? CDO->PickAreaClassForAgent(Actor, NavAgent)
+			: AreaClass;
+	}
+	
+
+protected:
+	/**
+	*	Picks an navigation area class that should be used for Actor when
+	*	queried by NavAgent. Call it via the UNavAreaBase::PickAreaClass
+	*/
+	virtual TSubclassOf<UNavAreaBase> PickAreaClassForAgent(const AActor& Actor, const FNavAgentProperties& NavAgent) const;
+};
+
diff --git a/Engine/Source/Runtime/Engine/Classes/AI/Navigation/NavLinkDefinition.h b/Engine/Source/Runtime/Engine/Classes/AI/Navigation/NavLinkDefinition.h
--- a/Engine/Source/Runtime/Engine/Classes/AI/Navigation/NavLinkDefinition.h
+++ b/Engine/Source/Runtime/Engine/Classes/AI/Navigation/NavLinkDefinition.h
@@ -10,6 +10,8 @@
 #include "AI/Navigation/NavigationTypes.h"
 #include "NavLinkDefinition.generated.h"
 
+class UNavAreaBase;
+
 UENUM()
 namespace ENavLinkDirection
 {
@@ -159,7 +161,7 @@ struct ENGINE_API FNavigationLinkBase
 private:
 	/** Area type of this link (empty = default) */
 	UPROPERTY(EditAnywhere, Category = Default)
-	TSubclassOf<class UNavArea> AreaClass;
+	TSubclassOf<UNavAreaBase> AreaClass;
 
 	TWeakObjectPtr<UClass> AreaClassOb;
 };
diff --git a/Engine/Source/Runtime/Engine/Classes/AI/Navigation/NavRelevantInterface.h b/Engine/Source/Runtime/Engine/Classes/AI/Navigation/NavRelevantInterface.h
--- a/Engine/Source/Runtime/Engine/Classes/AI/Navigation/NavRelevantInterface.h
+++ b/Engine/Source/Runtime/Engine/Classes/AI/Navigation/NavRelevantInterface.h
@@ -6,10 +6,113 @@
 #include "UObject/ObjectMacros.h"
 #include "UObject/Interface.h"
 #include "AI/Navigation/NavigationTypes.h"
+#include "AI/NavigationModifier.h"
 #include "NavRelevantInterface.generated.h"
 
 struct FNavigableGeometryExport;
-struct FNavigationRelevantData;
+
+struct FNavigationRelevantDataFilter 
+{
+	/** pass when actor has geometry */
+	uint32 bIncludeGeometry : 1;
+	/** pass when actor has any offmesh link modifier */
+	uint32 bIncludeOffmeshLinks : 1;
+	/** pass when actor has any area modifier */
+	uint32 bIncludeAreas : 1;
+	/** pass when actor has any modifier with meta area */
+	uint32 bIncludeMetaAreas : 1;
+
+	FNavigationRelevantDataFilter() 
+		: bIncludeGeometry(false)
+		, bIncludeOffmeshLinks(false)
+		, bIncludeAreas(false)
+		, bIncludeMetaAreas(false)
+	{}
+};
+
+// @todo consider optional structures that can contain a delegate instead of 
+// actual copy of collision data
+struct ENGINE_API FNavigationRelevantData : public TSharedFromThis<FNavigationRelevantData, ESPMode::ThreadSafe>
+{
+	DECLARE_DELEGATE_RetVal_OneParam(bool, FFilterNavDataDelegate, const struct FNavDataConfig*);
+
+	/** CollisionData should always start with this struct for validation purposes */
+	struct FCollisionDataHeader
+	{
+		int32 DataSize;
+
+		static bool IsValid(const uint8* RawData, int32 RawDataSize);
+	};
+
+	/** exported geometry (used by recast navmesh as FRecastGeometryCache) */
+	TNavStatArray<uint8> CollisionData;
+
+	/** cached voxels (used by recast navmesh as FRecastVoxelCache) */
+	TNavStatArray<uint8> VoxelData;
+
+	/** bounds of geometry (unreal coords) */
+	FBox Bounds;
+
+	/** Gathers per instance data for navigation geometry in a specified area box */
+	FNavDataPerInstanceTransformDelegate NavDataPerInstanceTransformDelegate;
+
+	/** called to check if hosted geometry should be used for given FNavDataConfig. If not set then "true" is assumed. */
+	FFilterNavDataDelegate ShouldUseGeometryDelegate;
+
+	/** additional modifiers: areas and external links */
+	FCompositeNavModifier Modifiers;
+
+	/** UObject these data represents */
+	const TWeakObjectPtr<UObject> SourceObject;
+
+	/** get set to true when lazy navigation exporting is enabled and this navigation data has "potential" of
+	*	containing geometry data. First access will result in gathering the data and setting this flag back to false.
+	*	Mind that this flag can go back to 'true' if related data gets cleared out. */
+	uint32 bPendingLazyGeometryGathering : 1;
+	uint32 bPendingLazyModifiersGathering : 1;
+
+	uint32 bSupportsGatheringGeometrySlices : 1;
+
+	FNavigationRelevantData(UObject& Source)
+		: SourceObject(&Source)
+		, bPendingLazyGeometryGathering(false)
+		, bPendingLazyModifiersGathering(false)
+	{}
+
+	FORCEINLINE bool HasGeometry() const { return VoxelData.Num() || CollisionData.Num(); }
+	FORCEINLINE bool HasModifiers() const { return !Modifiers.IsEmpty(); }
+	FORCEINLINE bool IsPendingLazyGeometryGathering() const { return bPendingLazyGeometryGathering; }
+	FORCEINLINE bool IsPendingLazyModifiersGathering() const { return bPendingLazyModifiersGathering; }
+	FORCEINLINE bool SupportsGatheringGeometrySlices() const { return bSupportsGatheringGeometrySlices; }
+	FORCEINLINE bool IsEmpty() const { return !HasGeometry() && !HasModifiers(); }
+	FORCEINLINE uint32 GetAllocatedSize() const { return CollisionData.GetAllocatedSize() + VoxelData.GetAllocatedSize() + Modifiers.GetAllocatedSize(); }
+	FORCEINLINE uint32 GetGeometryAllocatedSize() const { return CollisionData.GetAllocatedSize() + VoxelData.GetAllocatedSize(); }
+	FORCEINLINE int32 GetDirtyFlag() const
+	{
+		return ((HasGeometry() || IsPendingLazyGeometryGathering()) ? ENavigationDirtyFlag::Geometry : 0) |
+			((HasModifiers() || IsPendingLazyModifiersGathering()) ? ENavigationDirtyFlag::DynamicModifier : 0) |
+			(Modifiers.HasAgentHeightAdjust() ? ENavigationDirtyFlag::UseAgentHeight : 0);
+	}
+
+	bool HasPerInstanceTransforms() const;
+	bool IsMatchingFilter(const FNavigationRelevantDataFilter& Filter) const;
+	void Shrink();
+	bool IsCollisionDataValid() const;
+
+	void ValidateAndShrink()
+	{
+		if (IsCollisionDataValid())
+		{
+			Shrink();
+		}
+		else
+		{
+			CollisionData.Empty();
+		}
+	}
+
+	FORCEINLINE UObject* GetOwner() const { return SourceObject.Get(); }
+};
 
 UINTERFACE(MinimalAPI, meta=(CannotImplementInterfaceInBlueprint))
 class UNavRelevantInterface : public UInterface
