Copying //UE4/Odin-Staging to //UE4/Dev-Main (Source: //Odin/Main/Engine @ 3913580)

Change 3369800 by [REDACTED_PERSON]

	#odin (Main) - Add new OculusKioskAppID to INI and use it for entitlement checks when "-kiosk" is on the commandline. Takes precedence over -demo. Restored demo app ID to yesterday's value and switched kiosk to the new one requested by oculus.

Change 3335679 by [REDACTED_PERSON]

	Focus the Blueprint viewport for mod assets, when you open them (and they're still DataOnly) - frame a specified component so the user can tell what the Blueprint is.

Change 3311911 by [REDACTED_PERSON]

	Adding templates for Day and Night maps.

Change 3311411 by [REDACTED_PERSON]

	ODIN-6185 Celebration screen doesn╞t appear, ODIN-6178 2-2: bots reach player before intro timer completes

Change 3307677 by [REDACTED_PERSON]

	Forcing the Odin mod Blueprints (and sub-classes) to always open in the full Blueprint editor.

Change 3305535 by [REDACTED_PERSON]

	Add an option to UAT (-CookOutputDir=...) and the cooker (-OutputDir=...) which allows overriding the output directory for cooked files, and fix situations where the directory becomes too deep.

Change 3304566 by [REDACTED_PERSON]

	Add support for hot-reloading game plugin modules from Visual Studio, as long as their module returns IsGameModule() = true, and enable it by default for mods.

Change 3303091 by [REDACTED_PERSON]

	Fixes/improvements for mod editor and code mods:

	* A separate top-level project is generated for each code mod in the Visual Studio solution.
	* Plugin descriptors now have a flag to identify themselves as mod as opposed to a regular game plugin, which prevents project plugins (eg. OdinEditor) from getting their own VS project. New mods created with the mod editor will have this set by default, as do the three existing sample mods.
	* Cleaning and building code mods will never modify engine binaries. Presence of the Engine/Build/InstalledProjectBuild.txt file is used to indicate running in this environment. This flag also disables options to edit metadata for non-mod plugins in installed builds.
	* Plugin browser now includes a separate category for mods.
	* Mod editor now behaves as an "installed" program by default, and will use the user's home folder for storing settings.

Change 3302585 by [REDACTED_PERSON]

	All nativized functions/event declared for tha first time in Blueprint are BlueprintNativeEvent. A BP from mode, based on a nativized class, can override any function.

Change 3301859 by [REDACTED_PERSON]

	Making the new plugin wizard tab a bit bigger by default to accomodate the new information inside it.

Change 3301057 by [REDACTED_PERSON]

	Fixing use of the demo app id to be consistent, no matter where the plugin intiializes

Change 3298141 by [REDACTED_PERSON]

Change 3297043 by [REDACTED_PERSON]

	- Blank
	- OdinWeapon
	- OdinCharacter

	These still act as mixin templates, and are compatible with each other and with other mod content

Change 3295655 by [REDACTED_PERSON]

	Fix list of discovered assets being cleared by second call to FindFilesRecursive() when building DDC. Disable the -cookdir parameter again.

Change 3291094 by [REDACTED_PERSON]

	Fixing an issue where mod assets generated tags would not be correct after being created from the new mod wizard. Also now enables for PLUGIN_NAME macros to be used on uasset and umap files in plugin templates.

Change 3289510 by [REDACTED_PERSON]

	Fixed ODIN-5718 : A "Lost Tracker" prompt is displayed when launching the game

	- We now broadcast a VRTrackerAttached() on startup so that we can have an initial tracker count.
	- Fixed issue where VRTrackerAttached()/VRTrackerLost() would never be broadcast on Oculus platform in UE_BUILD_SHIPPING.
	- bIsInLoadingSequence is now native
	- VR Error prompts will now show up while loading
	- TrackerLost message persists across level transitions/loads

Change 3288170 by [REDACTED_PERSON]

	For multi-select plugin wizard definitions, the plugin wizard will now show a tile view instead of a list view when selecting templates. For created plugins that can contain content, the content browser will now be forced to show plugin or engine content if the user wants the content folder shown on plugin creation.

Change 3287456 by [REDACTED_PERSON]

	Improve CDN security by enabling randomized manifest filename for Odin Editor.
	Fix issue where Chunking and Posting would be attempted on non-buildfarm use.

Change 3286946 by [REDACTED_PERSON]

	Added ability to force disable auto loading splash screen on load.  This fixes the black HMD screen that appears in levels that did not manually turn off the loading movie in their level scripts

Change 3286888 by [REDACTED_PERSON]

	Disabling haptics when you begin changing levels

Change 3286200 by [REDACTED_PERSON]

	Odin - move GameAnalytics module back into Engine/Plugins. #jira AN-818

Change 3285398 by [REDACTED_PERSON]

	Enabled modular nativized building
	- Workaround for missing dependecies of NativizedAssets module.
	- An ugly hack to move NativizedAssets to plugins directory.

Change 3283738 by [REDACTED_PERSON]

	Removed old Odin-per-body-collision-settings hack that didn't work. Replaced bigbot usages of it with the new way that works, which fixes bigbot hands having to collision. ODIN-5638

Change 3283549 by [REDACTED_PERSON]

	Added OverrideCollisionResponseToChannelForBodies to BP library. Odin-specific hack to force specific collision responses on specific bodies. Using this to ignore WorldStatic on biped feet in some cases.

Change 3283088 by [REDACTED_PERSON]

	Fixing issue where Blueprint mod assets won't load in Odin nativized builds. Making it so we're leinent on matching import class types - BlueprintGeneratedClass (what the mod assets are expecting) and DynamicClass (what the imports actually are) are interchangable as far as these type checks are concerned.

Change 3282584 by [REDACTED_PERSON]

	Adding hooks for detecting if the touch controllers are present or not

Change 3282257 by [REDACTED_PERSON]

	Adding HasModules method to the plugin wizard definition

Change 3281499 by [REDACTED_PERSON]

	Add a CompatibleChangelist field to the plugin descriptor. If set, the engine will not load any plugins which have a different compatible changelist to the current engine version.

Change 3281286 by [REDACTED_PERSON]

	Exporting a number of required functions/classes/etc, used in nativized Odin source with the DLLEXPORT (*_API) tag - required for nativized and modular builds to work in tandem.

Change 3280775 by [REDACTED_PERSON]

	#odin - ODIN-3180: Piercing ammo and Disintegrator shot can be correctly blocked by objects that would normally overlap it (invincible characters, big bot, elites, boss). Implemented filtering of overlaps and blocking hit in custom FScopedMovementUpdate so we don't have to do any extra traces nor require bGenerateOverlapEvents. Minor engine change required to make some data protected on super class.

	Last step is just the disintegrator visual (needs to be scaled to match bullet distance).

	Also removed need for projectiles to update physics volume, because we don't care.

Change 3277573 by [REDACTED_PERSON]

	adding physx build to odin main

Change 3277289 by [REDACTED_PERSON]

	The Odin Plugin Packager will now get a current list of all available mod packages when gathering mods available for packaging

Change 3276949 by [REDACTED_PERSON]

	Basic Odin editor mod creation/packing telemetry. #jira AN-1304
	* Analytics event: Mod.Created
	* Analytics event: Mod.Pacikaged
	Associated Engine changes to enable this:
	* Enabled a UatTaskResultCallack TFunction to be given to IUatHelperModule::CreateUatTask to report the results.
	* Extended IPluginWizardDefinition interface with a PluginCreated callback when the plugin creation is finished.

Change 3276638 by [REDACTED_PERSON]

	moved GameAnalytics plugin temporarily into an Odin-only plugin location to get around issues with packaging a NotForLicensees plugin in a publicly distributed editor build.

	#tests build packaged game, editor, ran editor -game, verified anlaytics working.

Change 3276411 by [REDACTED_PERSON]

	Fixing copyrights and missing #pragma once

Change 3276394 by [REDACTED_PERSON]

	#Odin Initial Gameplay analytics integration
	* New shared engine plugin where common event infrastructure will reside (GameAnalytics). #jira AN-818
	  * CreateAnalyticsProviderSingleton
	  * CreateAnalyticsUserID
	  * TickClientHeartbeat
	  * More to come...
	* Created OdinAnalytics singleton to provide all analytics event functionality. #jira AN-1295
	  * All real data management in .cpp file to reduce header dependencies (Orion lesson learned).
	* Created heartbeat and Begin/EndMission hooks. #jira AN-1297.
	  * Hooked into UOdinBlueprintLibrary::FinishLevel to send EndMission event with some gameplay context (a bit more to come).
	  * Hooked into OdinGameMode::StartPlay to send StartMission event. Not quite symmetrical to FinishLevel, but it was easier that way.
	* Created UGameEngine And UUnrealEdEngine derived classes in Odin for analytics singleton management and heartbeat ticking. #jira AN-1298
	  * Editor support cfuns: OdinAnalytics.EnableInEditor and OdinAnalytics.DisableInEditor to test analytlcs while running the editor.
	* OdinGame modes exposes Analytics primary and secondary context to blueprints.
	* BP GameModes set the defauly analytics primary context (Title, Hub, InMission)
	* Menu maps set MainMenu vs Title secondary contexts.
	* BP_HoloStation_Menu_Control sets primary context of HoloTable and secondary contexts to match top-level menu states.
	* Removed "LoadingPhase": "PostEngineInit" from OdinEditor.uplugin so it can register the UnrealEdEngine class.
	* Ensure Odin editor has a unique analytics app id (UEEditor.Odin.Release). #jira AN-1296.
	  * used the UE4TypeOverride in Engine.ini so we could do it without modifying the engine.
	* Implented an INI override to disable Odin analytics. #jira AN-1300
	  * EngineIni[OdinAnalytics]:Enabled

	#tests running editor an PIE on start, hub, and pearl to verify analytics heartbeats and mission start/stop analytics fire correctly.

Change 3274664 by [REDACTED_PERSON]

	New and improved hack for disabling collision per-body on skel meshes.

Change 3273733 by [REDACTED_PERSON]

	Duplicating changes made to the plugin wizard definition that allows for the loading phase of a plugin to be specified from a template.

Change 3272613 by [REDACTED_PERSON]

	Fixing ODIN-5158 : Celebration - UMG panels are all redrawing at insane framerate; should be throttled/invalidated
	After finishing the animation, each panel sets itself to ManualRedraw mode, which stops it from redrawing.

Change 3272313 by [REDACTED_PERSON]

	It is now possible to display the content folder for a plugin once the plugin is created if the plugin can contain content. The user can opt out of this behavior from the new plugin wizard before the plugin is created.

Change 3271955 by [REDACTED_PERSON]

	#engine #nomerge Engine hack to allow setting skelmesh body collision properties individually.

Change 3269099 by [REDACTED_PERSON]

	Backed out all appId-related stuff.

Change 3269002 by [REDACTED_PERSON]

	Hopefully fix CIS for editor and tools (Test|Shipping probably broken.)

Change 3268749 by [REDACTED_PERSON]

	Oculus appId now stored in MiscSupportLib and statically linked into UE4.
	Resolves ODIN-4857.

Change 3264838 by [REDACTED_PERSON]

	Clarifying instructions within the new mod wizard, and allowing the user to restart the editor from the notification that appears when a new mod is created. This also prevents the user from changing the plugin folder location when creating a mod.

Change 3257225 by [REDACTED_PERSON]

	#odin - ODIN-4130: fixed Airborne detection when capsule has collision response to all channels disabled. Made a couple methods in CharacterMovement virtual and fixed K2_FindFloor modifying bForceNextFloorCheck when it shouldn't.

Change 3252319 by [REDACTED_PERSON]

Change 3247352 by [REDACTED_PERSON]

	Fix for undefined macro warnings in build

Change 3245755 by [REDACTED_PERSON]

	Merging using UE4_MAIN_to_ODIN @ CL 3242852

Change 3242301 by [REDACTED_PERSON]

	[Oculus] Avoid crash in ovr_GetPerfStats by not calling it until mirror window/texture sets are created.

Change 3237593 by [REDACTED_PERSON]

	Fixed gamma correction issue with widget component being blitted to the backbuffer.

Change 3236677 by [REDACTED_PERSON]

	ODIN-2552 Flash on pierce overlap WIP

Change 3233777 by [REDACTED_PERSON]

	Cooked Audio Bug: Pushing up test case for cooking "CookedAudioTest" console command, Optimizations disabled in AudioDevice.cpp and GameplayStatics::SpawnSoundAttached for debugging

Change 3232567 by [REDACTED_PERSON]

	Duplicating the HostType module descriptor API changes made to the plugin wizard definition interface

Change 3232092 by [REDACTED_PERSON]

	Merged from UE4/Main to Odin at CL 3232059

Change 3231886 by [REDACTED_PERSON]

	Background building shader 2.0, plus updates to filterstep and random uv samples functions.

Change 3227552 by [REDACTED_PERSON]

	Fixes an issue with OdinEditor where selecting mod templates, closing the plugin creator UI, and then reopening it would cause the mod to be created with the previously selected templates if no new templates are selected.

Change 3227514 by [REDACTED_PERSON]

	Adding templates for creating the following mod content: weapons, enemies, and an empty map.

diff --git a/Engine/Plugins/FX/Niagara/Source/Niagara/Public/NiagaraCommon.h b/Engine/Plugins/FX/Niagara/Source/Niagara/Public/NiagaraCommon.h
--- a/Engine/Plugins/FX/Niagara/Source/Niagara/Public/NiagaraCommon.h
+++ b/Engine/Plugins/FX/Niagara/Source/Niagara/Public/NiagaraCommon.h
@@ -317,19 +317,23 @@ struct FVMExternalFunctionBindingInfo
 struct NIAGARA_API FNiagaraSystemUpdateContext
 {
 	FNiagaraSystemUpdateContext(const UNiagaraSystem* System, bool bReInit) { Add(System, bReInit); }
+#if WITH_EDITORONLY_DATA
 	FNiagaraSystemUpdateContext(const UNiagaraEmitter* Emitter, bool bReInit) { Add(Emitter, bReInit); }
 	FNiagaraSystemUpdateContext(const UNiagaraScript* Script, bool bReInit) { Add(Script, bReInit); }
 	//FNiagaraSystemUpdateContext(UNiagaraDataInterface* Interface, bool bReinit) : Add(Interface, bReinit) {}
 	FNiagaraSystemUpdateContext(const UNiagaraParameterCollection* Collection, bool bReInit) { Add(Collection, bReInit); }
+#endif
 	FNiagaraSystemUpdateContext() { }
 
 	~FNiagaraSystemUpdateContext();
 
 	void Add(const UNiagaraSystem* System, bool bReInit);
+#if WITH_EDITORONLY_DATA
 	void Add(const UNiagaraEmitter* Emitter, bool bReInit);
 	void Add(const UNiagaraScript* Script, bool bReInit);
 	//void Add(UNiagaraDataInterface* Interface, bool bReinit);
 	void Add(const UNiagaraParameterCollection* Collection, bool bReInit);
+#endif
 
 	/** Adds all currently active systems.*/
 	void AddAll(bool bReInit);
diff --git a/Engine/Plugins/Runtime/Oculus/OculusVR/Source/OculusInput/Private/OculusInput.cpp b/Engine/Plugins/Runtime/Oculus/OculusVR/Source/OculusInput/Private/OculusInput.cpp
--- a/Engine/Plugins/Runtime/Oculus/OculusVR/Source/OculusInput/Private/OculusInput.cpp
+++ b/Engine/Plugins/Runtime/Oculus/OculusVR/Source/OculusInput/Private/OculusInput.cpp
@@ -838,6 +838,18 @@ float FOculusInput::GetHapticAmplitudeScale() const
 	return 1.f;
 }
 
+uint32 FOculusInput::GetNumberOfTouchControllers() const
+{
+	uint32 RetVal = 0;
+
+	for (FOculusTouchControllerPair Pair : ControllerPairs)
+	{
+		RetVal += (Pair.ControllerStates[0].bIsConnected ? 1 : 0);
+		RetVal += (Pair.ControllerStates[1].bIsConnected ? 1 : 0);
+	}
+
+	return RetVal;
+}
 
 } // namespace OculusInput
 
diff --git a/Engine/Plugins/Runtime/Oculus/OculusVR/Source/OculusInput/Private/OculusInput.h b/Engine/Plugins/Runtime/Oculus/OculusVR/Source/OculusInput/Private/OculusInput.h
--- a/Engine/Plugins/Runtime/Oculus/OculusVR/Source/OculusInput/Private/OculusInput.h
+++ b/Engine/Plugins/Runtime/Oculus/OculusVR/Source/OculusInput/Private/OculusInput.h
@@ -65,6 +65,8 @@ public:
 	virtual void GetHapticFrequencyRange(float& MinFrequency, float& MaxFrequency) const override;
 	virtual float GetHapticAmplitudeScale() const override;
 
+	uint32 GetNumberOfTouchControllers() const;
+
 private:
 
 	/** Applies force feedback settings to the controller */
diff --git a/Engine/Plugins/Runtime/Oculus/OculusVR/Source/OculusInput/Private/OculusInputModule.cpp b/Engine/Plugins/Runtime/Oculus/OculusVR/Source/OculusInput/Private/OculusInputModule.cpp
--- a/Engine/Plugins/Runtime/Oculus/OculusVR/Source/OculusInput/Private/OculusInputModule.cpp
+++ b/Engine/Plugins/Runtime/Oculus/OculusVR/Source/OculusInput/Private/OculusInputModule.cpp
@@ -26,7 +26,9 @@ TSharedPtr< class IInputDevice > FOculusInputModule::CreateInputDevice( const TS
 	{
 		if (FOculusHMDModule::Get().PreInit())
 		{
-			return TSharedPtr< class IInputDevice >(new OculusInput::FOculusInput(InMessageHandler));
+			TSharedPtr< OculusInput::FOculusInput > InputDevice(new OculusInput::FOculusInput(InMessageHandler));
+			OculusInputDevice = InputDevice;
+			return InputDevice;
 		}
 		// else, they may just not have a oculus headset plugged in (which we have to account for - no need for a warning)
 	}
@@ -37,6 +39,15 @@ TSharedPtr< class IInputDevice > FOculusInputModule::CreateInputDevice( const TS
 	return nullptr;		
 }
 
+uint32 FOculusInputModule::GetNumberOfTouchControllers() const
+{
+	if (OculusInputDevice.IsValid())
+	{
+		return OculusInputDevice.Pin()->GetNumberOfTouchControllers();
+	}
+	return 0;
+}
+
 
 #endif	// OCULUS_INPUT_SUPPORTED_PLATFORMS
 
diff --git a/Engine/Plugins/Runtime/Oculus/OculusVR/Source/OculusInput/Private/OculusInputModule.h b/Engine/Plugins/Runtime/Oculus/OculusVR/Source/OculusInput/Private/OculusInputModule.h
--- a/Engine/Plugins/Runtime/Oculus/OculusVR/Source/OculusInput/Private/OculusInputModule.h
+++ b/Engine/Plugins/Runtime/Oculus/OculusVR/Source/OculusInput/Private/OculusInputModule.h
@@ -3,6 +3,7 @@
 #pragma once
 #include "IOculusInputModule.h"
 #include "IInputDevice.h"
+#include "Templates/SharedPointer.h"
 
 #define LOCTEXT_NAMESPACE "OculusInput"
 
@@ -13,17 +14,31 @@
 
 #if OCULUS_INPUT_SUPPORTED_PLATFORMS
 
+namespace OculusInput
+{
+	class FOculusInput;
+}
+
 class FOculusInputModule : public IOculusInputModule
 {
+	TWeakPtr<OculusInput::FOculusInput> OculusInputDevice;
+
 	// IInputDeviceModule overrides
 	virtual void StartupModule() override;
 	virtual TSharedPtr< class IInputDevice > CreateInputDevice( const TSharedRef< FGenericApplicationMessageHandler >& InMessageHandler ) override;
+
+	// IOculusInputModule overrides
+	virtual uint32 GetNumberOfTouchControllers() const override;
 };
 
 #else	//	OCULUS_INPUT_SUPPORTED_PLATFORMS
 
 class FOculusInputModule : public FDefaultModuleImpl
 {
+	virtual uint32 GetNumberOfTouchControllers() const
+	{
+		return 0;
+	};
 };
 
 #endif	// OCULUS_INPUT_SUPPORTED_PLATFORMS
diff --git a/Engine/Plugins/Runtime/Oculus/OculusVR/Source/OculusInput/Public/IOculusInputModule.h b/Engine/Plugins/Runtime/Oculus/OculusVR/Source/OculusInput/Public/IOculusInputModule.h
--- a/Engine/Plugins/Runtime/Oculus/OculusVR/Source/OculusInput/Public/IOculusInputModule.h
+++ b/Engine/Plugins/Runtime/Oculus/OculusVR/Source/OculusInput/Public/IOculusInputModule.h
@@ -36,5 +36,12 @@ public:
 	{
 		return FModuleManager::Get().IsModuleLoaded( "OculusInput" );
 	}
+
+	/**
+	 * Gets the number of Touch controllers that are active, so that games that require them can check to make sure they're present
+	 *
+	 * @return The number of Touch controllers that are active (but not necessarily tracked)
+	 */
+	virtual uint32 GetNumberOfTouchControllers() const = 0;
 };
 
diff --git a/Engine/Source/Editor/Kismet/Private/BlueprintCompilationManager.cpp b/Engine/Source/Editor/Kismet/Private/BlueprintCompilationManager.cpp
--- a/Engine/Source/Editor/Kismet/Private/BlueprintCompilationManager.cpp
+++ b/Engine/Source/Editor/Kismet/Private/BlueprintCompilationManager.cpp
@@ -1551,7 +1551,7 @@ UClass* FBlueprintCompilationManagerImpl::FastGenerateSkeletonClass(UBlueprint*
 			UField**& InCurrentParamStorageLocation, 
 			EFunctionFlags InFunctionFlags, 
 			const TArray<UK2Node_FunctionResult*>& ReturnNodes, 
-			const TArray<UEdGraphPin*>& InputPins, 
+			const TArray<UEdGraphPin*>& InputPins,
 			bool bIsStaticFunction, 
 			bool bForceArrayStructRefsConst, 
 			UFunction* SignatureOverride) -> UFunction*
@@ -1791,7 +1791,7 @@ UClass* FBlueprintCompilationManagerImpl::FastGenerateSkeletonClass(UBlueprint*
 	UField** CurrentFieldStorageLocation = &Ret->Children;
 	
 	// Helper function for making UFunctions generated for 'event' nodes, e.g. custom event and timelines
-	const auto MakeEventFunction = [&CurrentFieldStorageLocation, MakeFunction, Schema]( FName InName, EFunctionFlags ExtraFnFlags, const TArray<UEdGraphPin*>& InputPins, UFunction* InSourceFN, bool bInCallInEditor )
+	const auto MakeEventFunction = [&CurrentFieldStorageLocation, MakeFunction, Schema]( FName InName, EFunctionFlags ExtraFnFlags, const TArray<UEdGraphPin*>& InputPins, const TArray< TSharedPtr<FUserPinInfo> >& UserPins, UFunction* InSourceFN, bool bInCallInEditor )
 	{
 		UField** CurrentParamStorageLocation = nullptr;
 
@@ -1809,18 +1809,7 @@ UClass* FBlueprintCompilationManagerImpl::FastGenerateSkeletonClass(UBlueprint*
 
 		if(NewFunction)
 		{
-			for (UEdGraphPin* InputPin : InputPins)
-			{
-				// No defaults for object/class pins
-				if(	!Schema->IsMetaPin(*InputPin) && 
-					(InputPin->PinType.PinCategory != UEdGraphSchema_K2::PC_Object) && 
-					(InputPin->PinType.PinCategory != UEdGraphSchema_K2::PC_Class) && 
-					(InputPin->PinType.PinCategory != UEdGraphSchema_K2::PC_Interface) && 
-					!InputPin->DefaultValue.IsEmpty() )
-				{
-					NewFunction->SetMetaData(InputPin->PinName, *InputPin->DefaultValue);
-				}
-			}
+			FKismetCompilerContext::SetDefaultInputValueMetaData(NewFunction, UserPins);
 
 			if(bInCallInEditor)
 			{
@@ -1870,6 +1859,7 @@ UClass* FBlueprintCompilationManagerImpl::FastGenerateSkeletonClass(UBlueprint*
 				CompilerContext.GetEventStubFunctionName(Event), 
 				(EFunctionFlags)Event->FunctionFlags, 
 				Event->Pins, 
+				Event->UserDefinedPins,
 				Event->FindEventSignatureFunction(),
 				bCallInEditor
 			);
@@ -1880,11 +1870,11 @@ UClass* FBlueprintCompilationManagerImpl::FastGenerateSkeletonClass(UBlueprint*
 	{
 		for(int32 EventTrackIdx=0; EventTrackIdx<Timeline->EventTracks.Num(); EventTrackIdx++)
 		{
-			MakeEventFunction(Timeline->GetEventTrackFunctionName(EventTrackIdx), EFunctionFlags::FUNC_None, TArray<UEdGraphPin*>(), nullptr, false);
+			MakeEventFunction(Timeline->GetEventTrackFunctionName(EventTrackIdx), EFunctionFlags::FUNC_None, TArray<UEdGraphPin*>(), TArray< TSharedPtr<FUserPinInfo> >(), nullptr, false);
 		}
 		
-		MakeEventFunction(Timeline->GetUpdateFunctionName(), EFunctionFlags::FUNC_None, TArray<UEdGraphPin*>(), nullptr, false);
-		MakeEventFunction(Timeline->GetFinishedFunctionName(), EFunctionFlags::FUNC_None, TArray<UEdGraphPin*>(), nullptr, false);
+		MakeEventFunction(Timeline->GetUpdateFunctionName(), EFunctionFlags::FUNC_None, TArray<UEdGraphPin*>(), TArray< TSharedPtr<FUserPinInfo> >(), nullptr, false);
+		MakeEventFunction(Timeline->GetFinishedFunctionName(), EFunctionFlags::FUNC_None, TArray<UEdGraphPin*>(), TArray< TSharedPtr<FUserPinInfo> >(), nullptr, false);
 	}
 
 	CompilerContext.NewClass = Ret;
diff --git a/Engine/Source/Editor/KismetCompiler/Private/KismetCompiler.cpp b/Engine/Source/Editor/KismetCompiler/Private/KismetCompiler.cpp
--- a/Engine/Source/Editor/KismetCompiler/Private/KismetCompiler.cpp
+++ b/Engine/Source/Editor/KismetCompiler/Private/KismetCompiler.cpp
@@ -37,6 +37,7 @@
 #include "K2Node_TunnelBoundary.h"
 #include "K2Node_VariableGet.h"
 #include "K2Node_VariableSet.h"
+#include "K2Node_EditablePinBase.h" // for FUserPinInfo
 #include "KismetCompilerBackend.h"
 #include "Kismet2/KismetReinstanceUtilities.h"
 #include "Engine/SCS_Node.h"
@@ -2088,30 +2089,7 @@ void FKismetCompilerContext::SetCalculatedMetaDataAndFlags(UFunction* Function,
 		Function->SetMetaData(FBlueprintMetadata::MD_WorldContext, *WorldContextPin->PinName.ToString());
 	}
 
-	for (int32 EntryPinIndex = 0; EntryPinIndex < EntryNode->Pins.Num(); ++EntryPinIndex)
-	{
-		UEdGraphPin* EntryPin = EntryNode->Pins[EntryPinIndex];
-		// No defaults for object/class pins
-		if(	!K2Schema->IsMetaPin(*EntryPin) && 
-			(EntryPin->PinType.PinCategory != UEdGraphSchema_K2::PC_Object) && 
-			(EntryPin->PinType.PinCategory != UEdGraphSchema_K2::PC_Class) && 
-			(EntryPin->PinType.PinCategory != UEdGraphSchema_K2::PC_Interface))
-		{
-			if (!EntryPin->DefaultValue.IsEmpty())
-			{
-				Function->SetMetaData(EntryPin->PinName, *EntryPin->DefaultValue);
-			}
-			else if (!EntryPin->DefaultTextValue.IsEmpty())
-			{
-				FString TextAsString;
-				if (FTextStringHelper::WriteToString(TextAsString, EntryPin->DefaultTextValue))
-				{
-					Function->SetMetaData(EntryPin->PinName, *TextAsString);
-				}
-			}
-			
-		}
-	}
+	SetDefaultInputValueMetaData(Function, EntryNode->UserDefinedPins);
 
 	if(UFunction* OverriddenFunction = Function->GetSuperFunction())
 	{
@@ -2120,6 +2098,22 @@ void FKismetCompilerContext::SetCalculatedMetaDataAndFlags(UFunction* Function,
 	}
 }
 
+void FKismetCompilerContext::SetDefaultInputValueMetaData(UFunction* Function, const TArray< TSharedPtr<FUserPinInfo> >& InputData)
+{
+	for (const TSharedPtr<FUserPinInfo>& InputDataPtr : InputData)
+	{
+		if ( InputDataPtr.IsValid() &&
+			(InputDataPtr->PinType.PinCategory != UEdGraphSchema_K2::PC_Exec) && 
+			(InputDataPtr->PinName != UEdGraphSchema_K2::PN_Self) && 
+			(InputDataPtr->PinType.PinCategory != UEdGraphSchema_K2::PC_Object) &&
+			(InputDataPtr->PinType.PinCategory != UEdGraphSchema_K2::PC_Class) &&
+			(InputDataPtr->PinType.PinCategory != UEdGraphSchema_K2::PC_Interface) )
+		{
+			Function->SetMetaData(InputDataPtr->PinName, *InputDataPtr->PinDefaultValue);
+		}
+	}
+}
+
 /**
  * Handles adding the implemented interface information to the class
  */
diff --git a/Engine/Source/Editor/KismetCompiler/Public/KismetCompiler.h b/Engine/Source/Editor/KismetCompiler/Public/KismetCompiler.h
--- a/Engine/Source/Editor/KismetCompiler/Public/KismetCompiler.h
+++ b/Engine/Source/Editor/KismetCompiler/Public/KismetCompiler.h
@@ -20,6 +20,7 @@ class UK2Node_TemporaryVariable;
 class UK2Node_Timeline;
 class UK2Node_Tunnel;
 class FKismetCompilerVMBackend;
+struct FUserPinInfo;
 
 KISMETCOMPILER_API DECLARE_LOG_CATEGORY_EXTERN(LogK2Compiler, Log, All);
 
@@ -497,6 +498,9 @@ protected:
 
 	static void SetCalculatedMetaDataAndFlags(UFunction* Function, UK2Node_FunctionEntry* EntryNode, const UEdGraphSchema_K2* Schema );
 
+	/** Reflects each pin's user set, default value into the function's metadata (so it can be queried for later by CallFunction nodes, etc.) */
+	static void SetDefaultInputValueMetaData(UFunction* Function, const TArray< TSharedPtr<FUserPinInfo> >& InputData);
+
 	/**
 	 * Handles adding the implemented interface information to the class
 	 */
diff --git a/Engine/Source/Runtime/CoreUObject/Private/Blueprint/BlueprintSupport.cpp b/Engine/Source/Runtime/CoreUObject/Private/Blueprint/BlueprintSupport.cpp
--- a/Engine/Source/Runtime/CoreUObject/Private/Blueprint/BlueprintSupport.cpp
+++ b/Engine/Source/Runtime/CoreUObject/Private/Blueprint/BlueprintSupport.cpp
@@ -689,6 +689,110 @@ static UObject* FindExistingImportObject(const int32 Index, const TArray<FObject
 	return FoundObject;
 }
 
+/**
+ * This utility struct helps track blueprint structs/linkers that are currently 
+ * in the middle of a call to ResolveDeferredDependencies(). This can be used  
+ * to know if a dependency's resolve needs to be finished (to avoid unwanted 
+ * placeholder references ending up in script-code).
+ */
+struct FUnresolvedStructTracker
+{
+public:
+	/** Marks the specified struct (and its linker) as "resolving" for the lifetime of this instance */
+	FUnresolvedStructTracker(UStruct* LoadStruct)
+		: TrackedStruct(LoadStruct)
+	{
+		DEFERRED_DEPENDENCY_CHECK((LoadStruct != nullptr) && (LoadStruct->GetLinker() != nullptr));
+		FScopeLock UnresolvedStructsLock(&UnresolvedStructsCritical);
+		UnresolvedStructs.Add(LoadStruct);
+	}
+
+	/** Removes the wrapped struct from the "resolving" set (it has been fully "resolved") */
+	~FUnresolvedStructTracker()
+	{
+		// even if another FUnresolvedStructTracker added this struct earlier,  
+		// we want the most nested one removing it from the set (because this 
+		// means the struct is fully resolved, even if we're still in the middle  
+		// of a ResolveDeferredDependencies() call further up the stack)
+		FScopeLock UnresolvedStructsLock(&UnresolvedStructsCritical);
+		UnresolvedStructs.Remove(TrackedStruct);
+	}
+
+	/**
+	 * Checks to see if the specified import object is a blueprint class/struct 
+	 * that is currently in the midst of resolving (and hasn't completed that  
+	 * resolve elsewhere in some nested call).
+	 * 
+	 * @param  ImportObject    The object you wish to check.
+	 * @return True if the specified object is a class/struct that hasn't been fully resolved yet (otherwise false).
+	 */
+	static bool IsImportStructUnresolved(UObject* ImportObject)
+	{
+		FScopeLock UnresolvedStructsLock(&UnresolvedStructsCritical);
+		return UnresolvedStructs.Contains(ImportObject);
+	}
+
+	/**
+	 * Checks to see if the specified linker is associated with any of the 
+	 * unresolved structs that this is currently tracking.
+	 *
+	 * NOTE: This could return false, even if the linker is in a 
+	 *       ResolveDeferredDependencies() call futher up the callstack... in 
+	 *       that scenario, the associated struct was fully resolved by a 
+	 *       subsequent call to the same function (for the same linker/struct).
+	 * 
+	 * @param  Linker	The linker you want to check.
+	 * @return True if the specified linker is in the midst of an unfinished ResolveDeferredDependencies() call (otherwise false).
+	 */
+	static bool IsAssociatedStructUnresolved(const FLinkerLoad* Linker)
+	{
+		FScopeLock UnresolvedStructsLock(&UnresolvedStructsCritical);
+		for (UObject* UnresolvedObj : UnresolvedStructs)
+		{
+			// each unresolved struct should have a linker set on it, because 
+			// they would have had to go through Preload()
+			if (UnresolvedObj->GetLinker() == Linker)
+			{
+				return true;
+			}
+		}
+		return false;
+	}
+
+	/**  */
+	static void Reset(const FLinkerLoad* Linker)
+	{
+		FScopeLock UnresolvedStructsLock(&UnresolvedStructsCritical);
+		TArray<UObject*> ToRemove;
+		for (UObject* UnresolvedObj : UnresolvedStructs)
+		{
+			if (UnresolvedObj->GetLinker() == Linker)
+			{
+				ToRemove.Add(UnresolvedObj);
+			}
+		}
+		for (UObject* ResetingObj : ToRemove)
+		{
+			UnresolvedStructs.Remove(ResetingObj);
+		}
+	}
+
+private:
+	/** The struct that is currently being "resolved" */
+	UStruct* TrackedStruct;
+
+	/** 
+	 * A set of blueprint structs & classes that are currently being "resolved"  
+	 * by ResolveDeferredDependencies() (using UObject* instead of UStruct, so
+	 * we don't have to cast import objects before checking for their presence).
+	 */
+	static TSet<UObject*> UnresolvedStructs;
+	static FCriticalSection UnresolvedStructsCritical;
+};
+/** A global set that tracks structs currently being ran through (and unfinished by) FLinkerLoad::ResolveDeferredDependencies() */
+TSet<UObject*> FUnresolvedStructTracker::UnresolvedStructs;
+FCriticalSection FUnresolvedStructTracker::UnresolvedStructsCritical;
+
 bool FLinkerLoad::DeferPotentialCircularImport(const int32 Index)
 {
 #if USE_CIRCULAR_DEPENDENCY_LOAD_DEFERRING
@@ -702,10 +806,65 @@ bool FLinkerLoad::DeferPotentialCircularImport(const int32 Index)
 	//--------------------------------------
 
 	FObjectImport& Import = ImportMap[Index];
+
 	if (Import.XObject != nullptr)
 	{
-		bool const bIsPlaceHolderClass = Import.XObject->IsA<ULinkerPlaceholderClass>();
-		return bIsPlaceHolderClass;
+		FLinkerPlaceholderBase* ImportPlaceholder = nullptr;
+		if (ULinkerPlaceholderClass* AsPlaceholderClass = Cast<ULinkerPlaceholderClass>(Import.XObject))
+		{
+			ImportPlaceholder = AsPlaceholderClass;
+		}
+		else if (ULinkerPlaceholderFunction* AsPlaceholderFunc = Cast<ULinkerPlaceholderFunction>(Import.XObject))
+		{
+			ImportPlaceholder = AsPlaceholderFunc;
+		}
+
+		const bool bIsResolvingPlaceholders = ImportPlaceholder && (LoadFlags & LOAD_DeferDependencyLoads) == LOAD_None;
+		// if this import already had a placeholder spawned for it, but the package 
+		// has passed the need for placeholders (it's in the midst of ResolveDeferredDependencies)
+		if (bIsResolvingPlaceholders)
+		{
+			// this is to validate our assumption that this package is in ResolveDeferredDependencies() earlier up the stack
+			DEFERRED_DEPENDENCY_CHECK(FUnresolvedStructTracker::IsAssociatedStructUnresolved(this));
+		
+			UClass* LoadClass = nullptr;
+			// Get the LoadClass that is currently in the midst of being resolved (needed to pass to ResolveDependencyPlaceholder)
+			{
+				// if DeferredCDOIndex is not set, then this is presumably a struct package (it should always be  
+				// set at this point for class BP packages - see Preload() where DeferredCDOIndex is assigned)
+				if (DeferredCDOIndex != INDEX_NONE)
+				{
+					FPackageIndex ClassIndex = ExportMap[DeferredCDOIndex].ClassIndex;
+					DEFERRED_DEPENDENCY_CHECK(ClassIndex.IsExport());
+
+					if (ClassIndex.IsExport())
+					{
+						FObjectExport ClassExport = ExportMap[ClassIndex.ToExport()];
+						LoadClass = Cast<UClass>(ClassExport.Object);
+					}
+
+					DEFERRED_DEPENDENCY_CHECK(LoadClass != nullptr);
+				}
+			}
+
+			// go ahead and resolve the placeholder here (since someone's requesting it and we're already in the 
+			// midst of resolving placeholders earlier in the stack) - the idea is that the resolve, already in progress, will 
+			// eventually get to this placeholder, it just hasn't looped there yet
+			//
+			// this will prevent other, needless placeholders from being created (export templates that are relying on this class, etc.)
+			ResolveDependencyPlaceholder(ImportPlaceholder, LoadClass);
+
+#if USE_DEFERRED_DEPENDENCY_CHECK_VERIFICATION_TESTS
+			const bool bIsStillPlaceholder = Import.XObject->IsA<ULinkerPlaceholderClass>() || Import.XObject->IsA<ULinkerPlaceholderFunction>();
+			DEFERRED_DEPENDENCY_CHECK(!bIsStillPlaceholder);
+			return bIsStillPlaceholder;
+#else 
+			// presume that ResolveDependencyPlaceholder() worked and the import is no longer a placeholder
+			return false;
+#endif 
+
+		}
+		return (ImportPlaceholder != nullptr);
 	}
 
 	if ((LoadFlags & LOAD_DeferDependencyLoads) && !IsImportNative(Index))
@@ -1019,110 +1178,6 @@ int32 FLinkerLoad::FindCDOExportIndex(UClass* LoadClass)
 	return INDEX_NONE;
 }
 
-/**
- * This utility struct helps track blueprint structs/linkers that are currently 
- * in the middle of a call to ResolveDeferredDependencies(). This can be used  
- * to know if a dependency's resolve needs to be finished (to avoid unwanted 
- * placeholder references ending up in script-code).
- */
-struct FUnresolvedStructTracker
-{
-public:
-	/** Marks the specified struct (and its linker) as "resolving" for the lifetime of this instance */
-	FUnresolvedStructTracker(UStruct* LoadStruct)
-		: TrackedStruct(LoadStruct)
-	{
-		DEFERRED_DEPENDENCY_CHECK((LoadStruct != nullptr) && (LoadStruct->GetLinker() != nullptr));
-		FScopeLock UnresolvedStructsLock(&UnresolvedStructsCritical);
-		UnresolvedStructs.Add(LoadStruct);
-	}
-
-	/** Removes the wrapped struct from the "resolving" set (it has been fully "resolved") */
-	~FUnresolvedStructTracker()
-	{
-		// even if another FUnresolvedStructTracker added this struct earlier,  
-		// we want the most nested one removing it from the set (because this 
-		// means the struct is fully resolved, even if we're still in the middle  
-		// of a ResolveDeferredDependencies() call further up the stack)
-		FScopeLock UnresolvedStructsLock(&UnresolvedStructsCritical);
-		UnresolvedStructs.Remove(TrackedStruct);
-	}
-
-	/**
-	 * Checks to see if the specified import object is a blueprint class/struct 
-	 * that is currently in the midst of resolving (and hasn't completed that  
-	 * resolve elsewhere in some nested call).
-	 * 
-	 * @param  ImportObject    The object you wish to check.
-	 * @return True if the specified object is a class/struct that hasn't been fully resolved yet (otherwise false).
-	 */
-	static bool IsImportStructUnresolved(UObject* ImportObject)
-	{
-		FScopeLock UnresolvedStructsLock(&UnresolvedStructsCritical);
-		return UnresolvedStructs.Contains(ImportObject);
-	}
-
-	/**
-	 * Checks to see if the specified linker is associated with any of the 
-	 * unresolved structs that this is currently tracking.
-	 *
-	 * NOTE: This could return false, even if the linker is in a 
-	 *       ResolveDeferredDependencies() call futher up the callstack... in 
-	 *       that scenario, the associated struct was fully resolved by a 
-	 *       subsequent call to the same function (for the same linker/struct).
-	 * 
-	 * @param  Linker	The linker you want to check.
-	 * @return True if the specified linker is in the midst of an unfinished ResolveDeferredDependencies() call (otherwise false).
-	 */
-	static bool IsAssociatedStructUnresolved(const FLinkerLoad* Linker)
-	{
-		FScopeLock UnresolvedStructsLock(&UnresolvedStructsCritical);
-		for (UObject* UnresolvedObj : UnresolvedStructs)
-		{
-			// each unresolved struct should have a linker set on it, because 
-			// they would have had to go through Preload()
-			if (UnresolvedObj->GetLinker() == Linker)
-			{
-				return true;
-			}
-		}
-		return false;
-	}
-
-	/**  */
-	static void Reset(const FLinkerLoad* Linker)
-	{
-		FScopeLock UnresolvedStructsLock(&UnresolvedStructsCritical);
-		TArray<UObject*> ToRemove;
-		for (UObject* UnresolvedObj : UnresolvedStructs)
-		{
-			if (UnresolvedObj->GetLinker() == Linker)
-			{
-				ToRemove.Add(UnresolvedObj);
-			}
-		}
-		for (UObject* ResetingObj : ToRemove)
-		{
-			UnresolvedStructs.Remove(ResetingObj);
-		}
-	}
-
-private:
-	/** The struct that is currently being "resolved" */
-	UStruct* TrackedStruct;
-
-	/** 
-	 * A set of blueprint structs & classes that are currently being "resolved"  
-	 * by ResolveDeferredDependencies() (using UObject* instead of UStruct, so
-	 * we don't have to cast import objects before checking for their presence).
-	 */
-	static TSet<UObject*> UnresolvedStructs;
-	static FCriticalSection UnresolvedStructsCritical;
-};
-/** A global set that tracks structs currently being ran through (and unfinished by) FLinkerLoad::ResolveDeferredDependencies() */
-TSet<UObject*> FUnresolvedStructTracker::UnresolvedStructs;
-FCriticalSection FUnresolvedStructTracker::UnresolvedStructsCritical;
-
 UPackage* LoadPackageInternal(UPackage* InOuter, const TCHAR* InLongPackageName, uint32 LoadFlags, FLinkerLoad* ImportLinker);
 
 void FLinkerLoad::ResolveDeferredDependencies(UStruct* LoadStruct)
@@ -1150,117 +1205,134 @@ void FLinkerLoad::ResolveDeferredDependencies(UStruct* LoadStruct)
 		
 		UClass* const LoadClass = Cast<UClass>(LoadStruct);
 
-		bool bImportMapResolved = false;
+		int32 StartingImportIndex = 0;
 		// this function (for this linker) could be reentrant (see where we 
 		// recursively call ResolveDeferredDependencies() for super-classes below);  
 		// if that's the case, then we want to finish resolving the pending class 
 		// before we continue on
-		if (ResolvingDeferredPlaceholder != nullptr)
-		{
-			FName ReplacementPkgPath = NAME_None;
-			// if this placeholder is not an ImportMap placeholder, then we need
-			// additional info to resolve it (namely the true object's package
-			// path)... if we don't have that, then this resolve will fail
-			if (ResolvingDeferredPlaceholder->PackageIndex.IsNull())
+		if (ResolvingPlaceholderStack.Num() > 0)
+		{
+			// Since this method is recursive, we don't need to needlessly loop over all the imports we've already 
+			// resolved. However, we can only guarantee that the oldest entry in the 'resolving' stack is from a loop below.
+			// Now that other places call ResolveDependencyPlaceholder(), the ResolvingPlaceholderStack may jump around and
+			// skip some entries. The only certainty is that this function is the initial entry point for ResolveDependencyPlaceholder().
+			const FPackageIndex& FirstResolvingIndex = ResolvingPlaceholderStack[0]->PackageIndex;
+			if (FirstResolvingIndex.IsNull())
+			{
+				// if the placeholder's package index is null, that means we've already looped over the entire 
+				// ImportMap, and moved on to the loop below it (where we resolve placeholders from ImportText() 
+				// and such - they don't have entries in the ImportMap), so skip the ImportMap loop
+				StartingImportIndex = ImportMap.Num();
+			}
+			else
 			{
-				const FName* ImportObjectPath = ImportPlaceholders.FindKey(ResolvingDeferredPlaceholder);
-				DEFERRED_DEPENDENCY_CHECK(ImportObjectPath != nullptr);
+				DEFERRED_DEPENDENCY_CHECK(FirstResolvingIndex.IsImport());
+
+				// Since the ImportMap loop below resolves ULinkerPlaceholderFunction's owner first (out of order), we cannot 
+				// even guarantee that we've resolved everything prior to FirstResolvingIndex, so don't set StartingImportIndex in this case
+// 				if (FirstResolvingIndex.IsImport())
+// 				{
+// 					StartingImportIndex = FirstResolvingIndex.ToImport()+1;
+// 				}
+			}
+
+			while (ResolvingPlaceholderStack.Num() > 0)
+			{
+				FLinkerPlaceholderBase* ResolvingPlaceholder = ResolvingPlaceholderStack.Pop();
+				// If this is a placeholder outside the ImportMap (from ImportText(), etc.), then it needs a PackagePath to
+				// resolve. Don't worry that one isn't passed in as a param here, ResolveDependencyPlaceholder() will 
+				// look it up itself in ImportPlaceholders (the param is just an optimization)
+				ResolveDependencyPlaceholder(ResolvingPlaceholder, LoadClass);
+			}
 
-				if (ImportObjectPath != nullptr)
+#if USE_DEFERRED_DEPENDENCY_CHECK_VERIFICATION_TESTS
+			for (int32 ImportIndex = 0; ImportIndex < StartingImportIndex; ++ImportIndex)
+			{
+				if (UObject* ImportObj = ImportMap[ImportIndex].XObject)
 				{
-					ReplacementPkgPath = *ImportObjectPath;
-					// save us from looping through the ImportMap again; if 
-					// we're already resolving an entry from ImportPlaceholders,  
-					// then we've already been through this function and are in  
-					// the ImportPlaceholders loop somewhere up the stack
-					bImportMapResolved = true;
+					DEFERRED_DEPENDENCY_CHECK(Cast<ULinkerPlaceholderClass>(ImportObj) == nullptr);
+					DEFERRED_DEPENDENCY_CHECK(Cast<ULinkerPlaceholderFunction>(ImportObj) == nullptr);
 				}
 			}
-
-			int32 ResolvedRefCount = ResolveDependencyPlaceholder(ResolvingDeferredPlaceholder, LoadClass, ReplacementPkgPath);
-			// @TODO: can we reliably count on this resolving some dependencies?... 
-			//        if so, check verify that!
-			ResolvingDeferredPlaceholder = nullptr;
-			ImportPlaceholders.Remove(ReplacementPkgPath);
+#endif // USE_DEFERRED_DEPENDENCY_CHECK_VERIFICATION_TESTS
 		}
-
-		if (!bImportMapResolved)
-		{
-			// because this loop could recurse (and end up finishing all of this for
-			// us), we check HasUnresolvedDependencies() so we can early out  
-			// from this loop in that situation (the loop has been finished elsewhere)
-			for (int32 ImportIndex = 0; ImportIndex < ImportMap.Num() && HasUnresolvedDependencies(); ++ImportIndex)
+		
+		// because this loop could recurse (and end up finishing all of this for
+		// us), we check HasUnresolvedDependencies() so we can early out  
+		// from this loop in that situation (the loop has been finished elsewhere)
+		for (int32 ImportIndex = StartingImportIndex; ImportIndex < ImportMap.Num() && HasUnresolvedDependencies(); ++ImportIndex)
+		{
+			FObjectImport& Import = ImportMap[ImportIndex];
+
+			FLinkerLoad* SourceLinker = Import.SourceLinker;
+			// we cannot rely on Import.SourceLinker being set, if you look 
+			// at FLinkerLoad::CreateImport(), you'll see in game builds 
+			// that we try to circumvent the normal Import loading with a
+			// FindImportFast() call... if this is successful (the import 
+			// has already been somewhat loaded), then we don't fill out the 
+			// SourceLinker field
+			if (SourceLinker == nullptr)
 			{
-				FObjectImport& Import = ImportMap[ImportIndex];
-
-				FLinkerLoad* SourceLinker = Import.SourceLinker;
-				// we cannot rely on Import.SourceLinker being set, if you look 
-				// at FLinkerLoad::CreateImport(), you'll see in game builds 
-				// that we try to circumvent the normal Import loading with a
-				// FindImportFast() call... if this is successful (the import 
-				// has already been somewhat loaded), then we don't fill out the 
-				// SourceLinker field
-				if (SourceLinker == nullptr)
+				if (Import.XObject != nullptr)
 				{
-					if (Import.XObject != nullptr)
-					{
-						SourceLinker = Import.XObject->GetLinker();
-						//if (SourceLinker == nullptr)
-						//{
-						//	UPackage* ImportPkg = Import.XObject->GetOutermost();
-						//	// we use this package as placeholder for our 
-						//	// placeholders, so make sure to skip those (all other
-						//	// imports should belong to another package)
-						//	if (ImportPkg && ImportPkg != LinkerRoot)
-						//	{
-						//		SourceLinker = FindExistingLinkerForPackage(ImportPkg);
-						//	}
-						//}
-					}					
-				}
+					SourceLinker = Import.XObject->GetLinker();
+					//if (SourceLinker == nullptr)
+					//{
+					//	UPackage* ImportPkg = Import.XObject->GetOutermost();
+					//	// we use this package as placeholder for our 
+					//	// placeholders, so make sure to skip those (all other
+					//	// imports should belong to another package)
+					//	if (ImportPkg && ImportPkg != LinkerRoot)
+					//	{
+					//		SourceLinker = FindExistingLinkerForPackage(ImportPkg);
+					//	}
+					//}
+				}					
+			}
 
-				const UPackage* SourcePackage = (SourceLinker != nullptr) ? SourceLinker->LinkerRoot : nullptr;
-				// this package may not have introduced any (possible) cyclic 
-				// dependencies, but it still could have been deferred (kept from
-				// fully loading... we need to make sure metadata gets loaded, etc.)
-				if ((SourcePackage != nullptr) && !SourcePackage->HasAnyFlags(RF_WasLoaded))
-				{
-					uint32 InternalLoadFlags = LoadFlags & (LOAD_NoVerify | LOAD_NoWarn | LOAD_Quiet);
-					// make sure LoadAllObjects() is called for this package
-					LoadPackageInternal(/*Outer =*/nullptr, *SourceLinker->Filename, InternalLoadFlags, this); //-V595
-				}
+			const UPackage* SourcePackage = (SourceLinker != nullptr) ? SourceLinker->LinkerRoot : nullptr;
+			// this package may not have introduced any (possible) cyclic 
+			// dependencies, but it still could have been deferred (kept from
+			// fully loading... we need to make sure metadata gets loaded, etc.)
+			if ((SourcePackage != nullptr) && !SourcePackage->HasAnyFlags(RF_WasLoaded))
+			{
+				uint32 InternalLoadFlags = LoadFlags & (LOAD_NoVerify | LOAD_NoWarn | LOAD_Quiet);
+				// make sure LoadAllObjects() is called for this package
+				LoadPackageInternal(/*Outer =*/nullptr, *SourceLinker->Filename, InternalLoadFlags, this); //-V595
+			}
 
-				if (ULinkerPlaceholderClass* PlaceholderClass = Cast<ULinkerPlaceholderClass>(Import.XObject))
-				{
-					DEFERRED_DEPENDENCY_CHECK(PlaceholderClass->PackageIndex.ToImport() == ImportIndex);
+			DEFERRED_DEPENDENCY_CHECK(ResolvingPlaceholderStack.Num() == 0);
+			if (ULinkerPlaceholderClass* PlaceholderClass = Cast<ULinkerPlaceholderClass>(Import.XObject))
+			{
+				DEFERRED_DEPENDENCY_CHECK(PlaceholderClass->PackageIndex.ToImport() == ImportIndex);
 
-					// NOTE: we don't check that this resolve successfully replaced any
-					//       references (by the return value), because this resolve 
-					//       could have been re-entered and completed by a nested call
-					//       to the same function (for the same placeholder)
-					ResolveDependencyPlaceholder(PlaceholderClass, LoadClass);
-				}
-				else if (ULinkerPlaceholderFunction* PlaceholderFunction = Cast<ULinkerPlaceholderFunction>(Import.XObject))
+				// NOTE: we don't check that this resolve successfully replaced any
+				//       references (by the return value), because this resolve 
+				//       could have been re-entered and completed by a nested call
+				//       to the same function (for the same placeholder)
+				ResolveDependencyPlaceholder(PlaceholderClass, LoadClass);
+			}
+			else if (ULinkerPlaceholderFunction* PlaceholderFunction = Cast<ULinkerPlaceholderFunction>(Import.XObject))
+			{
+				if (ULinkerPlaceholderClass* PlaceholderOwner = Cast<ULinkerPlaceholderClass>(PlaceholderFunction->GetOwnerClass()))
 				{
-					if (ULinkerPlaceholderClass* PlaceholderOwner = Cast<ULinkerPlaceholderClass>(PlaceholderFunction->GetOwnerClass()))
-					{
-						ResolveDependencyPlaceholder(PlaceholderOwner, LoadClass);
-					}
-
-					DEFERRED_DEPENDENCY_CHECK(PlaceholderFunction->PackageIndex.ToImport() == ImportIndex);
-					ResolveDependencyPlaceholder(PlaceholderFunction, LoadClass);
+					ResolveDependencyPlaceholder(PlaceholderOwner, LoadClass);
 				}
-				else if (UScriptStruct* StructObj = Cast<UScriptStruct>(Import.XObject))
+
+				DEFERRED_DEPENDENCY_CHECK(PlaceholderFunction->PackageIndex.ToImport() == ImportIndex);
+				ResolveDependencyPlaceholder(PlaceholderFunction, LoadClass);
+			}
+			else if (UScriptStruct* StructObj = Cast<UScriptStruct>(Import.XObject))
+			{
+				// in case this is a user defined struct, we have to resolve any 
+				// deferred dependencies in the struct 
+				if (SourceLinker != nullptr)
 				{
-					// in case this is a user defined struct, we have to resolve any 
-					// deferred dependencies in the struct 
-					if (SourceLinker != nullptr)
-					{
-						SourceLinker->ResolveDeferredDependencies(StructObj);
-					}
+					SourceLinker->ResolveDeferredDependencies(StructObj);
 				}
 			}
-		} // if (!bImportMapResolved)
+			DEFERRED_DEPENDENCY_CHECK(ResolvingPlaceholderStack.Num() == 0);
+		}
 
 		// resolve any placeholders that were imported through methods like 
 		// ImportText() (meaning the ImportMap wouldn't reference them)
@@ -1305,16 +1377,34 @@ void FLinkerLoad::ResolveDeferredDependencies(UStruct* LoadStruct)
 	//LoadAllObjects();
 
 #if USE_DEFERRED_DEPENDENCY_CHECK_VERIFICATION_TESTS
-	for (TObjectIterator<ULinkerPlaceholderClass> PlaceholderIt; PlaceholderIt; ++PlaceholderIt)
+	auto CheckPlaceholderReferences = [this](FLinkerPlaceholderBase* Placeholder)
 	{
-		ULinkerPlaceholderClass* PlaceholderClass = *PlaceholderIt;
-		if (PlaceholderClass->GetOuter() == LinkerRoot)
+		UObject* PlaceholderObj = Placeholder->GetPlaceholderAsUObject();
+		if (PlaceholderObj->GetOuter() == LinkerRoot)
 		{
 			// there shouldn't be any deferred dependencies (belonging to this 
 			// linker) that need to be resolved by this point
-			DEFERRED_DEPENDENCY_CHECK(!PlaceholderClass->HasKnownReferences());
+			DEFERRED_DEPENDENCY_CHECK(!Placeholder->HasKnownReferences());
+
+			if (!Placeholder->PackageIndex.IsNull() && ensure(Placeholder->PackageIndex.IsImport()))
+			{
+				const UObject* ImportObj = ImportMap[Placeholder->PackageIndex.ToImport()].XObject;
+				DEFERRED_DEPENDENCY_CHECK(ImportObj != PlaceholderObj);
+				DEFERRED_DEPENDENCY_CHECK(Cast<ULinkerPlaceholderClass>(ImportObj) == nullptr);
+				DEFERRED_DEPENDENCY_CHECK(Cast<ULinkerPlaceholderFunction>(ImportObj) == nullptr);
+			}
 		}
+	};
+
+	for (TObjectIterator<ULinkerPlaceholderClass> PlaceholderIt; PlaceholderIt; ++PlaceholderIt)
+	{
+		CheckPlaceholderReferences(*PlaceholderIt);
+	}
+	for (TObjectIterator<ULinkerPlaceholderFunction> PlaceholderIt; PlaceholderIt; ++PlaceholderIt)
+	{
+		CheckPlaceholderReferences(*PlaceholderIt);
 	}
+
 	DEFERRED_DEPENDENCY_CHECK(ImportPlaceholders.Num() == 0);
 #endif // USE_DEFERRED_DEPENDENCY_CHECK_VERIFICATION_TESTS
 #endif // USE_CIRCULAR_DEPENDENCY_LOAD_DEFERRING
@@ -1323,14 +1413,14 @@ void FLinkerLoad::ResolveDeferredDependencies(UStruct* LoadStruct)
 bool FLinkerLoad::HasUnresolvedDependencies() const
 {
 #if USE_CIRCULAR_DEPENDENCY_LOAD_DEFERRING
-	// checking (ResolvingDeferredPlaceholder != nullptr) is not sufficient, 
+	// checking (ResolvingPlaceholderStack.Num() <= 0) is not sufficient, 
 	// because the linker could be in the midst of a nested resolve (for a 
 	// struct, or super... see FLinkerLoad::ResolveDeferredDependencies)
 	bool bIsClassExportUnresolved = FUnresolvedStructTracker::IsAssociatedStructUnresolved(this);
 
-	// (ResolvingDeferredPlaceholder != nullptr) should imply 
+	// (ResolvingPlaceholderStack.Num() <= 0) should imply 
 	// bIsClassExportUnresolved is true (but not the other way around)
-	DEFERRED_DEPENDENCY_CHECK((ResolvingDeferredPlaceholder == nullptr) || bIsClassExportUnresolved);
+	DEFERRED_DEPENDENCY_CHECK((ResolvingPlaceholderStack.Num() <= 0) || bIsClassExportUnresolved);
 	
 	return bIsClassExportUnresolved;
 
@@ -1339,24 +1429,39 @@ bool FLinkerLoad::HasUnresolvedDependencies() const
 #endif // USE_CIRCULAR_DEPENDENCY_LOAD_DEFERRING
 }
 
-int32 FLinkerLoad::ResolveDependencyPlaceholder(FLinkerPlaceholderBase* PlaceholderIn, UClass* ReferencingClass, const FName ObjectPath)
+int32 FLinkerLoad::ResolveDependencyPlaceholder(FLinkerPlaceholderBase* PlaceholderIn, UClass* ReferencingClass, const FName ObjectPathIn)
 {
 #if USE_CIRCULAR_DEPENDENCY_LOAD_DEFERRING
 	TGuardValue<uint32>  LoadFlagsGuard(LoadFlags, (LoadFlags & ~LOAD_DeferDependencyLoads));
-	TGuardValue<FLinkerPlaceholderBase*> ResolvingClassGuard(ResolvingDeferredPlaceholder, PlaceholderIn);
+	ResolvingPlaceholderStack.Push(PlaceholderIn);
 
 	UObject* PlaceholderObj = PlaceholderIn->GetPlaceholderAsUObject();
+#if USE_DEFERRED_DEPENDENCY_CHECK_VERIFICATION_TESTS
 	DEFERRED_DEPENDENCY_CHECK(PlaceholderObj != nullptr);
 	DEFERRED_DEPENDENCY_CHECK(PlaceholderObj->GetOutermost() == LinkerRoot);
+	const int32 ResolvingStackDepth = ResolvingPlaceholderStack.Num();
+#endif
 	
 	UObject* RealImportObj = nullptr;
 
+	FName ObjectPathName = NAME_None;
 	if (PlaceholderIn->PackageIndex.IsNull())
 	{
-		DEFERRED_DEPENDENCY_CHECK(ObjectPath.IsValid() && !ObjectPath.IsNone());
+		ObjectPathName = ObjectPathIn;
+		if (!ObjectPathIn.IsValid() || ObjectPathIn.IsNone())
+		{
+			const FName* ObjectPathPtr = ImportPlaceholders.FindKey(PlaceholderIn);
+			DEFERRED_DEPENDENCY_CHECK(ObjectPathPtr != nullptr);
+			if (ObjectPathPtr)
+			{
+				ObjectPathName = *ObjectPathPtr;
+			}
+		}
+		DEFERRED_DEPENDENCY_CHECK(ObjectPathName.IsValid() && !ObjectPathName.IsNone());
+
 		// emulating the StaticLoadObject() call in UObjectPropertyBase::FindImportedObject(),
 		// since this was most likely a placeholder 
-		RealImportObj = StaticLoadObject(UObject::StaticClass(), /*Outer =*/nullptr, *ObjectPath.ToString(), /*Filename =*/nullptr, (LOAD_NoWarn | LOAD_FindIfFail));
+		RealImportObj = StaticLoadObject(UObject::StaticClass(), /*Outer =*/nullptr, *ObjectPathName.ToString(), /*Filename =*/nullptr, (LOAD_NoWarn | LOAD_FindIfFail));
 	}
 	else
 	{
@@ -1366,7 +1471,9 @@ int32 FLinkerLoad::ResolveDependencyPlaceholder(FLinkerPlaceholderBase* Placehol
 
 		if ((Import.XObject != nullptr) && (Import.XObject != PlaceholderObj))
 		{
-			DEFERRED_DEPENDENCY_CHECK(ResolvingDeferredPlaceholder == PlaceholderIn);
+			DEFERRED_DEPENDENCY_CHECK(ResolvingPlaceholderStack.Num() > 0 && ResolvingPlaceholderStack.Top() == PlaceholderIn);
+			DEFERRED_DEPENDENCY_CHECK(ResolvingPlaceholderStack.Num() == ResolvingStackDepth);
+
 			RealImportObj = Import.XObject;
 		}
 		else
@@ -1452,6 +1559,16 @@ int32 FLinkerLoad::ResolveDependencyPlaceholder(FLinkerPlaceholderBase* Placehol
 	DEFERRED_DEPENDENCY_CHECK(!bIsReferenced || bIsAsyncLoadRef);
 #endif // USE_DEFERRED_DEPENDENCY_CHECK_VERIFICATION_TESTS
 
+	// this could recurse back into ResolveDeferredDependencies(), which resolves all placeholders from this list,
+	// so by the time we're returned here, the list may be empty
+	if (ResolvingPlaceholderStack.Num() > 0)
+	{
+		DEFERRED_DEPENDENCY_CHECK(ResolvingPlaceholderStack.Top() == PlaceholderIn);
+		DEFERRED_DEPENDENCY_CHECK(ResolvingPlaceholderStack.Num() == ResolvingStackDepth);
+
+		ResolvingPlaceholderStack.Pop();
+	}
+	ImportPlaceholders.Remove(ObjectPathName);
 
 	return ReplacementCount;
 #else  // USE_CIRCULAR_DEPENDENCY_LOAD_DEFERRING
@@ -1892,7 +2009,7 @@ void FLinkerLoad::ResolveDeferredExports(UClass* LoadClass)
 		// regenerated the class layout (and property offsets) may no longer 
 		// match the layout that sub-class CDOs were constructed with (making 
 		// property copying dangerous)
-		FDeferredObjInitializerTracker::ResolveDeferredSubClassObjects(LoadClass);
+		FDeferredObjInitializationHelper::ResolveDeferredInitsFromArchetype(BlueprintCDO);
 
 		DEFERRED_DEPENDENCY_CHECK(BlueprintCDO->HasAnyFlags(RF_LoadCompleted));
 	}
@@ -2002,7 +2119,7 @@ void FLinkerLoad::ResetDeferredLoadingState()
 #if USE_CIRCULAR_DEPENDENCY_LOAD_DEFERRING
 	DeferredCDOIndex = INDEX_NONE;
 	bForceBlueprintFinalization = false;
-	ResolvingDeferredPlaceholder = nullptr;
+	ResolvingPlaceholderStack.Empty();
 	ImportPlaceholders.Empty();
 	LoadFlags &= ~(LOAD_DeferDependencyLoads);
 
@@ -2380,251 +2497,359 @@ bool FObjectInitializer::InitNonNativeProperty(UProperty* Property, UObject* Dat
 }
 
 /*******************************************************************************
- * FDeferredObjInitializerTracker
+ * FDeferredInitializationTrackerBase
  ******************************************************************************/
 
-FObjectInitializer* FDeferredObjInitializerTracker::Add(const FObjectInitializer& DeferringInitializer)
+FObjectInitializer* FDeferredInitializationTrackerBase::Add(const UObject* InitDependecy, const FObjectInitializer& DeferringInitializer)
 {
-	UObject* InitingObj = DeferringInitializer.GetObj();
-	DEFERRED_DEPENDENCY_CHECK(InitingObj != nullptr);
-	const bool bIsSubObjTemplate = InitingObj && InitingObj->HasAnyFlags(RF_InheritableComponentTemplate);
+	FObjectInitializer* DeferredInitializerCopy = nullptr;
 
-	UClass* LoadClass = nullptr;
-	if (bIsSubObjTemplate)
-	{
-		LoadClass = Cast<UClass>(InitingObj->GetOuter());
-	}
-	else if (InitingObj != nullptr)
+	DEFERRED_DEPENDENCY_CHECK(InitDependecy);
+	if (InitDependecy)
 	{
-		DEFERRED_DEPENDENCY_CHECK(InitingObj->HasAnyFlags(RF_ClassDefaultObject));
-		LoadClass = InitingObj->GetClass();
+		UObject* InstanceObj = DeferringInitializer.GetObj();
+		ArchetypeInstanceMap.AddUnique(InitDependecy, InstanceObj);
+
+		DEFERRED_DEPENDENCY_CHECK(DeferredInitializers.Find(InstanceObj) == nullptr); // did we try to init the object twice?
+
+		// NOTE: we copy the FObjectInitializer, because it is most likely in the process of being destroyed
+		DeferredInitializerCopy = &DeferredInitializers.Add(InstanceObj, DeferringInitializer);
 	}
-	
-	FObjectInitializer* DeferredCopy = nullptr;
-	if (LoadClass != nullptr)
-	{
-		FDeferredObjInitializerTracker& ThreadInst = FDeferredObjInitializerTracker::Get();
-		TMultiMap<UClass*, UClass*>& SuperClassMap = ThreadInst.SuperClassMap;
+	return DeferredInitializerCopy;
+}
 
-		UClass* SuperClass = LoadClass->GetSuperClass();
-		SuperClassMap.AddUnique(SuperClass, LoadClass);
-#if USE_DEFERRED_DEPENDENCY_CHECK_VERIFICATION_TESTS
-		UObject* SuperCDO = SuperClass->GetDefaultObject(/*bCreateIfNeeded =*/false);
-		DEFERRED_DEPENDENCY_CHECK( SuperCDO && (SuperCDO->HasAnyFlags(RF_NeedLoad) ||
-			(SuperClass->GetLinker() && SuperClass->GetLinker()->IsBlueprintFinalizationPending()) || IsCdoDeferred(SuperClass)) );
+void FDeferredInitializationTrackerBase::ResolveArchetypeInstances(UObject* InitDependecy)
+{
+	TArray<UObject*> ArchetypeInstances;
+	ArchetypeInstanceMap.MultiFind(InitDependecy, ArchetypeInstances);
 
-		FLinkerLoad* ClassLinker = LoadClass->GetLinker();
-		DEFERRED_DEPENDENCY_CHECK((bIsSubObjTemplate && ClassLinker->IsBlueprintFinalizationPending()) || 
-			(ClassLinker->LoadFlags & LOAD_DeferDependencyLoads) != 0);
-#endif // USE_DEFERRED_DEPENDENCY_CHECK_VERIFICATION_TESTS
-		
-		if (bIsSubObjTemplate)
-		{
-#if USE_DEFERRED_DEPENDENCY_CHECK_VERIFICATION_TESTS
-			TArray<FObjectInitializer*> DeferredSubObjInitializers;
-			ThreadInst.DeferredSubObjInitializers.MultiFindPointer(LoadClass, DeferredSubObjInitializers);
-			// check to make sure that we haven't already added this one for deferral
-			for (FObjectInitializer* SubObjInitializer : DeferredSubObjInitializers)
-			{
-				DEFERRED_DEPENDENCY_CHECK(SubObjInitializer->GetObj() != DeferringInitializer.GetObj());
-			}
-#endif // USE_DEFERRED_DEPENDENCY_CHECK_VERIFICATION_TESTS
-			DeferredCopy = &ThreadInst.DeferredSubObjInitializers.Add(LoadClass, DeferringInitializer);
-		}
-		else
-		{
-			TMap<UClass*, FObjectInitializer>& InitializersCache = ThreadInst.DeferredInitializers;
-			DEFERRED_DEPENDENCY_CHECK(InitializersCache.Find(LoadClass) == nullptr); // did we try to init the CDO twice?
+	for (UObject* Instance : ArchetypeInstances)
+	{
+		DEFERRED_DEPENDENCY_CHECK(ResolvingObjects.Contains(Instance) == false);
+		ResolvingObjects.Push(Instance);
 
-			// NOTE: we copy the FObjectInitializer, because it is most likely being destroyed
-			DeferredCopy = &InitializersCache.Add(LoadClass, DeferringInitializer);
+		if (ResolveDeferredInitialization(InitDependecy, Instance))
+		{
+			// For sub-objects, this has to come after ResolveDeferredInitialization(), since InitSubObjectProperties() is 
+			// invoked there (which is where we fill this sub-object with values from the super)
+			PreloadDeferredDependents(Instance);
 		}
+		
+		DEFERRED_DEPENDENCY_CHECK(ResolvingObjects.Top() == Instance);
+		ResolvingObjects.Pop();
 	}
-	return DeferredCopy;
+
+	ArchetypeInstanceMap.Remove(InitDependecy);
 }
 
-FObjectInitializer* FDeferredObjInitializerTracker::Find(UClass* LoadClass)
+bool FDeferredInitializationTrackerBase::IsInitializationDeferred(const UObject* Object) const
 {
-	FDeferredObjInitializerTracker& ThreadInst = FDeferredObjInitializerTracker::Get();
-	return ThreadInst.DeferredInitializers.Find(LoadClass);
+	return DeferredInitializers.Contains(Object);
 }
 
-bool FDeferredObjInitializerTracker::IsCdoDeferred(UClass* LoadClass)
+bool FDeferredInitializationTrackerBase::DeferPreload(UObject* Object)
 {
-	return (Find(LoadClass) != nullptr);
+	const bool bDeferPreload = IsInitializationDeferred(Object);
+	if (bDeferPreload && !IsResolving(Object))
+	{
+		DeferredPreloads.AddUnique(Object, Object);
+	}
+	return bDeferPreload;
 }
 
-bool FDeferredObjInitializerTracker::DeferSubObjectPreload(UObject* SubObject)
+bool FDeferredInitializationTrackerBase::IsResolving(UObject* ArchetypeInstance) const
 {
-	bool bDeferralNeeded = false;
-	// if this is an "InheritableComponentTemplate" we expect it's outer to be
-	// the class it belongs to, and know that this overrides a inherited 
-	// component (belonging to the owner's super)
-	const bool bIsComponentOverride = SubObject->HasAnyFlags(RF_InheritableComponentTemplate);
+	return ResolvingObjects.Contains(ArchetypeInstance);
+}
 
-	UClass* OwningClass = nullptr;
-	if (bIsComponentOverride)
+bool FDeferredInitializationTrackerBase::ResolveDeferredInitialization(UObject* /*ResolvingObject*/, UObject* ArchetypeInstance)
+{
+	if (FObjectInitializer* DeferredInitializer = DeferredInitializers.Find(ArchetypeInstance))
 	{
-		OwningClass = Cast<UClass>(SubObject->GetOuter());
-		DEFERRED_DEPENDENCY_CHECK(OwningClass != nullptr);
+		// initializes and instances CDO properties (copies inherited values 
+		// from the super's CDO)
+		FScriptIntegrationObjectHelper::PostConstructInitObject(*DeferredInitializer);
+
+		DeferredInitializers.Remove(ArchetypeInstance);
 	}
-	else
-	{
-		DEFERRED_DEPENDENCY_CHECK(SubObject->HasAnyFlags(RF_DefaultSubObject));
 
-		UObject* SubObjOuter = SubObject->GetOuter();
-		OwningClass = SubObjOuter->GetClass();
-		// NOTE: the outer of a DSO may not be a CDO like we want, it could 
-		//       be something like a component template; right now we ignore 
-		//       those cases (IsCdoDeferred() below will reject this), but if 
-		//       this case proves to be a problem, then we may need to look up 
-		//       the outer chain, or see if the outer sub-obj is deferred itself
+	return true;
+}
+
+void FDeferredInitializationTrackerBase::PreloadDeferredDependents(UObject* ArchetypeInstance)
+{
+	TArray<UObject*> ObjsToPreload;
+	DeferredPreloads.MultiFind(ArchetypeInstance, ObjsToPreload);
+
+	for (UObject* Object : ObjsToPreload)
+	{
+		FLinkerLoad* Linker = Object->GetLinker();
+		DEFERRED_DEPENDENCY_CHECK(Linker != nullptr);
+		if (Linker)
+		{
+			Linker->Preload(Object);
+		}
 	}
 
-	FDeferredObjInitializerTracker& ThreadInst = FDeferredObjInitializerTracker::Get();	
-	if (IsCdoDeferred(OwningClass) && (OwningClass != ThreadInst.ResolvingClass))
+	DeferredPreloads.Remove(ArchetypeInstance);
+}
+
+/*******************************************************************************
+ * FDeferredCdoInitializationTracker
+ ******************************************************************************/
+
+bool FDeferredCdoInitializationTracker::DeferPreload(UObject* Object)
+{
+	bool bDeferPostload = false;
+
+	if (Object->HasAnyFlags(RF_ClassDefaultObject))
 	{
-		// no need to check the archetype, we know that this is overriding a
-		// super component - we need to defer Preload()
-		if (bIsComponentOverride)
+		// When the initialization has been deferred we have to make sure to
+		// defer serialization as well - don't worry, for CDOs, Preload() will be invoked 
+		// again from FinalizeBlueprint()->ResolveDeferredExports()
+		bDeferPostload = !IsResolving(Object) && IsInitializationDeferred(Object);
+	}
+	else
+	{
+		auto ShouldDeferSubObjectPreload = [this, Object](UObject* OwnerObject)->bool
+		{
+			if (IsInitializationDeferred(OwnerObject))
+			{
+				const bool bDeferSubObjPostload = !IsResolving(OwnerObject);
+				if (bDeferSubObjPostload)
+				{
+					DeferredPreloads.AddUnique(OwnerObject, Object);
+				}
+				return bDeferSubObjPostload;
+			}
+#if USE_DEFERRED_DEPENDENCY_CHECK_VERIFICATION_TESTS
+			else if (OwnerObject)
+			{
+				UObject* OwnerClass = OwnerObject->GetClass();
+				for (auto& DeferredCdo : DeferredInitializers)
+				{
+					// we used to index these by class, so to ensure the same behavior validate
+					// our assumption that we can use the CDO object itself as the key (and that 
+					// using the class wouldn't find a match instead)
+					DEFERRED_DEPENDENCY_CHECK(DeferredCdo.Key->GetClass() != OwnerClass);
+				}
+			}
+#endif // USE_DEFERRED_DEPENDENCY_CHECK_VERIFICATION_TESTS
+			return false;
+		};
+
+		if (Object->HasAnyFlags(RF_DefaultSubObject))
 		{
-			ThreadInst.DeferredSubObjects.AddUnique(OwningClass, SubObject);
-			bDeferralNeeded = true;
+			UObject* SubObjOuter = Object->GetOuter();
+			// NOTE: The outer of a DSO may not be a CDO like we want. It could 
+			//       be something like a component template. Right now we ignore 
+			//       those cases (IsDeferred() will reject this - only CDOs are 
+			//       deferred in this struct), but if this case proves to be a problem, 
+			//       then we may need to look up the outer chain, or see if the outer 
+			//       sub-obj is deferred itself.
+			bDeferPostload = ShouldDeferSubObjectPreload(SubObjOuter);
 		}
-		else
+		else if (Object->HasAnyFlags(RF_InheritableComponentTemplate))
 		{
-			DEFERRED_DEPENDENCY_CHECK(SubObject->GetOuter()->HasAnyFlags(RF_ClassDefaultObject));
+			UClass* OwningClass = Cast<UClass>(Object->GetOuter());
 
-			UObject* SubObjTemplate = SubObject->GetArchetype();
-			// check that this is an inherited sub-object (that it is defined on the 
-			// parent class)... we only need to defer Preload() for inherited 
-			// components because they're what gets filled out in the CDO's InitSubobjectProperties()
-			if (SubObjTemplate && (SubObjTemplate->GetOuter() != SubObject->GetOuter()))
+			DEFERRED_DEPENDENCY_CHECK(OwningClass && OwningClass->ClassDefaultObject);
+			if (OwningClass)
 			{
-				ThreadInst.DeferredSubObjects.AddUnique(OwningClass, SubObject);
-				bDeferralNeeded = true;
+				bDeferPostload = ShouldDeferSubObjectPreload(OwningClass->ClassDefaultObject);
 			}
-		}		
-	}
-
-#if USE_DEFERRED_DEPENDENCY_CHECK_VERIFICATION_TESTS
-	TArray<FObjectInitializer*> DeferredSubObjInitializers;
-	ThreadInst.DeferredSubObjInitializers.MultiFindPointer(OwningClass, DeferredSubObjInitializers);
-	// check to make sure that we haven't already added this one for deferral
-	for (FObjectInitializer* SubObjInitializer : DeferredSubObjInitializers)
-	{
-		DEFERRED_DEPENDENCY_CHECK(bDeferralNeeded || SubObjInitializer->GetObj() != SubObject);
+		}
 	}
-#endif // USE_DEFERRED_DEPENDENCY_CHECK_VERIFICATION_TESTS
-	return bDeferralNeeded;
+	return bDeferPostload;
 }
 
-void FDeferredObjInitializerTracker::Remove(UClass* LoadClass)
-{
-	FDeferredObjInitializerTracker& ThreadInst = FDeferredObjInitializerTracker::Get();
-	ThreadInst.DeferredInitializers.Remove(LoadClass);
-	ThreadInst.DeferredSubObjects.Remove(LoadClass);
-	ThreadInst.SuperClassMap.RemoveSingle(LoadClass->GetSuperClass(), LoadClass);
-	ThreadInst.DeferredSubObjInitializers.Remove(LoadClass);
-}
+/*******************************************************************************
+ * FDeferredSubObjInitializationTracker
+ ******************************************************************************/
 
-bool FDeferredObjInitializerTracker::ResolveDeferredInitialization(UClass* LoadClass)
+bool FDeferredSubObjInitializationTracker::ResolveDeferredInitialization(UObject* ResolvingObject, UObject* ArchetypeInstance)
 {
-	if (FObjectInitializer* DeferredInitializer = FDeferredObjInitializerTracker::Find(LoadClass))
-	{
-		FDeferredObjInitializerTracker& ThreadInst = FDeferredObjInitializerTracker::Get();
-		TGuardValue<UClass*> ResolvingGuard(ThreadInst.ResolvingClass, LoadClass);
+	bool bInitializerRan = false;
 
-		DEFERRED_DEPENDENCY_CHECK(!LoadClass->GetSuperClass()->HasAnyClassFlags(CLASS_NewerVersionExists));
-		// initializes and instances CDO properties (copies inherited values 
-		// from the super's CDO)
-		FScriptIntegrationObjectHelper::PostConstructInitObject(*DeferredInitializer);
+	// If we deferred the sub-object because the super CDO wasn't ready, we still
+	// need to check that its archetype is in a ready state (ready to be copied from)
+	if (ResolvingObject->HasAnyFlags(RF_ClassDefaultObject))
+	{
+		if (FObjectInitializer* DeferredInitializer = DeferredInitializers.Find(ArchetypeInstance))
+		{
+			UObject* Archetype = DeferredInitializer->GetArchetype();
+			// When this sub-object was created it's archetype object (the 
+			// super's sub-obj) may not have been created yet. In that scenario, the 
+			// component class's CDO would have been used in its place; now that 
+			// the super is good, we should update the archetype
+			if (ArchetypeInstance->HasAnyFlags(RF_ClassDefaultObject))
+			{
+				Archetype = UObject::GetArchetypeFromRequiredInfo(ArchetypeInstance->GetClass(), ArchetypeInstance->GetOuter(), ArchetypeInstance->GetFName(), ArchetypeInstance->GetFlags());
+			}
 
-		ResolveDeferredSubObjects(LoadClass->GetDefaultObject());
-		FDeferredObjInitializerTracker::Remove(LoadClass);
+			const bool bArchetypeLoadPending = Archetype &&
+				(Archetype->HasAnyFlags(RF_NeedLoad) || (Archetype->HasAnyFlags(RF_WasLoaded) && !Archetype->HasAnyFlags(RF_LoadCompleted)));
 
-		return true;
+			if (bArchetypeLoadPending)
+			{
+				// Archetype isn't ready, move the deferred initializer to wait for its archetype 
+				ArchetypeInstanceMap.Add(Archetype, ArchetypeInstance);
+				// don't need to add this to DeferredInitializers, as it is already there
+			}
+			else
+			{
+				bInitializerRan = FDeferredInitializationTrackerBase::ResolveDeferredInitialization(ResolvingObject, ArchetypeInstance);
+			}
+		}
 	}
 	else
 	{
-		// make sure we're not missing sub-obj initializers that would have been 
-		// ran in the above if case
-		DEFERRED_DEPENDENCY_CHECK(FDeferredObjInitializerTracker::Get().DeferredSubObjInitializers.Find(LoadClass) == nullptr);
+		bInitializerRan = FDeferredInitializationTrackerBase::ResolveDeferredInitialization(ResolvingObject, ArchetypeInstance);
 	}
-	return false;
+
+	return bInitializerRan;
 }
 
-void FDeferredObjInitializerTracker::ResolveDeferredSubObjects(UObject* CDO)
+/*******************************************************************************
+ * FDeferredObjInitializationHelper
+ ******************************************************************************/
+
+FObjectInitializer* FDeferredObjInitializationHelper::DeferObjectInitializerIfNeeded(const FObjectInitializer& DeferringInitializer)
 {
-	DEFERRED_DEPENDENCY_CHECK(CDO->HasAnyFlags(RF_ClassDefaultObject));
-	UClass* LoadClass = CDO->GetClass();
+	FObjectInitializer* DeferredInitializerCopy = nullptr;
 
-	FDeferredObjInitializerTracker& ThreadInst = FDeferredObjInitializerTracker::Get();
-	// guard to keep sub-object preloads from ending up back in DeferSubObjectPreload()
-	TGuardValue<UClass*> ResolvingGuard(ThreadInst.ResolvingClass, LoadClass);
-	
-	TArray<FObjectInitializer*> DeferredSubObjInitializers;
-	ThreadInst.DeferredSubObjInitializers.MultiFindPointer(LoadClass, DeferredSubObjInitializers);
-
-	for (FObjectInitializer* DeferredInitializer : DeferredSubObjInitializers)
-	{
-		UObject* SubObjArchetype = DeferredInitializer->GetArchetype();
-		// just because the class's CDO archetype has completed serialization 
-		// does not mean that the archetypes for these sub-objects have... 
-		// unfortunately there isn't anywhere else where we ensure that these 
-		// archetypes are loaded (else we should defer this further, until then 
-		// when we can guarantee that each archetype is complete) - we don't  
-		// even force a load of these sub-object archetypes prior to their own 
-		// Blueprint's regeneration, meaning it's compiled-on-load potentially 
-		// with incomplete component templates; this is potentially bad in 
-		// itself (TODO: investigate)
-		if (SubObjArchetype && !SubObjArchetype->HasAnyFlags(RF_LoadCompleted))
-		{
-			if (FLinkerLoad* SubObjLinker = SubObjArchetype->GetLinker())
+	UObject* TargetObj = DeferringInitializer.GetObj();
+	if (TargetObj)
+	{
+		FDeferredCdoInitializationTracker& CdoInitDeferalSys = FDeferredCdoInitializationTracker::Get();
+		auto IsSuperCdoReadyToBeCopied = [&CdoInitDeferalSys](const UClass* LoadClass, const UObject* SuperCDO)->bool
+		{
+			// RF_WasLoaded indicates that this Super was loaded from disk (and hasn't been regenerated on load)
+			// regenerated CDOs will not have the RF_LoadCompleted
+			const bool bSuperCdoLoadPending = CdoInitDeferalSys.IsInitializationDeferred(SuperCDO) ||
+				SuperCDO->HasAnyFlags(RF_NeedLoad) || (SuperCDO->HasAnyFlags(RF_WasLoaded) && !SuperCDO->HasAnyFlags(RF_LoadCompleted));
+
+			if (bSuperCdoLoadPending)
 			{
-				check(!GEventDrivenLoaderEnabled || !EVENT_DRIVEN_ASYNC_LOAD_ACTIVE_AT_RUNTIME);
-				SubObjArchetype->SetFlags(RF_NeedLoad);
-				SubObjLinker->Preload(SubObjArchetype);
+				const FLinkerLoad* ObjLinker = LoadClass->GetLinker();
+				const bool bIsBpClassSerializing = ObjLinker && (ObjLinker->LoadFlags & LOAD_DeferDependencyLoads);
+				const bool bIsResolvingDeferredObjs = LoadClass->HasAnyFlags(RF_LoadCompleted) &&
+					ObjLinker && ObjLinker->IsBlueprintFinalizationPending();
+
+				DEFERRED_DEPENDENCY_CHECK(bIsBpClassSerializing || bIsResolvingDeferredObjs);
+				return !bIsBpClassSerializing && !bIsResolvingDeferredObjs;
+			}
+			return true;
+		};
+
+		const bool bIsCDO = TargetObj->HasAnyFlags(RF_ClassDefaultObject);
+		if (bIsCDO)
+		{
+			const UClass* CdoClass = DeferringInitializer.GetClass();
+			UClass* SuperClass = CdoClass->GetSuperClass();
+
+			if (!CdoClass->IsNative() && !SuperClass->IsNative())
+			{
+				DEFERRED_DEPENDENCY_CHECK(CdoClass->HasAnyClassFlags(CLASS_CompiledFromBlueprint));
+				DEFERRED_DEPENDENCY_CHECK(SuperClass->HasAnyClassFlags(CLASS_CompiledFromBlueprint));
+
+				const UObject* SuperCDO = DeferringInitializer.GetArchetype();
+				DEFERRED_DEPENDENCY_CHECK(SuperCDO && SuperCDO->HasAnyFlags(RF_ClassDefaultObject));
+				// use the ObjectArchetype for the super CDO because the SuperClass may have a REINST CDO cached currently
+				SuperClass = SuperCDO->GetClass();
+
+				if (!IsSuperCdoReadyToBeCopied(CdoClass, SuperCDO))
+				{
+					DeferredInitializerCopy = CdoInitDeferalSys.Add(SuperCDO, DeferringInitializer);
+				}
 			}
 		}
-		FScriptIntegrationObjectHelper::PostConstructInitObject(*DeferredInitializer);
-	}
-	ThreadInst.DeferredSubObjInitializers.Remove(LoadClass);
-	
-	TArray<UObject*> DeferredSubObjects;
-	ThreadInst.DeferredSubObjects.MultiFind(LoadClass, DeferredSubObjects);
+		// since "InheritableComponentTemplate"s are not default sub-objects, 
+		// they won't be fixed up by the owner's FObjectInitializer (CDO 
+		// FObjectInitializers will init default sub-object properties, copying  
+		// from the super's DSOs) - this means that we need to separately defer 
+		// init'ing these sub-objects when their archetype hasn't been loaded yet 
+		else if (TargetObj->HasAnyFlags(RF_InheritableComponentTemplate))
+		{
+			const UClass* OwnerClass = Cast<UClass>(TargetObj->GetOuter());
+			DEFERRED_DEPENDENCY_CHECK(OwnerClass && OwnerClass->HasAnyClassFlags(CLASS_CompiledFromBlueprint));
+			const UClass* SuperClass = OwnerClass->GetSuperClass();
 
-	FLinkerLoad* ClassLinker = LoadClass->GetLinker();
-	DEFERRED_DEPENDENCY_CHECK(ClassLinker != nullptr);
-	if (ClassLinker != nullptr)
-	{
-		// this all needs to happen after PostConstructInitObject() (in 
-		// ResolveDeferredInitialization), since InitSubObjectProperties() is 
-		// invoked there (which is where we fill this sub-object with values 
-		// from the super)... here we account for any Preload() calls that were 
-		// skipped on account of the deferred CDO initialization
-		for (UObject* SubObj : DeferredSubObjects)
-		{
-			DEFERRED_DEPENDENCY_CHECK( (SubObj->GetOuter() == CDO && SubObj->HasAnyFlags(RF_DefaultSubObject)) || 
-				(SubObj->GetOuter() == LoadClass && SubObj->HasAnyFlags(RF_InheritableComponentTemplate)) );
+			if (SuperClass && !SuperClass->IsNative())
+			{
+				// It is possible that the archetype isn't even correct, if the 
+				// super's sub-object hasn't even been created yet (in this case the 
+				// component's CDO is used, which is probably wrong)
+				// 
+				// So if the super CDO isn't ready, we need to defer this sub-object
+				const UObject* SuperCDO = SuperClass->ClassDefaultObject;
+				if (!IsSuperCdoReadyToBeCopied(OwnerClass, SuperCDO))
+				{
+					FDeferredSubObjInitializationTracker& SubObjInitDeferalSys = FDeferredSubObjInitializationTracker::Get();
+					DeferredInitializerCopy = SubObjInitDeferalSys.Add(SuperCDO, DeferringInitializer);
+				}
+			}
+
+			// if it passed the super CDO check above, assume the archetype is kosher
+			if (!DeferredInitializerCopy)
+			{
+				UObject* Archetype = DeferringInitializer.GetArchetype();
 			
-			ClassLinker->Preload(SubObj);
+				const bool bArchetypeLoadPending = Archetype &&
+					( Archetype->HasAnyFlags(RF_NeedLoad) || (Archetype->HasAnyFlags(RF_WasLoaded) && !Archetype->HasAnyFlags(RF_LoadCompleted)) );
+
+				if (bArchetypeLoadPending)
+				{
+					FDeferredSubObjInitializationTracker& SubObjInitDeferalSys = FDeferredSubObjInitializationTracker::Get();
+					DeferredInitializerCopy = SubObjInitDeferalSys.Add(Archetype, DeferringInitializer);
+				}
+			}
 		}
 	}
 
-	ThreadInst.DeferredSubObjects.Remove(LoadClass);
+	return DeferredInitializerCopy;
 }
 
-void FDeferredObjInitializerTracker::ResolveDeferredSubClassObjects(UClass* SuperClass)
+bool FDeferredObjInitializationHelper::DeferObjectPreload(UObject* Object)
 {
-	FDeferredObjInitializerTracker& ThreadInst = FDeferredObjInitializerTracker::Get();
-	TArray<UClass*> DeferredSubClasses;
-	ThreadInst.SuperClassMap.MultiFind(SuperClass, DeferredSubClasses);
+	return FDeferredCdoInitializationTracker::Get().DeferPreload(Object) || FDeferredSubObjInitializationTracker::Get().DeferPreload(Object);
+}
 
-	for (UClass* SubClass : DeferredSubClasses)
+void FDeferredObjInitializationHelper::ResolveDeferredInitsFromArchetype(UObject* Archetype)
+{
+	FDeferredCdoInitializationTracker& DeferredCdoTracker = FDeferredCdoInitializationTracker::Get();
+	FDeferredSubObjInitializationTracker& DeferredSubObjTracker = FDeferredSubObjInitializationTracker::Get();
+
+#if USE_DEFERRED_DEPENDENCY_CHECK_VERIFICATION_TESTS
+	if (Archetype->HasAnyFlags(RF_ClassDefaultObject))
 	{
-		ResolveDeferredInitialization(SubClass);
+		// we used to index the deferred initialization by class, so to ensure the same behavior validate
+		// our assumption that we can use the CDO object itself as the key (and that using the class wouldn't find a match instead)
+		auto IsDeferredByClass = [Archetype](const TMultiMap<const UObject*, UObject*>& ArchetypeMap)->bool
+		{
+			for (auto& DeferredObj : ArchetypeMap)
+			{
+				if (DeferredObj.Key->GetClass() == Archetype->GetClass())
+				{
+					return true;
+				}
+			}
+			return false;
+		};
+
+		if (!DeferredCdoTracker.ArchetypeInstanceMap.Contains(Archetype))
+		{
+			DEFERRED_DEPENDENCY_CHECK(IsDeferredByClass(DeferredCdoTracker.ArchetypeInstanceMap) == false);
+		}
+		if (!DeferredSubObjTracker.ArchetypeInstanceMap.Contains(Archetype))
+		{
+			DEFERRED_DEPENDENCY_CHECK(IsDeferredByClass(DeferredSubObjTracker.ArchetypeInstanceMap) == false);
+		}
 	}
+#endif // USE_DEFERRED_DEPENDENCY_CHECK_VERIFICATION_TESTS
+
+	DeferredCdoTracker.ResolveArchetypeInstances(Archetype);
+	DeferredSubObjTracker.ResolveArchetypeInstances(Archetype);
 }
 
 // don't want other files ending up with this internal define
@@ -2770,10 +2995,3 @@ void IBlueprintNativeCodeGenCore::Register(const IBlueprintNativeCodeGenCore* Co
 }
 
 #endif // WITH_EDITOR
-
-
-
-
-
-
-
diff --git a/Engine/Source/Runtime/CoreUObject/Private/UObject/LinkerLoad.cpp b/Engine/Source/Runtime/CoreUObject/Private/UObject/LinkerLoad.cpp
--- a/Engine/Source/Runtime/CoreUObject/Private/UObject/LinkerLoad.cpp
+++ b/Engine/Source/Runtime/CoreUObject/Private/UObject/LinkerLoad.cpp
@@ -799,7 +799,6 @@ FLinkerLoad::FLinkerLoad(UPackage* InParent, const TCHAR* InFilename, uint32 InL
 #if USE_CIRCULAR_DEPENDENCY_LOAD_DEFERRING
 ,	bForceBlueprintFinalization(false)
 ,	DeferredCDOIndex(INDEX_NONE)
-,	ResolvingDeferredPlaceholder(nullptr)
 #endif // USE_CIRCULAR_DEPENDENCY_LOAD_DEFERRING
 {
 	FMemory::Memset(ExportHash, INDEX_NONE, sizeof(ExportHash));
@@ -3172,34 +3171,11 @@ void FLinkerLoad::Preload( UObject* Object )
 #endif // USE_CIRCULAR_DEPENDENCY_LOAD_DEFERRING
 
 #if USE_CIRCULAR_DEPENDENCY_LOAD_DEFERRING
-			// Because of delta serialization, we require that a parent's CDO be 
-			// fully serialized before its children's CDOs are created. However, 
-			// due to cyclic parent/child dependencies, we have some cases where 
-			// the linker breaks that expected behavior. In those cases, we 
-			// defer the child's initialization (i.e. defer copying of parent  
-			// property values, etc.), and wait until we can guarantee that the 
-			// parent CDO has been fully loaded.
-			//
-			// In a normal scenario, the order of property initialization is:
-			// Creation (zeroed) -> Initialization (copied super's values) -> Serialization (overridden values loaded)
-			// When the initialization has been deferred we have to make sure to
-			// defer serialization here as well (don't worry, it will be invoked 
-			// again from FinalizeBlueprint()->ResolveDeferredExports())
-			if (Object->HasAnyFlags(RF_ClassDefaultObject) && FDeferredObjInitializerTracker::IsCdoDeferred(Object->GetClass()))
-			{
-				return;
-			}
-			// if this is an inherited sub-object on a CDO, and that CDO has had
-			// its initialization deferred (for reasons explained above), then 
-			// we shouldn't serialize in data for this quite yet... not until 
-			// its owner has had a chaTnce to initialize itself (because, as part
-			// of CDO initialization, inherited sub-objects get filled in with 
-			// values inherited from the super)
-			else if (Object->HasAnyFlags(RF_DefaultSubObject|RF_InheritableComponentTemplate) && FDeferredObjInitializerTracker::DeferSubObjectPreload(Object))
+			// In certain situations, a constructed object has its initializer deferred (when its archetype hasn't been serialized).
+			// In those cases, we shouldn't serialize the object yet (initialization needs to run first).
+			// See the comment on DeferObjectPreload() for more info on the issue.
+			if (FDeferredObjInitializationHelper::DeferObjectPreload(Object))
 			{
-				// don't worry, FDeferredObjInitializerTracker::DeferSubObjectPreload() 
-				// should have cached this object, and it will run Preload() on 
-				// this later (once the super CDO has been initialized)
 				return;
 			}
 #endif // USE_CIRCULAR_DEPENDENCY_LOAD_DEFERRING
@@ -3321,8 +3297,8 @@ void FLinkerLoad::Preload( UObject* Object )
 					}
 				}
 
-				{
 #if USE_CIRCULAR_DEPENDENCY_LOAD_DEFERRING
+				{
 					SCOPE_CYCLE_COUNTER(STAT_LinkerLoadDeferred);
 					if ((LoadFlags & LOAD_DeferDependencyLoads) != (*LoadFlagsGuard & LOAD_DeferDependencyLoads))
 					{
@@ -3366,9 +3342,24 @@ void FLinkerLoad::Preload( UObject* Object )
 							FinalizeBlueprint(ObjectAsClass);
 						}
 					}
-#endif // USE_CIRCULAR_DEPENDENCY_LOAD_DEFERRING
 				}
 
+				// Conceptually, we could run this here for CDOs and it shouldn't be a problem.
+				// 
+				// We don't do it here for CDOs because we were already doing it for them in 
+				// ResolveDeferredExports(), and we don't want to destabalize the functional 
+				// load order of things (doing it here could cause subsequent loads which would 
+				// happen from a point in ResolveDeferredExports() where they didn't happen before - again, this 
+				// should be fine; we're just keeping the surface area of this to a minimum at this time)
+				if (!Object->HasAnyFlags(RF_ClassDefaultObject))
+				{
+					// If this was an archetype object, there may be some initializers/preloads that
+					// were waiting for it to be fully serialized
+					FDeferredObjInitializationHelper::ResolveDeferredInitsFromArchetype(Object);
+				}
+#endif // USE_CIRCULAR_DEPENDENCY_LOAD_DEFERRING
+				
+
 				// Make sure we serialized the right amount of stuff.
 				int64 Pos = Tell();
 				int64 SizeSerialized = Pos - Export.SerialOffset;
diff --git a/Engine/Source/Runtime/CoreUObject/Private/UObject/UObjectGlobals.cpp b/Engine/Source/Runtime/CoreUObject/Private/UObject/UObjectGlobals.cpp
--- a/Engine/Source/Runtime/CoreUObject/Private/UObject/UObjectGlobals.cpp
+++ b/Engine/Source/Runtime/CoreUObject/Private/UObject/UObjectGlobals.cpp
@@ -2646,94 +2646,16 @@ FObjectInitializer::~FObjectInitializer()
 #if USE_CIRCULAR_DEPENDENCY_LOAD_DEFERRING
 	bool bIsPostConstructInitDeferred = false;
 	if (!FBlueprintSupport::IsDeferredCDOInitializationDisabled())
-	{		
-		UClass* BlueprintClass = nullptr;
-		// since "InheritableComponentTemplate"s are not default sub-objects, 
-		// they won't be fixed up by the owner's FObjectInitializer (CDO 
-		// FObjectInitializers will init default sub-object properties, copying  
-		// from the super's DSOs) - this means that we need to separately defer 
-		// init'ing these sub-objects when their archetype hasn't been loaded 
-		// yet (it is possible that the archetype isn't even correct, as the 
-		// super's sub-object hasn't even been created yet; in this case the 
-		// component's CDO is used, which is probably wrong)
-		if (Obj->HasAnyFlags(RF_InheritableComponentTemplate))
-		{
-			BlueprintClass = Cast<UClass>(Obj->GetOuter());
-#if USE_DEFERRED_DEPENDENCY_CHECK_VERIFICATION_TESTS
-			check(BlueprintClass != nullptr);
-#endif // USE_DEFERRED_DEPENDENCY_CHECK_VERIFICATION_TESTS
-		}
-		else if (bIsCDO && !Class->IsNative())
-		{
-			BlueprintClass = Class;
-#if USE_DEFERRED_DEPENDENCY_CHECK_VERIFICATION_TESTS
-			check(Class->HasAnyClassFlags(CLASS_CompiledFromBlueprint));
-#endif // USE_DEFERRED_DEPENDENCY_CHECK_VERIFICATION_TESTS
-		}
-
-		if (BlueprintClass != nullptr)
+	{
+		if (FObjectInitializer* DeferredCopy = FDeferredObjInitializationHelper::DeferObjectInitializerIfNeeded(*this))
 		{
-#if USE_DEFERRED_DEPENDENCY_CHECK_VERIFICATION_TESTS
-			check(!BlueprintClass->IsNative());
-#endif // USE_DEFERRED_DEPENDENCY_CHECK_VERIFICATION_TESTS
+			DeferredCopy->bIsDeferredInitializer = true;
+			// make sure this wasn't mistakenly pushed into ObjectInitializers
+			// (the copy constructor should have been what was invoked, 
+			// which doesn't push to ObjectInitializers)
+			check(FUObjectThreadContext::Get().TopInitializer() != DeferredCopy);
 
-			UClass* SuperClass = BlueprintClass->GetSuperClass();
-			if (SuperClass && !SuperClass->IsNative())
-			{
-				UObject* SuperBpCDO = nullptr;
-				// if this is a CDO (then we know/assume the archetype is the 
-				// CDO from the super class), use the ObjectArchetype for the 
-				// SuperBpCDO (because the SuperClass may have a REINST CDO 
-				// cached currently)
-				if (bIsCDO)
-				{
-					SuperBpCDO = ObjectArchetype;
-					SuperClass = ObjectArchetype->GetClass();
-
-#if USE_DEFERRED_DEPENDENCY_CHECK_VERIFICATION_TESTS
-					check(ObjectArchetype->HasAnyFlags(RF_ClassDefaultObject));
-#endif // USE_DEFERRED_DEPENDENCY_CHECK_VERIFICATION_TESTS
-				}
-				else
-				{
-					SuperBpCDO = SuperClass->GetDefaultObject(/*bCreateIfNeeded =*/false);
-				}
-
-				FLinkerLoad* SuperClassLinker = SuperClass->GetLinker();
-				const bool bSuperLoadPending = FDeferredObjInitializerTracker::IsCdoDeferred(SuperClass) ||
-					(SuperBpCDO && (SuperBpCDO->HasAnyFlags(RF_NeedLoad) || (SuperBpCDO->HasAnyFlags(RF_WasLoaded) && !SuperBpCDO->HasAnyFlags(RF_LoadCompleted))));
-
-				FLinkerLoad* ObjLinker = BlueprintClass->GetLinker();
-				const bool bIsBpClassSerializing    = ObjLinker && (ObjLinker->LoadFlags & LOAD_DeferDependencyLoads);
-				const bool bIsResolvingDeferredObjs = BlueprintClass->HasAnyFlags(RF_LoadCompleted) &&
-					ObjLinker && ObjLinker->IsBlueprintFinalizationPending();
-
-				if (bSuperLoadPending && (bIsBpClassSerializing || bIsResolvingDeferredObjs))
-				{
-#if USE_DEFERRED_DEPENDENCY_CHECK_VERIFICATION_TESTS
-					// make sure we haven't already deferred this once, if we have
-					// then something is destroying this one prematurely 
-					check(bIsDeferredInitializer == false);
-
-					for (const FSubobjectsToInit::FSubobjectInit& SubObjInfo : ComponentInits.SubobjectInits)
-					{
-						check(!SubObjInfo.Subobject->HasAnyFlags(RF_NeedLoad));
-					}
-#endif // USE_DEFERRED_DEPENDENCY_CHECK_VERIFICATION_TESTS
-					
-					// makes a copy of this and saves it off, to be ran later
-					if (FObjectInitializer* DeferredCopy = FDeferredObjInitializerTracker::Add(*this))
-					{
-						bIsPostConstructInitDeferred = true;
-						DeferredCopy->bIsDeferredInitializer = true;
-
-						// make sure this wasn't mistakenly pushed into ObjectInitializers
-						// (the copy constructor should have been what was invoked, 
-						// which doesn't push to ObjectInitializers)
-						check(FUObjectThreadContext::Get().TopInitializer() != DeferredCopy);
-					}
-				}
-			}
+			bIsPostConstructInitDeferred = true;
 		}
 	}
 
diff --git a/Engine/Source/Runtime/CoreUObject/Public/Blueprint/BlueprintSupport.h b/Engine/Source/Runtime/CoreUObject/Public/Blueprint/BlueprintSupport.h
--- a/Engine/Source/Runtime/CoreUObject/Public/Blueprint/BlueprintSupport.h
+++ b/Engine/Source/Runtime/CoreUObject/Public/Blueprint/BlueprintSupport.h
@@ -206,54 +206,177 @@ struct IBlueprintNativeCodeGenCore
 #endif // WITH_EDITOR
 
 /** 
- * A helper struct for storing FObjectInitializers that were not run on 
- * Blueprint CDO's (or CDO sub-objects) post-construction (presumably because 
- * that CDO's super had not been fully serialized yet). 
+ * A base struct for storing FObjectInitializers that were not run on 
+ * Blueprint objects post-construction (presumably because the object's super/archetype
+ * had not been fully serialized yet). 
  * 
  * This was designed to hold onto FObjectInitializers until a later point, when 
- * they can properly be ran (presumably in FLinkerLoad::ResolveDeferredExports,
- * after the super has been finalized).
+ * they can properly be ran (after the archetype has been serialized).
  */
-struct FDeferredObjInitializerTracker : TThreadSingleton<FDeferredObjInitializerTracker>
+struct FDeferredInitializationTrackerBase
 {
 public:
-	FDeferredObjInitializerTracker() : ResolvingClass(nullptr) {}
+	virtual ~FDeferredInitializationTrackerBase() {}
 
-	/** Stores a copy of the specified FObjectInitializer and returns a pointer to it (could be null if a corresponding class could not be determined). */
-	static FObjectInitializer* Add(const FObjectInitializer& DeferringInitializer);
+	/** 
+	 * Makes a copy of the specified initializer and stores it (mapped under its 
+	 * dependency), so that it can instead be executed later via ResolveArchetypeInstances().
+	 * 
+	 * @param  InitDependecy			The object (usually the initializer's archetype) that this initializer is dependent on. The key 
+	 *									you'll pass to ResolveArchetypeInstances() later to run this initializer.
+	 * @param  DeferringInitializer		The initializer you want to defer.
+	 *
+	 * @raturn A copy of the specified initializer. This should always succeed (only returns null if InitDependecy is null).
+	 */
+	FObjectInitializer* Add(const UObject* InitDependecy, const FObjectInitializer& DeferringInitializer);
+
+	/** 
+	 * Runs all deferred initializers that were dependent on the specified archetype 
+	 * (unless they're dependent on another), and runs Preload() on any objects that 
+	 * had their Preload() skipped as a result.
+	 * 
+	 * @param  ArchetypeKey		The initializer dependency that has been fully loaded/serialized.
+	 */
+	void ResolveArchetypeInstances(UObject* ArchetypeKey);
 
-	/** Looks up a FObjectInitializer that was deferred for the specified class (an FObjectInitializer for that class's CDO). */
-	static FObjectInitializer* Find(UClass* LoadClass);
+	/** 
+	 * Checks to see if the specified object has had its initialization deferred (meaning a super/archetype
+	 * hasn't had FObjectInitializer::PostConstructInit() ran on it yet). 
+	 */
+	bool IsInitializationDeferred(const UObject* Object) const;
 
-	/** Checks to see if the specified class has been logged as deferred (meaning its CDO hasn't had FObjectInitializer::PostConstructInit() ran on it yet). */
-	static bool IsCdoDeferred(UClass* LoadClass);
+	/**
+	 * Determines if the specified object needs to have its Preload() call deferred 
+	 * (this is meant to be called from Preload() itself). If so, this will record the object
+	 * and serialize it in later, once it's initializer dependency has been resolved 
+	 * (from ResolveArchetypeInstances).
+	 *
+	 * This should be the case for any object that's had its initialization deferred (need 
+	 * to initialize before you serialize), and any dependencies (sub-objects, etc.) waiting 
+	 * on that object's initialization.
+	 * 
+	 * @return True if Object's load/serialization should be skipped (for now).
+	 */
+	virtual bool DeferPreload(UObject* Object);
 
-	/** Determines if the specified sub-object should have its Preload() skipped; if so, this will cache the sub-object and return true. */
-	static bool DeferSubObjectPreload(UObject* SubObject);
+protected:
+	/** 
+	 * Used to keep DeferPreload() from re-adding objects while we're in the midst of resolving.
+	 */
+	bool IsResolving(UObject* ArchetypeInstance) const;
 
-	/** Destroys any FObjectInitializers that were cached corresponding to the specified class. */
-	static void Remove(UClass* LoadClass);
+	/** 
+	 * Runs the deferred initializer for the specified archetype object if its not dependent 
+	 * on other archetypes (like a sub-object that first requires the super's CDO to be constructed, 
+	 * and then for its archetype to be serialized).
+	 *
+	 * If the initializer needs to be further deferred, this should re-register under its new dependency.
+	 *
+	 * @param  ResolvingObject		The dependency that has been fully loaded/serialized, that the object's initializer was logged under.
+	 * @param  ArchetypeInstance	The object which has had its initializer deferred.
+	 *
+	 * @return True if the initializer was ran, other wise false.
+	 */
+	virtual bool ResolveDeferredInitialization(UObject* ResolvingObject, UObject* ArchetypeInstance);
 
-	/** Runs FObjectInitializer::PostConstructInit() on the specified class's CDO (if it was deferred), and preloads any sub-objects that were skipped. */
-	static bool ResolveDeferredInitialization(UClass* LoadClass);
+	/** 
+	 * Runs through all objects that had their Preload() skipped due to an initializer 
+	 * being deferred. Runs serialization on each object.
+	 */
+	void PreloadDeferredDependents(UObject* ArchetypeInstance);
 
-	/**  */
-	static void ResolveDeferredSubObjects(UObject* CDO);
+protected:
+#if USE_DEFERRED_DEPENDENCY_CHECK_VERIFICATION_TESTS
+	friend struct FDeferredObjInitializationHelper;
+#endif
 
-	/**  */
-	static void ResolveDeferredSubClassObjects(UClass* SuperClass);
+	/** Tracks objects that have had their initialization deferred as a result of their archetype not being fully serialized (maps archetype => list-o-instances) */
+	TMultiMap<const UObject*, UObject*> ArchetypeInstanceMap;
+	/** A map that lets us look up FObjectInitializers by their UObject */
+	TMap<UObject*, FObjectInitializer> DeferredInitializers;
+	/** Used to keep ResolveArchetypeInstances() from re-adding sub-objects via DeferPreload() */
+	TArray<UObject*> ResolvingObjects;
+	/** Track default objects that had their Preload() skipped, because a archetype dependency should initialize first */
+	TMultiMap<UObject*, UObject*> DeferredPreloads;
+};
 
-private:
-	/** A map that tracks the relationship between Blueprint classes and FObjectInitializers for their CDOs */
-	TMap<UClass*, FObjectInitializer> DeferredInitializers;
-	/** Track default sub-objects that had their Preload() skipped, because the owning CDO's initialization should happen first */
-	TMultiMap<UClass*, UObject*> DeferredSubObjects;
-	/** Used to keep ResolveDeferredSubObjects() from re-adding sub-objects via DeferSubObjectPreload() */
-	UClass* ResolvingClass;
-	/** Tracks sub-classes that have had their CDO deferred as a result of the super not being fully serialized */
-	TMultiMap<UClass*, UClass*> SuperClassMap;
-	/** A map that tracks deferred FObjectInitializers for a class/CDO's sub-objects (not DSOs, but things like component override templates)*/
-	TMultiMap<UClass*, FObjectInitializer> DeferredSubObjInitializers;
+/** 
+ * Specialized FDeferredInitializationTracker for tracking deferred Blueprint CDOs specifically. 
+ * (every object in DeferredInitializers should be a BP CDO).
+ */
+struct FDeferredCdoInitializationTracker : FDeferredInitializationTrackerBase, TThreadSingleton<FDeferredCdoInitializationTracker>
+{
+public:
+	//~ FDeferredInitializationTrackerBase interface
+	virtual bool DeferPreload(UObject* Object) override;
+};
+
+/**
+* Specialized FDeferredInitializationTracker for tracking deferred Blueprint sub-objects specifically.
+* (every object in DeferredInitializers should be a default sub-object or component template).
+*/
+struct FDeferredSubObjInitializationTracker : FDeferredInitializationTrackerBase, TThreadSingleton<FDeferredSubObjInitializationTracker>
+{
+protected:
+	//~ FDeferredInitializationTrackerBase interface
+	virtual bool ResolveDeferredInitialization(UObject* ResolvingObject, UObject* ArchetypeInstance) override;
+};
+
+/**
+ * Access points for making FDeferredInitializationTracker calls. Takes care of 
+ * routing calls to the right tracker (CDOs vs sub-objects), and wraps the TThreadSingleton access for each.
+ */
+struct FDeferredObjInitializationHelper
+{
+	/** 
+	 * Determines if the specified initializer needs to be deferred (does it have a 
+	 * archetype dependency that needs to be serialized first?). If so, the FObjectInitializer will
+	 * be copied and stored with the appropriate tracker (CDO vs. sub-object, etc.).
+	 *
+	 * Designed to be called from the FObjectInitializer itself (before it runs initialization).
+	 *
+	 * @return A pointer to the initializer copy (if one was made), null if no deferral was needed.
+	 */
+	static FObjectInitializer* DeferObjectInitializerIfNeeded(const FObjectInitializer& DeferringInitializer);
+
+	/** 
+	 * Determines if the specified object should have its Preload() skipped. If so, this
+	 * should cache the sub-object so it can be loaded later, when its dependency is resolved.
+	 *
+	 * Designed to be called from the Preload() itself (before it runs serializes the object).
+	 *
+	 * More info: Because of delta serialization, we require that a parent's CDO be
+	 *            fully serialized before its children's CDOs are created. However,
+	 *            due to cyclic parent/child dependencies, we have some cases where
+	 *            the linker breaks that expected behavior. In those cases, we
+	 *            defer the child's initialization (i.e. defer copying of parent
+	 *            property values, etc.), and wait until we can guarantee that the
+	 *            parent CDO has been fully loaded.
+	 *
+	 *            In a normal scenario, the order of property initialization is:
+	 *            Creation (zeroed) -> Initialization (copied super's values) -> Serialization (overridden values loaded)
+	 *            When the initialization has been deferred we have to make sure to
+	 *            defer serialization here as well (don't worry, it will be invoked
+	 *            again from FinalizeBlueprint()->ResolveDeferredExports())
+	 *
+	 *            also, if this is an inherited sub-object on a CDO, and that CDO has had
+	 *            its initialization deferred (for reasons explained above), then
+	 *            we shouldn't serialize in data for this quite yet... not until
+	 *            its owner has had a chaTnce to initialize itself (because, as part
+	 *            of CDO initialization, inherited sub-objects get filled in with
+	 *            values inherited from the super)
+	 *
+	 * @return True if the object's Preload() should be skipped.
+	 */
+	static bool DeferObjectPreload(UObject* Object);
+
+	/** 
+	 * Loops through all object initializers and preloads that were skipped due 
+	 * to this archetype object not being ready yet.
+	 *
+	 * Should be called once the object has been fully serialized in.
+	 */
+	static void ResolveDeferredInitsFromArchetype(UObject* Archetype);
 };
 
 struct FBlueprintDependencyType
diff --git a/Engine/Source/Runtime/CoreUObject/Public/UObject/LinkerLoad.h b/Engine/Source/Runtime/CoreUObject/Public/UObject/LinkerLoad.h
--- a/Engine/Source/Runtime/CoreUObject/Public/UObject/LinkerLoad.h
+++ b/Engine/Source/Runtime/CoreUObject/Public/UObject/LinkerLoad.h
@@ -1098,7 +1098,7 @@ private:
 	 * to make sure this placeholder is completely resolved before continuing on 
 	 * to the next.
 	 */
-	class FLinkerPlaceholderBase* ResolvingDeferredPlaceholder;
+	TArray<class FLinkerPlaceholderBase*> ResolvingPlaceholderStack;
 
 	/** 
 	 * Internal list to track imports that were deferred, but don't belong to 
diff --git a/Engine/Source/Runtime/Engine/Classes/Components/SkeletalMeshComponent.h b/Engine/Source/Runtime/Engine/Classes/Components/SkeletalMeshComponent.h
--- a/Engine/Source/Runtime/Engine/Classes/Components/SkeletalMeshComponent.h
+++ b/Engine/Source/Runtime/Engine/Classes/Components/SkeletalMeshComponent.h
@@ -1470,7 +1470,8 @@ public:
 	int32 FindRootBodyIndex() const;
 
 
-	/** Terminate physics on all bodies below the named bone */
+	/** Terminate physics on all bodies below the named bone, effectively disabling collision forever. If you terminate, you won't be able to re-init later. */
+	UFUNCTION(BlueprintCallable, Category="Components|SkeletalMesh")
 	void TermBodiesBelow(FName ParentBoneName);
 
 	/** Find instance of the constraint that matches the name supplied. */
diff --git a/Engine/Source/Runtime/Engine/Classes/GameFramework/CharacterMovementComponent.h b/Engine/Source/Runtime/Engine/Classes/GameFramework/CharacterMovementComponent.h
--- a/Engine/Source/Runtime/Engine/Classes/GameFramework/CharacterMovementComponent.h
+++ b/Engine/Source/Runtime/Engine/Classes/GameFramework/CharacterMovementComponent.h
@@ -1810,7 +1810,7 @@ public:
 	* @param FloorResult			Result of the floor check
 	*/
 	UFUNCTION(BlueprintCallable, Category="Pawn|Components|CharacterMovement", meta=(DisplayName="FindFloor", ScriptName="FindFloor"))
-	void K2_FindFloor(FVector CapsuleLocation, FFindFloorResult& FloorResult) const;
+	virtual void K2_FindFloor(FVector CapsuleLocation, FFindFloorResult& FloorResult) const;
 
 	/**
 	 * Compute distance to the floor from bottom sphere of capsule and store the result in OutFloorResult.
@@ -1839,7 +1839,7 @@ public:
 	* @param FloorResult			Result of the floor check
 	*/
 	UFUNCTION(BlueprintCallable, Category="Pawn|Components|CharacterMovement", meta=(DisplayName="ComputeFloorDistance", ScriptName="ComputeFloorDistance"))
-	void K2_ComputeFloorDist(FVector CapsuleLocation, float LineDistance, float SweepDistance, float SweepRadius, FFindFloorResult& FloorResult) const;
+	virtual void K2_ComputeFloorDist(FVector CapsuleLocation, float LineDistance, float SweepDistance, float SweepRadius, FFindFloorResult& FloorResult) const;
 
 	/**
 	 * Sweep against the world and return the first blocking hit.
diff --git a/Engine/Source/Runtime/Engine/Classes/GameFramework/PlayerController.h b/Engine/Source/Runtime/Engine/Classes/GameFramework/PlayerController.h
--- a/Engine/Source/Runtime/Engine/Classes/GameFramework/PlayerController.h
+++ b/Engine/Source/Runtime/Engine/Classes/GameFramework/PlayerController.h
@@ -1059,6 +1059,14 @@ public:
 	UFUNCTION(BlueprintCallable, Category="Game|Feedback")
 	void SetHapticsByValue(const float Frequency, const float Amplitude, EControllerHand Hand);
 	
+	/**
+	 * Allows the player controller to disable all haptic requests from being fired, e.g. in the case of a level loading
+	 *
+	 * @param	bNewDisabled	If TRUE, the haptics will stop and prevented from being enabled again until set to FALSE
+	 */
+	UFUNCTION(BlueprintCallable, Category = "Game|Feedback")
+	virtual void SetDisableHaptics(bool bNewDisabled);
+
 	/**
 	 * Sets the light color of the player's controller
 	 * @param	Color					The color for the light to be
@@ -1824,4 +1832,9 @@ public:
 	 * Designate this player controller as local (public for GameModeBase to use, not expected to be called anywhere else)
 	 */
 	void SetAsLocalPlayerController() { bIsLocalPlayerController = true; }
+
+private:
+	/** If true, prevent any haptic effects from playing */
+	bool bDisableHaptics : 1;
+
 };
diff --git a/Engine/Source/Runtime/Engine/Classes/PhysicsEngine/BodyInstance.h b/Engine/Source/Runtime/Engine/Classes/PhysicsEngine/BodyInstance.h
--- a/Engine/Source/Runtime/Engine/Classes/PhysicsEngine/BodyInstance.h
+++ b/Engine/Source/Runtime/Engine/Classes/PhysicsEngine/BodyInstance.h
@@ -338,6 +338,18 @@ public:
 	*/
 	UPROPERTY(EditAnywhere, AdvancedDisplay, BlueprintReadWrite, Category = Physics)
 	uint8 bUseAsyncScene : 1;
+
+
+	/** 
+	 * @HACK:
+	 * These are ONLY used when the 'p.EnableDynamicPerBodyFilterHacks' CVar is set (disabled by default).
+	 * Some games need to dynamically modify collision per skeletal body. These provide game code a way to 
+	 * do that, until we're able to refactor how skeletal bodies work.
+	 */
+	uint32 bHACK_DisableCollisionResponse : 1;
+	/* By default, an owning skel mesh component will override the body's collision filter. This will disable that behavior. */
+	uint32 bHACK_DisableSkelComponentFilterOverriding : 1;
+
 protected:
 
 	/** Whether this body instance has its own custom MaxDepenetrationVelocity*/
diff --git a/Engine/Source/Runtime/Engine/Private/Components/CharacterMovementComponent.cpp b/Engine/Source/Runtime/Engine/Private/Components/CharacterMovementComponent.cpp
--- a/Engine/Source/Runtime/Engine/Private/Components/CharacterMovementComponent.cpp
+++ b/Engine/Source/Runtime/Engine/Private/Components/CharacterMovementComponent.cpp
@@ -5948,7 +5948,12 @@ void UCharacterMovementComponent::FindFloor(const FVector& CapsuleLocation, FFin
 
 void UCharacterMovementComponent::K2_FindFloor(FVector CapsuleLocation, FFindFloorResult& FloorResult) const
 {
+	const bool SavedForceNextFloorCheck(bForceNextFloorCheck);
 	FindFloor(CapsuleLocation, FloorResult, false);
+	
+	// FindFloor clears this, but this is only a test not done during normal movement.
+	UCharacterMovementComponent* MutableThis = const_cast<UCharacterMovementComponent*>(this);
+	MutableThis->bForceNextFloorCheck = SavedForceNextFloorCheck;
 }
 
 void UCharacterMovementComponent::K2_ComputeFloorDist(FVector CapsuleLocation, float LineDistance, float SweepDistance, float SweepRadius, FFindFloorResult& FloorResult) const
diff --git a/Engine/Source/Runtime/Engine/Private/PhysicsEngine/BodyInstance.cpp b/Engine/Source/Runtime/Engine/Private/PhysicsEngine/BodyInstance.cpp
--- a/Engine/Source/Runtime/Engine/Private/PhysicsEngine/BodyInstance.cpp
+++ b/Engine/Source/Runtime/Engine/Private/PhysicsEngine/BodyInstance.cpp
@@ -23,7 +23,7 @@
 #include "PhysicsEngine/SphylElem.h"
 #include "PhysicsEngine/BodySetup.h"
 #include "HAL/LowLevelMemTracker.h"
-
+#include "HAL/IConsoleManager.h"
 #include "Logging/TokenizedMessage.h"
 #include "Logging/MessageLog.h"
 #include "Misc/UObjectToken.h"
@@ -60,6 +60,14 @@ DECLARE_CYCLE_STAT(TEXT("CreatePhysicsShapesAndActors"), STAT_CreatePhysicsShape
 DECLARE_CYCLE_STAT(TEXT("BodyInstance SetCollisionProfileName"), STAT_BodyInst_SetCollisionProfileName, STATGROUP_Physics);
 DECLARE_CYCLE_STAT(TEXT("Phys SetBodyTransform"), STAT_SetBodyTransform, STATGROUP_Physics);
 
+// @HACK Guard to better encapsulate game related hacks introduced into UpdatePhysicsFilterData()
+TAutoConsoleVariable<int32> CVarEnableDynamicPerBodyFilterHacks(
+	TEXT("p.EnableDynamicPerBodyFilterHacks"), 
+	0, 
+	TEXT("Enables/Disables the use of a set of game focused hacks - allowing users to modify skel body collision dynamically (changes the behavior of per-body collision filtering)."),
+	ECVF_ReadOnly
+);
+
 #if WITH_PHYSX
 int32 FillInlinePxShapeArray_AssumesLocked(FInlinePxShapeArray& Array, const physx::PxRigidActor& RigidActor)
 {
@@ -1110,9 +1118,12 @@ void FBodyInstance::UpdatePhysicsFilterData()
 	bool bResponseOverride = false;
 	bool bNotifyOverride = false;
 
+	// @TODO: The skel mesh really shouldn't be the (pseudo-)authority here on the body's collision.
+	//        This block should ultimately be removed, and outside of this (in the skel component) we 
+	//        should configure the bodies to reflect this desired behavior.
 	if (USkeletalMeshComponent* SkelMeshComp = Cast<USkeletalMeshComponent>(OwnerComponentInst))
 	{
-		// In skeletal case, collision enable/disable/movement should be overriden by mesh component
+		// In skeletal case, collision enable/disable/movement should be overridden by mesh component
 		// being in the physics asset, and not having collision is a waste and it can cause a bug where disconnected bodies
 		if (Owner)
 		{
@@ -1125,24 +1136,44 @@ void FBodyInstance::UpdatePhysicsFilterData()
 				UseCollisionEnabled = ECollisionEnabled::NoCollision;
 			}
 		}
-		
 
 		ObjectType = SkelMeshComp->GetCollisionObjectType();
 		bUseCollisionEnabledOverride = true;
 
-		if (BodySetup->CollisionReponse == EBodyCollisionResponse::BodyCollision_Enabled)
+		// @HACK: 
+		// Some of our games have needed to dynamically disable & re-enable collision per body.
+		// We support that here as a hack, but as noted by the @TODO above, this should all be moved and 
+		// refactored (gameplay code should be able to modify the body instance's collision without fear 
+		// of the skel component overriding it).
+		if (CVarEnableDynamicPerBodyFilterHacks.GetValueOnGameThread() && bHACK_DisableCollisionResponse)
+		{
+			UseResponse.SetAllChannels(ECR_Ignore);
+			UseCollisionEnabled = ECollisionEnabled::PhysicsOnly;
+		}
+		else if (BodySetup->CollisionReponse == EBodyCollisionResponse::BodyCollision_Enabled)
 		{
 			UseResponse.SetAllChannels(ECR_Block);
 		}
 		else if (BodySetup->CollisionReponse == EBodyCollisionResponse::BodyCollision_Disabled)
 		{
 			UseResponse.SetAllChannels(ECR_Ignore);
+			UseCollisionEnabled = ECollisionEnabled::PhysicsOnly;		// this will prevent object traces hitting this as well
 		}
 
 		UseResponse = FCollisionResponseContainer::CreateMinContainer(UseResponse, SkelMeshComp->BodyInstance.CollisionResponses.GetResponseContainer());
 		bUseNotifyRBCollision = bUseNotifyRBCollision && SkelMeshComp->BodyInstance.bNotifyRigidBodyCollision;
 		bResponseOverride = true;
 		bNotifyOverride = true;
+
+		// @HACK
+		// Some of our games want to prevent the skeletal mesh component from overriding the 
+		// body's original setup (on a per body basis). This hack allows them to dynamically prevent that until we refactor this
+		// and remove this whole block of skel specialization (as detailed above).
+		if (CVarEnableDynamicPerBodyFilterHacks.GetValueOnGameThread() && bHACK_DisableSkelComponentFilterOverriding)
+		{
+			// if we have a custom response set up, trust that, don't override with the above.
+			bResponseOverride = false;
+		}
 	}
 
 #if WITH_EDITOR
@@ -1351,7 +1382,6 @@ struct FInitBodiesHelper
 			Instance->RigidActorSync->setName(Instance->CharDebugName.IsValid() ? Instance->CharDebugName->GetData() : nullptr);
 
 			check(FPhysxUserData::Get<FBodyInstance>(Instance->RigidActorSync->userData) == Instance);
-
 		}
 
 		if (Instance->RigidActorAsync)
@@ -2360,8 +2390,8 @@ bool FBodyInstance::UpdateBodyScale(const FVector& InScale3D, bool bForceUpdate)
 					float HalfHeight = HalfLength - Radius;
 					HalfHeight = FMath::Max(FCollisionShape::MinCapsuleAxisHalfHeight(), HalfHeight);
 
-					PCapsuleGeom.halfHeight = HalfHeight;
-					PCapsuleGeom.radius = Radius;
+					PCapsuleGeom.halfHeight = FMath::Max(HalfHeight, KINDA_SMALL_NUMBER);
+					PCapsuleGeom.radius = FMath::Max(Radius, KINDA_SMALL_NUMBER);
 
 					PLocalPose = PxTransform(U2PVector(RelativeTM.TransformPosition(SphylElem->Center)), U2PQuat(SphylElem->Rotation.Quaternion()) * U2PSphylBasis);
 					PLocalPose.p.x *= AdjustedScale3D.X;
diff --git a/Engine/Source/Runtime/Engine/Private/PhysicsEngine/PhysScene.cpp b/Engine/Source/Runtime/Engine/Private/PhysicsEngine/PhysScene.cpp
--- a/Engine/Source/Runtime/Engine/Private/PhysicsEngine/PhysScene.cpp
+++ b/Engine/Source/Runtime/Engine/Private/PhysicsEngine/PhysScene.cpp
@@ -2030,6 +2030,10 @@ void FPhysScene::InitPhysScene(uint32 SceneType)
 
 	// Create scene, and add to map
 	PxScene* PScene = GPhysXSDK->createScene(PSceneDesc);
+	if(PxPvdSceneClient* PVDClient = PScene->getScenePvdClient())
+	{
+		PVDClient->setScenePvdFlags(PxPvdSceneFlag::eTRANSMIT_CONSTRAINTS | PxPvdSceneFlag::eTRANSMIT_CONTACTS | PxPvdSceneFlag::eTRANSMIT_SCENEQUERIES);
+	}
 
 #if WITH_APEX
 	// Build the APEX scene descriptor for the PhysX scene
diff --git a/Engine/Source/Runtime/Engine/Private/PhysicsEngine/PhysUtils.cpp b/Engine/Source/Runtime/Engine/Private/PhysicsEngine/PhysUtils.cpp
--- a/Engine/Source/Runtime/Engine/Private/PhysicsEngine/PhysUtils.cpp
+++ b/Engine/Source/Runtime/Engine/Private/PhysicsEngine/PhysUtils.cpp
@@ -444,7 +444,7 @@ void PvdConnect(FString Host, bool bVisualization)
 	int32	Port = 5425;         // TCP port to connect to, where PVD is listening
 	uint32	Timeout = 100;          // timeout in milliseconds to wait for PVD to respond, consoles and remote PCs need a higher timeout.
 
-	PxPvdInstrumentationFlags ConnectionFlags = bVisualization ? PxPvdInstrumentationFlag::eALL : PxPvdInstrumentationFlag::ePROFILE;
+	PxPvdInstrumentationFlags ConnectionFlags = bVisualization  ? PxPvdInstrumentationFlag::eALL : (PxPvdInstrumentationFlag::ePROFILE | PxPvdInstrumentationFlag::eMEMORY);
 
 	PxPvdTransport* transport = PxDefaultPvdSocketTransportCreate(TCHAR_TO_ANSI(*Host), Port, Timeout);
 	GPhysXVisualDebugger->connect(*transport, ConnectionFlags);
diff --git a/Engine/Source/Runtime/Engine/Private/PhysicsEngine/PhysicalAnimationComponent.cpp b/Engine/Source/Runtime/Engine/Private/PhysicsEngine/PhysicalAnimationComponent.cpp
--- a/Engine/Source/Runtime/Engine/Private/PhysicsEngine/PhysicalAnimationComponent.cpp
+++ b/Engine/Source/Runtime/Engine/Private/PhysicsEngine/PhysicalAnimationComponent.cpp
@@ -6,6 +6,7 @@
 #include "PhysicsEngine/PhysicsAsset.h"
 #include "PhysXPublic.h"
 
+
 const FConstraintProfileProperties UPhysicalAnimationComponent::PhysicalAnimationProfile = []()
 {
 	//Setup the default constraint profile for all joints created by physical animation system
@@ -232,10 +233,17 @@ FTransform ComputeLocalSpaceTargetTM(const USkeletalMeshComponent& SkeletalMeshC
 
 		if (SkeletalMeshComponent.Bodies.IsValidIndex(BodyIndex))
 		{
-
-			FBodyInstance* ParentBody = SkeletalMeshComponent.Bodies[BodyIndex];
-			const FTransform NewWorldTM = AccumulatedDelta * ParentBody->GetUnrealWorldTransform_AssumesLocked();
-			return NewWorldTM;
+			if (BodyIndex < SkeletalMeshComponent.Bodies.Num())
+			{
+				FBodyInstance* ParentBody = SkeletalMeshComponent.Bodies[BodyIndex];
+				const FTransform NewWorldTM = AccumulatedDelta * ParentBody->GetUnrealWorldTransform_AssumesLocked();
+				return NewWorldTM;
+			}
+			else
+			{
+				// Bodies array has changed on us?
+				break;
+			}
 		}
 
 		AccumulatedDelta = AccumulatedDelta * SkeletalMeshComponent.BoneSpaceTransforms[CurBoneIdx];
diff --git a/Engine/Source/Runtime/Engine/Private/PlayerController.cpp b/Engine/Source/Runtime/Engine/Private/PlayerController.cpp
--- a/Engine/Source/Runtime/Engine/Private/PlayerController.cpp
+++ b/Engine/Source/Runtime/Engine/Private/PlayerController.cpp
@@ -117,6 +117,7 @@ APlayerController::APlayerController(const FObjectInitializer& ObjectInitializer
 
 	bIsPlayerController = true;
 	bIsLocalPlayerController = false;
+	bDisableHaptics = false;
 
 	ClickEventKeys.Add(EKeys::LeftMouseButton);
 
@@ -3848,12 +3849,24 @@ void APlayerController::StopHapticEffect(EControllerHand Hand)
 	SetHapticsByValue(0.f, 0.f, Hand);
 }
 
+void APlayerController::SetDisableHaptics(bool bNewDisabled)
+{
+	if (bNewDisabled)
+	{
+		StopHapticEffect(EControllerHand::Left);
+		StopHapticEffect(EControllerHand::Right);
+		StopHapticEffect(EControllerHand::Gun);
+	}
+
+	bDisableHaptics = bNewDisabled;
+}
+
 static TAutoConsoleVariable<int32> CVarDisableHaptics(TEXT("input.DisableHaptics"),0,TEXT("If greater than zero, no haptic feedback is processed."));
 
 void APlayerController::SetHapticsByValue(const float Frequency, const float Amplitude, EControllerHand Hand)
 {
-	bool bDisableHaptics = (CVarDisableHaptics.GetValueOnGameThread() > 0);
-	if (bDisableHaptics)
+	bool bAreHapticsDisabled = bDisableHaptics || (CVarDisableHaptics.GetValueOnGameThread() > 0);
+	if (bAreHapticsDisabled)
 	{
 		return;
 	}
@@ -4009,8 +4022,8 @@ void APlayerController::ProcessForceFeedbackAndHaptics(const float DeltaTime, co
 		{
 			InputInterface->SetForceFeedbackChannelValues(ControllerId, (bForceFeedbackEnabled ? ForceFeedbackValues : FForceFeedbackValues()));
 
-			const bool bDisableHaptics = (CVarDisableHaptics.GetValueOnGameThread() > 0);
-			if (!bDisableHaptics)
+			bool bAreHapticsDisabled = (CVarDisableHaptics.GetValueOnGameThread() > 0) || bDisableHaptics;
+			if (!bAreHapticsDisabled)
 			{
 				// Haptic Updates
 				if (bLeftHapticsNeedUpdate)
@@ -4021,6 +4034,7 @@ void APlayerController::ProcessForceFeedbackAndHaptics(const float DeltaTime, co
 				{
 					InputInterface->SetHapticFeedbackValues(ControllerId, (int32)EControllerHand::Right, RightHaptics);
 				}
+
 				if (bGunHapticsNeedUpdate)
 				{
 					InputInterface->SetHapticFeedbackValues(ControllerId, (int32)EControllerHand::Gun, GunHaptics);
diff --git a/Engine/Source/Runtime/Engine/Private/PrimitiveComponentPhysics.cpp b/Engine/Source/Runtime/Engine/Private/PrimitiveComponentPhysics.cpp
--- a/Engine/Source/Runtime/Engine/Private/PrimitiveComponentPhysics.cpp
+++ b/Engine/Source/Runtime/Engine/Private/PrimitiveComponentPhysics.cpp
@@ -724,7 +724,7 @@ void UPrimitiveComponent::UnWeldFromParent()
 
 	FBodyInstance* NewRootBI = GetBodyInstance(NAME_None, false);
 	UWorld* CurrentWorld = GetWorld();
-	if (NewRootBI == NULL || NewRootBI->WeldParent == nullptr || CurrentWorld == nullptr || CurrentWorld->GetPhysicsScene() == nullptr || IsPendingKill())
+	if (NewRootBI == NULL || NewRootBI->WeldParent == nullptr || CurrentWorld == nullptr || CurrentWorld->GetPhysicsScene() == nullptr || IsPendingKillOrUnreachable())
 	{
 		return;
 	}
diff --git a/Engine/Source/Runtime/Engine/Public/PhysicsPublic.h b/Engine/Source/Runtime/Engine/Public/PhysicsPublic.h
--- a/Engine/Source/Runtime/Engine/Public/PhysicsPublic.h
+++ b/Engine/Source/Runtime/Engine/Public/PhysicsPublic.h
@@ -676,7 +676,7 @@ private:
 **/
 FORCEINLINE bool PhysSingleThreadedMode()
 {
-	if (IsRunningDedicatedServer() || FPlatformMisc::NumberOfCores() < 3 || !FPlatformProcess::SupportsMultithreading())
+	if (IsRunningDedicatedServer() || FPlatformMisc::NumberOfCores() < 3 || !FPlatformProcess::SupportsMultithreading() || FParse::Param(FCommandLine::Get(), TEXT("SingleThreadedPhysics")))
 	{
 		return true;
 	}
diff --git a/Engine/Source/Runtime/LevelSequence/Private/LevelSequenceBindingReference.cpp b/Engine/Source/Runtime/LevelSequence/Private/LevelSequenceBindingReference.cpp
--- a/Engine/Source/Runtime/LevelSequence/Private/LevelSequenceBindingReference.cpp
+++ b/Engine/Source/Runtime/LevelSequence/Private/LevelSequenceBindingReference.cpp
@@ -90,6 +90,29 @@ UObject* ResolveByPath(UObject* InContext, const FString& InObjectPath)
 			return FoundObject;
 		}
 
+#if WITH_EDITOR
+		UWorld* WorldContext = InContext->GetWorld();
+		if (WorldContext && WorldContext->IsPlayInEditor())
+		{
+			FString PackageRoot, PackagePath, PackageName;
+			if (FPackageName::SplitLongPackageName(InObjectPath, PackageRoot, PackagePath, PackageName))
+			{
+				int32 ObjectDelimiterIdx = INDEX_NONE;
+				PackageName.FindChar('.', ObjectDelimiterIdx);
+				const FString SubLevelObjPath = PackageName.Mid(ObjectDelimiterIdx + 1);
+
+				for (ULevel* Level : WorldContext->GetLevels())
+				{
+					UPackage* Pkg = Level->GetOutermost();
+					if (UObject* FoundObject = FindObject<UObject>(Pkg, *SubLevelObjPath, false))
+ 					{
+ 						return FoundObject;
+					}
+				}
+			}
+		}
+#endif
+
 		if (UObject* FoundObject = FindObject<UObject>(ANY_PACKAGE, *InObjectPath, false))
 		{
 			return FoundObject;
@@ -114,7 +137,7 @@ UObject* FLevelSequenceLegacyObjectReference::Resolve(UObject* InContext) const
 				return FoundObject;
 			}
 
-			UE_LOG(LogMovieScene, Warning, TEXT("Attempted to resolve object with a PIE instance that has not been fixed up yet. This is probably due to a streamed level not being available yet."));
+			UE_LOG(LogMovieScene, Warning, TEXT("Attempted to resolve object (%s) with a PIE instance that has not been fixed up yet. This is probably due to a streamed level not being available yet."), *ObjectPath);
 			return nullptr;
 		}
 		FLazyObjectPtr LazyPtr;
diff --git a/Engine/Source/Runtime/SlateCore/Public/Fonts/SlateFontInfo.h b/Engine/Source/Runtime/SlateCore/Public/Fonts/SlateFontInfo.h
--- a/Engine/Source/Runtime/SlateCore/Public/Fonts/SlateFontInfo.h
+++ b/Engine/Source/Runtime/SlateCore/Public/Fonts/SlateFontInfo.h
@@ -71,8 +71,6 @@ struct SLATECORE_API FFontOutlineSettings
 	bool operator==(const FFontOutlineSettings& Other) const
 	{
 		return OutlineSize == Other.OutlineSize
-			&& OutlineMaterial == Other.OutlineMaterial
-			&& OutlineColor == Other.OutlineColor
 			&& bSeparateFillAlpha == Other.bSeparateFillAlpha;
 	}
 
@@ -80,8 +78,6 @@ struct SLATECORE_API FFontOutlineSettings
 	{
 		uint32 Hash = 0;
 		Hash = HashCombine(Hash, GetTypeHash(OutlineSettings.OutlineSize));
-		Hash = HashCombine(Hash, GetTypeHash(OutlineSettings.OutlineMaterial));
-		Hash = HashCombine(Hash, GetTypeHash(OutlineSettings.OutlineColor));
 		Hash = HashCombine(Hash, GetTypeHash(OutlineSettings.bSeparateFillAlpha));
 		return Hash;
 	}
diff --git a/Engine/Source/Runtime/UMG/Private/Components/WidgetComponent.cpp b/Engine/Source/Runtime/UMG/Private/Components/WidgetComponent.cpp
--- a/Engine/Source/Runtime/UMG/Private/Components/WidgetComponent.cpp
+++ b/Engine/Source/Runtime/UMG/Private/Components/WidgetComponent.cpp
@@ -502,6 +502,9 @@ public:
 
 		Result.bDrawRelevance = IsShown(View) && bVisible && View->Family->EngineShowFlags.WidgetComponents;
 		Result.bDynamicRelevance = true;
+		Result.bRenderCustomDepth = ShouldRenderCustomDepth();
+		Result.bRenderInMainPass = ShouldRenderInMainPass();
+		Result.bUsesLightingChannels = GetLightingChannelMask() != GetDefaultLightingChannelMask();
 		Result.bShadowRelevance = IsShadowCast(View);
 		Result.bEditorPrimitiveRelevance = false;
 
@@ -557,6 +560,7 @@ UWidgetComponent::UWidgetComponent( const FObjectInitializer& PCIP )
 	, LastWidgetRenderTime(0)
 	, bReceiveHardwareInput(false)
 	, bWindowFocusable(true)
+	, bApplyGammaCorrection(false)
 	, BackgroundColor( FLinearColor::Transparent )
 	, TintColorAndOpacity( FLinearColor::White )
 	, OpacityFromTexture( 1.0f )
@@ -599,7 +603,7 @@ UWidgetComponent::UWidgetComponent( const FObjectInitializer& PCIP )
 
 	Space = EWidgetSpace::World;
 	TimingPolicy = EWidgetTimingPolicy::RealTime;
-	Pivot = FVector2D(0.5, 0.5);
+	Pivot = FVector2D(0.5f, 0.5f);
 
 	bAddedToScreen = false;
 }
@@ -772,7 +776,7 @@ void UWidgetComponent::OnRegister()
 
 			if ( !WidgetRenderer.IsValid() && !GUsingNullRHI )
 			{
-				WidgetRenderer = MakeShareable(new FWidgetRenderer());
+				WidgetRenderer = MakeShareable(new FWidgetRenderer(bApplyGammaCorrection));
 			}
 		}
 
@@ -920,7 +924,7 @@ void UWidgetComponent::TickComponent(float DeltaTime, enum ELevelTick TickType,
 				// the world is paused, this also takes care of the case where the widget component is rendering at
 				// a different rate than the rest of the world.
 				const float DeltaTimeFromLastDraw = LastWidgetRenderTime == 0 ? 0 : (GetCurrentTime() - LastWidgetRenderTime);
-			    DrawWidgetToRenderTarget(DeltaTimeFromLastDraw);
+				DrawWidgetToRenderTarget(DeltaTimeFromLastDraw);
 		    }
 	    }
 	    else
@@ -1295,6 +1299,11 @@ void UWidgetComponent::SetOwnerPlayer(ULocalPlayer* LocalPlayer)
 	}
 }
 
+void UWidgetComponent::SetManuallyRedraw(bool bUseManualRedraw)
+{
+	bManuallyRedraw = bUseManualRedraw;
+}
+
 ULocalPlayer* UWidgetComponent::GetOwnerPlayer() const
 {
 	if (OwnerPlayer)
diff --git a/Engine/Source/Runtime/UMG/Public/Components/WidgetComponent.h b/Engine/Source/Runtime/UMG/Public/Components/WidgetComponent.h
--- a/Engine/Source/Runtime/UMG/Public/Components/WidgetComponent.h
+++ b/Engine/Source/Runtime/UMG/Public/Components/WidgetComponent.h
@@ -184,6 +184,10 @@ public:
 	UFUNCTION(BlueprintCallable, Category=UserInterface)
 	void SetOwnerPlayer(ULocalPlayer* LocalPlayer);
 
+	/** @see bManuallyRedraw */
+	UFUNCTION(BlueprintCallable, Category=UserInterface)
+	void SetManuallyRedraw(bool bUseManualRedraw);
+
 	/** Gets the local player that owns this widget component. */
 	UFUNCTION(BlueprintCallable, Category=UserInterface)
 	ULocalPlayer* GetOwnerPlayer() const;
@@ -360,6 +364,13 @@ protected:
 	UPROPERTY(EditAnywhere, Category=Interaction)
 	bool bWindowFocusable;
 
+	/**
+	 * Widget components that appear in the world will be gamma corrected by the 3D renderer.
+	 * In some cases, widget components are blitted directly into the backbuffer, in which case gamma correction should be enabled.
+	 */
+	UPROPERTY(EditAnywhere, Category = UserInterface, AdvancedDisplay)
+	bool bApplyGammaCorrection;
+
 	/**
 	 * The owner player for a widget component, if this widget is drawn on the screen, this controls
 	 * what player's screen it appears on for split screen, if not set, users player 0.
