Fix package linker crash on load for Blueprint-nativized packaged builds.

diff --git a/Engine/Source/Runtime/CoreUObject/Private/UObject/LinkerLoad.cpp b/Engine/Source/Runtime/CoreUObject/Private/UObject/LinkerLoad.cpp
--- a/Engine/Source/Runtime/CoreUObject/Private/UObject/LinkerLoad.cpp
+++ b/Engine/Source/Runtime/CoreUObject/Private/UObject/LinkerLoad.cpp
@@ -915,7 +915,6 @@ FLinkerLoad::ELinkerStatus FLinkerLoad::CreateLoader(
 		}
 		else
 		{
-			StructuredArchiveFormatter = nullptr;
 #if WITH_TEXT_ARCHIVE_SUPPORT
 			if (Filename.EndsWith(FPackageName::GetTextAssetPackageExtension()) || Filename.EndsWith(FPackageName::GetTextMapPackageExtension()))
 			{
@@ -993,14 +992,18 @@ FLinkerLoad::ELinkerStatus FLinkerLoad::CreateLoader(
 				{
 					bLoaderIsFArchiveAsync2 = true;
 				}
-
-				StructuredArchiveFormatter = new FBinaryArchiveFormatter(*this);
 			}
-
-			StructuredArchive = new FStructuredArchive(*StructuredArchiveFormatter);
-			StructuredArchiveRootRecord.Emplace(StructuredArchive->Open().EnterRecord());
 		} 
 
+		// Create structured archive wrapper
+		if (StructuredArchiveFormatter == nullptr)
+		{
+			StructuredArchiveFormatter = new FBinaryArchiveFormatter(*this);
+		}
+
+		StructuredArchive = new FStructuredArchive(*StructuredArchiveFormatter);
+		StructuredArchiveRootRecord.Emplace(StructuredArchive->Open().EnterRecord());
+
 		check(bDynamicClassLinker || Loader);
 		check(bDynamicClassLinker || !Loader->IsError());
 
