Fixed placeholder string table entries being considered identical as their entry may appear in the future

diff --git a/Engine/Source/Runtime/Core/Private/Internationalization/StringTableCore.cpp b/Engine/Source/Runtime/Core/Private/Internationalization/StringTableCore.cpp
--- a/Engine/Source/Runtime/Core/Private/Internationalization/StringTableCore.cpp
+++ b/Engine/Source/Runtime/Core/Private/Internationalization/StringTableCore.cpp
@@ -50,6 +50,18 @@ FTextDisplayStringPtr FStringTableEntry::GetDisplayString() const
 	return DisplayString;
 }
 
+const FString& FStringTableEntry::GetPlaceholderSourceString()
+{
+	static const FString MissingSourceString = TEXT("<MISSING STRING TABLE ENTRY>");
+	return MissingSourceString;
+}
+
+FTextDisplayStringRef FStringTableEntry::GetPlaceholderDisplayString()
+{
+	static const FTextDisplayStringRef MissingDisplayString = MakeShared<FString, ESPMode::ThreadSafe>(TEXT("<MISSING STRING TABLE ENTRY>"));
+	return MissingDisplayString;
+}
+
 
 FStringTable::FStringTable()
 	: OwnerAsset(nullptr)
diff --git a/Engine/Source/Runtime/Core/Private/Internationalization/TextHistory.cpp b/Engine/Source/Runtime/Core/Private/Internationalization/TextHistory.cpp
--- a/Engine/Source/Runtime/Core/Private/Internationalization/TextHistory.cpp
+++ b/Engine/Source/Runtime/Core/Private/Internationalization/TextHistory.cpp
@@ -1177,9 +1177,7 @@ const FString* FTextHistory_StringTableEntry::GetSourceString() const
 	{
 		return &StringTableEntryPin->GetSourceString();
 	}
-
-	static const FString MissingSourceString = TEXT("<MISSING STRING TABLE ENTRY>");
-	return &MissingSourceString;
+	return &FStringTableEntry::GetPlaceholderSourceString();
 }
 
 FTextDisplayStringRef FTextHistory_StringTableEntry::GetDisplayString() const
@@ -1193,9 +1191,7 @@ FTextDisplayStringRef FTextHistory_StringTableEntry::GetDisplayString() const
 			return DisplayString.ToSharedRef();
 		}
 	}
-
-	static const FTextDisplayStringRef MissingDisplayString = MakeShared<FString, ESPMode::ThreadSafe>(TEXT("<MISSING STRING TABLE ENTRY>"));
-	return MissingDisplayString;
+	return FStringTableEntry::GetPlaceholderDisplayString();
 }
 
 void FTextHistory_StringTableEntry::GetTableIdAndKey(FName& OutTableId, FString& OutKey) const
diff --git a/Engine/Source/Runtime/Core/Public/Internationalization/StringTableCore.h b/Engine/Source/Runtime/Core/Public/Internationalization/StringTableCore.h
--- a/Engine/Source/Runtime/Core/Public/Internationalization/StringTableCore.h
+++ b/Engine/Source/Runtime/Core/Public/Internationalization/StringTableCore.h
@@ -38,6 +38,12 @@ public:
 	/** Get the display string of this string table entry */
 	FTextDisplayStringPtr GetDisplayString() const;
 
+	/** Get the placeholder source string to use for string table entries that are missing */
+	static const FString& GetPlaceholderSourceString();
+
+	/** Get the placeholder display string to use for string table entries that are missing */
+	static FTextDisplayStringRef GetPlaceholderDisplayString();
+
 private:
 	/** The string table that owns us (if any) */
 	FStringTableConstWeakPtr OwnerTable;
diff --git a/Engine/Source/Runtime/CoreUObject/Private/UObject/TextProperty.cpp b/Engine/Source/Runtime/CoreUObject/Private/UObject/TextProperty.cpp
--- a/Engine/Source/Runtime/CoreUObject/Private/UObject/TextProperty.cpp
+++ b/Engine/Source/Runtime/CoreUObject/Private/UObject/TextProperty.cpp
@@ -7,6 +7,7 @@
 #include "Internationalization/TextNamespaceUtil.h"
 #include "Internationalization/TextPackageNamespaceUtil.h"
 #include "Internationalization/StringTableRegistry.h"
+#include "Internationalization/StringTableCore.h"
 
 EConvertFromTypeResult UTextProperty::ConvertFromType(const FPropertyTag& Tag, FArchive& Ar, uint8* Data, UStruct* DefaultsStruct)
 {
@@ -55,6 +56,21 @@ bool UTextProperty::Identical_Implementation(const FText& ValueA, const FText& V
 	// If both texts share the same pointer, then they must be equal
 	if (ValueA.IdenticalTo(ValueB))
 	{
+		// Placeholder string table entries will have the same pointer, but should only be considered equal if they're using the same string table and key
+		if (ValueA.IsFromStringTable() && FTextInspector::GetSourceString(ValueA) == &FStringTableEntry::GetPlaceholderSourceString())
+		{
+			FName ValueAStringTableId;
+			FString ValueAStringTableEntryKey;
+			FTextInspector::GetTableIdAndKey(ValueA, ValueAStringTableId, ValueAStringTableEntryKey);
+
+			FName ValueBStringTableId;
+			FString ValueBStringTableEntryKey;
+			FTextInspector::GetTableIdAndKey(ValueB, ValueBStringTableId, ValueBStringTableEntryKey);
+
+			return ValueAStringTableId == ValueBStringTableId && ValueAStringTableEntryKey.Equals(ValueBStringTableEntryKey, ESearchCase::CaseSensitive);
+		}
+
+		// Otherwise they're equal
 		return true;
 	}
 
