Fixed crash after attempting to load a package that is too new

diff --git a/Engine/Source/Runtime/CoreUObject/Private/UObject/PackageReload.cpp b/Engine/Source/Runtime/CoreUObject/Private/UObject/PackageReload.cpp
--- a/Engine/Source/Runtime/CoreUObject/Private/UObject/PackageReload.cpp
+++ b/Engine/Source/Runtime/CoreUObject/Private/UObject/PackageReload.cpp
@@ -187,6 +187,18 @@ void MakeObjectPurgeable(UObject* InObject)
 	InObject->ClearFlags(RF_Public | RF_Standalone);
 }
 
+/**
+ * Given a package, put it into a state where a GC may purge it (assuming there are no external references).
+ */
+void MakePackagePurgeable(UPackage* InPackage)
+{
+	MakeObjectPurgeable(InPackage);
+	ForEachObjectWithOuter(InPackage, [](UObject* InObject)
+	{
+		MakeObjectPurgeable(InObject);
+	});
+}
+
 /**
  * Given an object, dump anything that is still externally referencing it to the log.
  */
@@ -266,6 +278,13 @@ UPackage* LoadReplacementPackage(UPackage* InExistingPackage, const uint32 InLoa
 	{
 		UE_LOG(LogUObjectGlobals, Warning, TEXT("ReloadPackage cannot reload '%s' as the new package failed to load. The old package will be restored."), *ExistingPackageName);
 
+		// Make sure that the failed load attempt hasn't left any objects behind that would prevent the rename
+		if (UPackage* FailedPackage = FindObject<UPackage>(Cast<UPackage>(InExistingPackage->GetOuter()), *ExistingPackageName))
+		{
+			FailedPackage->Rename(*MakeUniqueObjectName(Cast<UPackage>(FailedPackage->GetOuter()), UPackage::StaticClass(), *FString::Printf(TEXT("%s_DEADPACKAGE"), *FailedPackage->GetName())).ToString(), nullptr, PkgRenameFlags);
+			MakePackagePurgeable(FailedPackage);
+		}
+
 		// Failed to load the new package, give the old package its original name and bail!
 		InExistingPackage->Rename(*ExistingPackageName, nullptr, PkgRenameFlags);
 		ClearPackageReplaced(InExistingPackage);
@@ -406,7 +425,7 @@ void ReloadPackages(const TArrayView<FReloadPackageData>& InPackagesToReload, TA
 
 	// Cache the current dirty state of all packages so we can restore it after the reload
 	TSet<FName> DirtyPackages;
-	ForEachObjectOfClass(UPackage::StaticClass(), [&](UObject* InPackageObj)
+	ForEachObjectOfClass(UPackage::StaticClass(), [&DirtyPackages](UObject* InPackageObj)
 	{
 		UPackage* Package = CastChecked<UPackage>(InPackageObj);
 		if (Package->IsDirty())
@@ -523,20 +542,18 @@ void ReloadPackages(const TArrayView<FReloadPackageData>& InPackagesToReload, TA
 
 			FCoreUObjectDelegates::OnPackageReloaded.Broadcast(EPackageReloadPhase::PostBatchPreGC, nullptr);
 
-			// Purge old packages
+			// Purge old packages that have had a replacement package loaded
 			for (int32 BatchPackageIndex = BatchStartIndex; BatchPackageIndex < PackageIndex; ++BatchPackageIndex)
 			{
 				FixingUpReferencesSlowTask.EnterProgressFrame(1.0f);
 
 				UPackage* ExistingPackage = ExistingPackages[BatchPackageIndex].RawRef;
-				if (ExistingPackage)
+				UPackage* NewPackage = NewPackages[BatchPackageIndex].Package.Get();
+
+				if (ExistingPackage && NewPackage)
 				{
 					// Allow the old package to be GC'd
-					PackageReloadInternal::MakeObjectPurgeable(ExistingPackage);
-					ForEachObjectWithOuter(ExistingPackage, [&](UObject* InExistingObject)
-					{
-						PackageReloadInternal::MakeObjectPurgeable(InExistingObject);
-					});
+					PackageReloadInternal::MakePackagePurgeable(ExistingPackage);
 					ExistingPackages[BatchPackageIndex].StrongRef.Reset();
 				}
 			}
@@ -547,7 +564,7 @@ void ReloadPackages(const TArrayView<FReloadPackageData>& InPackagesToReload, TA
 	}
 
 	// Clean any packages that we dirtied as part of the replacement process
-	ForEachObjectOfClass(UPackage::StaticClass(), [&](UObject* InPackageObj)
+	ForEachObjectOfClass(UPackage::StaticClass(), [&DirtyPackages](UObject* InPackageObj)
 	{
 		UPackage* Package = CastChecked<UPackage>(InPackageObj);
 		if (Package->IsDirty() && !DirtyPackages.Contains(Package->GetFName()))
@@ -575,7 +592,7 @@ void ReloadPackages(const TArrayView<FReloadPackageData>& InPackagesToReload, TA
 			if (bDumpExternalReferences)
 			{
 				PackageReloadInternal::DumpExternalReferences(ExistingPackage, ExistingPackage);
-				//ForEachObjectWithOuter(ExistingPackage, [&](UObject* InExistingObject)
+				//ForEachObjectWithOuter(ExistingPackage, [](UObject* InExistingObject)
 				//{
 				//	PackageReloadInternal::DumpExternalReferences(InExistingObject, ExistingPackage);
 				//}, true, RF_NoFlags, EInternalObjectFlags::PendingKill);
