Mirror 4078822 from Dev-Framework
Continue trying to finish reinstancing when instances of the old class are being async loaded

diff --git a/Engine/Source/Editor/Kismet/Private/BlueprintCompilationManager.cpp b/Engine/Source/Editor/Kismet/Private/BlueprintCompilationManager.cpp
--- a/Engine/Source/Editor/Kismet/Private/BlueprintCompilationManager.cpp
+++ b/Engine/Source/Editor/Kismet/Private/BlueprintCompilationManager.cpp
@@ -1063,7 +1063,24 @@ void FBlueprintCompilationManagerImpl::FlushReinstancingQueueImpl()
 		TGuardValue<bool> ReinstancingGuard(GIsReinstancing, true);
 		FBlueprintCompileReinstancer::BatchReplaceInstancesOfClass(ClassesToReinstance, true);
 
-		ClassesToReinstance.Empty();
+		if (IsAsyncLoading())
+		{
+			// While async loading we only remove classes that have no instances being
+			// async loaded. Those instances will need to be reinstanced once they finish
+			// loading, there's no race here because if any instances are created after
+			// we check ClassHasInstancesAsyncLoading they will be created with the new class:
+			for( TMap<UClass*, UClass*>::TIterator It(ClassesToReinstance); It; ++It )
+			{
+				if (!ClassHasInstancesAsyncLoading(It->Key))
+				{
+					It.RemoveCurrent();
+				}
+			}
+		}
+		else
+		{
+			ClassesToReinstance.Empty();
+		}
 	}
 	
 #if VERIFY_NO_STALE_CLASS_REFERENCES
diff --git a/Engine/Source/Runtime/CoreUObject/Private/UObject/UObjectHash.cpp b/Engine/Source/Runtime/CoreUObject/Private/UObject/UObjectHash.cpp
--- a/Engine/Source/Runtime/CoreUObject/Private/UObject/UObjectHash.cpp
+++ b/Engine/Source/Runtime/CoreUObject/Private/UObject/UObjectHash.cpp
@@ -884,6 +884,35 @@ void GetDerivedClasses(UClass* ClassToLookFor, TArray<UClass *>& Results, bool b
 	}
 }
 
+bool ClassHasInstancesAsyncLoading(UClass* ClassToLookFor)
+{
+	auto& ThreadHash = FUObjectHashTables::Get();
+	FHashTableLock HashLock(ThreadHash);
+	TSet<UClass*> ClassesToSearch;
+	ClassesToSearch.Add(ClassToLookFor);
+	{
+		RecursivelyPopulateDerivedClasses(ThreadHash, ClassToLookFor, ClassesToSearch);
+	}
+
+	for (auto ClassIt = ClassesToSearch.CreateConstIterator(); ClassIt; ++ClassIt)
+	{
+		TSet<UObjectBase*> const* List = ThreadHash.ClassToObjectListMap.Find(*ClassIt);
+		if (List)
+		{
+			for (auto ObjectIt = List->CreateConstIterator(); ObjectIt; ++ObjectIt)
+			{
+				UObject *Object = static_cast<UObject*>(*ObjectIt);
+				if (Object->HasAnyInternalFlags(EInternalObjectFlags::AsyncLoading))
+				{
+					return true;
+				}
+			}
+		}
+	}
+
+	return false;
+}
+
 void AllocateUObjectIndexForCurrentThread(UObjectBase* Object)
 {
 	GUObjectArray.AllocateUObjectIndex(Object);
diff --git a/Engine/Source/Runtime/CoreUObject/Public/UObject/UObjectHash.h b/Engine/Source/Runtime/CoreUObject/Public/UObject/UObjectHash.h
--- a/Engine/Source/Runtime/CoreUObject/Public/UObject/UObjectHash.h
+++ b/Engine/Source/Runtime/CoreUObject/Public/UObject/UObjectHash.h
@@ -104,6 +104,14 @@ COREUOBJECT_API void ForEachObjectOfClass(UClass* ClassToLookFor, TFunctionRef<v
  */
 COREUOBJECT_API void GetDerivedClasses(UClass* ClassToLookFor, TArray<UClass *>& Results, bool bRecursive = true);
 
+/**
+ * Returns true if any instances of the class in question are currently being async loaded.
+ *
+ * @param	ClassToLookFor				The class in question
+ * @return	True if there are any instances of the class being async loaded - includes instances based on derived classes. Otherwise, false
+ */
+COREUOBJECT_API bool ClassHasInstancesAsyncLoading(UClass* ClassToLookFor);
+
 /**
  * Add an object to the name hash tables
  *
