Update assert to handle subobjects nested in components

diff --git a/Engine/Source/Runtime/CoreUObject/Private/UObject/UObjectGlobals.cpp b/Engine/Source/Runtime/CoreUObject/Private/UObject/UObjectGlobals.cpp
--- a/Engine/Source/Runtime/CoreUObject/Private/UObject/UObjectGlobals.cpp
+++ b/Engine/Source/Runtime/CoreUObject/Private/UObject/UObjectGlobals.cpp
@@ -2422,20 +2422,21 @@ UObject* StaticAllocateObject
 		{
 #if !UE_BUILD_SHIPPING
 			// Handle nested DSOs
-			bool bIsOwnedByCDO = false;
+			bool bIsOwnedByCDOOrArchetype = false;
 			UObject* Iter = InOuter;
 			while (Iter)
 			{
-				if (Iter->HasAnyFlags(RF_ClassDefaultObject))
+				if (Iter->HasAnyFlags(RF_ClassDefaultObject|RF_ArchetypeObject))
 				{
-					bIsOwnedByCDO = true;
+					bIsOwnedByCDOOrArchetype = true;
 					break;
 				}
 				Iter = Iter->GetOuter();
 			}
+
 			// Should only get in here if we're NOT creating a subobject of a CDO.  CDO subobjects may still need to be serialized off of disk after being created by the constructor
 			// if really necessary there was code to allow replacement of object just needing postload, but lets not go there unless we have to
-			checkf(!Obj->HasAnyFlags(RF_NeedLoad|RF_NeedPostLoad|RF_ClassDefaultObject) || bIsOwnedByCDO,
+			checkf(!Obj->HasAnyFlags(RF_NeedLoad|RF_NeedPostLoad|RF_ClassDefaultObject) || bIsOwnedByCDOOrArchetype,
 				TEXT("Attempting to replace an object that hasn't been fully loaded: %s (Outer=%s, Flags=0x%08x)"),
 				*Obj->GetFullName(),
 				InOuter ? *InOuter->GetFullName() : TEXT("NULL"),
