[NO NON-PII COMMIT MESSAGE]

diff --git a/Engine/Source/Runtime/CoreUObject/Private/UObject/PropertyObject.cpp b/Engine/Source/Runtime/CoreUObject/Private/UObject/PropertyObject.cpp
--- a/Engine/Source/Runtime/CoreUObject/Private/UObject/PropertyObject.cpp
+++ b/Engine/Source/Runtime/CoreUObject/Private/UObject/PropertyObject.cpp
@@ -40,8 +40,24 @@ EConvertFromTypeResult UObjectProperty::ConvertFromType(const FPropertyTag& Tag,
 		FSoftObjectPtr PreviousValue;
 		Ar << PreviousValue;
 
-		// now copy the value into the object's address space
-		UObject* PreviousValueObj = PreviousValue.LoadSynchronous();
+		UObject* PreviousValueObj = nullptr;
+
+		// If we're async loading it's not safe to do a sync load because it may crash or fail to set the variable, so throw an error if it's not already in memory
+		if (IsInAsyncLoadingThread())
+		{
+			PreviousValueObj = PreviousValue.Get();
+
+			if (!PreviousValueObj && !PreviousValue.IsNull())
+			{
+				UE_LOG(LogClass, Error, TEXT("Failed to convert soft path %s to unloaded object as this is not safe during async loading. Load and resave %s in the editor to fix!"), *PreviousValue.ToString(), *Ar.GetArchiveName());
+			}
+		}
+		else
+		{
+			PreviousValueObj = PreviousValue.LoadSynchronous();
+		}
+
+		// Now copy the value into the object's address space
 		SetPropertyValue_InContainer(Data, PreviousValueObj, Tag.ArrayIndex);
 
 		// Validate the type is proper
