fixed a crash when GC triggered while Sequencer FBX import window is opened (i.e. Autosave)
added missing static asserts and fixed resulting compilation error

diff --git a/Engine/Source/Editor/MovieSceneTools/Private/MovieSceneToolHelpers.cpp b/Engine/Source/Editor/MovieSceneTools/Private/MovieSceneToolHelpers.cpp
--- a/Engine/Source/Editor/MovieSceneTools/Private/MovieSceneToolHelpers.cpp
+++ b/Engine/Source/Editor/MovieSceneTools/Private/MovieSceneToolHelpers.cpp
@@ -1475,7 +1475,6 @@ class SMovieSceneImportFBXSettings : public SCompoundWidget, public FGCObject
 	virtual void AddReferencedObjects( FReferenceCollector& Collector ) override
 	{
 		Collector.AddReferencedObject(MovieScene);
-		Collector.AddReferencedObject(Sequencer);
 	}
 
 	void SetObjectBindingMap(const TMap<FGuid, FString>& InObjectBindingMap)
diff --git a/Engine/Source/Runtime/CoreUObject/Public/UObject/UObjectGlobals.h b/Engine/Source/Runtime/CoreUObject/Public/UObject/UObjectGlobals.h
--- a/Engine/Source/Runtime/CoreUObject/Public/UObject/UObjectGlobals.h
+++ b/Engine/Source/Runtime/CoreUObject/Public/UObject/UObjectGlobals.h
@@ -1620,6 +1620,8 @@ public:
 	template<class UObjectType>
 	void AddReferencedObject(UObjectType*& Object, const UObject* ReferencingObject = nullptr, const UProperty* ReferencingProperty = nullptr)
 	{
+		static_assert(sizeof(UObjectType) > 0, "AddReferencedObject: Element must be a pointer to a fully-defined type");
+		static_assert(TPointerIsConvertibleFromTo<UObjectType, const UObjectBase>::Value, "AddReferencedObject: Element must be a pointer to a type derived from UObject");
 		HandleObjectReference(*(UObject**)&Object, ReferencingObject, ReferencingProperty);
 	}
 
diff --git a/Engine/Source/Runtime/Engine/Private/HModel.cpp b/Engine/Source/Runtime/Engine/Private/HModel.cpp
--- a/Engine/Source/Runtime/Engine/Private/HModel.cpp
+++ b/Engine/Source/Runtime/Engine/Private/HModel.cpp
@@ -7,6 +7,7 @@
 #include "HModel.h"
 #include "SceneView.h"
 #include "Model.h"
+#include "Components/ModelComponent.h"
 #include "Raster.h"
 #include "Engine/Polys.h"
 #include "Engine/Brush.h"
@@ -137,3 +138,9 @@ bool HModel::ResolveSurface(const FSceneView* View,int32 X,int32 Y,uint32& OutSu
 	OutSurfaceIndex = HitState.SurfaceIndex;
 	return HitState.bHitSurface;
 }
+
+void HModel::AddReferencedObjects(FReferenceCollector& Collector)
+{
+	Collector.AddReferencedObject(Component);
+	Collector.AddReferencedObject(Model);
+}
diff --git a/Engine/Source/Runtime/Engine/Public/HModel.h b/Engine/Source/Runtime/Engine/Public/HModel.h
--- a/Engine/Source/Runtime/Engine/Public/HModel.h
+++ b/Engine/Source/Runtime/Engine/Public/HModel.h
@@ -35,11 +35,7 @@ public:
 	ENGINE_API bool ResolveSurface(const FSceneView* View,int32 X,int32 Y,uint32& OutSurfaceIndex) const;
 
 	// HHitProxy interface.
-	virtual void AddReferencedObjects( FReferenceCollector& Collector ) override
-	{
-		Collector.AddReferencedObject( Component );
-		Collector.AddReferencedObject( Model );
-	}
+	virtual void AddReferencedObjects(FReferenceCollector& Collector) override;
 	virtual EMouseCursor::Type GetMouseCursor() override
 	{
 		return EMouseCursor::Crosshairs;
