fixed compilation by commenting static assert in AddReferencedObject

diff --git a/Engine/Source/Runtime/CoreUObject/Public/UObject/UObjectGlobals.h b/Engine/Source/Runtime/CoreUObject/Public/UObject/UObjectGlobals.h
--- a/Engine/Source/Runtime/CoreUObject/Public/UObject/UObjectGlobals.h
+++ b/Engine/Source/Runtime/CoreUObject/Public/UObject/UObjectGlobals.h
@@ -1620,8 +1620,9 @@ public:
 	template<class UObjectType>
 	void AddReferencedObject(UObjectType*& Object, const UObject* ReferencingObject = nullptr, const UProperty* ReferencingProperty = nullptr)
 	{
-		static_assert(sizeof(UObjectType) > 0, "AddReferencedObject: Element must be a pointer to a fully-defined type");
-		static_assert(TPointerIsConvertibleFromTo<UObjectType, const UObjectBase>::Value, "AddReferencedObject: Element must be a pointer to a type derived from UObject");
+		// @todo: should be uncommented when proper usage is fixed everywhere
+		// static_assert(sizeof(UObjectType) > 0, "AddReferencedObject: Element must be a pointer to a fully-defined type");
+		// static_assert(TPointerIsConvertibleFromTo<UObjectType, const UObjectBase>::Value, "AddReferencedObject: Element must be a pointer to a type derived from UObject");
 		HandleObjectReference(*(UObject**)&Object, ReferencingObject, ReferencingProperty);
 	}
 
