Copy of fix from CL# 4158650.

diff --git a/Engine/Source/Runtime/CoreUObject/Private/Serialization/AsyncLoading.cpp b/Engine/Source/Runtime/CoreUObject/Private/Serialization/AsyncLoading.cpp
--- a/Engine/Source/Runtime/CoreUObject/Private/Serialization/AsyncLoading.cpp
+++ b/Engine/Source/Runtime/CoreUObject/Private/Serialization/AsyncLoading.cpp
@@ -5334,6 +5334,8 @@ FAsyncPackage::~FAsyncPackage()
 void FAsyncPackage::AddReferencedObjects(FReferenceCollector& Collector)
 {
 	Collector.AddReferencedObjects(ReferencedObjects);
+	Collector.AddReferencedObjects(DeferredFinalizeObjects);
+	Collector.AddReferencedObjects(PackageObjLoaded);
 }
 
 void FAsyncPackage::AddObjectReference(UObject* InObject)
