Change fatal error on TMap/TSet replication to a non-fatal error and improve the log message.

diff --git a/Engine/Source/Runtime/CoreUObject/Private/UObject/PropertyMap.cpp b/Engine/Source/Runtime/CoreUObject/Private/UObject/PropertyMap.cpp
--- a/Engine/Source/Runtime/CoreUObject/Private/UObject/PropertyMap.cpp
+++ b/Engine/Source/Runtime/CoreUObject/Private/UObject/PropertyMap.cpp
@@ -436,7 +436,7 @@ void UMapProperty::SerializeItem(FArchive& Ar, void* Value, const void* Defaults
 
 bool UMapProperty::NetSerializeItem( FArchive& Ar, UPackageMap* Map, void* Data, TArray<uint8> * MetaData ) const
 {
-	UE_LOG( LogProperty, Fatal, TEXT( "Deprecated code path" ) );
+	UE_LOG( LogProperty, Error, TEXT( "Replicated TMaps are not supported." ) );
 	return 1;
 }
 
diff --git a/Engine/Source/Runtime/CoreUObject/Private/UObject/PropertySet.cpp b/Engine/Source/Runtime/CoreUObject/Private/UObject/PropertySet.cpp
--- a/Engine/Source/Runtime/CoreUObject/Private/UObject/PropertySet.cpp
+++ b/Engine/Source/Runtime/CoreUObject/Private/UObject/PropertySet.cpp
@@ -388,7 +388,7 @@ void USetProperty::SerializeItem(FArchive& Ar, void* Value, const void* Defaults
 
 bool USetProperty::NetSerializeItem( FArchive& Ar, UPackageMap* Map, void* Data, TArray<uint8>* MetaData ) const
 {
-	UE_LOG( LogProperty, Fatal, TEXT( "Deprecated code path" ) );
+	UE_LOG( LogProperty, Error, TEXT( "Replicated TSets are not supported." ) );
 	return 1;
 }
 
