Fix default object properties of user defined structs failing to be loaded as a result of the load deferred flag
#rnx

diff --git a/Engine/Source/Runtime/CoreUObject/Private/UObject/PropertyBaseObject.cpp b/Engine/Source/Runtime/CoreUObject/Private/UObject/PropertyBaseObject.cpp
--- a/Engine/Source/Runtime/CoreUObject/Private/UObject/PropertyBaseObject.cpp
+++ b/Engine/Source/Runtime/CoreUObject/Private/UObject/PropertyBaseObject.cpp
@@ -442,12 +442,7 @@ UObject* UObjectPropertyBase::FindImportedObject( const UProperty* Property, UOb
 			if (Result == nullptr)
 #endif // USE_CIRCULAR_DEPENDENCY_LOAD_DEFERRING
 			{
-				uint32 LoadFlags = LOAD_NoWarn | LOAD_FindIfFail;
-				if (Linker)
-				{
-					LoadFlags |= Linker->LoadFlags & LOAD_DeferDependencyLoads;
-				}
-					
+				const uint32 LoadFlags = LOAD_NoWarn | LOAD_FindIfFail;		
 
 				UE_LOG(LogProperty, Verbose, TEXT("FindImportedObject is attempting to import [%s] (class = %s) with StaticLoadObject"), Text, *GetFullNameSafe(ObjectClass));
 				Result = StaticLoadObject(ObjectClass, nullptr, Text, nullptr, LoadFlags, nullptr);
