Only run force preload on subobjects that are being loaded by the linker, prevents objects regenerated after load from being loaded 'on top of'

diff --git a/Engine/Source/Runtime/CoreUObject/Private/UObject/LinkerLoad.cpp b/Engine/Source/Runtime/CoreUObject/Private/UObject/LinkerLoad.cpp
--- a/Engine/Source/Runtime/CoreUObject/Private/UObject/LinkerLoad.cpp
+++ b/Engine/Source/Runtime/CoreUObject/Private/UObject/LinkerLoad.cpp
@@ -3334,6 +3334,7 @@ void FLinkerLoad::Preload( UObject* Object )
 			// right now there are no known scenarios where someone requests a Preloa() 
 			// on a temporary ULinkerPlaceholderExportObject
 			check(!Object->IsA<ULinkerPlaceholderExportObject>());
+			ensure(Object->HasAnyFlags(RF_WasLoaded));
 #endif // USE_DEFERRED_DEPENDENCY_CHECK_VERIFICATION_TESTS
 #endif // USE_CIRCULAR_DEPENDENCY_LOAD_DEFERRING
 
@@ -4105,7 +4106,8 @@ UObject* FLinkerLoad::CreateExport( int32 Index )
 				for (UObject* SubObject : SuperSubObjects)
 				{
 					// Matching behavior in UBlueprint::ForceLoad to ensure that the subobject is actually loaded:
-					if (SubObject->HasAnyFlags(RF_NeedLoad) || !SubObject->HasAnyFlags(RF_LoadCompleted))
+					if (SubObject->HasAnyFlags(RF_WasLoaded) &&
+						(SubObject->HasAnyFlags(RF_NeedLoad) || !SubObject->HasAnyFlags(RF_LoadCompleted)))
 					{
 						SubObject->SetFlags(RF_NeedLoad);
 						Preload(SubObject);
