Optimized Mark Phase of GC by up to 10ms by making it run in parallel and removing a huge array presize which we didn't need.

diff --git a/Engine/Source/Runtime/CoreUObject/Private/UObject/GarbageCollection.cpp b/Engine/Source/Runtime/CoreUObject/Private/UObject/GarbageCollection.cpp
--- a/Engine/Source/Runtime/CoreUObject/Private/UObject/GarbageCollection.cpp
+++ b/Engine/Source/Runtime/CoreUObject/Private/UObject/GarbageCollection.cpp
@@ -23,6 +23,7 @@
 #include "UObject/UObjectClusters.h"
 #include "HAL/LowLevelMemTracker.h"
 #include "UObject/GarbageCollectionVerification.h"
+#include "Async/ParallelFor.h"
 
 /*-----------------------------------------------------------------------------
    Garbage collection.
@@ -33,9 +34,8 @@
 
 DEFINE_LOG_CATEGORY(LogGarbage);
 
-
 /** Object count during last mark phase																				*/
-int32		GObjectCountDuringLastMarkPhase			= 0;
+FThreadSafeCounter		GObjectCountDuringLastMarkPhase;
 /** Count of objects purged since last mark phase																	*/
 int32		GPurgedObjectCountSinceLastMarkPhase	= 0;
 /** Whether incremental object purge is in progress										*/
@@ -707,111 +707,146 @@ public:
 	 * Marks all objects that don't have KeepFlags and EInternalObjectFlags::GarbageCollectionKeepFlags as unreachable
 	 * This function is a template to speed up the case where we don't need to assemble the token stream (saves about 6ms on PS4)
 	 */
-	void MarkObjectsAsUnreachable(TArray<UObject*>& ObjectsToSerialize, const EObjectFlags KeepFlags)
+	void MarkObjectsAsUnreachable(TArray<UObject*>& ObjectsToSerialize, const EObjectFlags KeepFlags, bool bForceSingleThreaded)
 	{
 		const EInternalObjectFlags FastKeepFlags = EInternalObjectFlags::GarbageCollectionKeepFlags;
 
+		TLockFreePointerListFIFO<UObject, PLATFORM_CACHE_LINE_SIZE> ObjectsToSerializeList;
+		TLockFreePointerListFIFO<FUObjectItem, PLATFORM_CACHE_LINE_SIZE> ClustersToDissolveList;
+		TLockFreePointerListFIFO<FUObjectItem, PLATFORM_CACHE_LINE_SIZE> KeepClusterRefsList;
+
+		int32 MaxNumberOfObjects = GUObjectArray.GetObjectArrayNum() - GUObjectArray.GetFirstGCIndex();
+		int32 NumThreads = FMath::Max(1, FTaskGraphInterface::Get().GetNumWorkerThreads());
+		int32 NumberOfObjectsPerThread = (MaxNumberOfObjects / NumThreads) + 1;		
+
 		// Iterate over all objects. Note that we iterate over the UObjectArray and usually check only internal flags which
 		// are part of the array so we don't suffer from cache misses as much as we would if we were to check ObjectFlags.
-		TArray<FUObjectItem*> KeepClusterRefs;
-		for (FRawObjectIterator It(true); It; ++It)
+		ParallelFor(NumThreads, [&ObjectsToSerializeList, &ClustersToDissolveList, &KeepClusterRefsList, FastKeepFlags, KeepFlags, NumberOfObjectsPerThread, NumThreads, MaxNumberOfObjects](int32 ThreadIndex)
 		{
-			FUObjectItem* ObjectItem = *It;
-			checkSlow(ObjectItem);
-			UObject* Object = (UObject*)ObjectItem->Object;
+			int32 FirstObjectIndex = ThreadIndex * NumberOfObjectsPerThread + GUObjectArray.GetFirstGCIndex();
+			int32 NumObjects = (ThreadIndex < (NumThreads - 1)) ? NumberOfObjectsPerThread : (MaxNumberOfObjects - (NumThreads - 1) * NumberOfObjectsPerThread);
+			int32 LastObjectIndex = FMath::Min(GUObjectArray.GetObjectArrayNum() - 1, FirstObjectIndex + NumObjects - 1);
+			int32 ObjectCountDuringMarkPhase = 0;
 
-			// We can't collect garbage during an async load operation and by now all unreachable objects should've been purged.
-			checkf(!ObjectItem->IsUnreachable(), TEXT("%s"), *Object->GetFullName());
-
-			// Keep track of how many objects are around.
-			GObjectCountDuringLastMarkPhase++;
-			ObjectItem->ClearFlags(EInternalObjectFlags::ReachableInCluster);
-			// Special case handling for objects that are part of the root set.
-			if (ObjectItem->IsRootSet())
+			for (int32 ObjectIndex = FirstObjectIndex; ObjectIndex <= LastObjectIndex; ++ObjectIndex)
 			{
-				// IsValidLowLevel is extremely slow in this loop so only do it in debug
-				checkSlow(Object->IsValidLowLevel());
-				// We cannot use RF_PendingKill on objects that are part of the root set.
-#if DO_GUARD_SLOW
-				checkCode(if (ObjectItem->IsPendingKill()) { UE_LOG(LogGarbage, Fatal, TEXT("Object %s is part of root set though has been marked RF_PendingKill!"), *Object->GetFullName()); });
-#endif
-				if (ObjectItem->HasAnyFlags(EInternalObjectFlags::ClusterRoot) || ObjectItem->GetOwnerIndex() > 0)
+				FUObjectItem* ObjectItem = &GUObjectArray.GetObjectItemArrayUnsafe()[ObjectIndex];
+				if (ObjectItem->Object)
 				{
-					KeepClusterRefs.Add(ObjectItem);
-				}
+					UObject* Object = (UObject*)ObjectItem->Object;
 
-				ObjectsToSerialize.Add(Object);
-			}
-			// Regular objects or cluster root objects
-			else if (ObjectItem->GetOwnerIndex() <= 0)
-			{
-				bool bMarkAsUnreachable = true;
-				if (!ObjectItem->IsPendingKill())
-				{
-					// Internal flags are super fast to check
-					if (ObjectItem->HasAnyFlags(FastKeepFlags))
+					// We can't collect garbage during an async load operation and by now all unreachable objects should've been purged.
+					checkf(!ObjectItem->IsUnreachable(), TEXT("%s"), *Object->GetFullName());
+
+					// Keep track of how many objects are around.
+					ObjectCountDuringMarkPhase++;
+					
+					ObjectItem->ClearFlags(EInternalObjectFlags::ReachableInCluster);
+					// Special case handling for objects that are part of the root set.
+					if (ObjectItem->IsRootSet())
 					{
-						bMarkAsUnreachable = false;
+						// IsValidLowLevel is extremely slow in this loop so only do it in debug
+						checkSlow(Object->IsValidLowLevel());
+						// We cannot use RF_PendingKill on objects that are part of the root set.
+#if DO_GUARD_SLOW
+						checkCode(if (ObjectItem->IsPendingKill()) { UE_LOG(LogGarbage, Fatal, TEXT("Object %s is part of root set though has been marked RF_PendingKill!"), *Object->GetFullName()); });
+#endif
+						if (ObjectItem->HasAnyFlags(EInternalObjectFlags::ClusterRoot) || ObjectItem->GetOwnerIndex() > 0)
+						{
+							KeepClusterRefsList.Push(ObjectItem);
+						}
+
+						ObjectsToSerializeList.Push(Object);
 					}
-					// If KeepFlags is non zero this is going to be very slow due to cache misses
-					else if (KeepFlags != RF_NoFlags && Object->HasAnyFlags(KeepFlags))
+					// Regular objects or cluster root objects
+					else if (ObjectItem->GetOwnerIndex() <= 0)
 					{
-						bMarkAsUnreachable = false;
-					}
-				}
-				else if (ObjectItem->HasAnyFlags(EInternalObjectFlags::ClusterRoot))
-				{
-					GUObjectClusters.DissolveClusterAndMarkObjectsAsUnreachable(It.GetIndex(), ObjectItem);
-					GUObjectClusters.SetClustersNeedDissolving();					
-				}
+						bool bMarkAsUnreachable = true;
+						if (!ObjectItem->IsPendingKill())
+						{
+							// Internal flags are super fast to check
+							if (ObjectItem->HasAnyFlags(FastKeepFlags))
+							{
+								bMarkAsUnreachable = false;
+							}
+							// If KeepFlags is non zero this is going to be very slow due to cache misses
+							else if (KeepFlags != RF_NoFlags && Object->HasAnyFlags(KeepFlags))
+							{
+								bMarkAsUnreachable = false;
+							}
+						}
+						else if (ObjectItem->HasAnyFlags(EInternalObjectFlags::ClusterRoot))
+						{
+							ClustersToDissolveList.Push(ObjectItem);
+						}
 
-				// Mark objects as unreachable unless they have any of the passed in KeepFlags set and it's not marked for elimination..
-				if (!bMarkAsUnreachable)
-				{
-					// IsValidLowLevel is extremely slow in this loop so only do it in debug
-					checkSlow(Object->IsValidLowLevel());
-					ObjectsToSerialize.Add(Object);
+						// Mark objects as unreachable unless they have any of the passed in KeepFlags set and it's not marked for elimination..
+						if (!bMarkAsUnreachable)
+						{
+							// IsValidLowLevel is extremely slow in this loop so only do it in debug
+							checkSlow(Object->IsValidLowLevel());
+							ObjectsToSerializeList.Push(Object);
 
-					if (ObjectItem->HasAnyFlags(EInternalObjectFlags::ClusterRoot))
-					{
-						KeepClusterRefs.Add(ObjectItem);
+							if (ObjectItem->HasAnyFlags(EInternalObjectFlags::ClusterRoot))
+							{
+								KeepClusterRefsList.Push(ObjectItem);
+							}
+						}
+						else
+						{
+							ObjectItem->SetFlags(EInternalObjectFlags::Unreachable);
+						}
 					}
 				}
-				else
-				{
-					ObjectItem->SetFlags(EInternalObjectFlags::Unreachable);
-				}
+			}
+
+			GObjectCountDuringLastMarkPhase.Add(ObjectCountDuringMarkPhase);
+		}, bForceSingleThreaded);
+		
+		ObjectsToSerializeList.PopAll(ObjectsToSerialize);
+
+		{
+			TArray<FUObjectItem*> ClustersToDissolve;
+			ClustersToDissolveList.PopAll(ClustersToDissolve);
+			for (FUObjectItem* ObjectItem : ClustersToDissolve)
+			{
+				GUObjectClusters.DissolveClusterAndMarkObjectsAsUnreachable(ObjectItem);
+				GUObjectClusters.SetClustersNeedDissolving();
 			}
 		}
 
-		for (FUObjectItem* ObjectItem : KeepClusterRefs)
 		{
-			if (ObjectItem->GetOwnerIndex() > 0)
+			TArray<FUObjectItem*> KeepClusterRefs;
+			KeepClusterRefsList.PopAll(KeepClusterRefs);
+			for (FUObjectItem* ObjectItem : KeepClusterRefs)
 			{
-				checkSlow(!ObjectItem->HasAnyFlags(EInternalObjectFlags::ClusterRoot));
-				bool bNeedsDoing = !ObjectItem->HasAnyFlags(EInternalObjectFlags::ReachableInCluster);
-				if (bNeedsDoing)
+				if (ObjectItem->GetOwnerIndex() > 0)
 				{
-					ObjectItem->SetFlags(EInternalObjectFlags::ReachableInCluster);
-					// Make sure cluster root object is reachable too
-					const int32 OwnerIndex = ObjectItem->GetOwnerIndex();
-					FUObjectItem* RootObjectItem = GUObjectArray.IndexToObjectUnsafeForGC(OwnerIndex);
-					checkSlow(RootObjectItem->HasAnyFlags(EInternalObjectFlags::ClusterRoot));
-					// if it is reachable via keep flags we will do this below (or maybe already have)
-					if (RootObjectItem->IsUnreachable()) 
+					checkSlow(!ObjectItem->HasAnyFlags(EInternalObjectFlags::ClusterRoot));
+					bool bNeedsDoing = !ObjectItem->HasAnyFlags(EInternalObjectFlags::ReachableInCluster);
+					if (bNeedsDoing)
 					{
-						RootObjectItem->ClearFlags(EInternalObjectFlags::Unreachable);
-						// Make sure all referenced clusters are marked as reachable too
-						FGCReferenceProcessorSinglethreaded::MarkReferencedClustersAsReachable(RootObjectItem->GetClusterIndex(), ObjectsToSerialize);
+						ObjectItem->SetFlags(EInternalObjectFlags::ReachableInCluster);
+						// Make sure cluster root object is reachable too
+						const int32 OwnerIndex = ObjectItem->GetOwnerIndex();
+						FUObjectItem* RootObjectItem = GUObjectArray.IndexToObjectUnsafeForGC(OwnerIndex);
+						checkSlow(RootObjectItem->HasAnyFlags(EInternalObjectFlags::ClusterRoot));
+						// if it is reachable via keep flags we will do this below (or maybe already have)
+						if (RootObjectItem->IsUnreachable()) 
+						{
+							RootObjectItem->ClearFlags(EInternalObjectFlags::Unreachable);
+							// Make sure all referenced clusters are marked as reachable too
+							FGCReferenceProcessorSinglethreaded::MarkReferencedClustersAsReachable(RootObjectItem->GetClusterIndex(), ObjectsToSerialize);
+						}
 					}
 				}
-			}
-			else
-			{
-				checkSlow(ObjectItem->HasAnyFlags(EInternalObjectFlags::ClusterRoot));
-				// this thing is definitely not marked unreachable, so don't test it here
-				// Make sure all referenced clusters are marked as reachable too
-				FGCReferenceProcessorSinglethreaded::MarkReferencedClustersAsReachable(ObjectItem->GetClusterIndex(), ObjectsToSerialize);
+				else
+				{
+					checkSlow(ObjectItem->HasAnyFlags(EInternalObjectFlags::ClusterRoot));
+					// this thing is definitely not marked unreachable, so don't test it here
+					// Make sure all referenced clusters are marked as reachable too
+					FGCReferenceProcessorSinglethreaded::MarkReferencedClustersAsReachable(ObjectItem->GetClusterIndex(), ObjectsToSerialize);
+				}
 			}
 		}
 	}
@@ -833,18 +868,25 @@ public:
 		TArray<UObject*>& ObjectsToSerialize = ArrayStruct->ObjectsToSerialize;
 
 		// Reset object count.
-		GObjectCountDuringLastMarkPhase = 0;
+		GObjectCountDuringLastMarkPhase.Reset();
 
-		// Presize array and add a bit of extra slack for prefetching.
-		ObjectsToSerialize.Reset( GUObjectArray.GetObjectArrayNumMinusPermanent() + 3 );
 		// Make sure GC referencer object is checked for references to other objects even if it resides in permanent object pool
 		if (FPlatformProperties::RequiresCookedData() && FGCObject::GGCObjectReferencer && GUObjectArray.IsDisregardForGC(FGCObject::GGCObjectReferencer))
 		{
 			ObjectsToSerialize.Add(FGCObject::GGCObjectReferencer);
 		}
 
-		MarkObjectsAsUnreachable(ObjectsToSerialize, KeepFlags);
-		PerformReachabilityAnalysisOnObjects(ArrayStruct, bForceSingleThreaded);
+		{
+			const double StartTime = FPlatformTime::Seconds();
+			MarkObjectsAsUnreachable(ObjectsToSerialize, KeepFlags, bForceSingleThreaded);
+			UE_LOG(LogGarbage, Verbose, TEXT("%f ms for Mark Phase (%d Objects To Serialize"), (FPlatformTime::Seconds() - StartTime) * 1000, ObjectsToSerialize.Num());
+		}
+
+		{
+			const double StartTime = FPlatformTime::Seconds();
+			PerformReachabilityAnalysisOnObjects(ArrayStruct, bForceSingleThreaded);
+			UE_LOG(LogGarbage, Verbose, TEXT("%f ms for Reachability Analysis"), (FPlatformTime::Seconds() - StartTime) * 1000);
+		}
         
 		// Allowing external systems to add object roots. This can't be done through AddReferencedObjects
 		// because it may require tracing objects (via FGarbageCollectionTracer) multiple times
@@ -1144,7 +1186,7 @@ void IncrementalPurgeGarbage( bool bUseTimeLimit, float TimeLimit )
 			GObjCurrentPurgeObjectIndexResetPastPermanent	= true;
 
 			// Log status information.
-			UE_LOG(LogGarbage, Log, TEXT("GC purged %i objects (%i -> %i)"), GPurgedObjectCountSinceLastMarkPhase, GObjectCountDuringLastMarkPhase, GObjectCountDuringLastMarkPhase - GPurgedObjectCountSinceLastMarkPhase );
+			UE_LOG(LogGarbage, Log, TEXT("GC purged %i objects (%i -> %i)"), GPurgedObjectCountSinceLastMarkPhase, GObjectCountDuringLastMarkPhase.GetValue(), GObjectCountDuringLastMarkPhase.GetValue() - GPurgedObjectCountSinceLastMarkPhase );
 
 #if PERF_DETAILED_PER_CLASS_GC_STATS
 			LogClassCountInfo( TEXT("objects of"), GClassToPurgeCountMap, 10, GPurgedObjectCountSinceLastMarkPhase );
diff --git a/Engine/Source/Runtime/CoreUObject/Private/UObject/Obj.cpp b/Engine/Source/Runtime/CoreUObject/Private/UObject/Obj.cpp
--- a/Engine/Source/Runtime/CoreUObject/Private/UObject/Obj.cpp
+++ b/Engine/Source/Runtime/CoreUObject/Private/UObject/Obj.cpp
@@ -4218,8 +4218,8 @@ void StaticExit()
 	IncrementalPurgeGarbage( false );
 
 	// Keep track of how many objects there are for GC stats as we simulate a mark pass.
-	extern int32 GObjectCountDuringLastMarkPhase;
-	GObjectCountDuringLastMarkPhase = 0;
+	extern FThreadSafeCounter GObjectCountDuringLastMarkPhase;
+	GObjectCountDuringLastMarkPhase.Reset();
 
 	// Tag all non template & class objects as unreachable. We can't use object iterators for this as they ignore certain objects.
 	//
@@ -4231,7 +4231,7 @@ void StaticExit()
 	for ( FRawObjectIterator It; It; ++It )
 	{
 		// Valid object.
-		GObjectCountDuringLastMarkPhase++;
+		GObjectCountDuringLastMarkPhase.Increment();
 
 		FUObjectItem* ObjItem = *It;
 		checkSlow(ObjItem);
diff --git a/Engine/Source/Runtime/CoreUObject/Private/UObject/UObjectClusters.cpp b/Engine/Source/Runtime/CoreUObject/Private/UObject/UObjectClusters.cpp
--- a/Engine/Source/Runtime/CoreUObject/Private/UObject/UObjectClusters.cpp
+++ b/Engine/Source/Runtime/CoreUObject/Private/UObject/UObjectClusters.cpp
@@ -163,7 +163,7 @@ void FUObjectClusterContainer::DissolveCluster(FUObjectCluster& Cluster)
 	}
 }
 
-void FUObjectClusterContainer::DissolveClusterAndMarkObjectsAsUnreachable(const int32 CurrentIndex, FUObjectItem* RootObjectItem)
+void FUObjectClusterContainer::DissolveClusterAndMarkObjectsAsUnreachable(FUObjectItem* RootObjectItem)
 {
 	const int32 OldClusterIndex = RootObjectItem->GetClusterIndex();
 	FUObjectCluster& Cluster = Clusters[OldClusterIndex];
@@ -180,10 +180,7 @@ void FUObjectClusterContainer::DissolveClusterAndMarkObjectsAsUnreachable(const
 	{
 		FUObjectItem* ClusterObjectItem = GUObjectArray.IndexToObjectUnsafeForGC(ClusterObjectIndex);
 		ClusterObjectItem->SetOwnerIndex(0);
-		if (ClusterObjectIndex < CurrentIndex)
-		{
-			ClusterObjectItem->SetFlags(EInternalObjectFlags::Unreachable);
-		}
+		ClusterObjectItem->SetFlags(EInternalObjectFlags::Unreachable);
 	}
 
 #if !UE_GCCLUSTER_VERBOSE_LOGGING
@@ -199,11 +196,8 @@ void FUObjectClusterContainer::DissolveClusterAndMarkObjectsAsUnreachable(const
 		FUObjectItem* ReferencedByClusterRootItem = GUObjectArray.IndexToObjectUnsafeForGC(ReferencedByClusterRootIndex);
 		if (ReferencedByClusterRootItem->HasAnyFlags(EInternalObjectFlags::ClusterRoot))
 		{
-			if (ReferencedByClusterRootIndex < CurrentIndex)
-			{
-				ReferencedByClusterRootItem->SetFlags(EInternalObjectFlags::Unreachable);
-			}
-			DissolveClusterAndMarkObjectsAsUnreachable(CurrentIndex, ReferencedByClusterRootItem);
+			ReferencedByClusterRootItem->SetFlags(EInternalObjectFlags::Unreachable);
+			DissolveClusterAndMarkObjectsAsUnreachable(ReferencedByClusterRootItem);
 		}
 	}
 }
diff --git a/Engine/Source/Runtime/CoreUObject/Private/UObject/UObjectGlobals.cpp b/Engine/Source/Runtime/CoreUObject/Private/UObject/UObjectGlobals.cpp
--- a/Engine/Source/Runtime/CoreUObject/Private/UObject/UObjectGlobals.cpp
+++ b/Engine/Source/Runtime/CoreUObject/Private/UObject/UObjectGlobals.cpp
@@ -3441,8 +3441,8 @@ public:
 	void PerformReachabilityAnalysis( EObjectFlags KeepFlags, EInternalObjectFlags InternalKeepFlags, EObjectFlags SearchFlags = RF_NoFlags, FReferencerInformationList* FoundReferences = NULL)
 	{
 		// Reset object count.
-		extern int32 GObjectCountDuringLastMarkPhase;
-		GObjectCountDuringLastMarkPhase = 0;
+		extern FThreadSafeCounter GObjectCountDuringLastMarkPhase;
+		GObjectCountDuringLastMarkPhase.Reset();
 		ReferenceSearchFlags = SearchFlags;
 		FoundReferencesList = FoundReferences;
 
@@ -3451,7 +3451,7 @@ public:
 		{
 			UObject* Object	= *It;
 			checkSlow(Object->IsValidLowLevel());
-			GObjectCountDuringLastMarkPhase++;
+			GObjectCountDuringLastMarkPhase.Increment();
 
 			// Special case handling for objects that are part of the root set.
 			if( Object->IsRooted() )
diff --git a/Engine/Source/Runtime/CoreUObject/Public/UObject/UObjectArray.h b/Engine/Source/Runtime/CoreUObject/Public/UObject/UObjectArray.h
--- a/Engine/Source/Runtime/CoreUObject/Public/UObject/UObjectArray.h
+++ b/Engine/Source/Runtime/CoreUObject/Public/UObject/UObjectArray.h
@@ -1058,7 +1058,7 @@ public:
 	void DissolveClusters();
 
 	/** Dissolve the specified cluster and all clusters that reference it */
-	void DissolveClusterAndMarkObjectsAsUnreachable(const int32 CurrentIndex, FUObjectItem* RootObjectItem);
+	void DissolveClusterAndMarkObjectsAsUnreachable(FUObjectItem* RootObjectItem);
 
 	/** Gets the clusters array (for internal use only!) */
 	TArray<FUObjectCluster>& GetClustersUnsafe() 
