Set the minimum GC cluster size to 5 so that GC doesn't have to process micro clusters which are more expensive than processing individual objects

+ Exposed the minimum cluster size to ini and project settings as gc.MinGCClusterSize
+ Added the ability to sort clusters by name/object count/mutable object count/referenced clusters count when dumping them with gc.ListClusters command

diff --git a/Engine/Config/BaseEngine.ini b/Engine/Config/BaseEngine.ini
--- a/Engine/Config/BaseEngine.ini
+++ b/Engine/Config/BaseEngine.ini
@@ -692,6 +692,7 @@ gc.TimeBetweenPurgingPendingKillObjects=61.1
 gc.MaxObjectsInEditor=8388607
 gc.CreateGCClusters=True
 gc.MergeGCClusters=False
+gc.MinGCClusterSize=5
 gc.ActorClusteringEnabled=False
 gc.BlueprintClusteringEnabled=False
 gc.UseDisregardForGCOnDedicatedServers=False
diff --git a/Engine/Source/Runtime/CoreUObject/Private/UObject/UObjectClusters.cpp b/Engine/Source/Runtime/CoreUObject/Private/UObject/UObjectClusters.cpp
--- a/Engine/Source/Runtime/CoreUObject/Private/UObject/UObjectClusters.cpp
+++ b/Engine/Source/Runtime/CoreUObject/Private/UObject/UObjectClusters.cpp
@@ -35,6 +35,13 @@ static FAutoConsoleVariableRef CMergeGCClusters(
 	ECVF_Default
 	);
 
+int32 GMinGCClusterSize = 5;
+static FAutoConsoleVariableRef CMinGCClusterSize(
+	TEXT("gc.MinGCClusterSize"),
+	GMinGCClusterSize,
+	TEXT("Minimum GC cluster size"),
+	ECVF_Default
+);
 
 FUObjectClusterContainer::FUObjectClusterContainer()
 	: NumAllocatedClusters(0)
@@ -214,6 +221,11 @@ void FUObjectClusterContainer::DissolveClusters()
 	bClustersNeedDissolving = false;
 }
 
+int32 FUObjectClusterContainer::GetMinClusterSize() const
+{
+	return FMath::Max(1, GMinGCClusterSize);
+}
+
 #if !UE_BUILD_SHIPPING
 
 void DumpClusterToLog(const FUObjectCluster& Cluster, bool bHierarchy, bool bIndexOnly)
@@ -286,35 +298,67 @@ void DumpClusterToLog(const FUObjectCluster& Cluster, bool bHierarchy, bool bInd
 // Dumps all clusters to log.
 void ListClusters(const TArray<FString>& Args)
 {
-	const bool bHierarchy = Args.Num() && Args[0] == TEXT("Hierarchy");
+	const bool bHierarchy = Args.Contains(TEXT("Hierarchy"));
 	int32 MaxInterClusterReferences = 0;
 	int32 TotalInterClusterReferences = 0;
 	int32 MaxClusterSize = 0;
 	int32 TotalClusterObjects = 0;	
-	int32 NumClusters = 0;
 
-	for (const FUObjectCluster& Cluster : GUObjectClusters.GetClustersUnsafe())
+	TArray<FUObjectCluster*> AllClusters;
+	for (FUObjectCluster& Cluster : GUObjectClusters.GetClustersUnsafe())
 	{
-		if (Cluster.RootIndex == INDEX_NONE)
+		if (Cluster.RootIndex != INDEX_NONE)
 		{
-			continue;
+			AllClusters.Add(&Cluster);
 		}
+	}
+
+	if (Args.Contains(TEXT("SortByName")))
+	{
+		Algo::SortBy(AllClusters, [](FUObjectCluster* A)
+		{
+			return GUObjectArray.IndexToObject(A->RootIndex)->Object->GetFName();
+		});
+	}
+	else if (Args.Contains(TEXT("SortByObjectCount")))
+	{
+		Algo::SortBy(AllClusters, [](FUObjectCluster* A)
+		{
+			return A->Objects.Num();
+		});
+	}
+	else if (Args.Contains(TEXT("SortByMutableObjectCount")))
+	{
+		Algo::SortBy(AllClusters, [](FUObjectCluster* A)
+		{
+			return A->MutableObjects.Num();
+		});
+	}
+	else if (Args.Contains(TEXT("SortByReferencedClustersCount")))
+	{
+		Algo::SortBy(AllClusters, [](FUObjectCluster* A)
+		{
+			return A->ReferencedClusters.Num();
+		});
+	}
 
-		NumClusters++;
+	for (FUObjectCluster* Cluster : AllClusters)
+	{
+		check(Cluster->RootIndex != INDEX_NONE);
 
-		MaxInterClusterReferences = FMath::Max(MaxInterClusterReferences, Cluster.ReferencedClusters.Num());
-		TotalInterClusterReferences += Cluster.ReferencedClusters.Num();
-		MaxClusterSize = FMath::Max(MaxClusterSize, Cluster.Objects.Num());
-		TotalClusterObjects += Cluster.Objects.Num();
+		MaxInterClusterReferences = FMath::Max(MaxInterClusterReferences, Cluster->ReferencedClusters.Num());
+		TotalInterClusterReferences += Cluster->ReferencedClusters.Num();
+		MaxClusterSize = FMath::Max(MaxClusterSize, Cluster->Objects.Num());
+		TotalClusterObjects += Cluster->Objects.Num();
 
-		DumpClusterToLog(Cluster, bHierarchy, false);
+		DumpClusterToLog(*Cluster, bHierarchy, false);
 	}
-	UE_LOG(LogObj, Display, TEXT("Number of clusters: %d"), NumClusters);
+	UE_LOG(LogObj, Display, TEXT("Number of clusters: %d"), AllClusters.Num());
 	UE_LOG(LogObj, Display, TEXT("Maximum cluster size: %d"), MaxClusterSize);
-	UE_LOG(LogObj, Display, TEXT("Average cluster size: %d"), NumClusters ? (TotalClusterObjects / NumClusters) : 0);
+	UE_LOG(LogObj, Display, TEXT("Average cluster size: %d"), AllClusters.Num() ? (TotalClusterObjects / AllClusters.Num()) : 0);
 	UE_LOG(LogObj, Display, TEXT("Number of objects in GC clusters: %d"), TotalClusterObjects);
 	UE_LOG(LogObj, Display, TEXT("Maximum number of custer-to-cluster references: %d"), MaxInterClusterReferences);
-	UE_LOG(LogObj, Display, TEXT("Average number of custer-to-cluster references: %d"), NumClusters ? (TotalInterClusterReferences / NumClusters) : 0);
+	UE_LOG(LogObj, Display, TEXT("Average number of custer-to-cluster references: %d"), AllClusters.Num() ? (TotalInterClusterReferences / AllClusters.Num()) : 0);
 }
 
 void FindStaleClusters(const TArray<FString>& Args)
@@ -677,15 +721,16 @@ void UObjectBaseUtility::CreateCluster()
 	FGCArrayPool::Get().CheckLeaks();
 #endif
 
-	if (Cluster.Objects.Num())
+	check(RootItem->GetOwnerIndex() == 0);
+	RootItem->SetClusterIndex(ClusterIndex);
+	RootItem->SetFlags(EInternalObjectFlags::ClusterRoot);
+
+	if (Cluster.Objects.Num() >= GUObjectClusters.GetMinClusterSize())
 	{
 		// Add new cluster to the global cluster map.
 		Cluster.Objects.Sort();
 		Cluster.ReferencedClusters.Sort();		
 		Cluster.MutableObjects.Sort();		
-		check(RootItem->GetOwnerIndex() == 0);
-		RootItem->SetClusterIndex(ClusterIndex);
-		RootItem->SetFlags(EInternalObjectFlags::ClusterRoot);
 
 #if UE_GCCLUSTER_VERBOSE_LOGGING
 		UE_LOG(LogObj, Log, TEXT("Created Cluster (%d) with %d objects, %d referenced clusters and %d mutable objects."),
@@ -696,9 +741,14 @@ void UObjectBaseUtility::CreateCluster()
 	}
 	else
 	{
-		check(RootItem->GetOwnerIndex() == 0);
-		RootItem->SetClusterIndex(ClusterIndex);
+		for (int32 ClusterObjectIndex : Cluster.Objects)
+		{
+			FUObjectItem* ClusterObjectItem = GUObjectArray.IndexToObjectUnsafeForGC(ClusterObjectIndex);
+			ClusterObjectItem->SetOwnerIndex(0);
+		}
 		GUObjectClusters.FreeCluster(ClusterIndex);
+		check(RootItem->GetOwnerIndex() == 0);
+		check(!RootItem->HasAnyFlags(EInternalObjectFlags::ClusterRoot));
 	}
 }
 
diff --git a/Engine/Source/Runtime/CoreUObject/Public/UObject/UObjectArray.h b/Engine/Source/Runtime/CoreUObject/Public/UObject/UObjectArray.h
--- a/Engine/Source/Runtime/CoreUObject/Public/UObject/UObjectArray.h
+++ b/Engine/Source/Runtime/CoreUObject/Public/UObject/UObjectArray.h
@@ -1060,6 +1060,9 @@ public:
 	/** Dissolve the specified cluster and all clusters that reference it */
 	void DissolveClusterAndMarkObjectsAsUnreachable(FUObjectItem* RootObjectItem);
 
+	/*** Returns the minimum cluster size as specified in ini settings */
+	int32 GetMinClusterSize() const;
+
 	/** Gets the clusters array (for internal use only!) */
 	TArray<FUObjectCluster>& GetClustersUnsafe() 
 	{ 
diff --git a/Engine/Source/Runtime/Engine/Classes/Engine/CoreSettings.h b/Engine/Source/Runtime/Engine/Classes/Engine/CoreSettings.h
--- a/Engine/Source/Runtime/Engine/Classes/Engine/CoreSettings.h
+++ b/Engine/Source/Runtime/Engine/Classes/Engine/CoreSettings.h
@@ -195,6 +195,11 @@ protected:
 		ToolTip = "If false, DisregardForGC will be disabled for dedicated servers."))
 	uint32 UseDisregardForGCOnDedicatedServers : 1;
 
+	UPROPERTY(EditAnywhere, config, Category = Optimization, meta = (
+		ConsoleVariable = "gc.MinGCClusterSize", DisplayName = "Minimum GC Cluster size",
+		ToolTip = "Minimum GC cluster size."))
+	int32 MinGCClusterSize;
+
 	UPROPERTY(EditAnywhere, config, Category = General, meta = (
 		ConsoleVariable = "gc.NumRetriesBeforeForcingGC", DisplayName = "Number Of Retries Before Forcing GC",
 		ToolTip = "Maximum number of times GC can be skipped if worker threads are currently modifying UObject state. 0 = never force GC"))
diff --git a/Engine/Source/Runtime/Engine/Private/CoreSettings.cpp b/Engine/Source/Runtime/Engine/Private/CoreSettings.cpp
--- a/Engine/Source/Runtime/Engine/Private/CoreSettings.cpp
+++ b/Engine/Source/Runtime/Engine/Private/CoreSettings.cpp
@@ -175,6 +175,7 @@ UGarbageCollectionSettings::UGarbageCollectionSettings()
 	MaxObjectsInGame = 2 * 1024 * 1024;
 	CreateGCClusters = true;
 	MergeGCClusters = false;
+	MinGCClusterSize = 5;
 	ActorClusteringEnabled = true;
 	BlueprintClusteringEnabled = false;
 	UseDisregardForGCOnDedicatedServers = false;
diff --git a/Engine/Source/Runtime/Engine/Private/LevelActorContainer.cpp b/Engine/Source/Runtime/Engine/Private/LevelActorContainer.cpp
--- a/Engine/Source/Runtime/Engine/Private/LevelActorContainer.cpp
+++ b/Engine/Source/Runtime/Engine/Private/LevelActorContainer.cpp
@@ -189,15 +189,16 @@ void ULevelActorContainer::CreateCluster()
 	FGCArrayPool::Get().CheckLeaks();
 #endif
 
-	if (Cluster.Objects.Num())
+	check(RootItem->GetOwnerIndex() == 0);
+	RootItem->SetClusterIndex(ClusterIndex);
+	RootItem->SetFlags(EInternalObjectFlags::ClusterRoot);
+
+	if (Cluster.Objects.Num() >= GUObjectClusters.GetMinClusterSize())
 	{
 		// Sort all objects and set up the cluster root
 		Cluster.Objects.Sort();
 		Cluster.ReferencedClusters.Sort();
 		Cluster.MutableObjects.Sort();
-		check(RootItem->GetOwnerIndex() == 0);
-		RootItem->SetClusterIndex(ClusterIndex);
-		RootItem->SetFlags(EInternalObjectFlags::ClusterRoot);
 
 		UE_LOG(LogLevelActorContainer, Log, TEXT("Created LevelActorCluster (%d) for %s with %d objects, %d referenced clusters and %d mutable objects."),
 			ClusterIndex, *GetOuter()->GetPathName(), Cluster.Objects.Num(), Cluster.ReferencedClusters.Num(), Cluster.MutableObjects.Num());
@@ -208,9 +209,14 @@ void ULevelActorContainer::CreateCluster()
 	}
 	else
 	{
-		check(RootItem->GetOwnerIndex() == 0);
-		RootItem->SetClusterIndex(ClusterIndex);
+		for (int32 ClusterObjectIndex : Cluster.Objects)
+		{
+			FUObjectItem* ClusterObjectItem = GUObjectArray.IndexToObjectUnsafeForGC(ClusterObjectIndex);
+			ClusterObjectItem->SetOwnerIndex(0);
+		}
 		GUObjectClusters.FreeCluster(ClusterIndex);
+		check(RootItem->GetOwnerIndex() == 0);
+		check(!RootItem->HasAnyFlags(EInternalObjectFlags::ClusterRoot));
 	}
 }
 
