Reworked how GC and other threads (ALT specifically) interact - GC will now notify the ALT it wants to run and ALT will immediately try to finish its current work to allow that. Also the entire ALT tick is now protected against GC running at the same time to improve ALT stability.

+ added gc.ForceCollectGarbageEveryFrame console variable that triggers a forced GC every frame

diff --git a/Engine/Source/Runtime/CoreUObject/Private/Serialization/AsyncLoading.cpp b/Engine/Source/Runtime/CoreUObject/Private/Serialization/AsyncLoading.cpp
--- a/Engine/Source/Runtime/CoreUObject/Private/Serialization/AsyncLoading.cpp
+++ b/Engine/Source/Runtime/CoreUObject/Private/Serialization/AsyncLoading.cpp
@@ -100,6 +100,9 @@ static FName StaticGetNativeClassName(UClass* InClass)
 /** Returns true if we're inside a FGCScopeLock */
 bool IsGarbageCollectionLocked();
 
+/** Returns true if GC wants to run */
+bool IsGarbageCollectionWaiting();
+
 /** Global request ID counter */
 static FThreadSafeCounter GPackageRequestID;
 
@@ -262,7 +265,11 @@ static FORCEINLINE bool IsTimeLimitExceeded(double InTickStartTime, bool bUseTim
 		}
 		LastTestTime = CurrentTime;
 	}
-
+	if (!bTimeLimitExceeded)
+	{
+		bTimeLimitExceeded = IsGarbageCollectionWaiting();
+		UE_CLOG(bTimeLimitExceeded, LogStreaming, Verbose, TEXT("Timing out async loading due to Garbage Collection request"));
+	}
 	return bTimeLimitExceeded;
 }
 FORCEINLINE bool FAsyncPackage::IsTimeLimitExceeded()
@@ -4751,6 +4758,7 @@ EAsyncPackageState::Type FAsyncLoadingThread::TickAsyncLoading(bool bUseTimeLimi
 	LLM_SCOPE(ELLMTag::AsyncLoading);
 
 	check(IsInGameThread());
+	check(!IsGarbageCollecting());
 
 	const bool bLoadingSuspended = IsAsyncLoadingSuspended();
 	EAsyncPackageState::Type Result = bLoadingSuspended ? EAsyncPackageState::PendingImports : EAsyncPackageState::Complete;
@@ -4969,9 +4977,12 @@ uint32 FAsyncLoadingThread::Run()
 			{
 				bWasSuspendedLastFrame = false;
 				ThreadResumedEvent->Trigger();
-			}		
-			bool bDidSomething = false;
-			TickAsyncThread(false, true, 0.0f, bDidSomething);
+			}
+			if (!IsGarbageCollectionWaiting())
+			{
+				bool bDidSomething = false;
+				TickAsyncThread(false, true, 0.0f, bDidSomething);
+			}
 		}
 		else if (!bWasSuspendedLastFrame)
 		{
@@ -5026,7 +5037,10 @@ EAsyncPackageState::Type  FAsyncLoadingThread::TickAsyncThread(bool bUseTimeLimi
 			{
 				GetGEDLBootNotificationManager().FireCompletedCompiledInImports();
 			}
-			CreateAsyncPackagesFromQueue(bUseTimeLimit, bUseFullTimeLimit, TimeLimit, FlushTree);
+			{
+				FGCScopeGuard GCGuard;
+				CreateAsyncPackagesFromQueue(bUseTimeLimit, bUseFullTimeLimit, TimeLimit, FlushTree);
+			}
 			float TimeUsed = (float)(FPlatformTime::Seconds() - TickStartTime);
 			const float RemainingTimeLimit = FMath::Max(0.0f, TimeLimit - TimeUsed);
 			if (RemainingTimeLimit <= 0.0f && bUseTimeLimit && !IsMultithreaded())
@@ -5035,6 +5049,7 @@ EAsyncPackageState::Type  FAsyncLoadingThread::TickAsyncThread(bool bUseTimeLimi
 			}
 			else
 			{
+				FGCScopeGuard GCGuard;
 				Result = ProcessAsyncLoading(ProcessedRequests, bUseTimeLimit, bUseFullTimeLimit, RemainingTimeLimit, FlushTree);
 				bDidSomething = bDidSomething || ProcessedRequests > 0;
 			}
@@ -5047,9 +5062,9 @@ EAsyncPackageState::Type  FAsyncLoadingThread::TickAsyncThread(bool bUseTimeLimi
 
 				if (!GetGEDLBootNotificationManager().IsWaitingForSomething())
 				{
-				CheckForCycles();
-				IsTimeLimitExceeded(TickStartTime, bUseTimeLimit, TimeLimit, TEXT("CheckForCycles (non-shipping)"));
-			}
+					CheckForCycles();
+					IsTimeLimitExceeded(TickStartTime, bUseTimeLimit, TimeLimit, TEXT("CheckForCycles (non-shipping)"));
+				}
 				else
 				{
 					WaitTime = 1; // we are waiting for the game thread to deal with the boot constructors, so lets spin tighter
@@ -5623,7 +5638,7 @@ EAsyncPackageState::Type FAsyncPackage::TickAsyncPackage(bool InbUseTimeLimit, b
 		}
 	} while (!IsTimeLimitExceeded() && LoadingState == EAsyncPackageState::TimeOut);
 
-	check(bUseTimeLimit || LoadingState != EAsyncPackageState::TimeOut || AsyncLoadingThread.IsAsyncLoadingSuspended());
+	check(bUseTimeLimit || LoadingState != EAsyncPackageState::TimeOut || AsyncLoadingThread.IsAsyncLoadingSuspended() || IsGarbageCollectionWaiting());
 
 	if (LinkerRoot && LoadingState == EAsyncPackageState::Complete)
 	{
diff --git a/Engine/Source/Runtime/CoreUObject/Private/UObject/GCScopeLock.h b/Engine/Source/Runtime/CoreUObject/Private/UObject/GCScopeLock.h
--- a/Engine/Source/Runtime/CoreUObject/Private/UObject/GCScopeLock.h
+++ b/Engine/Source/Runtime/CoreUObject/Private/UObject/GCScopeLock.h
@@ -48,24 +48,40 @@ public:
 */
 class FGCCSyncObject
 {
+	/** Non zero if any of the non-game threads is blocking GC */
 	FThreadSafeCounter AsyncCounter;
+	/** Non zero if GC is running */
 	FThreadSafeCounter GCCounter;
+	/** Non zero if GC wants to run but is blocked by some other thread \
+	    This flag is not automatically enforced on the async threads, instead
+			threads have to manually implement support for it. */
+	FThreadSafeCounter GCWantsToRunCounter;
+	/** Critical section for thread safe operations */
 	FCriticalSection Critical;
+
 public:
 	/** Lock on non-game thread. Will block if GC is running. */
 	void LockAsync()
 	{
 		if (!IsInGameThread())
 		{
-			FScopeLock CriticalLock(&Critical);
-
-			// Wait until GC is done if it's currently running
-			FPlatformProcess::ConditionalSleep([&]()
+			// Wait until GC is done if it was running when entering this function
+			int32 OldAsyncCounterValue = 0;
+			do
 			{
-				return GCCounter.GetValue() == 0;
-			});
-
-			AsyncCounter.Increment();
+				FPlatformProcess::ConditionalSleep([&]()
+				{
+					return GCCounter.GetValue() == 0;
+				});
+				{
+					FScopeLock CriticalLock(&Critical);
+					OldAsyncCounterValue = AsyncCounter.GetValue();
+					if (GCCounter.GetValue() == 0)
+					{
+						AsyncCounter.Increment();
+					}
+				}
+			} while (AsyncCounter.GetValue() == OldAsyncCounterValue);
 		}
 	}
 	/** Release lock from non-game thread */
@@ -79,18 +95,33 @@ public:
 	/** Lock for GC. Will block if any other thread has locked. */
 	void GCLock()
 	{
-		FScopeLock CriticalLock(&Critical);
+		// Signal other threads that GC wants to run
+		SetGCIsWaiting();
 
 		// Wait until all other threads are done if they're currently holding the lock
-		FPlatformProcess::ConditionalSleep([&]()
+		int32 OldGCCounterValue = 0;
+		do
 		{
-			return AsyncCounter.GetValue() == 0;
-		});
-
-		GCCounter.Increment();
+			FPlatformProcess::ConditionalSleep([&]()
+			{
+				return AsyncCounter.GetValue() == 0;
+			});
+			{
+				FScopeLock CriticalLock(&Critical);
+				OldGCCounterValue = GCCounter.GetValue();
+				if (AsyncCounter.GetValue() == 0)
+				{
+					int32 GCCounterValue = GCCounter.Increment();
+					check(GCCounterValue == 1); // GCLock doesn't support recursive locks
+					// At this point GC can run so remove the signal that it's waiting
+					FPlatformMisc::MemoryBarrier();
+					ResetGCIsWaiting();
+				}
+			}
+		} while (GCCounter.GetValue() == OldGCCounterValue);
 	}
 	/** Checks if any async thread has a lock */
-	bool IsAsyncLocked()
+	bool IsAsyncLocked() const
 	{
 		return AsyncCounter.GetValue() != 0;
 	}
@@ -112,4 +143,22 @@ public:
 	{
 		GCCounter.Decrement();
 	}
+
+	/** Manually mark GC state as 'waiting to run' */
+	void SetGCIsWaiting()
+	{
+		GCWantsToRunCounter.Increment();
+	}
+
+	/** Manually reset GC 'waiting to run' state*/
+	void ResetGCIsWaiting()
+	{
+		GCWantsToRunCounter.Set(0);
+	}
+
+	/** True if GC wants to run on the game thread but is maybe blocked by some other thread */
+	bool IsGCWaiting() const
+	{
+		return GCWantsToRunCounter.GetValue() > 0;
+	}
 };
diff --git a/Engine/Source/Runtime/CoreUObject/Private/UObject/GarbageCollection.cpp b/Engine/Source/Runtime/CoreUObject/Private/UObject/GarbageCollection.cpp
--- a/Engine/Source/Runtime/CoreUObject/Private/UObject/GarbageCollection.cpp
+++ b/Engine/Source/Runtime/CoreUObject/Private/UObject/GarbageCollection.cpp
@@ -92,6 +92,11 @@ bool IsGarbageCollectionLocked()
 	return GGarbageCollectionGuardCritical.IsAsyncLocked();
 }
 
+bool IsGarbageCollectionWaiting()
+{
+	return GGarbageCollectionGuardCritical.IsGCWaiting();
+}
+
 /** Called on shutdown to free GC memory */
 void CleanupGCArrayPools()
 {
@@ -1334,6 +1339,8 @@ void CollectGarbageInternal(EObjectFlags KeepFlags, bool bPerformFullPurge)
 	SCOPE_TIME_GUARD(TEXT("Collect Garbage"));
 	SCOPED_NAMED_EVENT(CollectGarbageInternal, FColor::Red);
 
+	GGarbageCollectionGuardCritical.ResetGCIsWaiting();
+
 #if defined(WITH_CODE_GUARD_HANDLER) && WITH_CODE_GUARD_HANDLER
 	void CheckImageIntegrityAtRuntime();
 	CheckImageIntegrityAtRuntime();
@@ -1519,29 +1526,50 @@ void CollectGarbageInternal(EObjectFlags KeepFlags, bool bPerformFullPurge)
 	STAT_ADD_CUSTOMMESSAGE_NAME( STAT_NamedMarker, TEXT( "GarbageCollection - End" ) );
 }
 
-void CollectGarbage(EObjectFlags KeepFlags, bool bPerformFullPurge)
+static void AcquireGCLock()
 {
-	// No other thread may be performing UOBject operations while we're running
+	const double StartTime = FPlatformTime::Seconds();
 	GGarbageCollectionGuardCritical.GCLock();
+	const double ElapsedTime = FPlatformTime::Seconds() - StartTime;
+	if (ElapsedTime > 0.001)
+	{
+		UE_LOG(LogGarbage, Warning, TEXT("%f ms for acquiring GC lock"), ElapsedTime * 1000);
+	}
+}
+
+static void ReleaseGCLock()
+{
+	GGarbageCollectionGuardCritical.GCUnlock();
+}
+
+void CollectGarbage(EObjectFlags KeepFlags, bool bPerformFullPurge)
+{
+	// No other thread may be performing UObject operations while we're running
+	AcquireGCLock();
 
 	// Perform actual garbage collection
 	CollectGarbageInternal(KeepFlags, bPerformFullPurge);
 
 	// Other threads are free to use UObjects
-	GGarbageCollectionGuardCritical.GCUnlock();
+	ReleaseGCLock();
 }
 
 bool TryCollectGarbage(EObjectFlags KeepFlags, bool bPerformFullPurge)
 {
-	// No other thread may be performing UOBject operations while we're running
+	// No other thread may be performing UObject operations while we're running
 	bool bCanRunGC = GGarbageCollectionGuardCritical.TryGCLock();
 	if (!bCanRunGC)
 	{
+		// TryGCLock won't set the 'GC is waiting' flag so we do it here manually so that other threads that respect this flag
+		// have the time to finish working with UObjects
+		GGarbageCollectionGuardCritical.SetGCIsWaiting();
 		if (GNumRetriesBeforeForcingGC > 0 && GNumAttemptsSinceLastGC > GNumRetriesBeforeForcingGC)
 		{
-			// Force GC and block main thread
+			// Force GC and block main thread			
+			UE_LOG(LogGarbage, Warning, TEXT("TryCollectGarbage: forcing GC after %d skipped attempts."), GNumAttemptsSinceLastGC);
+			GNumAttemptsSinceLastGC = 0;
+			AcquireGCLock();
 			bCanRunGC = true;
-			UE_LOG(LogGarbage, Warning, TEXT("TryCollectGarbage: forcing GC after %d skipped attempts."), GNumAttemptsSinceLastGC)
 		}
 	}
 	if (bCanRunGC)
@@ -1550,7 +1578,7 @@ bool TryCollectGarbage(EObjectFlags KeepFlags, bool bPerformFullPurge)
 		CollectGarbageInternal(KeepFlags, bPerformFullPurge);
 
 		// Other threads are free to use UObjects
-		GGarbageCollectionGuardCritical.GCUnlock();
+		ReleaseGCLock();
 	}
 	else
 	{
diff --git a/Engine/Source/Runtime/Engine/Private/UnrealEngine.cpp b/Engine/Source/Runtime/Engine/Private/UnrealEngine.cpp
--- a/Engine/Source/Runtime/Engine/Private/UnrealEngine.cpp
+++ b/Engine/Source/Runtime/Engine/Private/UnrealEngine.cpp
@@ -997,6 +997,10 @@ static TAutoConsoleVariable<int32> CVarStressTestGCWhileStreaming(
 	TEXT("gc.StressTestGC"),
 	0,
 	TEXT("If set to 1, the engine will attempt to trigger GC each frame while async loading."));
+static TAutoConsoleVariable<int32> CVarForceCollectGarbageEveryFrame(
+	TEXT("gc.ForceCollectGarbageEveryFrame"),
+	0,
+	TEXT("If set to 1, the engine will force GC each frame."));
 #endif
 
 static TAutoConsoleVariable<int32> CVarCollectGarbageEveryFrame(
@@ -1086,6 +1090,10 @@ void UEngine::ConditionalCollectGarbage()
 		{
 			TryCollectGarbage(GARBAGE_COLLECTION_KEEPFLAGS, true);
 		}
+		else if (CVarForceCollectGarbageEveryFrame.GetValueOnGameThread())
+		{
+			CollectGarbage(GARBAGE_COLLECTION_KEEPFLAGS, true);
+		}
 		else
 #endif
 			if (bFullPurgeTriggered)
