Changed FGCCSyncObject to use events when waiting for GC to finish so that it doesn't spin on non-game threads when GC is running

diff --git a/Engine/Source/Runtime/CoreUObject/Private/UObject/GCScopeLock.h b/Engine/Source/Runtime/CoreUObject/Private/UObject/GCScopeLock.h
--- a/Engine/Source/Runtime/CoreUObject/Private/UObject/GCScopeLock.h
+++ b/Engine/Source/Runtime/CoreUObject/Private/UObject/GCScopeLock.h
@@ -11,6 +11,7 @@
 #include "HAL/PlatformProcess.h"
 #include "Misc/ScopeLock.h"
 #include "HAL/ThreadSafeBool.h"
+#include "HAL/Event.h"
 
 /** Locks all UObject hash tables when performing GC */
 class FGCScopeLock
@@ -58,8 +59,27 @@ class FGCCSyncObject
 	FThreadSafeCounter GCWantsToRunCounter;
 	/** Critical section for thread safe operations */
 	FCriticalSection Critical;
+	/** Event used to block non-game threads when GC is running */
+	FEvent* GCUnlockedEvent;
+
+	/** One and only global instance of FGCCSyncObject */
+	static TUniquePtr<FGCCSyncObject> Singleton;
 
 public:
+
+	FGCCSyncObject();
+	~FGCCSyncObject();
+
+	/** Creates the singleton object */
+	static void Create();
+
+	/** Gets the singleton object */
+	static FGCCSyncObject& Get()
+	{
+		check(Singleton.IsValid());
+		return *Singleton.Get();
+	}
+
 	/** Lock on non-game thread. Will block if GC is running. */
 	void LockAsync()
 	{
@@ -69,10 +89,10 @@ public:
 			int32 OldAsyncCounterValue = 0;
 			do
 			{
-				FPlatformProcess::ConditionalSleep([&]()
+				if (GCCounter.GetValue() > 0)
 				{
-					return GCCounter.GetValue() == 0;
-				});
+					GCUnlockedEvent->Wait();
+				}
 				{
 					FScopeLock CriticalLock(&Critical);
 					OldAsyncCounterValue = AsyncCounter.GetValue();
@@ -111,6 +131,7 @@ public:
 				OldGCCounterValue = GCCounter.GetValue();
 				if (AsyncCounter.GetValue() == 0)
 				{
+					GCUnlockedEvent->Reset();
 					int32 GCCounterValue = GCCounter.Increment();
 					check(GCCounterValue == 1); // GCLock doesn't support recursive locks
 					// At this point GC can run so remove the signal that it's waiting
@@ -133,7 +154,9 @@ public:
 		// If any other thread is currently locking we just exit
 		if (AsyncCounter.GetValue() == 0)
 		{
-			GCCounter.Increment();
+			GCUnlockedEvent->Reset();
+			int32 GCCounterValue = GCCounter.Increment();
+			check(GCCounterValue == 1); // GCLock doesn't support recursive locks
 			bSuccess = true;
 		}
 		return bSuccess;
@@ -141,6 +164,7 @@ public:
 	/** Unlock GC */
 	void GCUnlock()
 	{
+		GCUnlockedEvent->Trigger();
 		GCCounter.Decrement();
 	}
 
diff --git a/Engine/Source/Runtime/CoreUObject/Private/UObject/GarbageCollection.cpp b/Engine/Source/Runtime/CoreUObject/Private/UObject/GarbageCollection.cpp
--- a/Engine/Source/Runtime/CoreUObject/Private/UObject/GarbageCollection.cpp
+++ b/Engine/Source/Runtime/CoreUObject/Private/UObject/GarbageCollection.cpp
@@ -71,15 +71,45 @@ FThreadSafeBool& FGCScopeLock::GetGarbageCollectingFlag()
 	return IsGarbageCollecting;
 }
 
-FGCCSyncObject GGarbageCollectionGuardCritical;
+TUniquePtr<FGCCSyncObject> FGCCSyncObject::Singleton;
+
+FGCCSyncObject::FGCCSyncObject()
+{
+	GCUnlockedEvent = FPlatformProcess::GetSynchEventFromPool(true);
+}
+FGCCSyncObject::~FGCCSyncObject()
+{
+	FPlatformProcess::ReturnSynchEventToPool(GCUnlockedEvent);
+	GCUnlockedEvent = nullptr;
+}
+
+void FGCCSyncObject::Create()
+{
+	check(!Singleton.IsValid());
+	Singleton = MakeUnique<FGCCSyncObject>();
+}
+
+#define UE_LOG_FGCScopeGuard_LockAsync_Time 0
 
 FGCScopeGuard::FGCScopeGuard()
 {
-	GGarbageCollectionGuardCritical.LockAsync();
+#if UE_LOG_FGCScopeGuard_LockAsync_Time
+	const double StartTime = FPlatformTime::Seconds();
+#endif
+	FGCCSyncObject::Get().LockAsync();
+#if UE_LOG_FGCScopeGuard_LockAsync_Time
+	const double ElapsedTime = FPlatformTime::Seconds() - StartTime;
+	if (ElapsedTime > 0.001)
+	{
+		// Note this is expected to take roughly the time it takes to collect garbage and verify GC assumptions, so up to 300ms in development
+		UE_LOG(LogGarbage, Warning, TEXT("%f ms for acquiring ASYNC lock"), ElapsedTime * 1000);
+	}
+#endif
 }
+
 FGCScopeGuard::~FGCScopeGuard()
 {
-	GGarbageCollectionGuardCritical.UnlockAsync();
+	FGCCSyncObject::Get().UnlockAsync();
 }
 
 bool IsGarbageCollecting()
@@ -89,12 +119,12 @@ bool IsGarbageCollecting()
 
 bool IsGarbageCollectionLocked()
 {
-	return GGarbageCollectionGuardCritical.IsAsyncLocked();
+	return FGCCSyncObject::Get().IsAsyncLocked();
 }
 
 bool IsGarbageCollectionWaiting()
 {
-	return GGarbageCollectionGuardCritical.IsGCWaiting();
+	return FGCCSyncObject::Get().IsGCWaiting();
 }
 
 /** Called on shutdown to free GC memory */
@@ -1339,7 +1369,7 @@ void CollectGarbageInternal(EObjectFlags KeepFlags, bool bPerformFullPurge)
 	SCOPE_TIME_GUARD(TEXT("Collect Garbage"));
 	SCOPED_NAMED_EVENT(CollectGarbageInternal, FColor::Red);
 
-	GGarbageCollectionGuardCritical.ResetGCIsWaiting();
+	FGCCSyncObject::Get().ResetGCIsWaiting();
 
 #if defined(WITH_CODE_GUARD_HANDLER) && WITH_CODE_GUARD_HANDLER
 	void CheckImageIntegrityAtRuntime();
@@ -1362,9 +1392,9 @@ void CollectGarbageInternal(EObjectFlags KeepFlags, bool bPerformFullPurge)
 		{
 			UE_LOG(LogGarbage, Log, TEXT("CollectGarbageInternal() is flushing async loading"));
 		}
-		GGarbageCollectionGuardCritical.GCUnlock();
+		FGCCSyncObject::Get().GCUnlock();
 		FlushAsyncLoading();
-		GGarbageCollectionGuardCritical.GCLock();
+		FGCCSyncObject::Get().GCLock();
 	}
 
 	// Route callbacks so we can ensure that we are e.g. not in the middle of loading something by flushing
@@ -1529,7 +1559,7 @@ void CollectGarbageInternal(EObjectFlags KeepFlags, bool bPerformFullPurge)
 static void AcquireGCLock()
 {
 	const double StartTime = FPlatformTime::Seconds();
-	GGarbageCollectionGuardCritical.GCLock();
+	FGCCSyncObject::Get().GCLock();
 	const double ElapsedTime = FPlatformTime::Seconds() - StartTime;
 	if (ElapsedTime > 0.001)
 	{
@@ -1539,7 +1569,7 @@ static void AcquireGCLock()
 
 static void ReleaseGCLock()
 {
-	GGarbageCollectionGuardCritical.GCUnlock();
+	FGCCSyncObject::Get().GCUnlock();
 }
 
 void CollectGarbage(EObjectFlags KeepFlags, bool bPerformFullPurge)
@@ -1557,12 +1587,12 @@ void CollectGarbage(EObjectFlags KeepFlags, bool bPerformFullPurge)
 bool TryCollectGarbage(EObjectFlags KeepFlags, bool bPerformFullPurge)
 {
 	// No other thread may be performing UObject operations while we're running
-	bool bCanRunGC = GGarbageCollectionGuardCritical.TryGCLock();
+	bool bCanRunGC = FGCCSyncObject::Get().TryGCLock();
 	if (!bCanRunGC)
 	{
 		// TryGCLock won't set the 'GC is waiting' flag so we do it here manually so that other threads that respect this flag
 		// have the time to finish working with UObjects
-		GGarbageCollectionGuardCritical.SetGCIsWaiting();
+		FGCCSyncObject::Get().SetGCIsWaiting();
 		if (GNumRetriesBeforeForcingGC > 0 && GNumAttemptsSinceLastGC > GNumRetriesBeforeForcingGC)
 		{
 			// Force GC and block main thread			
diff --git a/Engine/Source/Runtime/CoreUObject/Private/UObject/Obj.cpp b/Engine/Source/Runtime/CoreUObject/Private/UObject/Obj.cpp
--- a/Engine/Source/Runtime/CoreUObject/Private/UObject/Obj.cpp
+++ b/Engine/Source/Runtime/CoreUObject/Private/UObject/Obj.cpp
@@ -42,7 +42,7 @@
 #include "Serialization/BulkData.h"
 #include "UObject/LinkerLoad.h"
 #include "Misc/RedirectCollector.h"
-
+#include "UObject/GCScopeLock.h"
 
 #include "Serialization/ArchiveUObjectFromStructuredArchive.h"
 #include "Serialization/ArchiveDescribeReference.h"
@@ -4133,6 +4133,8 @@ void InitUObject()
 {
 	LLM_SCOPE(ELLMTag::InitUObject);
 
+	FGCCSyncObject::Create();
+
 	// Initialize redirects map
 	for (const TPair<FString,FConfigFile>& It : *GConfig)
 	{
diff --git a/Engine/Source/Runtime/CoreUObject/Private/UObject/SavePackage.cpp b/Engine/Source/Runtime/CoreUObject/Private/UObject/SavePackage.cpp
--- a/Engine/Source/Runtime/CoreUObject/Private/UObject/SavePackage.cpp
+++ b/Engine/Source/Runtime/CoreUObject/Private/UObject/SavePackage.cpp
@@ -3388,9 +3388,6 @@ void VerifyEDLCookInfo()
 	FEDLCookChecker::Verify();
 }
 
-
-extern FGCCSyncObject GGarbageCollectionGuardCritical;
-
 FSavePackageResultStruct UPackage::Save(UPackage* InOuter, UObject* Base, EObjectFlags TopLevelFlags, const TCHAR* Filename,
 	FOutputDevice* Error, FLinkerLoad* Conform, bool bForceByteSwapping, bool bWarnOfLongFilename, uint32 SaveFlags, 
 	const class ITargetPlatform* TargetPlatform, const FDateTime&  FinalTimeStamp, bool bSlowTask, FArchiveDiffMap* InOutDiffMap)
@@ -3683,7 +3680,7 @@ FSavePackageResultStruct UPackage::Save(UPackage* InOuter, UObject* Base, EObjec
 						: bSavingConcurrent(InSavingConcurrent)
 					{
 						// We need the same lock as GC so that no StaticFindObject can happen in parallel to saveing a package
-						GGarbageCollectionGuardCritical.GCLock();
+						FGCCSyncObject::Get().GCLock();
 
 						// Do not change GIsSavingPackage while saving concurrently. It should have been set before and after all packages are saved
 						if (!bSavingConcurrent)
@@ -3697,7 +3694,7 @@ FSavePackageResultStruct UPackage::Save(UPackage* InOuter, UObject* Base, EObjec
 						{
 							GIsSavingPackage = false;
 						}
-						GGarbageCollectionGuardCritical.GCUnlock();
+						FGCCSyncObject::Get().GCUnlock();
 					}
 
 					bool bSavingConcurrent;
