Unhashing unreachable objects (ConditionalBeginDestroy) will now also be done incrementally, just like the purge phase of Garbage Collection

diff --git a/Engine/Source/Runtime/CoreUObject/Private/UObject/GarbageCollection.cpp b/Engine/Source/Runtime/CoreUObject/Private/UObject/GarbageCollection.cpp
--- a/Engine/Source/Runtime/CoreUObject/Private/UObject/GarbageCollection.cpp
+++ b/Engine/Source/Runtime/CoreUObject/Private/UObject/GarbageCollection.cpp
@@ -62,8 +62,16 @@ static bool GObjCurrentPurgeObjectIndexResetPastPermanent = false;
 /** Whether we are currently purging an object in the GC purge pass. */
 static bool GIsPurgingObject = false;
 
+/** Contains a list of objects that stayed marked as unreachable after the last reachability analysis */
+static TArray<FUObjectItem*> GUnreachableObjects;
+static FCriticalSection GUnreachableObjectsCritical;
+static int32 GUnrechableObjectIndex = 0;
+
 /** Helpful constant for determining how many token slots we need to store a pointer **/
 static const uint32 GNumTokensPerPointer = sizeof(void*) / sizeof(uint32); //-V514
+/** Calls ConditionalBeginDestroy on unreachable objects */
+static bool UnhashUnreachableObjects(bool bUseTimeLimit, float TimeLimit = 0.0f);
+
 
 FThreadSafeBool& FGCScopeLock::GetGarbageCollectingFlag()
 {
@@ -951,6 +959,23 @@ public:
 	}
 };
 
+
+static void AcquireGCLock()
+{
+	const double StartTime = FPlatformTime::Seconds();
+	FGCCSyncObject::Get().GCLock();
+	const double ElapsedTime = FPlatformTime::Seconds() - StartTime;
+	if (ElapsedTime > 0.001)
+	{
+		UE_LOG(LogGarbage, Warning, TEXT("%f ms for acquiring GC lock"), ElapsedTime * 1000);
+	}
+}
+
+static void ReleaseGCLock()
+{
+	FGCCSyncObject::Get().GCUnlock();
+}
+
 /**
  * Incrementally purge garbage by deleting all unreferenced objects after routing Destroy.
  *
@@ -1001,20 +1026,29 @@ void IncrementalPurgeGarbage( bool bUseTimeLimit, float TimeLimit )
 
 	} ResetPurgeProgress(bCompleted);
 
-	// Set 'I'm garbage collecting' flag - might be checked inside UObject::Destroy etc.
-	FGCScopeLock GCLock;
-
 	// Keep track of start time to enforce time limit unless bForceFullPurge is true;
 	const double		StartTime							= FPlatformTime::Seconds();
 	bool		bTimeLimitReached							= false;
 	// Depending on platform FPlatformTime::Seconds might take a noticeable amount of time if called thousands of times so we avoid
 	// enforcing the time limit too often, especially as neither Destroy nor actual deletion should take significant
 	// amounts of time.
-	const int32	TimeLimitEnforcementGranularityForDestroy	= 10;
+	const int32	TimeLimitEnforcementGranularityForDestroy	= 10;	
 	const int32	TimeLimitEnforcementGranularityForDeletion	= 100;
 
+	if (GUnrechableObjectIndex < GUnreachableObjects.Num())
+	{
+		AcquireGCLock();
+		bTimeLimitReached = UnhashUnreachableObjects(bUseTimeLimit, TimeLimit);
+		ReleaseGCLock();
+	}
+
+	// Set 'I'm garbage collecting' flag - might be checked inside UObject::Destroy etc.
+	FGCScopeLock GCLock;
+
 	if( !GObjFinishDestroyHasBeenRoutedToAllObjects && !bTimeLimitReached )
 	{
+		check(GUnrechableObjectIndex >= GUnreachableObjects.Num());
+
 		// Try to dispatch all FinishDestroy messages to unreachable objects.  We'll iterate over every
 		// single object and destroy any that are ready to be destroyed.  The objects that aren't yet
 		// ready will be added to a list to be processed afterwards.
@@ -1357,6 +1391,43 @@ struct FScopedCBDProfile
 };
 #endif
 
+void GatherUnreachableObjects(bool bForceSingleThreaded)
+{
+	const double StartTime = FPlatformTime::Seconds();
+
+	GUnreachableObjects.Reset();
+	GUnrechableObjectIndex = 0;
+
+	int32 MaxNumberOfObjects = GUObjectArray.GetObjectArrayNum() - GUObjectArray.GetFirstGCIndex();
+	int32 NumThreads = FMath::Max(1, FTaskGraphInterface::Get().GetNumWorkerThreads());
+	int32 NumberOfObjectsPerThread = (MaxNumberOfObjects / NumThreads) + 1;
+
+	// Iterate over all objects. Note that we iterate over the UObjectArray and usually check only internal flags which
+	// are part of the array so we don't suffer from cache misses as much as we would if we were to check ObjectFlags.
+	ParallelFor(NumThreads, [NumberOfObjectsPerThread, NumThreads, MaxNumberOfObjects](int32 ThreadIndex)
+	{
+		int32 FirstObjectIndex = ThreadIndex * NumberOfObjectsPerThread + GUObjectArray.GetFirstGCIndex();
+		int32 NumObjects = (ThreadIndex < (NumThreads - 1)) ? NumberOfObjectsPerThread : (MaxNumberOfObjects - (NumThreads - 1) * NumberOfObjectsPerThread);
+		int32 LastObjectIndex = FMath::Min(GUObjectArray.GetObjectArrayNum() - 1, FirstObjectIndex + NumObjects - 1);
+		TArray<FUObjectItem*> ThisThreadUnreachableObjects;
+
+		for (int32 ObjectIndex = FirstObjectIndex; ObjectIndex <= LastObjectIndex; ++ObjectIndex)
+		{
+			FUObjectItem* ObjectItem = &GUObjectArray.GetObjectItemArrayUnsafe()[ObjectIndex];
+			if (ObjectItem->IsUnreachable())
+			{
+				ThisThreadUnreachableObjects.Add(ObjectItem);
+			}
+		}
+		if (ThisThreadUnreachableObjects.Num())
+		{
+			FScopeLock UnreachableObjectsLock(&GUnreachableObjectsCritical);
+			GUnreachableObjects.Append(ThisThreadUnreachableObjects);
+		}
+	}, bForceSingleThreaded);
+
+	UE_LOG(LogGarbage, Log, TEXT("%f ms for Gather Unreachable Objects (%d objects collected)"), (FPlatformTime::Seconds() - StartTime) * 1000, GUnreachableObjects.Num());
+}
 
 /** 
  * Deletes all unreferenced objects, keeping objects that have any of the passed in KeepFlags set
@@ -1461,82 +1532,23 @@ void CollectGarbageInternal(EObjectFlags KeepFlags, bool bPerformFullPurge)
 		FCoreUObjectDelegates::PostReachabilityAnalysis.Broadcast();
 
 		{
-			DECLARE_SCOPE_CYCLE_COUNTER(TEXT("CollectGarbageInternal.UnhashUnreachable"), STAT_CollectGarbageInternal_UnhashUnreachable, STATGROUP_GC);
-
-			TGuardValue<bool> GuardObjUnhashUnreachableIsInProgress(GObjUnhashUnreachableIsInProgress, true);
-
-			FCoreUObjectDelegates::PreGarbageCollectConditionalBeginDestroy.Broadcast();
+			DECLARE_SCOPE_CYCLE_COUNTER(TEXT("CollectGarbageInternal.UnhashUnreachable"), STAT_CollectGarbageInternal_UnhashUnreachable, STATGROUP_GC);			
 
 			FGCArrayPool::Get().ClearWeakReferences(bPerformFullPurge);
 
-			// Unhash all unreachable objects.
-			const double StartTime = FPlatformTime::Seconds();
-			int32 ClustersRemoved = 0;
-			int32 Items = 0;
-			int32 ClusterItems = 0;
-			for (FRawObjectIterator It(true); It; ++It)
-			{
-				//@todo UE4 - A prefetch was removed here. Re-add it. It wasn't right anyway, since it was ten items ahead and the consoles on have 8 prefetch slots
-
-				FUObjectItem* ObjectItem = *It;
-				checkSlow(ObjectItem);
-				if (ObjectItem->IsUnreachable())
-				{
-					Items++;
-#if UE_GCCLUSTER_VERBOSE_LOGGING
-					UObject* Object = (UObject*)ObjectItem->Object;
-#endif
-					if ((ObjectItem->GetFlags() & EInternalObjectFlags::ClusterRoot) == EInternalObjectFlags::ClusterRoot)
-					{
-#if UE_GCCLUSTER_VERBOSE_LOGGING
-						UE_LOG(LogGarbage, Log, TEXT("Destroying cluster (%d) %s"), ObjectItem->GetClusterIndex(), *Object->GetFullName());
-#endif
-						// Nuke the entire cluster
-						ObjectItem->ClearFlags(EInternalObjectFlags::ClusterRoot);
-						const int32 ClusterRootIndex = It.GetIndex();
-						FUObjectCluster& Cluster = GUObjectClusters[ObjectItem->GetClusterIndex()];
-						for (int32 ClusterObjectIndex : Cluster.Objects)
-						{
-							FUObjectItem* ClusterObjectItem = GUObjectArray.IndexToObjectUnsafeForGC(ClusterObjectIndex);
-							ClusterObjectItem->SetOwnerIndex(0);
-
-							if (!ClusterObjectItem->HasAnyFlags(EInternalObjectFlags::ReachableInCluster))
-							{
-								ClusterObjectItem->SetFlags(EInternalObjectFlags::Unreachable);
-								if (ClusterObjectIndex < ClusterRootIndex)
-								{
-									UObject* ClusterObject = (UObject*)ClusterObjectItem->Object;
-									FScopedCBDProfile Profile(ClusterObject);
-									ClusterObject->ConditionalBeginDestroy();
-									ClusterItems++;
-								}
-							}
-						}
-						GUObjectClusters.FreeCluster(ObjectItem->GetClusterIndex());
-						ClustersRemoved++;
-					}
-
-					// Begin the object's asynchronous destruction.
-#if !UE_GCCLUSTER_VERBOSE_LOGGING
-					UObject* Object = (UObject*)ObjectItem->Object;
-#endif
-					FScopedCBDProfile Profile(Object);
-					Object->ConditionalBeginDestroy();
-
-					ObjectItem->ClearFlags(EInternalObjectFlags::HadReferenceKilled);
-				}
-
+			GatherUnreachableObjects(bForceSingleThreadedGC);
 
+			if (bPerformFullPurge)
+			{
+				UnhashUnreachableObjects(/**bUseTimeLimit = */ false);
 			}
-
-			UE_LOG(LogGarbage, Log, TEXT("%f ms for unhashing unreachable objects. Clusters removed: %d.   Items %d Cluster Items %d"), (FPlatformTime::Seconds() - StartTime) * 1000, ClustersRemoved, Items, ClusterItems);
-			FCoreUObjectDelegates::PostGarbageCollectConditionalBeginDestroy.Broadcast();
 		}
 		FScopedCBDProfile::DumpProfile();
 		// Set flag to indicate that we are relying on a purge to be performed.
 		GObjPurgeIsRequired = true;
 		// Reset purged count.
 		GPurgedObjectCountSinceLastMarkPhase = 0;
+		GObjCurrentPurgeObjectIndexResetPastPermanent = true;
 
 		// Perform a full purge by not using a time limit for the incremental purge. The Editor always does a full purge.
 		if (bPerformFullPurge || GIsEditor)
@@ -1556,20 +1568,88 @@ void CollectGarbageInternal(EObjectFlags KeepFlags, bool bPerformFullPurge)
 	STAT_ADD_CUSTOMMESSAGE_NAME( STAT_NamedMarker, TEXT( "GarbageCollection - End" ) );
 }
 
-static void AcquireGCLock()
+bool UnhashUnreachableObjects(bool bUseTimeLimit, float TimeLimit)
 {
+	TGuardValue<bool> GuardObjUnhashUnreachableIsInProgress(GObjUnhashUnreachableIsInProgress, true);
+
+	FCoreUObjectDelegates::PreGarbageCollectConditionalBeginDestroy.Broadcast();
+
+	// Unhash all unreachable objects.
 	const double StartTime = FPlatformTime::Seconds();
-	FGCCSyncObject::Get().GCLock();
-	const double ElapsedTime = FPlatformTime::Seconds() - StartTime;
-	if (ElapsedTime > 0.001)
+	const int32 TimeLimitEnforcementGranularityForBeginDestroy = 10;
+	int32 ClustersRemoved = 0;
+	int32 Items = 0;
+	int32 ClusterItems = 0;
+	int32 TimePollCounter = 0;	
+
+	while (GUnrechableObjectIndex < GUnreachableObjects.Num())
 	{
-		UE_LOG(LogGarbage, Warning, TEXT("%f ms for acquiring GC lock"), ElapsedTime * 1000);
-	}
-}
+		//@todo UE4 - A prefetch was removed here. Re-add it. It wasn't right anyway, since it was ten items ahead and the consoles on have 8 prefetch slots
 
-static void ReleaseGCLock()
-{
-	FGCCSyncObject::Get().GCUnlock();
+		FUObjectItem* ObjectItem = GUnreachableObjects[GUnrechableObjectIndex++];
+
+		Items++;
+#if UE_GCCLUSTER_VERBOSE_LOGGING
+		UObject* Object = (UObject*)ObjectItem->Object;
+#endif
+		if ((ObjectItem->GetFlags() & EInternalObjectFlags::ClusterRoot) == EInternalObjectFlags::ClusterRoot)
+		{
+#if UE_GCCLUSTER_VERBOSE_LOGGING
+			UE_LOG(LogGarbage, Log, TEXT("Destroying cluster (%d) %s"), ObjectItem->GetClusterIndex(), *Object->GetFullName());
+#endif
+			// Nuke the entire cluster
+			ObjectItem->ClearFlags(EInternalObjectFlags::ClusterRoot);
+			FUObjectCluster& Cluster = GUObjectClusters[ObjectItem->GetClusterIndex()];
+			for (int32 ClusterObjectIndex : Cluster.Objects)
+			{
+				FUObjectItem* ClusterObjectItem = GUObjectArray.IndexToObjectUnsafeForGC(ClusterObjectIndex);
+				ClusterObjectItem->SetOwnerIndex(0);
+
+				if (!ClusterObjectItem->HasAnyFlags(EInternalObjectFlags::ReachableInCluster))
+				{
+					ClusterObjectItem->SetFlags(EInternalObjectFlags::Unreachable);
+					UObject* ClusterObject = (UObject*)ClusterObjectItem->Object;
+					{
+						FScopedCBDProfile Profile(ClusterObject);
+						ClusterObject->ConditionalBeginDestroy();
+						ClusterItems++;
+						TimePollCounter++;
+					}
+				}
+			}
+			GUObjectClusters.FreeCluster(ObjectItem->GetClusterIndex());
+			ClustersRemoved++;
+		}
+
+		// Begin the object's asynchronous destruction.
+#if !UE_GCCLUSTER_VERBOSE_LOGGING
+		UObject* Object = (UObject*)ObjectItem->Object;
+#endif
+		{
+			FScopedCBDProfile Profile(Object);
+			Object->ConditionalBeginDestroy();
+			ObjectItem->ClearFlags(EInternalObjectFlags::HadReferenceKilled);
+		}
+
+		const bool bPollTimeLimit = ((TimePollCounter++) % TimeLimitEnforcementGranularityForBeginDestroy == 0);
+		if (bUseTimeLimit && bPollTimeLimit && ((FPlatformTime::Seconds() - StartTime) > TimeLimit))
+		{
+			break;
+		}
+	}	
+
+	UE_LOG(LogGarbage, Log, TEXT("%f ms for %sunhashing unreachable objects. Clusters removed: %d. Items %d (%d/%d) Cluster Items %d"), 
+		(FPlatformTime::Seconds() - StartTime) * 1000, 
+		bUseTimeLimit ? TEXT("incrementally ") : TEXT(""),
+		ClustersRemoved,
+		Items, 
+		GUnrechableObjectIndex, GUnreachableObjects.Num(),
+		ClusterItems);
+
+	FCoreUObjectDelegates::PostGarbageCollectConditionalBeginDestroy.Broadcast();
+
+	// Return true if time limit has been reached
+	return GUnrechableObjectIndex < GUnreachableObjects.Num();
 }
 
 void CollectGarbage(EObjectFlags KeepFlags, bool bPerformFullPurge)
diff --git a/Engine/Source/Runtime/CoreUObject/Private/UObject/UObjectHash.cpp b/Engine/Source/Runtime/CoreUObject/Private/UObject/UObjectHash.cpp
--- a/Engine/Source/Runtime/CoreUObject/Private/UObject/UObjectHash.cpp
+++ b/Engine/Source/Runtime/CoreUObject/Private/UObject/UObjectHash.cpp
@@ -366,6 +366,8 @@ static FORCEINLINE int32 GetObjectOuterHash(FName ObjName,PTRINT Outer)
 
 UObject* StaticFindObjectFastExplicitThreadSafe(FUObjectHashTables& ThreadHash, UClass* ObjectClass, FName ObjectName, const FString& ObjectPathName, bool bExactClass, EObjectFlags ExcludeFlags/*=0*/)
 {
+	const EInternalObjectFlags ExclusiveInternalFlags = EInternalObjectFlags::Unreachable;
+
 	// Find an object with the specified name and (optional) class, in any package; if bAnyPackage is false, only matches top-level packages
 	int32 Hash = GetObjectHash(ObjectName);
 	FHashTableLock HashLock(ThreadHash);
@@ -381,6 +383,8 @@ UObject* StaticFindObjectFastExplicitThreadSafe(FUObjectHashTables& ThreadHash,
 				/* Don't return objects that have any of the exclusive flags set */
 				&& !Object->HasAnyFlags(ExcludeFlags)
 
+				&& !Object->HasAnyInternalFlags(ExclusiveInternalFlags)
+
 				/** If a class was specified, check that the object is of the correct class */
 				&& (ObjectClass == nullptr || (bExactClass ? Object->GetClass() == ObjectClass : Object->IsA(ObjectClass)))
 				)
@@ -425,6 +429,8 @@ UObject* StaticFindObjectFastExplicit( UClass* ObjectClass, FName ObjectName, co
 
 UObject* StaticFindObjectFastInternalThreadSafe(FUObjectHashTables& ThreadHash, UClass* ObjectClass, UObject* ObjectPackage, FName ObjectName, bool bExactClass, bool bAnyPackage, EObjectFlags ExcludeFlags, EInternalObjectFlags ExclusiveInternalFlags)
 {
+	ExclusiveInternalFlags |= EInternalObjectFlags::Unreachable;
+
 	// If they specified an outer use that during the hashing
 	UObject* Result = nullptr;
 	if (ObjectPackage != nullptr)
