Fix for a potential assert when incremental purge garbage is pending and something forces a full purge

diff --git a/Engine/Source/Runtime/CoreUObject/Private/UObject/GCScopeLock.h b/Engine/Source/Runtime/CoreUObject/Private/UObject/GCScopeLock.h
--- a/Engine/Source/Runtime/CoreUObject/Private/UObject/GCScopeLock.h
+++ b/Engine/Source/Runtime/CoreUObject/Private/UObject/GCScopeLock.h
@@ -146,6 +146,11 @@ public:
 	{
 		return AsyncCounter.GetValue() != 0;
 	}
+	/** Checks if GC has a lock */
+	bool IsGCLocked() const
+	{
+		return GCCounter.GetValue() != 0;
+	}
 	/** Lock for GC. Will not block and return false if any other thread has already locked. */
 	bool TryGCLock()
 	{
diff --git a/Engine/Source/Runtime/CoreUObject/Private/UObject/GarbageCollection.cpp b/Engine/Source/Runtime/CoreUObject/Private/UObject/GarbageCollection.cpp
--- a/Engine/Source/Runtime/CoreUObject/Private/UObject/GarbageCollection.cpp
+++ b/Engine/Source/Runtime/CoreUObject/Private/UObject/GarbageCollection.cpp
@@ -976,6 +976,28 @@ static void ReleaseGCLock()
 	FGCCSyncObject::Get().GCUnlock();
 }
 
+/** Locks GC within a scope but only if it hasn't been locked already */
+struct FConditionalGCLock
+{
+	bool bNeedsUnlock;
+	FConditionalGCLock()
+		: bNeedsUnlock(false)
+	{
+		if (!FGCCSyncObject::Get().IsGCLocked())
+		{
+			AcquireGCLock();
+			bNeedsUnlock = true;
+		}
+	}
+	~FConditionalGCLock()
+	{
+		if (bNeedsUnlock)
+		{
+			ReleaseGCLock();
+		}
+	}
+};
+
 /**
  * Incrementally purge garbage by deleting all unreferenced objects after routing Destroy.
  *
@@ -1037,9 +1059,8 @@ void IncrementalPurgeGarbage( bool bUseTimeLimit, float TimeLimit )
 
 	if (GUnrechableObjectIndex < GUnreachableObjects.Num())
 	{
-		AcquireGCLock();
+		FConditionalGCLock ScopedGCLock;
 		bTimeLimitReached = UnhashUnreachableObjects(bUseTimeLimit, TimeLimit);
-		ReleaseGCLock();
 	}
 
 	// Set 'I'm garbage collecting' flag - might be checked inside UObject::Destroy etc.
