Fix for a possible hang when worker threads are preventing GC from running and something is later trying to FlushAsyncLoading with the Async Loading Thread enabled

diff --git a/Engine/Config/BaseEngine.ini b/Engine/Config/BaseEngine.ini
--- a/Engine/Config/BaseEngine.ini
+++ b/Engine/Config/BaseEngine.ini
@@ -685,7 +685,7 @@ s.MaxPackageSummarySize=16384
 gc.MaxObjectsNotConsideredByGC=0
 gc.SizeOfPermanentObjectPool=0
 gc.FlushStreamingOnGC=0
-gc.NumRetriesBeforeForcingGC=0
+gc.NumRetriesBeforeForcingGC=10
 gc.AllowParallelGC=True
 ; pick a fractional number to keep phase shifting and avoid collisions
 gc.TimeBetweenPurgingPendingKillObjects=61.1
diff --git a/Engine/Source/Runtime/CoreUObject/Private/UObject/GarbageCollection.cpp b/Engine/Source/Runtime/CoreUObject/Private/UObject/GarbageCollection.cpp
--- a/Engine/Source/Runtime/CoreUObject/Private/UObject/GarbageCollection.cpp
+++ b/Engine/Source/Runtime/CoreUObject/Private/UObject/GarbageCollection.cpp
@@ -1313,7 +1313,7 @@ static FAutoConsoleVariableRef CVarAllowParallelGC(
 static int32 GNumAttemptsSinceLastGC = 0;
 
 // Number of times GC can be skipped.
-static int32 GNumRetriesBeforeForcingGC = 0;
+static int32 GNumRetriesBeforeForcingGC = 10;
 static FAutoConsoleVariableRef CVarNumRetriesBeforeForcingGC(
 	TEXT("gc.NumRetriesBeforeForcingGC"),
 	GNumRetriesBeforeForcingGC,
@@ -1691,9 +1691,6 @@ bool TryCollectGarbage(EObjectFlags KeepFlags, bool bPerformFullPurge)
 	bool bCanRunGC = FGCCSyncObject::Get().TryGCLock();
 	if (!bCanRunGC)
 	{
-		// TryGCLock won't set the 'GC is waiting' flag so we do it here manually so that other threads that respect this flag
-		// have the time to finish working with UObjects
-		FGCCSyncObject::Get().SetGCIsWaiting();
 		if (GNumRetriesBeforeForcingGC > 0 && GNumAttemptsSinceLastGC > GNumRetriesBeforeForcingGC)
 		{
 			// Force GC and block main thread			
