Fix an invalid object assert in the following scenario:

1. ULevelStreaming::RequestLevel() is called.
2. ULevelStreaming::AsyncLevelLoadComplete() is called. The level package is stored in FAsyncLoadingThread::PackagesToDelete.
3. ULevelStreaming::DiscardPendingUnloadLevel() is called.
4. Perform GC. The level package is still stored in the FAsyncLoadingThread::PackagesToDelete due to IsTimeLimitExceeded().
5. PreGarbageCollectDelegate is called. The class of level script actor and its CDO are marked as pending kill in FLevelStreamingGCHelper::PrepareStreamedOutLevelsForGC.
6. During GC, level script actor's CDO's ClassPrivate is cleared to null because it is marked as pending kill.
7. During GC, FAsyncPackage::AddReferencedObjects is called. IsValidLowLevelFast() for level script actor's CDO is failed due to the class is null.

diff --git a/Engine/Source/Runtime/CoreUObject/Private/Serialization/AsyncLoading.cpp b/Engine/Source/Runtime/CoreUObject/Private/Serialization/AsyncLoading.cpp
--- a/Engine/Source/Runtime/CoreUObject/Private/Serialization/AsyncLoading.cpp
+++ b/Engine/Source/Runtime/CoreUObject/Private/Serialization/AsyncLoading.cpp
@@ -5316,11 +5316,7 @@ FAsyncPackage::~FAsyncPackage()
 
 void FAsyncPackage::AddReferencedObjects(FReferenceCollector& Collector)
 {
-	Collector.AllowEliminatingReferences(false);
-	{
-		Collector.AddReferencedObjects(ReferencedObjects);
-	}
-	Collector.AllowEliminatingReferences(true);
+	Collector.AddReferencedObjects(ReferencedObjects);
 }
 
 void FAsyncPackage::AddObjectReference(UObject* InObject)
@@ -5345,11 +5341,13 @@ void FAsyncPackage::EmptyReferencedObjects()
 	FScopeLock ReferencedObjectsLock(&ReferencedObjectsCritical);
 	for (UObject* Obj : ReferencedObjects)
 	{
-		// Temporary fatal messages instead of checks to find the cause for a one-time crash in shipping config
-		UE_CLOG(Obj == nullptr, LogStreaming, Fatal, TEXT("NULL object in Async Objects Referencer"));
-		UE_CLOG(!Obj->IsValidLowLevelFast(), LogStreaming, Fatal, TEXT("Invalid object in Async Objects Referencer"));
-		Obj->AtomicallyClearInternalFlags(AsyncFlags);
-		check(!Obj->HasAnyInternalFlags(AsyncFlags))
+		if (Obj)
+		{
+			// Temporary fatal messages instead of checks to find the cause for a one-time crash in shipping config
+			UE_CLOG(!Obj->IsValidLowLevelFast(), LogStreaming, Fatal, TEXT("Invalid object in Async Objects Referencer"));
+			Obj->AtomicallyClearInternalFlags(AsyncFlags);
+			check(!Obj->HasAnyInternalFlags(AsyncFlags))
+		}
 	}
 	ReferencedObjects.Reset();
 }
@@ -6743,7 +6741,10 @@ void FAsyncPackage::Cancel()
 		const EObjectFlags ObjectLoadFlags = EObjectFlags(RF_NeedLoad | RF_NeedPostLoad | RF_NeedPostLoadSubobjects | RF_WasLoaded);
 		for (UObject* ObjRef : ReferencedObjects)
 		{			
-			ObjRef->AtomicallyClearFlags(ObjectLoadFlags);
+			if (ObjRef)
+			{
+				ObjRef->AtomicallyClearFlags(ObjectLoadFlags);
+			}
 		}
 		// Release references
 		EmptyReferencedObjects();
