TBA: Fixes for import/export name cache and object resolving

diff --git a/Engine/Source/Runtime/CoreUObject/Private/UObject/LinkerLoad.cpp b/Engine/Source/Runtime/CoreUObject/Private/UObject/LinkerLoad.cpp
--- a/Engine/Source/Runtime/CoreUObject/Private/UObject/LinkerLoad.cpp
+++ b/Engine/Source/Runtime/CoreUObject/Private/UObject/LinkerLoad.cpp
@@ -928,16 +928,18 @@ FLinkerLoad::ELinkerStatus FLinkerLoad::CreateLoader(
 					
 					FName FullObjectPath = *InName.Right(InName.Len() - EndOfClassNameIndex - 1);
 					UObject* Result = nullptr;
-					const FPackageIndex* PackageIndex = ObjectNameToPackageIndex.Find(FullObjectPath);
-					if (PackageIndex != nullptr)
+					const FPackageIndex* ExportIndex = ObjectNameToPackageExportIndex.Find(FullObjectPath);
+					if (ExportIndex)
 					{
-						if (PackageIndex->IsExport())
-						{
-							Result = CreateExport(PackageIndex->ToExport());
-						}
-						else if (PackageIndex->IsImport())
+						Result = CreateExport(ExportIndex->ToExport());
+					}
+					else
+					{
+						const FPackageIndex* ImportIndex = ObjectNameToPackageImportIndex.Find(FullObjectPath);
+
+						if (ImportIndex)
 						{
-							Result = CreateImport(PackageIndex->ToImport());
+							Result = CreateImport(ImportIndex->ToImport());
 						}
 					}
 
@@ -1935,7 +1937,8 @@ FLinkerLoad::ELinkerStatus FLinkerLoad::FinalizeCreation()
 
 		if (IsTextFormat())
 		{
-			ObjectNameToPackageIndex.Empty(ExportMap.Num() + ImportMap.Num());
+			ObjectNameToPackageExportIndex.Empty(ExportMap.Num());
+			ObjectNameToPackageImportIndex.Empty(ImportMap.Num());
 
 			// Build full export map name lookup
 			for (int32 ExportIndex = 0; ExportIndex < ExportMap.Num(); ++ExportIndex)
@@ -1952,19 +1955,27 @@ FLinkerLoad::ELinkerStatus FLinkerLoad::FinalizeCreation()
 					}
 					FullName = Export.ObjectName.ToString() + FullName;
 					CurIndex = Export.OuterIndex;
-
 				} while (!CurIndex.IsNull());
 				
-				int32 DotIndex = INDEX_NONE;
-				if (FullName.FindChar('.', DotIndex))
+				// In some cases, the export in this package can be pathed into a different package (map BuiltInData stuff mainly) so we only want to attach
+				// this package path to the name if it doesn't look like a long package name
+				if (FullName[0] != '/')
 				{
-					FullName[DotIndex] = ':';
+					int32 DotIndex = INDEX_NONE;
+					if (FullName.FindChar('.', DotIndex))
+					{
+						FullName[DotIndex] = ':';
+					}
+
+					FullName = LinkerRoot->GetName() + TEXT(".") + FullName;
 				}
-				FullName = LinkerRoot->GetName() + TEXT(".") + FullName;
-				ObjectNameToPackageIndex.Add(*FullName, FPackageIndex::FromExport(ExportIndex));
+
+				check(ObjectNameToPackageExportIndex.Find(*FullName) == nullptr);
+				ObjectNameToPackageExportIndex.Add(*FullName, FPackageIndex::FromExport(ExportIndex));
 			}
 
 			// Same for import map
+			ObjectNameToPackageImportIndex.Reserve(ImportMap.Num());
 			for (int32 ImportIndex = 0; ImportIndex < ImportMap.Num(); ++ImportIndex)
 			{
 				FPackageIndex CurIndex = FPackageIndex::FromImport(ImportIndex);
@@ -1980,7 +1991,10 @@ FLinkerLoad::ELinkerStatus FLinkerLoad::FinalizeCreation()
 					CurIndex = Import.OuterIndex;
 				}
 				while (!CurIndex.IsNull());
-				ObjectNameToPackageIndex.Add(*FullName, FPackageIndex::FromImport(ImportIndex));
+
+				FName FullFName(*FullName);
+				check(!ObjectNameToPackageImportIndex.Contains(FullFName));
+				ObjectNameToPackageImportIndex.Add(FullFName, FPackageIndex::FromImport(ImportIndex));
 			}
 		}
 
@@ -3463,7 +3477,15 @@ void FLinkerLoad::Preload( UObject* Object )
 #if WITH_EDITOR && WITH_TEXT_ARCHIVE_SUPPORT
 						if (IsTextFormat())
 						{
-							FString ObjectName = Object->GetPathName(Object->GetOutermost());
+							const FName* FullObjectPath = ObjectNameToPackageExportIndex.FindKey(Export.ThisIndex);
+							check(FullObjectPath);
+							FString ObjectPath = FullObjectPath->ToString();
+							FString PackagePath = LinkerRoot->GetPathName(nullptr);
+							FString ObjectName = ObjectPath;
+							if (ObjectPath.StartsWith(PackagePath + TEXT(".")))
+							{
+								ObjectName = ObjectPath.Right(ObjectPath.Len() - PackagePath.Len() - 1);
+							}
 							FStructuredArchive::FSlot ExportSlot = StructuredArchiveRootRecord->EnterField(FIELD_NAME(*ObjectName));
 
 							if (bClassSupportsTextFormat)
@@ -4559,13 +4581,10 @@ void FLinkerLoad::Detach()
 		FUObjectThreadContext::Get().DelayedLinkerClosePackages.Remove(this);
 	}
 
-	if (StructuredArchiveFormatter)
-	{
-		delete StructuredArchive;
-		StructuredArchive = nullptr;
-		delete StructuredArchiveFormatter;
-		StructuredArchiveFormatter = nullptr;
-	}
+	delete StructuredArchive;
+	StructuredArchive = nullptr;
+	delete StructuredArchiveFormatter;
+	StructuredArchiveFormatter = nullptr;
 
 	if (Loader)
 	{
diff --git a/Engine/Source/Runtime/CoreUObject/Public/UObject/LinkerLoad.h b/Engine/Source/Runtime/CoreUObject/Public/UObject/LinkerLoad.h
--- a/Engine/Source/Runtime/CoreUObject/Public/UObject/LinkerLoad.h
+++ b/Engine/Source/Runtime/CoreUObject/Public/UObject/LinkerLoad.h
@@ -141,7 +141,8 @@ private:
 	TOptional<FStructuredArchive::FRecord> StructuredArchiveRootRecord;
 
 	/** A map of full object path name to package index. Used with text assets to resolve incoming string names to an export */
-	TMap<FName, FPackageIndex> ObjectNameToPackageIndex;
+	TMap<FName, FPackageIndex> ObjectNameToPackageImportIndex;
+	TMap<FName, FPackageIndex> ObjectNameToPackageExportIndex;
 
 protected:
 
