TBA: Weak/Soft/Lazy pointer serialization changes

* Remove FWeakObjectPtr::Serialize, move it's logic into, and replace usages of with calls to, FArchiveUObject::SerializeWeakObjectPtr(). Ensures that something is always sent to the archive so that structured archives can be kept happy in the future.
* Added Weak/Soft/Lazy pointer handling to the structured archive slot interface and all the formatters. Binary formatters just forward the call onto their inner and text archives store as a string path reference.
* FArchiveUObjectFromStructuredArchive caches all these pointer types and stores indices in the binary block, same as with a UObject*. All pointers are then forwarded to the underlying formatter in one go on finalization.

diff --git a/Engine/Source/Developer/CollisionAnalyzer/Private/CollisionAnalyzer.h b/Engine/Source/Developer/CollisionAnalyzer/Private/CollisionAnalyzer.h
--- a/Engine/Source/Developer/CollisionAnalyzer/Private/CollisionAnalyzer.h
+++ b/Engine/Source/Developer/CollisionAnalyzer/Private/CollisionAnalyzer.h
@@ -131,7 +131,6 @@ struct FCollisionAnalyzerProxyArchive : public FNameAsStringProxyArchive
 
 	 virtual FArchive& operator<< (struct FWeakObjectPtr& Value) override
 	 {
-		 Value.Serialize(*this);
-		 return *this;
+		 return FArchiveUObject::SerializeWeakObjectPtr(*this, Value);
 	 }
 };
\ No newline at end of file
diff --git a/Engine/Source/Developer/HotReload/Private/HotReloadClassReinstancer.cpp b/Engine/Source/Developer/HotReload/Private/HotReloadClassReinstancer.cpp
--- a/Engine/Source/Developer/HotReload/Private/HotReloadClassReinstancer.cpp
+++ b/Engine/Source/Developer/HotReload/Private/HotReloadClassReinstancer.cpp
@@ -185,8 +185,7 @@ void FHotReloadClassReinstancer::SerializeCDOProperties(UObject* InObject, FHotR
 		}
 		FArchive& operator<<(FWeakObjectPtr& WeakObjectPtr) override
 		{
-			WeakObjectPtr.Serialize(*this);
-			return *this;
+			return FArchiveUObject::SerializeWeakObjectPtr(*this, WeakObjectPtr);
 		}
 		/** Archive name, for debugging */
 		virtual FString GetArchiveName() const override { return TEXT("FCDOWriter"); }
@@ -425,8 +424,7 @@ void FHotReloadClassReinstancer::UpdateDefaultProperties()
 		}
 		FArchive& operator<<(FWeakObjectPtr& WeakObjectPtr) override
 		{
-			WeakObjectPtr.Serialize(*this);
-			return *this;
+			return FArchiveUObject::SerializeWeakObjectPtr(*this, WeakObjectPtr);
 		}
 	};
 
diff --git a/Engine/Source/Runtime/Core/Private/Serialization/ArchiveFromStructuredArchive.cpp b/Engine/Source/Runtime/Core/Private/Serialization/ArchiveFromStructuredArchive.cpp
--- a/Engine/Source/Runtime/Core/Private/Serialization/ArchiveFromStructuredArchive.cpp
+++ b/Engine/Source/Runtime/Core/Private/Serialization/ArchiveFromStructuredArchive.cpp
@@ -14,10 +14,11 @@ FArchiveFromStructuredArchive::FArchiveFromStructuredArchive(FStructuredArchive:
 	
 	if (IsTextFormat())
 	{
-		Record = Slot.EnterRecord();
+		Root = Slot.EnterRecord();
+
 		if (IsLoading())
 		{
-			SerializeInternal();
+			Commit();
 		}
 	}
 	else
@@ -28,18 +29,18 @@ FArchiveFromStructuredArchive::FArchiveFromStructuredArchive(FStructuredArchive:
 
 FArchiveFromStructuredArchive::~FArchiveFromStructuredArchive()
 {
-	SerializeInternal();
+	Commit();
 }
 
 void FArchiveFromStructuredArchive::Flush()
 {
-	SerializeInternal();
+	Commit();
 	FArchive::Flush();
 }
 
 bool FArchiveFromStructuredArchive::Close()
 {
-	SerializeInternal();
+	Commit();
 	return FArchive::Close();
 }
 
@@ -130,7 +131,7 @@ FArchive& FArchiveFromStructuredArchive::operator<<(class UObject*& Value)
 			}
 			else
 			{
-				TOptional<FStructuredArchive::FSlot> ObjectsSlot = Record->TryEnterField(FIELD_NAME_TEXT("Objects"), false);
+				TOptional<FStructuredArchive::FSlot> ObjectsSlot = Root->TryEnterField(FIELD_NAME_TEXT("Objects"), false);
 				if (ObjectsSlot.IsSet())
 				{
 					// We know exactly which stream index we want to load here, but because of the API we need to read through them
@@ -225,14 +226,22 @@ void FArchiveFromStructuredArchive::Serialize(void* V, int64 Length)
 	}
 }
 
-void FArchiveFromStructuredArchive::SerializeInternal()
+void FArchiveFromStructuredArchive::Commit()
+{
+	if (IsTextFormat())
+	{
+		SerializeInternal(Root.GetValue());
+	}
+}
+
+void FArchiveFromStructuredArchive::SerializeInternal(FStructuredArchive::FRecord Record)
 {
-	if (bPendingSerialize && IsTextFormat())
+	if (bPendingSerialize)
 	{
-		FStructuredArchive::FSlot DataSlot = Record.GetValue().EnterField(FIELD_NAME_TEXT("Data"));
+		FStructuredArchive::FSlot DataSlot = Record.EnterField(FIELD_NAME_TEXT("Data"));
 		DataSlot.Serialize(Buffer);
 
-		TOptional<FStructuredArchive::FSlot> ObjectsSlot = Record.GetValue().TryEnterField(FIELD_NAME_TEXT("Objects"), Objects.Num() > 0);
+		TOptional<FStructuredArchive::FSlot> ObjectsSlot = Record.TryEnterField(FIELD_NAME_TEXT("Objects"), Objects.Num() > 0);
 		if (ObjectsSlot.IsSet())
 		{
 			if (IsLoading())
@@ -256,7 +265,7 @@ void FArchiveFromStructuredArchive::SerializeInternal()
 			}
 		}
 
-		TOptional<FStructuredArchive::FSlot> NamesSlot = Record.GetValue().TryEnterField(FIELD_NAME_TEXT("Names"), Names.Num() > 0);
+		TOptional<FStructuredArchive::FSlot> NamesSlot = Record.TryEnterField(FIELD_NAME_TEXT("Names"), Names.Num() > 0);
 		if (NamesSlot.IsSet())
 		{
 			NamesSlot.GetValue() << Names;
diff --git a/Engine/Source/Runtime/Core/Private/Serialization/StructuredArchive.cpp b/Engine/Source/Runtime/Core/Private/Serialization/StructuredArchive.cpp
--- a/Engine/Source/Runtime/Core/Private/Serialization/StructuredArchive.cpp
+++ b/Engine/Source/Runtime/Core/Private/Serialization/StructuredArchive.cpp
@@ -346,6 +346,34 @@ void FStructuredArchive::FSlot::operator<< (UObject*& Value)
 	Ar.LeaveSlot();
 }
 
+void FStructuredArchive::FSlot::operator<< (FWeakObjectPtr& Value)
+{
+	Ar.EnterSlot(ElementId);
+	Ar.Formatter.Serialize(Value);
+	Ar.LeaveSlot();
+}
+
+void FStructuredArchive::FSlot::operator<< (FLazyObjectPtr& Value)
+{
+	Ar.EnterSlot(ElementId);
+	Ar.Formatter.Serialize(Value);
+	Ar.LeaveSlot();
+}
+
+void FStructuredArchive::FSlot::operator<< (FSoftObjectPtr& Value)
+{
+	Ar.EnterSlot(ElementId);
+	Ar.Formatter.Serialize(Value);
+	Ar.LeaveSlot();
+}
+
+void FStructuredArchive::FSlot::operator<< (FSoftObjectPath& Value)
+{
+	Ar.EnterSlot(ElementId);
+	Ar.Formatter.Serialize(Value);
+	Ar.LeaveSlot();
+}
+
 void FStructuredArchive::FSlot::Serialize(TArray<uint8>& Data)
 {
 	Ar.EnterSlot(ElementId);
diff --git a/Engine/Source/Runtime/Core/Public/Containers/Map.h b/Engine/Source/Runtime/Core/Public/Containers/Map.h
--- a/Engine/Source/Runtime/Core/Public/Containers/Map.h
+++ b/Engine/Source/Runtime/Core/Public/Containers/Map.h
@@ -532,6 +532,7 @@ public:
 	/** Structured archive serializer. */
 	FORCEINLINE friend void operator<<(FStructuredArchive::FSlot Slot, TMapBase& InMap)
 	{
+		/*
 		if (Slot.GetUnderlyingArchive().IsTextFormat())
 		{
 			int32 Num = InMap.Num();
@@ -563,6 +564,7 @@ public:
 			}
 		}
 		else
+		*/
 		{
 			Slot << InMap.Pairs;
 		}
diff --git a/Engine/Source/Runtime/Core/Public/Serialization/ArchiveFromStructuredArchive.h b/Engine/Source/Runtime/Core/Public/Serialization/ArchiveFromStructuredArchive.h
--- a/Engine/Source/Runtime/Core/Public/Serialization/ArchiveFromStructuredArchive.h
+++ b/Engine/Source/Runtime/Core/Public/Serialization/ArchiveFromStructuredArchive.h
@@ -14,7 +14,7 @@ class CORE_API FArchiveFromStructuredArchive : public FArchiveProxy
 {
 public:
 	FArchiveFromStructuredArchive(FStructuredArchive::FSlot Slot);
-	~FArchiveFromStructuredArchive();
+	virtual ~FArchiveFromStructuredArchive();
 
 	virtual void Flush() override;
 	virtual bool Close() override;
@@ -32,9 +32,15 @@ public:
 
 	virtual FArchive* GetCacheableArchive() override
 	{
-		return IsTextFormat() ? nullptr : Record->GetUnderlyingArchive().GetCacheableArchive();
+		return IsTextFormat() ? nullptr : Root->GetUnderlyingArchive().GetCacheableArchive();
 	}
-	
+
+protected:
+
+	TOptional<FStructuredArchive::FRecord> Root;
+	void Commit();
+	virtual void SerializeInternal(FStructuredArchive::FRecord Record);
+
 private:
 	static const int32 MaxBufferSize = 128;
 
@@ -49,8 +55,4 @@ private:
 	TArray<UObject*> Objects;
 	TBitArray<> ObjectsValid;
 	TMap<UObject*, int32> ObjectToIndex;
-
-	TOptional<FStructuredArchive::FRecord> Record;
-
-	void SerializeInternal();
-};
+};
\ No newline at end of file
diff --git a/Engine/Source/Runtime/Core/Public/Serialization/Formatters/BinaryArchiveFormatter.h b/Engine/Source/Runtime/Core/Public/Serialization/Formatters/BinaryArchiveFormatter.h
--- a/Engine/Source/Runtime/Core/Public/Serialization/Formatters/BinaryArchiveFormatter.h
+++ b/Engine/Source/Runtime/Core/Public/Serialization/Formatters/BinaryArchiveFormatter.h
@@ -49,6 +49,10 @@ public:
 	virtual void Serialize(FString& Value) override;
 	virtual void Serialize(FName& Value) override;
 	virtual void Serialize(UObject*& Value) override;
+	virtual void Serialize(FWeakObjectPtr& Value) override;
+	virtual void Serialize(FSoftObjectPtr& Value) override;
+	virtual void Serialize(FSoftObjectPath& Value) override;
+	virtual void Serialize(FLazyObjectPtr& Value) override;
 	virtual void Serialize(TArray<uint8>& Value) override;
 	virtual void Serialize(void* Data, uint64 DataSize) override;
 
@@ -210,6 +214,26 @@ inline void FBinaryArchiveFormatter::Serialize(UObject*& Value)
 	Inner << Value;
 }
 
+inline void FBinaryArchiveFormatter::Serialize(FWeakObjectPtr& Value)
+{
+	Inner << Value;
+}
+
+inline void FBinaryArchiveFormatter::Serialize(FSoftObjectPtr& Value)
+{
+	Inner << Value;
+}
+
+inline void FBinaryArchiveFormatter::Serialize(FSoftObjectPath& Value)
+{
+	Inner << Value;
+}
+
+inline void FBinaryArchiveFormatter::Serialize(FLazyObjectPtr& Value)
+{
+	Inner << Value;
+}
+
 inline void FBinaryArchiveFormatter::Serialize(TArray<uint8>& Data)
 {
 	Inner << Data;
diff --git a/Engine/Source/Runtime/Core/Public/Serialization/StructuredArchive.h b/Engine/Source/Runtime/Core/Public/Serialization/StructuredArchive.h
--- a/Engine/Source/Runtime/Core/Public/Serialization/StructuredArchive.h
+++ b/Engine/Source/Runtime/Core/Public/Serialization/StructuredArchive.h
@@ -126,6 +126,10 @@ public:
 		void operator << (FString& Value);
 		void operator << (FName& Value);
 		void operator << (UObject*& Value);
+		void operator << (FWeakObjectPtr& Value);
+		void operator << (FSoftObjectPtr& Value);
+		void operator << (FSoftObjectPath& Value);
+		void operator << (FLazyObjectPtr& Value);
 
 		void Serialize(TArray<uint8>& Data);
 		void Serialize(void* Data, uint64 DataSize);
@@ -593,6 +597,26 @@ private:
 		Ar.Formatter.Serialize(Value);
 	}
 
+	FORCEINLINE void FStructuredArchive::FSlot::operator<< (FWeakObjectPtr& Value)
+	{
+		Ar.Formatter.Serialize(Value);
+	}
+
+	FORCEINLINE void FStructuredArchive::FSlot::operator<< (FSoftObjectPath& Value)
+	{
+		Ar.Formatter.Serialize(Value);
+	}
+
+	FORCEINLINE void FStructuredArchive::FSlot::operator<< (FSoftObjectPtr& Value)
+	{
+		Ar.Formatter.Serialize(Value);
+	}
+
+	FORCEINLINE void FStructuredArchive::FSlot::operator<< (FLazyObjectPtr& Value)
+	{
+		Ar.Formatter.Serialize(Value);
+	}
+
 	FORCEINLINE void FStructuredArchive::FSlot::Serialize(TArray<uint8>& Data)
 	{
 		Ar.Formatter.Serialize(Data);
diff --git a/Engine/Source/Runtime/Core/Public/Serialization/StructuredArchiveFormatter.h b/Engine/Source/Runtime/Core/Public/Serialization/StructuredArchiveFormatter.h
--- a/Engine/Source/Runtime/Core/Public/Serialization/StructuredArchiveFormatter.h
+++ b/Engine/Source/Runtime/Core/Public/Serialization/StructuredArchiveFormatter.h
@@ -100,6 +100,10 @@ public:
 	virtual void Serialize(FString& Value) = 0;
 	virtual void Serialize(FName& Value) = 0;
 	virtual void Serialize(UObject*& Value) = 0;
+	virtual void Serialize(struct FWeakObjectPtr& Value) = 0;
+	virtual void Serialize(struct FSoftObjectPtr& Value) = 0;
+	virtual void Serialize(struct FSoftObjectPath& Value) = 0;
+	virtual void Serialize(struct FLazyObjectPtr& Value) = 0;
 	virtual void Serialize(TArray<uint8>& Value) = 0;
 	virtual void Serialize(void* Data, uint64 DataSize) = 0;
 };
diff --git a/Engine/Source/Runtime/CoreUObject/Private/Serialization/ArchiveUObject.cpp b/Engine/Source/Runtime/CoreUObject/Private/Serialization/ArchiveUObject.cpp
--- a/Engine/Source/Runtime/CoreUObject/Private/Serialization/ArchiveUObject.cpp
+++ b/Engine/Source/Runtime/CoreUObject/Private/Serialization/ArchiveUObject.cpp
@@ -90,7 +90,24 @@ FArchive& FArchiveUObject::SerializeSoftObjectPath(FArchive& Ar, FSoftObjectPath
 
 FArchive& FArchiveUObject::SerializeWeakObjectPtr(FArchive& Ar, FWeakObjectPtr& Value)
 {
-	Value.Serialize(Ar);
+	// NOTE: When changing this function, make sure to update the SavePackage.cpp version in the import and export tagger.
+	
+	// We never serialize our reference while the garbage collector is harvesting references
+	// to objects, because we don't want weak object pointers to keep objects from being garbage
+	// collected.  That would defeat the whole purpose of a weak object pointer!
+	// However, when modifying both kinds of references we want to serialize and writeback the updated value.
+	if (!Ar.IsObjectReferenceCollector() || Ar.IsModifyingWeakAndStrongReferences())
+	{
+		UObject* Object = Value.Get(true);
+	
+		Ar << Object;
+	
+		if (Ar.IsLoading() || Ar.IsModifyingWeakAndStrongReferences())
+		{
+			Value = Object;
+		}
+	}
+
 	return Ar;
 }
 
@@ -126,8 +143,7 @@ FArchive& FObjectAndNameAsStringProxyArchive::operator<<(UObject*& Obj)
 
 FArchive& FObjectAndNameAsStringProxyArchive::operator<<(FWeakObjectPtr& Obj)
 {
-	Obj.Serialize(*this);
-	return *this;
+	return FArchiveUObject::SerializeWeakObjectPtr(*this, Obj);
 }
 
 void FSerializedPropertyScope::PushProperty()
diff --git a/Engine/Source/Runtime/CoreUObject/Private/Serialization/ArchiveUObjectFromStructuredArchive.cpp b/Engine/Source/Runtime/CoreUObject/Private/Serialization/ArchiveUObjectFromStructuredArchive.cpp
new file mode 100644
--- /dev/null
+++ b/Engine/Source/Runtime/CoreUObject/Private/Serialization/ArchiveUObjectFromStructuredArchive.cpp
@@ -0,0 +1,164 @@
+// Copyright 1998-2018 Epic Games, Inc. All Rights Reserved.
+
+#include "Serialization/ArchiveUObjectFromStructuredArchive.h"
+
+FArchiveUObjectFromStructuredArchive::FArchiveUObjectFromStructuredArchive(FStructuredArchive::FSlot Slot)
+	: FArchiveFromStructuredArchive(Slot)
+	, bPendingSerialize(true)
+{
+	if (IsTextFormat() && IsLoading())
+	{
+		Commit();
+	}
+}
+
+FArchiveUObjectFromStructuredArchive::~FArchiveUObjectFromStructuredArchive()
+{
+	Commit();
+}
+
+FArchive& FArchiveUObjectFromStructuredArchive::operator<<(FLazyObjectPtr& Value)
+{
+	if (InnerArchive.IsTextFormat())
+	{
+		if (IsLoading())
+		{
+			int32 ObjectIdx = 0;
+			Serialize(&ObjectIdx, sizeof(ObjectIdx));
+			Value = LazyObjectPtrs[ObjectIdx];
+		}
+		else
+		{
+			int32* IdxPtr = LazyObjectPtrToIndex.Find(Value);
+			if (IdxPtr == nullptr)
+			{
+				IdxPtr = &(LazyObjectPtrToIndex.Add(Value));
+				*IdxPtr = LazyObjectPtrs.Add(Value);
+			}
+			Serialize(IdxPtr, sizeof(*IdxPtr));
+		}
+	}
+	else
+	{
+		InnerArchive << Value;
+	}
+	return *this;
+}
+
+FArchive& FArchiveUObjectFromStructuredArchive::operator<<(FSoftObjectPtr& Value)
+{
+	if (InnerArchive.IsTextFormat())
+	{
+		if (IsLoading())
+		{
+			int32 ObjectIdx = 0;
+			Serialize(&ObjectIdx, sizeof(ObjectIdx));
+			Value = SoftObjectPtrs[ObjectIdx];
+		}
+		else
+		{
+			int32* IdxPtr = SoftObjectPtrToIndex.Find(Value);
+			if (IdxPtr == nullptr)
+			{
+				IdxPtr = &(SoftObjectPtrToIndex.Add(Value));
+				*IdxPtr = SoftObjectPtrs.Add(Value);
+			}
+			Serialize(IdxPtr, sizeof(*IdxPtr));
+		}
+	}
+	else
+	{
+		InnerArchive << Value;
+	}
+	return *this;
+}
+
+FArchive& FArchiveUObjectFromStructuredArchive::operator<<(FSoftObjectPath& Value)
+{
+	if (InnerArchive.IsTextFormat())
+	{
+		if (IsLoading())
+		{
+			int32 ObjectIdx = 0;
+			Serialize(&ObjectIdx, sizeof(ObjectIdx));
+			Value = SoftObjectPaths[ObjectIdx];
+		}
+		else
+		{
+			int32* IdxPtr = SoftObjectPathToIndex.Find(Value);
+			if (IdxPtr == nullptr)
+			{
+				IdxPtr = &(SoftObjectPathToIndex.Add(Value));
+				*IdxPtr = SoftObjectPaths.Add(Value);
+			}
+			Serialize(IdxPtr, sizeof(*IdxPtr));
+		}
+	}
+	else
+	{
+		InnerArchive << Value;
+	}
+	return *this;
+}
+
+FArchive& FArchiveUObjectFromStructuredArchive::operator<<(FWeakObjectPtr& Value)
+{
+	if (InnerArchive.IsTextFormat())
+	{
+		if (IsLoading())
+		{
+			int32 ObjectIdx = 0;
+			Serialize(&ObjectIdx, sizeof(ObjectIdx));
+			Value = WeakObjectPtrs[ObjectIdx];
+		}
+		else
+		{
+			int32* IdxPtr = WeakObjectPtrToIndex.Find(Value);
+			if (IdxPtr == nullptr)
+			{
+				IdxPtr = &(WeakObjectPtrToIndex.Add(Value));
+				*IdxPtr = WeakObjectPtrs.Add(Value);
+			}
+			Serialize(IdxPtr, sizeof(*IdxPtr));
+		}
+	}
+	else
+	{
+		InnerArchive << Value;
+	}
+	return *this;
+}
+
+void FArchiveUObjectFromStructuredArchive::SerializeInternal(FStructuredArchive::FRecord Record)
+{
+	FArchiveFromStructuredArchive::SerializeInternal(Record);
+
+	if (bPendingSerialize)
+	{
+		TOptional<FStructuredArchive::FSlot> LazyObjectPtrsSlot = Record.TryEnterField(FIELD_NAME_TEXT("LazyObjectPtrs"), LazyObjectPtrs.Num() > 0);
+		if (LazyObjectPtrsSlot.IsSet())
+		{
+			LazyObjectPtrsSlot.GetValue() << LazyObjectPtrs;
+		}
+
+		TOptional<FStructuredArchive::FSlot> SoftObjectPtrsSlot = Record.TryEnterField(FIELD_NAME_TEXT("SoftObjectPtrs"), SoftObjectPtrs.Num() > 0);
+		if (SoftObjectPtrsSlot.IsSet())
+		{
+			SoftObjectPtrsSlot.GetValue() << SoftObjectPtrs;
+		}
+
+		TOptional<FStructuredArchive::FSlot> SoftObjectPathsSlot = Record.TryEnterField(FIELD_NAME_TEXT("SoftObjectPaths"), SoftObjectPaths.Num() > 0);
+		if (SoftObjectPathsSlot.IsSet())
+		{
+			SoftObjectPathsSlot.GetValue() << SoftObjectPaths;
+		}
+
+		TOptional<FStructuredArchive::FSlot> WeakObjectPtrsSlot = Record.TryEnterField(FIELD_NAME_TEXT("WeakObjectPtrs"), WeakObjectPtrs.Num() > 0);
+		if (WeakObjectPtrsSlot.IsSet())
+		{
+			WeakObjectPtrsSlot.GetValue() << WeakObjectPtrs;
+		}
+
+		bPendingSerialize = true;
+	}
+}
\ No newline at end of file
diff --git a/Engine/Source/Runtime/CoreUObject/Private/Serialization/Formatters/JsonArchiveInputFormatter.cpp b/Engine/Source/Runtime/CoreUObject/Private/Serialization/Formatters/JsonArchiveInputFormatter.cpp
--- a/Engine/Source/Runtime/CoreUObject/Private/Serialization/Formatters/JsonArchiveInputFormatter.cpp
+++ b/Engine/Source/Runtime/CoreUObject/Private/Serialization/Formatters/JsonArchiveInputFormatter.cpp
@@ -6,6 +6,10 @@
 #include "Serialization/JsonSerializer.h"
 #include "Misc/Base64.h"
 #include "Misc/SecureHash.h"
+#include "UObject/SoftObjectPath.h"
+#include "UObject/SoftObjectPtr.h"
+#include "UObject/WeakObjectPtr.h"
+#include "UObject/LazyObjectPtr.h"
 
 #if WITH_TEXT_ARCHIVE_SUPPORT
 
@@ -279,6 +283,50 @@ void FJsonArchiveInputFormatter::Serialize(UObject*& Value)
 	}
 }
 
+void FJsonArchiveInputFormatter::Serialize(FWeakObjectPtr& Value)
+{
+	UObject* Object = nullptr;
+	Serialize(Object);
+	Value = Object;
+}
+
+void FJsonArchiveInputFormatter::Serialize(FSoftObjectPtr& Value)
+{
+	FSoftObjectPath Path;
+	Serialize(Path);
+	Value = Path;
+}
+
+void FJsonArchiveInputFormatter::Serialize(FSoftObjectPath& Value)
+{
+	FString StringValue;
+	const TCHAR Prefix[] = TEXT("Object:");
+	if (ValueStack.Top()->TryGetString(StringValue) && StringValue.StartsWith(Prefix))
+	{
+		Value.SetPath(*StringValue + ARRAY_COUNT(Prefix) - 1);
+	}
+	else
+	{
+		Value.Reset();
+	}
+}
+
+void FJsonArchiveInputFormatter::Serialize(FLazyObjectPtr& Value)
+{
+	FString StringValue;
+	const TCHAR Prefix[] = TEXT("Lazy:");
+	if (ValueStack.Top()->TryGetString(StringValue) && StringValue.StartsWith(Prefix))
+	{
+		FUniqueObjectGuid Guid;
+		Guid.FromString(*StringValue + ARRAY_COUNT(Prefix) - 1);
+		Value = Guid;
+	}
+	else
+	{
+		Value.Reset();
+	}
+}
+
 void FJsonArchiveInputFormatter::Serialize(TArray<uint8>& Data)
 {
 	const FJsonValue& Value = *ValueStack.Top();
diff --git a/Engine/Source/Runtime/CoreUObject/Private/Serialization/Formatters/JsonArchiveOutputFormatter.cpp b/Engine/Source/Runtime/CoreUObject/Private/Serialization/Formatters/JsonArchiveOutputFormatter.cpp
--- a/Engine/Source/Runtime/CoreUObject/Private/Serialization/Formatters/JsonArchiveOutputFormatter.cpp
+++ b/Engine/Source/Runtime/CoreUObject/Private/Serialization/Formatters/JsonArchiveOutputFormatter.cpp
@@ -4,6 +4,10 @@
 #include "Misc/Base64.h"
 #include "Misc/SecureHash.h"
 #include "UObject/Object.h"
+#include "UObject/WeakObjectPtr.h"
+#include "UObject/SoftObjectPtr.h"
+#include "UObject/SoftObjectPath.h"
+#include "UObject/LazyObjectPtr.h"
 
 #if WITH_TEXT_ARCHIVE_SUPPORT
 
@@ -235,6 +239,54 @@ void FJsonArchiveOutputFormatter::Serialize(UObject*& Value)
 	}
 }
 
+void FJsonArchiveOutputFormatter::Serialize(FWeakObjectPtr& Value)
+{
+	if (Value.IsValid())
+	{
+		SerializeStringInternal(FString::Printf(TEXT("Object:%s"), *Value.Get()->GetFullName()));
+	}
+	else
+	{
+		WriteValue(TEXT("null"));
+	}
+}
+
+void FJsonArchiveOutputFormatter::Serialize(FSoftObjectPtr& Value)
+{
+	if (Value.IsValid())
+	{
+		SerializeStringInternal(FString::Printf(TEXT("Object:%s"), *Value.Get()->GetFullName()));
+	}
+	else
+	{
+		WriteValue(TEXT("null"));
+	}
+}
+
+void FJsonArchiveOutputFormatter::Serialize(FSoftObjectPath& Value)
+{
+	if (Value.IsValid())
+	{
+		SerializeStringInternal(FString::Printf(TEXT("Object:%s"), *Value.GetAssetPathName().ToString()));
+	}
+	else
+	{
+		WriteValue(TEXT("null"));
+	}
+}
+
+void FJsonArchiveOutputFormatter::Serialize(FLazyObjectPtr& Value)
+{
+	if (Value.IsValid())
+	{
+		SerializeStringInternal(FString::Printf(TEXT("Lazy:%s"), *Value.GetUniqueID().ToString()));
+	}
+	else
+	{
+		WriteValue(TEXT("null"));
+	}
+}
+
 void FJsonArchiveOutputFormatter::Serialize(TArray<uint8>& Data)
 {
 	Serialize(Data.GetData(), Data.Num());
diff --git a/Engine/Source/Runtime/CoreUObject/Private/Serialization/ObjectReader.cpp b/Engine/Source/Runtime/CoreUObject/Private/Serialization/ObjectReader.cpp
--- a/Engine/Source/Runtime/CoreUObject/Private/Serialization/ObjectReader.cpp
+++ b/Engine/Source/Runtime/CoreUObject/Private/Serialization/ObjectReader.cpp
@@ -50,8 +50,7 @@ FArchive& FObjectReader::operator<<(FSoftObjectPath& Value)
 
 FArchive& FObjectReader::operator<<(FWeakObjectPtr& Value)
 {
-	Value.Serialize(*this);
-	return *this;
+	return FArchiveUObject::SerializeWeakObjectPtr(*this, Value);
 }
 
 FString FObjectReader::GetArchiveName() const
diff --git a/Engine/Source/Runtime/CoreUObject/Private/Serialization/ObjectWriter.cpp b/Engine/Source/Runtime/CoreUObject/Private/Serialization/ObjectWriter.cpp
--- a/Engine/Source/Runtime/CoreUObject/Private/Serialization/ObjectWriter.cpp
+++ b/Engine/Source/Runtime/CoreUObject/Private/Serialization/ObjectWriter.cpp
@@ -43,8 +43,7 @@ FArchive& FObjectWriter::operator<<(FSoftObjectPath& Value)
 
 FArchive& FObjectWriter::operator<<(FWeakObjectPtr& Value)
 {
-	Value.Serialize(*this);
-	return *this;
+	return FArchiveUObject::SerializeWeakObjectPtr(*this, Value);
 }
 
 FString FObjectWriter::GetArchiveName() const
diff --git a/Engine/Source/Runtime/CoreUObject/Private/UObject/CoreNet.cpp b/Engine/Source/Runtime/CoreUObject/Private/UObject/CoreNet.cpp
--- a/Engine/Source/Runtime/CoreUObject/Private/UObject/CoreNet.cpp
+++ b/Engine/Source/Runtime/CoreUObject/Private/UObject/CoreNet.cpp
@@ -381,8 +381,7 @@ FArchive& FNetBitWriter::operator<<(FSoftObjectPath& Value)
 
 FArchive& FNetBitWriter::operator<<(struct FWeakObjectPtr& WeakObjectPtr)
 {
-	WeakObjectPtr.Serialize(*this);
-	return *this;
+	return FArchiveUObject::SerializeWeakObjectPtr(*this, WeakObjectPtr);
 }
 
 // ----------------------------------------------------------------
@@ -430,7 +429,7 @@ FArchive& FNetBitReader::operator<<(FSoftObjectPath& Value)
 
 FArchive& FNetBitReader::operator<<(struct FWeakObjectPtr& WeakObjectPtr)
 {
-	WeakObjectPtr.Serialize(*this);
+	return FArchiveUObject::SerializeWeakObjectPtr(*this, WeakObjectPtr);
 	return *this;
 }
 
diff --git a/Engine/Source/Runtime/CoreUObject/Private/UObject/Obj.cpp b/Engine/Source/Runtime/CoreUObject/Private/UObject/Obj.cpp
--- a/Engine/Source/Runtime/CoreUObject/Private/UObject/Obj.cpp
+++ b/Engine/Source/Runtime/CoreUObject/Private/UObject/Obj.cpp
@@ -1166,11 +1166,7 @@ void UObject::GetPreloadDependencies(TArray<UObject*>& OutDeps)
 	}
 }
 
-void UObject::Serialize(FStructuredArchive::FRecord Record)
-{
-	FArchiveUObjectFromStructuredArchive Ar(Record.EnterField(FIELD_NAME_TEXT("BaseClassAutoGen")));
-	UObject::Serialize(Ar);
-}
+IMPLEMENT_FSTRUCTUREDARCHIVE_SERIALIZER(UObject)
 
 void UObject::Serialize( FArchive& Ar )
 {
diff --git a/Engine/Source/Runtime/CoreUObject/Private/UObject/SavePackage.cpp b/Engine/Source/Runtime/CoreUObject/Private/UObject/SavePackage.cpp
--- a/Engine/Source/Runtime/CoreUObject/Private/UObject/SavePackage.cpp
+++ b/Engine/Source/Runtime/CoreUObject/Private/UObject/SavePackage.cpp
@@ -949,7 +949,7 @@ FArchive& FArchiveSaveTagExports::operator<<(FWeakObjectPtr& Value)
 	}
 	else
 	{
-		Value.Serialize(*this);
+		FArchiveUObject::SerializeWeakObjectPtr(*this, Value);
 	}
 	return *this;
 }
@@ -1163,7 +1163,7 @@ FArchive& FArchiveSaveTagImports::operator<< (struct FWeakObjectPtr& Value)
 	}
 	else
 	{
-		Value.Serialize(*this);
+		FArchiveUObject::SerializeWeakObjectPtr(*this, Value);
 	}
 	return *this;
 }
@@ -5240,7 +5240,6 @@ FSavePackageResultStruct UPackage::Save(UPackage* InOuter, UObject* Base, EObjec
 							if ( Export.Object->HasAnyFlags(RF_ClassDefaultObject) )
 							{
 								FArchiveUObjectFromStructuredArchive Adapter(ExportSlot);
-								Adapter.SetExternalObjectIndicesMap(&Linker->ObjectIndicesMap);
 								Export.Object->GetClass()->SerializeDefaultObject(Export.Object, Adapter);
 							}
 							else
@@ -5260,7 +5259,6 @@ FSavePackageResultStruct UPackage::Save(UPackage* InOuter, UObject* Base, EObjec
 								else
 								{
 									FArchiveUObjectFromStructuredArchive Adapter(ExportSlot);
-									Adapter.SetExternalObjectIndicesMap(&Linker->ObjectIndicesMap);
 									Export.Object->Serialize(Adapter);
 								}
 
@@ -5575,13 +5573,8 @@ FSavePackageResultStruct UPackage::Save(UPackage* InOuter, UObject* Base, EObjec
 					Linker->Detach();
 
 					delete StructuredArchive;
-
-					if (TextFormatArchive)
-					{
-						TextFormatArchive->Flush();
-						TextFormatArchive->Close();
-						delete TextFormatArchive;
-					}
+					delete Formatter;
+					delete TextFormatArchive;
 				}
 				UNCLOCK_CYCLES(Time);
 				UE_CLOG(!(SaveFlags & (SAVE_DiffCallstack | SAVE_DiffOnly)), LogSavePackage, Verbose,  TEXT("Save=%.2fms"), FPlatformTime::ToMilliseconds(Time) );
@@ -5645,6 +5638,10 @@ FSavePackageResultStruct UPackage::Save(UPackage* InOuter, UObject* Base, EObjec
 							}
 						}
 						Linker->Detach();
+
+						delete StructuredArchive;
+						delete Formatter;
+						delete TextFormatArchive;
 					}
 					// Move the temporary file.
 					else
diff --git a/Engine/Source/Runtime/CoreUObject/Private/UObject/WeakObjectPtr.cpp b/Engine/Source/Runtime/CoreUObject/Private/UObject/WeakObjectPtr.cpp
--- a/Engine/Source/Runtime/CoreUObject/Private/UObject/WeakObjectPtr.cpp
+++ b/Engine/Source/Runtime/CoreUObject/Private/UObject/WeakObjectPtr.cpp
@@ -95,26 +95,4 @@ UObject* FWeakObjectPtr::GetEvenIfUnreachable() const
 		Result = static_cast<UObject*>(ObjectItem->Object);
 	}
 	return Result;
-}
-
-void FWeakObjectPtr::Serialize(FArchive& Ar)
-{
-	// NOTE: When changing this function, make sure to update the SavePackage.cpp version in the import and export tagger.
-
-	// We never serialize our reference while the garbage collector is harvesting references
-	// to objects, because we don't want weak object pointers to keep objects from being garbage
-	// collected.  That would defeat the whole purpose of a weak object pointer!
-	// However, when modifying both kinds of references we want to serialize and writeback the updated value.
-	if (!Ar.IsObjectReferenceCollector() || Ar.IsModifyingWeakAndStrongReferences())
-	{
-		UObject* Object = Get(true);
-
-		Ar << Object;
-
-		if (Ar.IsLoading() || Ar.IsModifyingWeakAndStrongReferences())
-		{
-			*this = Object;
-		}
-	}
-}
-
+}
\ No newline at end of file
diff --git a/Engine/Source/Runtime/CoreUObject/Public/Serialization/ArchiveUObjectFromStructuredArchive.h b/Engine/Source/Runtime/CoreUObject/Public/Serialization/ArchiveUObjectFromStructuredArchive.h
--- a/Engine/Source/Runtime/CoreUObject/Public/Serialization/ArchiveUObjectFromStructuredArchive.h
+++ b/Engine/Source/Runtime/CoreUObject/Public/Serialization/ArchiveUObjectFromStructuredArchive.h
@@ -5,98 +5,41 @@
 #include "Serialization/ArchiveFromStructuredArchive.h"
 #include "Serialization/ArchiveUObject.h"
 #include "UObject/ObjectResource.h"
+
+#include "UObject/SoftObjectPath.h"
+#include "UObject/SoftObjectPtr.h"
+#include "UObject/LazyObjectPtr.h"
 #include "UObject/WeakObjectPtr.h"
 
-class FArchiveUObjectFromStructuredArchive : public FArchiveFromStructuredArchive
+class COREUOBJECT_API FArchiveUObjectFromStructuredArchive : public FArchiveFromStructuredArchive
 {
 public:
 
-	FArchiveUObjectFromStructuredArchive(FStructuredArchive::FSlot Slot)
-		: FArchiveFromStructuredArchive(Slot)
-		, ExternalObjectIndicesMap(nullptr)
-	{
-
-	}
-
-	void SetExternalObjectIndicesMap(const TMap<UObject*, FPackageIndex>* InObjectIndicesMap)
-	{
-		ExternalObjectIndicesMap = InObjectIndicesMap;
-	}
+	FArchiveUObjectFromStructuredArchive(FStructuredArchive::FSlot Slot);
+	virtual ~FArchiveUObjectFromStructuredArchive();
 
 	using FArchive::operator<<; // For visibility of the overloads we don't override
 
 	//~ Begin FArchive Interface
-	virtual FArchive& operator<<(FLazyObjectPtr& Value) override
-	{
-		if (InnerArchive.IsTextFormat())
-		{
-			FArchiveUObject::SerializeLazyObjectPtr(*this, Value);
-		}
-		else
-		{
-			InnerArchive << Value;
-		}
-		return *this;
-	}
-
-	virtual FArchive& operator<<(FSoftObjectPtr& Value) override
-	{
-		if (InnerArchive.IsTextFormat())
-		{
-			FArchiveUObject::SerializeSoftObjectPtr(*this, Value);
-		}
-		else
-		{
-			InnerArchive << Value;
-		}
-		return *this;
-	}
-
-	virtual FArchive& operator<<(FSoftObjectPath& Value) override
-	{
-		if (InnerArchive.IsTextFormat())
-		{
-			FArchiveUObject::SerializeSoftObjectPath(*this, Value);
-		}
-		else
-		{
-			InnerArchive << Value;
-		}
-		return *this;
-	}
-
-	virtual FArchive& operator<<(FWeakObjectPtr& Value) override
-	{
-		if (InnerArchive.IsTextFormat())
-		{
-			UObject* Object = Value.IsValid() ? Value.Get() : nullptr;
-
-			if (!IsObjectAllowed(Object))
-			{
-				Object = nullptr;
-			}
-
-			*this << Object;
-
-			if (IsLoading())
-			{
-				Value = Object;
-			}
-		}
-		else
-		{
-			InnerArchive << Value;
-		}
-		return *this;
-	}
+	virtual FArchive& operator<<(FLazyObjectPtr& Value) override;
+	virtual FArchive& operator<<(FSoftObjectPtr& Value) override;
+	virtual FArchive& operator<<(FSoftObjectPath& Value) override;
+	virtual FArchive& operator<<(FWeakObjectPtr& Value) override;
 	//~ End FArchive Interface
 
 private:
 
-	const TMap<UObject *, FPackageIndex>* ExternalObjectIndicesMap;
+	bool bPendingSerialize;
+
+	TArray<FLazyObjectPtr> LazyObjectPtrs;
+	TArray<FWeakObjectPtr> WeakObjectPtrs;
+	TArray<FSoftObjectPtr> SoftObjectPtrs;
+	TArray<FSoftObjectPath> SoftObjectPaths;
+
+	TMap<FLazyObjectPtr, int32> LazyObjectPtrToIndex;
+	TMap<FWeakObjectPtr, int32> WeakObjectPtrToIndex;
+	TMap<FSoftObjectPtr, int32> SoftObjectPtrToIndex;
+	TMap<FSoftObjectPath, int32> SoftObjectPathToIndex;
 
-	bool IsObjectAllowed(UObject* InObject) const
-	{
-		return IsLoading() || (ExternalObjectIndicesMap == nullptr) || (ExternalObjectIndicesMap->Contains(InObject));
-	}
+	virtual void SerializeInternal(FStructuredArchive::FRecord Record) override;
 };
\ No newline at end of file
diff --git a/Engine/Source/Runtime/CoreUObject/Public/Serialization/Formatters/JsonArchiveInputFormatter.h b/Engine/Source/Runtime/CoreUObject/Public/Serialization/Formatters/JsonArchiveInputFormatter.h
--- a/Engine/Source/Runtime/CoreUObject/Public/Serialization/Formatters/JsonArchiveInputFormatter.h
+++ b/Engine/Source/Runtime/CoreUObject/Public/Serialization/Formatters/JsonArchiveInputFormatter.h
@@ -60,6 +60,10 @@ public:
 	virtual void Serialize(FString& Value) override;
 	virtual void Serialize(FName& Value) override;
 	virtual void Serialize(UObject*& Value) override;
+	virtual void Serialize(FWeakObjectPtr& Value) override;
+	virtual void Serialize(FSoftObjectPtr& Value) override;
+	virtual void Serialize(FSoftObjectPath& Value) override;
+	virtual void Serialize(FLazyObjectPtr& Value) override;
 	virtual void Serialize(TArray<uint8>& Value) override;
 	virtual void Serialize(void* Data, uint64 DataSize) override;
 
diff --git a/Engine/Source/Runtime/CoreUObject/Public/Serialization/Formatters/JsonArchiveOutputFormatter.h b/Engine/Source/Runtime/CoreUObject/Public/Serialization/Formatters/JsonArchiveOutputFormatter.h
--- a/Engine/Source/Runtime/CoreUObject/Public/Serialization/Formatters/JsonArchiveOutputFormatter.h
+++ b/Engine/Source/Runtime/CoreUObject/Public/Serialization/Formatters/JsonArchiveOutputFormatter.h
@@ -7,7 +7,7 @@
 
 #if WITH_TEXT_ARCHIVE_SUPPORT
 
-class FJsonArchiveOutputFormatter : public FStructuredArchiveFormatter
+class COREUOBJECT_API FJsonArchiveOutputFormatter : public FStructuredArchiveFormatter
 {
 public:
 	FJsonArchiveOutputFormatter(FArchive& InInner);
@@ -50,6 +50,10 @@ public:
 	virtual void Serialize(FString& Value) override;
 	virtual void Serialize(FName& Value) override;
 	virtual void Serialize(UObject*& Value) override;
+	virtual void Serialize(FWeakObjectPtr& Value) override;
+	virtual void Serialize(FSoftObjectPtr& Value) override;
+	virtual void Serialize(FSoftObjectPath& Value) override;
+	virtual void Serialize(FLazyObjectPtr& Value) override;
 	virtual void Serialize(TArray<uint8>& Value) override;
 	virtual void Serialize(void* Data, uint64 DataSize) override;
 
diff --git a/Engine/Source/Runtime/CoreUObject/Public/UObject/WeakObjectPtr.h b/Engine/Source/Runtime/CoreUObject/Public/UObject/WeakObjectPtr.h
--- a/Engine/Source/Runtime/CoreUObject/Public/UObject/WeakObjectPtr.h
+++ b/Engine/Source/Runtime/CoreUObject/Public/UObject/WeakObjectPtr.h
@@ -143,14 +143,6 @@ public:
 		return uint32(WeakObjectPtr.ObjectIndex ^ WeakObjectPtr.ObjectSerialNumber);
 	}
 
-	/**
-	 * Weak object pointer serialization.  Weak object pointers only have weak references to objects and
-	 * won't serialize the object when gathering references for garbage collection.  So in many cases, you
-	 * don't need to bother serializing weak object pointers.  However, serialization is required if you
-	 * want to load and save your object.
-	 */
-	COREUOBJECT_API void Serialize(FArchive& Ar);
-
 protected:
 
 	FORCEINLINE int32 GetObjectIndex() const
