Fixing ConditionalBeginDestroy profiling after changes to incremental CBD.

diff --git a/Engine/Source/Runtime/CoreUObject/Private/UObject/GarbageCollection.cpp b/Engine/Source/Runtime/CoreUObject/Private/UObject/GarbageCollection.cpp
--- a/Engine/Source/Runtime/CoreUObject/Private/UObject/GarbageCollection.cpp
+++ b/Engine/Source/Runtime/CoreUObject/Private/UObject/GarbageCollection.cpp
@@ -1059,8 +1059,14 @@ void IncrementalPurgeGarbage( bool bUseTimeLimit, float TimeLimit )
 
 	if (GUnrechableObjectIndex < GUnreachableObjects.Num())
 	{
-		FConditionalGCLock ScopedGCLock;
-		bTimeLimitReached = UnhashUnreachableObjects(bUseTimeLimit, TimeLimit);
+		{
+			FConditionalGCLock ScopedGCLock;
+			bTimeLimitReached = UnhashUnreachableObjects(bUseTimeLimit, TimeLimit);
+		}
+		if (GUnrechableObjectIndex >= GUnreachableObjects.Num())
+		{
+			FScopedCBDProfile::DumpProfile();
+		}
 	}
 
 	// Set 'I'm garbage collecting' flag - might be checked inside UObject::Destroy etc.
@@ -1330,88 +1336,6 @@ static FAutoConsoleVariableRef CVarFlushStreamingOnGC(
 	ECVF_Default
 	);
 
-
-#define PROFILE_GCConditionalBeginDestroy 0
-#define PROFILE_GCConditionalBeginDestroyByClass 0
-
-#if PROFILE_GCConditionalBeginDestroy
-
-struct FCBDTime
-{
-	double TotalTime;
-	int32 Items;
-	FCBDTime()
-		: TotalTime(0.0)
-		, Items(0)
-	{
-	}
-
-	bool operator<(const FCBDTime& Other) const
-	{
-		return TotalTime > Other.TotalTime;
-	}
-};
-
-TMap<FName, FCBDTime> CBDTimings;
-TMap<UObject*, FName> CBDNameLookup;
-
-struct FScopedCBDProfile
-{
-	FName Obj;
-	double StartTime;
-
-	FORCEINLINE FScopedCBDProfile(UObject* InObj)
-		: StartTime(FPlatformTime::Seconds())
-	{
-		CBDNameLookup.Add(InObj, InObj->GetFName());
-#if PROFILE_GCConditionalBeginDestroyByClass
-		UObject* Outermost = ((UObject*)InObj->GetClass());
-#else
-		UObject* Outermost = ((UObject*)InObj->GetOutermost());
-#endif
-		Obj = Outermost->GetFName();
-		if (Obj == NAME_None)
-		{
-			Obj = CBDNameLookup.FindRef(Outermost);
-		}
-	}
-	FORCEINLINE ~FScopedCBDProfile()
-	{
-		double ThisTime = FPlatformTime::Seconds() - StartTime;
-		FCBDTime& Rec = CBDTimings.FindOrAdd(Obj);
-		Rec.Items++;
-		Rec.TotalTime += ThisTime;
-	}
-	FORCEINLINE static void DumpProfile()
-	{
-		CBDTimings.ValueSort(TLess<FCBDTime>());
-		int32 NumPrint = 0;
-		for (auto& Item : CBDTimings)
-		{
-			UE_LOG(LogGarbage, Log, TEXT("    %6d cnt %6.2fus per   %6.2fms total  %s"), Item.Value.Items, 1000.0f * 1000.0f * Item.Value.TotalTime / float(Item.Value.Items), 1000.0f * Item.Value.TotalTime, *Item.Key.ToString());
-			if (NumPrint++ > 3000000000)
-			{
-				break;
-			}
-		}
-		CBDTimings.Empty();
-		CBDNameLookup.Empty();
-	}
-};
-
-
-#else
-struct FScopedCBDProfile
-{
-	FORCEINLINE FScopedCBDProfile(UObject*)
-	{
-	}
-	FORCEINLINE static void DumpProfile()
-	{
-	}
-};
-#endif
-
 void GatherUnreachableObjects(bool bForceSingleThreaded)
 {
 	const double StartTime = FPlatformTime::Seconds();
@@ -1562,9 +1486,10 @@ void CollectGarbageInternal(EObjectFlags KeepFlags, bool bPerformFullPurge)
 			if (bPerformFullPurge)
 			{
 				UnhashUnreachableObjects(/**bUseTimeLimit = */ false);
+				FScopedCBDProfile::DumpProfile();
 			}
 		}
-		FScopedCBDProfile::DumpProfile();
+
 		// Set flag to indicate that we are relying on a purge to be performed.
 		GObjPurgeIsRequired = true;
 		// Reset purged count.
diff --git a/Engine/Source/Runtime/CoreUObject/Private/UObject/GarbageCollectionVerification.cpp b/Engine/Source/Runtime/CoreUObject/Private/UObject/GarbageCollectionVerification.cpp
--- a/Engine/Source/Runtime/CoreUObject/Private/UObject/GarbageCollectionVerification.cpp
+++ b/Engine/Source/Runtime/CoreUObject/Private/UObject/GarbageCollectionVerification.cpp
@@ -316,4 +316,27 @@ void VerifyClustersAssumptions()
 	UE_CLOG(NumErrors.GetValue() > 0, LogGarbage, Fatal, TEXT("Encountered %d object(s) breaking GC Clusters assumptions. Please check log for details."), NumErrors.GetValue());
 }
 
-#endif // VERIFY_DISREGARD_GC_ASSUMPTIONS
\ No newline at end of file
+#endif // VERIFY_DISREGARD_GC_ASSUMPTIONS
+
+#if PROFILE_GCConditionalBeginDestroy
+
+TMap<FName, FCBDTime> CBDTimings;
+TMap<UObject*, FName> CBDNameLookup;
+
+void FScopedCBDProfile::DumpProfile()
+{
+	CBDTimings.ValueSort(TLess<FCBDTime>());
+	int32 NumPrint = 0;
+	for (auto& Item : CBDTimings)
+	{
+		UE_LOG(LogGarbage, Log, TEXT("    %6d cnt %6.2fus per   %6.2fms total  %s"), Item.Value.Items, 1000.0f * 1000.0f * Item.Value.TotalTime / float(Item.Value.Items), 1000.0f * Item.Value.TotalTime, *Item.Key.ToString());
+		if (NumPrint++ > 3000000000)
+		{
+			break;
+		}
+	}
+	CBDTimings.Empty();
+	CBDNameLookup.Empty();
+}
+
+#endif // PROFILE_GCConditionalBeginDestroy
\ No newline at end of file
diff --git a/Engine/Source/Runtime/CoreUObject/Private/UObject/GarbageCollectionVerification.h b/Engine/Source/Runtime/CoreUObject/Private/UObject/GarbageCollectionVerification.h
--- a/Engine/Source/Runtime/CoreUObject/Private/UObject/GarbageCollectionVerification.h
+++ b/Engine/Source/Runtime/CoreUObject/Private/UObject/GarbageCollectionVerification.h
@@ -7,6 +7,9 @@
 #pragma once
 
 #include "CoreMinimal.h"
+#include "Stats/Stats.h"
+#include "UObject/UObjectGlobals.h"
+#include "UObject/Class.h"
 
 // UE_BUILD_SHIPPING has GShouldVerifyGCAssumptions=false by default
 #define VERIFY_DISREGARD_GC_ASSUMPTIONS			!(UE_BUILD_SHIPPING || UE_BUILD_TEST)
@@ -20,3 +23,71 @@ void VerifyGCAssumptions();
 void VerifyClustersAssumptions();
 
 #endif // VERIFY_DISREGARD_GC_ASSUMPTIONS
+
+
+#define PROFILE_GCConditionalBeginDestroy 0
+#define PROFILE_GCConditionalBeginDestroyByClass 0
+
+#if PROFILE_GCConditionalBeginDestroy
+
+struct FCBDTime
+{
+	double TotalTime;
+	int32 Items;
+	FCBDTime()
+		: TotalTime(0.0)
+		, Items(0)
+	{
+	}
+
+	bool operator<(const FCBDTime& Other) const
+	{
+		return TotalTime > Other.TotalTime;
+	}
+};
+
+extern TMap<FName, FCBDTime> CBDTimings;
+extern TMap<UObject*, FName> CBDNameLookup;
+
+struct FScopedCBDProfile
+{
+	FName Obj;
+	double StartTime;
+
+	FORCEINLINE FScopedCBDProfile(UObject* InObj)
+		: StartTime(FPlatformTime::Seconds())
+	{
+		CBDNameLookup.Add(InObj, InObj->GetFName());
+#if PROFILE_GCConditionalBeginDestroyByClass
+		UObject* Outermost = ((UObject*)InObj->GetClass());
+#else
+		UObject* Outermost = ((UObject*)InObj->GetOutermost());
+#endif
+		Obj = Outermost->GetFName();
+		if (Obj == NAME_None)
+		{
+			Obj = CBDNameLookup.FindRef(Outermost);
+		}
+	}
+	FORCEINLINE ~FScopedCBDProfile()
+	{
+		double ThisTime = FPlatformTime::Seconds() - StartTime;
+		FCBDTime& Rec = CBDTimings.FindOrAdd(Obj);
+		Rec.Items++;
+		Rec.TotalTime += ThisTime;
+	}
+	static void DumpProfile();
+};
+
+
+#else
+struct FScopedCBDProfile
+{
+	FORCEINLINE FScopedCBDProfile(UObject*)
+	{
+	}
+	FORCEINLINE static void DumpProfile()
+	{
+	}
+};
+#endif
\ No newline at end of file
