TBA: High level tagged property / UObject base class text serialization
 - UObject serialize converted to structured archive
 - Properties written to text individually with text tags, and then binary adapted values
 - Only saves, doesn't load

diff --git a/Engine/Source/Runtime/CoreUObject/Private/UObject/Class.cpp b/Engine/Source/Runtime/CoreUObject/Private/UObject/Class.cpp
--- a/Engine/Source/Runtime/CoreUObject/Private/UObject/Class.cpp
+++ b/Engine/Source/Runtime/CoreUObject/Private/UObject/Class.cpp
@@ -35,6 +35,7 @@
 #include "UObject/PropertyHelper.h"
 #include "UObject/CoreRedirects.h"
 #include "Serialization/ArchiveScriptReferenceCollector.h"
+#include "Serialization/ArchiveUObjectFromStructuredArchive.h"
 #include "UObject/FrameworkObjectVersion.h"
 
 // This flag enables some expensive class tree validation that is meant to catch mutations of 
@@ -891,12 +892,20 @@ void UStruct::SerializeBinEx( FArchive& Ar, void* Data, void const* DefaultData,
 
 void UStruct::SerializeTaggedProperties(FArchive& Ar, uint8* Data, UStruct* DefaultsStruct, uint8* Defaults, const UObject* BreakRecursionIfFullyLoad) const
 {
+	SerializeTaggedProperties(FStructuredArchiveFromArchive(Ar).GetSlot(), Data, DefaultsStruct, Defaults, BreakRecursionIfFullyLoad);
+}
+
+void UStruct::SerializeTaggedProperties(FStructuredArchive::FSlot Slot, uint8* Data, UStruct* DefaultsStruct, uint8* Defaults, const UObject* BreakRecursionIfFullyLoad) const
+{
+	FArchive& UnderlyingArchive = Slot.GetUnderlyingArchive();
 	//SCOPED_LOADTIMER(SerializeTaggedPropertiesTime);
 
 	// Determine if this struct supports optional property guid's (UBlueprintGeneratedClasses Only)
-	const bool bArePropertyGuidsAvailable = (Ar.UE4Ver() >= VER_UE4_PROPERTY_GUID_IN_PROPERTY_TAG) && !FPlatformProperties::RequiresCookedData() && ArePropertyGuidsAvailable();
+	const bool bArePropertyGuidsAvailable = (UnderlyingArchive.UE4Ver() >= VER_UE4_PROPERTY_GUID_IN_PROPERTY_TAG) && !FPlatformProperties::RequiresCookedData() && ArePropertyGuidsAvailable();
 
-	if( Ar.IsLoading() )
+	FStructuredArchive::FStream PropertiesStream = Slot.EnterStream();
+
+	if( UnderlyingArchive.IsLoading() )
 	{
 		// Load tagged properties.
 
@@ -911,7 +920,7 @@ void UStruct::SerializeTaggedProperties(FArchive& Ar, uint8* Data, UStruct* Defa
 		while (1)
 		{
 			FPropertyTag Tag;
-			Ar << Tag;
+			UnderlyingArchive << Tag;
 
 			if( Tag.Name == NAME_None )
 			{
@@ -919,7 +928,7 @@ void UStruct::SerializeTaggedProperties(FArchive& Ar, uint8* Data, UStruct* Defa
 			}
 			if (!Tag.Name.IsValid())
 			{
-				UE_LOG(LogClass, Warning, TEXT("Invalid tag name: struct '%s', archive '%s'"), *GetName(), *Ar.GetArchiveName());
+				UE_LOG(LogClass, Warning, TEXT("Invalid tag name: struct '%s', archive '%s'"), *GetName(), *UnderlyingArchive.GetArchiveName());
 				break;
 			}
 
@@ -928,7 +937,7 @@ void UStruct::SerializeTaggedProperties(FArchive& Ar, uint8* Data, UStruct* Defa
 			{
 				Property = Property->PropertyLinkNext;
 				// Skip over properties that don't need to be serialized.
-				while( Property && !Property->ShouldSerializeValue( Ar ) )
+				while( Property && !Property->ShouldSerializeValue( UnderlyingArchive ) )
 				{
 					Property = Property->PropertyLinkNext;
 				}
@@ -950,7 +959,7 @@ void UStruct::SerializeTaggedProperties(FArchive& Ar, uint8* Data, UStruct* Defa
 			if( Property == nullptr || Property->GetFName() != Tag.Name )
 			{
 				// No need to check redirects on platforms where everything is cooked. Always check for save games
-				if ((!FPlatformProperties::RequiresCookedData() || Ar.IsSaveGame()) && !Ar.HasAnyPortFlags(PPF_DuplicateForPIE|PPF_Duplicate))
+				if ((!FPlatformProperties::RequiresCookedData() || UnderlyingArchive.IsSaveGame()) && !UnderlyingArchive.HasAnyPortFlags(PPF_DuplicateForPIE|PPF_Duplicate))
 				{
 					FName EachName = GetFName();
 					FName PackageName = GetOutermost()->GetFName();
@@ -1015,7 +1024,7 @@ void UStruct::SerializeTaggedProperties(FArchive& Ar, uint8* Data, UStruct* Defa
 
 			// Check if this is a struct property and we have a redirector
 			// No need to check redirects on platforms where everything is cooked. Always check for save games
-			if (!FPlatformProperties::RequiresCookedData() || Ar.IsSaveGame())
+			if (!FPlatformProperties::RequiresCookedData() || UnderlyingArchive.IsSaveGame())
 			{
 				if (Tag.Type == NAME_StructProperty && PropID == NAME_StructProperty)
 				{
@@ -1036,10 +1045,10 @@ void UStruct::SerializeTaggedProperties(FArchive& Ar, uint8* Data, UStruct* Defa
 				}
 			}
 
-			const int64 StartOfProperty = Ar.Tell();
+			const int64 StartOfProperty = UnderlyingArchive.Tell();
 			if( !Property )
 			{
-				//UE_LOG(LogClass, Warning, TEXT("Property %s of %s not found for package:  %s"), *Tag.Name.ToString(), *GetFullName(), *Ar.GetArchiveName() );
+				//UE_LOG(LogClass, Warning, TEXT("Property %s of %s not found for package:  %s"), *Tag.Name.ToString(), *GetFullName(), *UnderlyingArchive.GetArchiveName() );
 			}
 #if WITH_EDITOR
 			else if (BreakRecursionIfFullyLoad && BreakRecursionIfFullyLoad->HasAllFlags(RF_LoadCompleted))
@@ -1055,15 +1064,15 @@ void UStruct::SerializeTaggedProperties(FArchive& Ar, uint8* Data, UStruct* Defa
 			else if( Tag.ArrayIndex >= Property->ArrayDim || Tag.ArrayIndex < 0 )
 			{
 				UE_LOG(LogClass, Warning, TEXT("Array bound exceeded (var %s=%d, exceeds %s [0-%d] in package:  %s"),
-					*Tag.Name.ToString(), Tag.ArrayIndex, *GetName(), Property->ArrayDim-1, *Ar.GetArchiveName());
+					*Tag.Name.ToString(), Tag.ArrayIndex, *GetName(), Property->ArrayDim-1, *UnderlyingArchive.GetArchiveName());
 			}
-			else if( !Property->ShouldSerializeValue(Ar) )
+			else if( !Property->ShouldSerializeValue(UnderlyingArchive) )
 			{
-				UE_CLOG((Ar.IsPersistent() && FPlatformProperties::RequiresCookedData()), LogClass, Warning, TEXT("Skipping saved property %s of %s since it is no longer serializable for asset:  %s. (Maybe resave asset?)"), *Tag.Name.ToString(), *GetName(), *Ar.GetArchiveName() );
+				UE_CLOG((UnderlyingArchive.IsPersistent() && FPlatformProperties::RequiresCookedData()), LogClass, Warning, TEXT("Skipping saved property %s of %s since it is no longer serializable for asset:  %s. (Maybe resave asset?)"), *Tag.Name.ToString(), *GetName(), *UnderlyingArchive.GetArchiveName() );
 			}
 			else
 			{
-				switch (Property->ConvertFromType(Tag, Ar, Data, DefaultsStruct))
+				switch (Property->ConvertFromType(Tag, UnderlyingArchive, Data, DefaultsStruct))
 				{
 					case EConvertFromTypeResult::Converted:
 						bAdvanceProperty = true;
@@ -1072,7 +1081,7 @@ void UStruct::SerializeTaggedProperties(FArchive& Ar, uint8* Data, UStruct* Defa
 					case EConvertFromTypeResult::UseSerializeItem:
 						if (Tag.Type != PropID)
 						{
-							UE_LOG(LogClass, Warning, TEXT("Type mismatch in %s of %s - Previous (%s) Current(%s) for package:  %s"), *Tag.Name.ToString(), *GetName(), *Tag.Type.ToString(), *PropID.ToString(), *Ar.GetArchiveName() );
+							UE_LOG(LogClass, Warning, TEXT("Type mismatch in %s of %s - Previous (%s) Current(%s) for package:  %s"), *Tag.Name.ToString(), *GetName(), *Tag.Type.ToString(), *PropID.ToString(), *UnderlyingArchive.GetArchiveName() );
 						}
 						else
 						{
@@ -1080,8 +1089,8 @@ void UStruct::SerializeTaggedProperties(FArchive& Ar, uint8* Data, UStruct* Defa
 							uint8* DefaultsFromParent = Property->ContainerPtrToValuePtrForDefaults<uint8>(DefaultsStruct, Defaults, Tag.ArrayIndex);
 
 							// This property is ok.
-							Tag.SerializeTaggedProperty(Ar, Property, DestAddress, DefaultsFromParent);
-							bAdvanceProperty = !Ar.IsCriticalError();
+							Tag.SerializeTaggedProperty(UnderlyingArchive, Property, DestAddress, DefaultsFromParent);
+							bAdvanceProperty = !UnderlyingArchive.IsCriticalError();
 						}
 						break;
 
@@ -1095,13 +1104,13 @@ void UStruct::SerializeTaggedProperties(FArchive& Ar, uint8* Data, UStruct* Defa
 
 			if (!bAdvanceProperty)
 			{
-				Ar.Seek(StartOfProperty + Tag.Size);
+				UnderlyingArchive.Seek(StartOfProperty + Tag.Size);
 			}
 		}
 	}
 	else
 	{
-		check(Ar.IsSaving() || Ar.IsCountingMemory());
+		check(UnderlyingArchive.IsSaving() || UnderlyingArchive.IsCountingMemory());
 
 		UScriptStruct* DefaultsScriptStruct = dynamic_cast<UScriptStruct*>(DefaultsStruct);
 
@@ -1109,18 +1118,18 @@ void UStruct::SerializeTaggedProperties(FArchive& Ar, uint8* Data, UStruct* Defa
 		bool bUseAtomicSerialization = false;
 		if (DefaultsScriptStruct)
 		{
-			bUseAtomicSerialization = DefaultsScriptStruct->ShouldSerializeAtomically(Ar);
+			bUseAtomicSerialization = DefaultsScriptStruct->ShouldSerializeAtomically(UnderlyingArchive);
 		}
 
 		// Save tagged properties.
 
 		// Iterate over properties in the order they were linked and serialize them.
-		const FCustomPropertyListNode* CustomPropertyNode = Ar.ArUseCustomPropertyList ? Ar.ArCustomPropertyList : nullptr;
-		for (UProperty* Property = Ar.ArUseCustomPropertyList ? (CustomPropertyNode ? CustomPropertyNode->Property : nullptr) : PropertyLink;
+		const FCustomPropertyListNode* CustomPropertyNode = UnderlyingArchive.ArUseCustomPropertyList ? UnderlyingArchive.ArCustomPropertyList : nullptr;
+		for (UProperty* Property = UnderlyingArchive.ArUseCustomPropertyList ? (CustomPropertyNode ? CustomPropertyNode->Property : nullptr) : PropertyLink;
 			Property;
-			Property = Ar.ArUseCustomPropertyList ? FCustomPropertyListNode::GetNextPropertyAndAdvance(CustomPropertyNode) : Property->PropertyLinkNext)
+			Property = UnderlyingArchive.ArUseCustomPropertyList ? FCustomPropertyListNode::GetNextPropertyAndAdvance(CustomPropertyNode) : Property->PropertyLinkNext)
 		{
-			if( Property->ShouldSerializeValue(Ar) )
+			if( Property->ShouldSerializeValue(UnderlyingArchive) )
 			{
 				const int32 LoopMin = CustomPropertyNode ? CustomPropertyNode->ArrayIndex : 0;
 				const int32 LoopMax = CustomPropertyNode ? LoopMin + 1 : Property->ArrayDim;
@@ -1128,7 +1137,7 @@ void UStruct::SerializeTaggedProperties(FArchive& Ar, uint8* Data, UStruct* Defa
 				{
 					uint8* DataPtr      = Property->ContainerPtrToValuePtr           <uint8>(Data, Idx);
 					uint8* DefaultValue = Property->ContainerPtrToValuePtrForDefaults<uint8>(DefaultsStruct, Defaults, Idx);
-					if( CustomPropertyNode || !Ar.DoDelta() || Ar.IsTransacting() || (!Defaults && !dynamic_cast<const UClass*>(this)) || !Property->Identical( DataPtr, DefaultValue, Ar.GetPortFlags()) )
+					if( CustomPropertyNode || !UnderlyingArchive.DoDelta() || UnderlyingArchive.IsTransacting() || (!Defaults && !dynamic_cast<const UClass*>(this)) || !Property->Identical( DataPtr, DefaultValue, UnderlyingArchive.GetPortFlags()) )
 					{
 						if (bUseAtomicSerialization)
 						{
@@ -1136,58 +1145,70 @@ void UStruct::SerializeTaggedProperties(FArchive& Ar, uint8* Data, UStruct* Defa
 						}
 #if WITH_EDITOR
 						static const FName NAME_PropertySerialize = FName(TEXT("PropertySerialize"));
-						FArchive::FScopeAddDebugData P(Ar, NAME_PropertySerialize);
-						FArchive::FScopeAddDebugData S(Ar, Property->GetFName());
+						FArchive::FScopeAddDebugData P(UnderlyingArchive, NAME_PropertySerialize);
+						FArchive::FScopeAddDebugData S(UnderlyingArchive, Property->GetFName());
 #endif
-						FPropertyTag Tag( Ar, Property, Idx, DataPtr, DefaultValue );
+						FPropertyTag Tag( UnderlyingArchive, Property, Idx, DataPtr, DefaultValue );
 						// If available use the property guid from BlueprintGeneratedClasses, provided we aren't cooking data.
-						if (bArePropertyGuidsAvailable && !Ar.IsCooking())
+						if (bArePropertyGuidsAvailable && !UnderlyingArchive.IsCooking())
 						{
 							const FGuid PropertyGuid = FindPropertyGuidFromName(Tag.Name);
 							Tag.SetPropertyGuid(PropertyGuid);
 						}
-						Ar << Tag;
+
+						FStructuredArchive::FRecord PropertyRecord = PropertiesStream.EnterElement().EnterRecord();
+
+						PropertyRecord << NAMED_FIELD(Tag);
 
 						// need to know how much data this call to SerializeTaggedProperty consumes, so mark where we are
-						int64 DataOffset = Ar.Tell();
+						int64 DataOffset = UnderlyingArchive.Tell();
 
 						// if using it, save the current custom property list and switch to its sub property list (in case of UStruct serialization)
 						const FCustomPropertyListNode* SavedCustomPropertyList = nullptr;
-						if(Ar.ArUseCustomPropertyList && CustomPropertyNode)
+						if(UnderlyingArchive.ArUseCustomPropertyList && CustomPropertyNode)
 						{
-							SavedCustomPropertyList = Ar.ArCustomPropertyList;
-							Ar.ArCustomPropertyList = CustomPropertyNode->SubPropertyList;
+							SavedCustomPropertyList = UnderlyingArchive.ArCustomPropertyList;
+							UnderlyingArchive.ArCustomPropertyList = CustomPropertyNode->SubPropertyList;
 						}
 
-						Tag.SerializeTaggedProperty( Ar, Property, DataPtr, DefaultValue );
+						Tag.SerializeTaggedProperty(PropertyRecord.EnterField(FIELD_NAME_TEXT("Value")), Property, DataPtr, DefaultValue);
 
 						// restore the original custom property list after serializing
 						if (SavedCustomPropertyList)
 						{
-							Ar.ArCustomPropertyList = SavedCustomPropertyList;
+							UnderlyingArchive.ArCustomPropertyList = SavedCustomPropertyList;
 						}
 
 						// set the tag's size
-						Tag.Size = Ar.Tell() - DataOffset;
+						Tag.Size = UnderlyingArchive.Tell() - DataOffset;
 
-						if ( Tag.Size >  0 )
+						if ( Tag.Size >  0 && !UnderlyingArchive.IsTextFormat())
 						{
 							// mark our current location
-							DataOffset = Ar.Tell();
+							DataOffset = UnderlyingArchive.Tell();
 
 							// go back and re-serialize the size now that we know it
-							Ar.Seek(Tag.SizeOffset);
-							Ar << Tag.Size;
+							UnderlyingArchive.Seek(Tag.SizeOffset);
+							UnderlyingArchive << Tag.Size;
 
 							// return to the current location
-							Ar.Seek(DataOffset);
+							UnderlyingArchive.Seek(DataOffset);
 						}
 					}
 				}
 			}
 		}
-		static FName Temp(NAME_None);
-		Ar << Temp;
+
+		if (UnderlyingArchive.IsTextFormat())
+		{
+			FPropertyTag EmptyTag;
+			PropertiesStream.EnterElement().EnterRecord() << NAMED_ITEM("Tag", EmptyTag);
+		}
+		else
+		{
+			static FName Temp(NAME_None);
+			UnderlyingArchive << Temp;
+		}
 	}
 }
 void UStruct::FinishDestroy()
diff --git a/Engine/Source/Runtime/CoreUObject/Private/UObject/LazyObjectPtr.cpp b/Engine/Source/Runtime/CoreUObject/Private/UObject/LazyObjectPtr.cpp
--- a/Engine/Source/Runtime/CoreUObject/Private/UObject/LazyObjectPtr.cpp
+++ b/Engine/Source/Runtime/CoreUObject/Private/UObject/LazyObjectPtr.cpp
@@ -89,16 +89,18 @@ FThreadSafeCounter FUniqueObjectGuid::CurrentAnnotationTag(1);
 	FLazyObjectPtr
 -------------------------------------------------------------------------------------------------------------*/
 
-void FLazyObjectPtr::PossiblySerializeObjectGuid(UObject *Object, FArchive& Ar)
+void FLazyObjectPtr::PossiblySerializeObjectGuid(UObject *Object, FStructuredArchive::FRecord Record)
 {
-	if (Ar.IsSaving() || Ar.IsCountingMemory())
+	FArchive& UnderlyingArchive = Record.GetUnderlyingArchive();
+
+	if (UnderlyingArchive.IsSaving() || UnderlyingArchive.IsCountingMemory())
 	{
 		FUniqueObjectGuid Guid = GuidAnnotation.GetAnnotation(Object);
 		bool HasGuid = Guid.IsValid();
-		Ar << HasGuid;
+		Record << NAMED_FIELD(HasGuid);
 		if (HasGuid)
 		{
-			if (Ar.GetPortFlags() & PPF_DuplicateForPIE)
+			if (UnderlyingArchive.GetPortFlags() & PPF_DuplicateForPIE)
 			{
 				check(GPlayInEditorID != -1);
 				FGuid &FoundGuid = PIEGuidMap[GPlayInEditorID % MAX_PIE_INSTANCES].FindOrAdd(Guid.GetGuid());
@@ -112,22 +114,22 @@ void FLazyObjectPtr::PossiblySerializeObjectGuid(UObject *Object, FArchive& Ar)
 				}
 			}
 
-			Ar << Guid;
+			Record << NAMED_FIELD(Guid);
 		}
 	}
-	else if (Ar.IsLoading())
+	else if (UnderlyingArchive.IsLoading())
 	{
 		bool HasGuid = false;
-		Ar << HasGuid;
+		Record << NAMED_FIELD(HasGuid);
 		if (HasGuid)
 		{
 			FUniqueObjectGuid Guid;
-			Ar << Guid;
+			Record << NAMED_FIELD(Guid);
 
 			// Don't try and resolve GUIDs when loading a package for diffing
 			const UPackage* Package = Object->GetOutermost();
 			const bool bLoadedForDiff = Package->HasAnyPackageFlags(PKG_ForDiffing);
-			if (!bLoadedForDiff && (!(Ar.GetPortFlags() & PPF_Duplicate) || (Ar.GetPortFlags() & PPF_DuplicateForPIE)))
+			if (!bLoadedForDiff && (!(UnderlyingArchive.GetPortFlags() & PPF_Duplicate) || (UnderlyingArchive.GetPortFlags() & PPF_DuplicateForPIE)))
 			{
 				check(!Guid.IsDefault());
 				UObject* OtherObject = Guid.ResolveObject();
diff --git a/Engine/Source/Runtime/CoreUObject/Private/UObject/Obj.cpp b/Engine/Source/Runtime/CoreUObject/Private/UObject/Obj.cpp
--- a/Engine/Source/Runtime/CoreUObject/Private/UObject/Obj.cpp
+++ b/Engine/Source/Runtime/CoreUObject/Private/UObject/Obj.cpp
@@ -1166,10 +1166,12 @@ void UObject::GetPreloadDependencies(TArray<UObject*>& OutDeps)
 	}
 }
 
-IMPLEMENT_FSTRUCTUREDARCHIVE_SERIALIZER(UObject)
+IMPLEMENT_FARCHIVE_SERIALIZER(UObject)
 
-void UObject::Serialize( FArchive& Ar )
+void UObject::Serialize( FStructuredArchive::FRecord Record )
 {
+	FArchive& UnderlyingArchive = Record.GetUnderlyingArchive();
+
 	// These three items are very special items from a serialization standpoint. They aren't actually serialized.
 	UClass *ObjClass = GetClass();
 	UObject* LoadOuter = GetOuter();
@@ -1178,7 +1180,7 @@ void UObject::Serialize( FArchive& Ar )
 	// Make sure this object's class's data is loaded.
 	if(ObjClass->HasAnyFlags(RF_NeedLoad) )
 	{
-		Ar.Preload(ObjClass);
+		UnderlyingArchive.Preload(ObjClass);
 
 		// make sure this object's template data is loaded - the only objects
 		// this should actually affect are those that don't have any defaults
@@ -1187,31 +1189,31 @@ void UObject::Serialize( FArchive& Ar )
 		// we've hit this code.
 		if ( !HasAnyFlags(RF_ClassDefaultObject) && ObjClass->GetDefaultsCount() > 0 )
 		{
-			Ar.Preload(ObjClass->GetDefaultObject());
+			UnderlyingArchive.Preload(ObjClass->GetDefaultObject());
 		}
 	}
 
 	// Special info.
-	if( (!Ar.IsLoading() && !Ar.IsSaving() && !Ar.IsObjectReferenceCollector()) )
+	if( (!UnderlyingArchive.IsLoading() && !UnderlyingArchive.IsSaving() && !UnderlyingArchive.IsObjectReferenceCollector()) )
 	{
-		Ar << LoadName;
-		if(!Ar.IsIgnoringOuterRef())
+		Record << NAMED_FIELD(LoadName);
+		if(!UnderlyingArchive.IsIgnoringOuterRef())
 		{
-			Ar << LoadOuter;
+			Record << NAMED_FIELD(LoadOuter);
 		}
-		if ( !Ar.IsIgnoringClassRef() )
+		if ( !UnderlyingArchive.IsIgnoringClassRef() )
 		{
-			Ar << ObjClass;
+			Record << NAMED_FIELD(ObjClass);
 		}
 	}
 	// Special support for supporting undo/redo of renaming and changing Archetype.
-	else if( Ar.IsTransacting() )
+	else if( UnderlyingArchive.IsTransacting() )
 	{
-		if(!Ar.IsIgnoringOuterRef())
+		if(!UnderlyingArchive.IsIgnoringOuterRef())
 		{
-			if(Ar.IsLoading())
+			if(UnderlyingArchive.IsLoading())
 			{
-				Ar << LoadName << LoadOuter;
+				Record << NAMED_FIELD(LoadName) << NAMED_FIELD(LoadOuter);
 
 				// If the name we loaded is different from the current one,
 				// unhash the object, change the name and hash it again.
@@ -1224,7 +1226,7 @@ void UObject::Serialize( FArchive& Ar )
 			}
 			else
 			{
-				Ar << LoadName << LoadOuter;
+				Record << NAMED_FIELD(LoadName) << NAMED_FIELD(LoadOuter);
 			}
 		}
 	}
@@ -1233,16 +1235,16 @@ void UObject::Serialize( FArchive& Ar )
 	// Handle derived UClass objects (exact UClass objects are native only and shouldn't be touched)
 	if (ObjClass != UClass::StaticClass())
 	{
-		SerializeScriptProperties(Ar);
+		SerializeScriptProperties(Record.EnterField(FIELD_NAME_TEXT("Properties")));
 	}
 
 	// Keep track of pending kill
-	if( Ar.IsTransacting() )
+	if( UnderlyingArchive.IsTransacting() )
 	{
 		bool WasKill = IsPendingKill();
-		if( Ar.IsLoading() )
+		if( UnderlyingArchive.IsLoading() )
 		{
-			Ar << WasKill;
+			Record << NAMED_FIELD(WasKill);
 			if (WasKill)
 			{
 				MarkPendingKill();
@@ -1252,79 +1254,92 @@ void UObject::Serialize( FArchive& Ar )
 				ClearPendingKill();
 			}
 		}
-		else if( Ar.IsSaving() )
+		else if( UnderlyingArchive.IsSaving() )
 		{
-			Ar << WasKill;
+			Record << NAMED_FIELD(WasKill);
 		}
 	}
 
 	// Serialize a GUID if this object has one mapped to it
-	FLazyObjectPtr::PossiblySerializeObjectGuid(this, Ar);
+	FLazyObjectPtr::PossiblySerializeObjectGuid(this, Record);
 
 	// Invalidate asset pointer caches when loading a new object
-	if (Ar.IsLoading() )
+	if (UnderlyingArchive.IsLoading() )
 	{
 		FSoftObjectPath::InvalidateTag();
 	}
 
 	// Memory counting (with proper alignment to match C++)
 	SIZE_T Size = GetClass()->GetStructureSize();
-	Ar.CountBytes( Size, Size );
+	UnderlyingArchive.CountBytes( Size, Size );
 }
 
+void UObject::SerializeScriptProperties(FArchive& Ar) const
+{
+	SerializeScriptProperties(FStructuredArchiveFromArchive(Ar).GetSlot());
+}
 
-
-void UObject::SerializeScriptProperties( FArchive& Ar ) const
+void UObject::SerializeScriptProperties( FStructuredArchive::FSlot Slot ) const
 {
-	Ar.MarkScriptSerializationStart(this);
+	FArchive& UnderlyingArchive = Slot.GetUnderlyingArchive();
+
+	UnderlyingArchive.MarkScriptSerializationStart(this);
 	if( HasAnyFlags(RF_ClassDefaultObject) )
 	{
-		Ar.StartSerializingDefaults();
+		UnderlyingArchive.StartSerializingDefaults();
 	}
 
 	UClass *ObjClass = GetClass();
+	bool bSlotEntered = false;
 
-	if( (Ar.IsLoading() || Ar.IsSaving()) && !Ar.WantBinaryPropertySerialization() )
+	if(UnderlyingArchive.IsTextFormat() || ((UnderlyingArchive.IsLoading() || UnderlyingArchive.IsSaving()) && !UnderlyingArchive.WantBinaryPropertySerialization()))
 	{
 		//@todoio GetArchetype is pathological for blueprint classes and the event driven loader; the EDL already knows what the archetype is; just calling this->GetArchetype() tries to load some other stuff.
-		UObject* DiffObject = Ar.GetArchetypeFromLoader(this);
+		UObject* DiffObject = UnderlyingArchive.GetArchetypeFromLoader(this);
 		if (!DiffObject)
 		{
 			DiffObject = GetArchetype();
 		}
 #if WITH_EDITOR
 		static const FBoolConfigValueHelper BreakSerializationRecursion(TEXT("StructSerialization"), TEXT("BreakSerializationRecursion"));
-		const bool bBreakSerializationRecursion = BreakSerializationRecursion && Ar.IsLoading() && Ar.GetLinker();
+		const bool bBreakSerializationRecursion = BreakSerializationRecursion && UnderlyingArchive.IsLoading() && UnderlyingArchive.GetLinker();
 #else 
 		const bool bBreakSerializationRecursion = false;
 #endif
 #if WITH_EDITOR
 		static const FName NAME_SerializeScriptProperties = FName(TEXT("SerializeScriptProperties"));
-		FArchive::FScopeAddDebugData P(Ar, NAME_SerializeScriptProperties);
-		FArchive::FScopeAddDebugData S(Ar, ObjClass->GetFName());
+		FArchive::FScopeAddDebugData P(UnderlyingArchive, NAME_SerializeScriptProperties);
+		FArchive::FScopeAddDebugData S(UnderlyingArchive, ObjClass->GetFName());
 #endif
-		ObjClass->SerializeTaggedProperties(Ar, (uint8*)this, HasAnyFlags(RF_ClassDefaultObject) ? ObjClass->GetSuperClass() : ObjClass, (uint8*)DiffObject, bBreakSerializationRecursion ? this : NULL);
+
+		ObjClass->SerializeTaggedProperties(Slot, (uint8*)this, HasAnyFlags(RF_ClassDefaultObject) ? ObjClass->GetSuperClass() : ObjClass, (uint8*)DiffObject, bBreakSerializationRecursion ? this : NULL);
+		bSlotEntered = true;
 	}
-	else if ( Ar.GetPortFlags() != 0 && !Ar.ArUseCustomPropertyList )
+	else if (UnderlyingArchive.GetPortFlags() != 0 && !UnderlyingArchive.ArUseCustomPropertyList )
 	{
 		//@todoio GetArchetype is pathological for blueprint classes and the event driven loader; the EDL already knows what the archetype is; just calling this->GetArchetype() tries to load some other stuff.
-		UObject* DiffObject = Ar.GetArchetypeFromLoader(this);
+		UObject* DiffObject = UnderlyingArchive.GetArchetypeFromLoader(this);
 		if (!DiffObject)
 		{
 			DiffObject = GetArchetype();
 		}
-		ObjClass->SerializeBinEx( Ar, const_cast<UObject *>(this), DiffObject, DiffObject ? DiffObject->GetClass() : NULL );
+		ObjClass->SerializeBinEx( UnderlyingArchive, const_cast<UObject *>(this), DiffObject, DiffObject ? DiffObject->GetClass() : NULL );
 	}
 	else
 	{
-		ObjClass->SerializeBin( Ar, const_cast<UObject *>(this) );
+		ObjClass->SerializeBin(UnderlyingArchive, const_cast<UObject *>(this) );
 	}
 
 	if( HasAnyFlags(RF_ClassDefaultObject) )
 	{
-		Ar.StopSerializingDefaults();
+		UnderlyingArchive.StopSerializingDefaults();
+	}
+	UnderlyingArchive.MarkScriptSerializationEnd(this);
+
+	if (!bSlotEntered)
+	{
+		Slot.EnterStream();
 	}
-	Ar.MarkScriptSerializationEnd(this);
 }
 
 
diff --git a/Engine/Source/Runtime/CoreUObject/Private/UObject/PropertyTag.cpp b/Engine/Source/Runtime/CoreUObject/Private/UObject/PropertyTag.cpp
--- a/Engine/Source/Runtime/CoreUObject/Private/UObject/PropertyTag.cpp
+++ b/Engine/Source/Runtime/CoreUObject/Private/UObject/PropertyTag.cpp
@@ -3,6 +3,7 @@
 #include "UObject/PropertyTag.h"
 #include "UObject/DebugSerializationFlags.h"
 #include "Serialization/SerializedPropertyScope.h"
+#include "Serialization/ArchiveUObjectFromStructuredArchive.h"
 #include "UObject/UnrealType.h"
 #include "UObject/EnumProperty.h"
 #include "UObject/BlueprintsObjectVersion.h"
@@ -92,99 +93,118 @@ void FPropertyTag::SetPropertyGuid(const FGuid& InPropertyGuid)
 }
 
 // Serializer.
-FArchive& operator<<( FArchive& Ar, FPropertyTag& Tag )
+FArchive& operator<<(FArchive& Ar, FPropertyTag& Tag)
 {
+	FStructuredArchiveFromArchive(Ar).GetSlot() << Tag;
+	return Ar;
+}
+
+void operator<<(FStructuredArchive::FSlot Slot, FPropertyTag& Tag)
+{
+	FArchive& UnderlyingArchive = Slot.GetUnderlyingArchive();
+	FStructuredArchive::FRecord Record = Slot.EnterRecord();
+	int32 Version = UnderlyingArchive.UE4Ver();
+
 	// Name.
-	Ar << Tag.Name;
+	Record << NAMED_ITEM("Name", Tag.Name);
 	if ((Tag.Name == NAME_None) || !Tag.Name.IsValid())
 	{
-		return Ar;
+		return;
 	}
 
-	Ar << Tag.Type;
-	if ( Ar.IsSaving() )
+	Record << NAMED_ITEM("Type", Tag.Type);
+	if (UnderlyingArchive.IsSaving())
 	{
 		// remember the offset of the Size variable - UStruct::SerializeTaggedProperties will update it after the
 		// property has been serialized.
-		Tag.SizeOffset = Ar.Tell();
+		Tag.SizeOffset = UnderlyingArchive.Tell();
 	}
 	{
-		FArchive::FScopeSetDebugSerializationFlags S(Ar, DSF_IgnoreDiff);
-		Ar << Tag.Size << Tag.ArrayIndex;
+		FArchive::FScopeSetDebugSerializationFlags S(UnderlyingArchive, DSF_IgnoreDiff);
+		Record << NAMED_ITEM("Size", Tag.Size) << NAMED_ITEM("ArrayIndex", Tag.ArrayIndex);
 	}
 	// only need to serialize this for structs
 	if (Tag.Type == NAME_StructProperty)
 	{
-		Ar << Tag.StructName;
-		if (Ar.UE4Ver() >= VER_UE4_STRUCT_GUID_IN_PROPERTY_TAG)
+		Record << NAMED_ITEM("StructName", Tag.StructName);
+		if (Version >= VER_UE4_STRUCT_GUID_IN_PROPERTY_TAG)
 		{
-			Ar << Tag.StructGuid;
+			Record << NAMED_ITEM("StructGuid", Tag.StructGuid);
 		}
 	}
 	// only need to serialize this for bools
-	else if (Tag.Type == NAME_BoolProperty)
+	else if (Tag.Type == NAME_BoolProperty && !UnderlyingArchive.IsTextFormat())
 	{
-		Ar << Tag.BoolVal;
+		Record << NAMED_ITEM("BoolVal", Tag.BoolVal);
 	}
 	// only need to serialize this for bytes/enums
 	else if (Tag.Type == NAME_ByteProperty || Tag.Type == NAME_EnumProperty)
 	{
-		Ar << Tag.EnumName;
+		Record << NAMED_ITEM("EnumName", Tag.EnumName);
 	}
 	// only need to serialize this for arrays
 	else if (Tag.Type == NAME_ArrayProperty)
 	{
-		if (Ar.UE4Ver() >= VAR_UE4_ARRAY_PROPERTY_INNER_TAGS)
+		if (Version >= VAR_UE4_ARRAY_PROPERTY_INNER_TAGS)
 		{
-			Ar << Tag.InnerType;
+			Record << NAMED_ITEM("InnerType", Tag.InnerType);
 		}
 	}
 
-	if (Ar.UE4Ver() >= VER_UE4_PROPERTY_TAG_SET_MAP_SUPPORT)
+	if (Version >= VER_UE4_PROPERTY_TAG_SET_MAP_SUPPORT)
 	{
 		if (Tag.Type == NAME_SetProperty)
 		{
-			Ar << Tag.InnerType;
+			Record << NAMED_ITEM("InnerType", Tag.InnerType);
 		}
 		else if (Tag.Type == NAME_MapProperty)
 		{
-			Ar << Tag.InnerType;
-			Ar << Tag.ValueType;
+			Record << NAMED_ITEM("InnerType", Tag.InnerType);
+			Record << NAMED_ITEM("ValueType", Tag.ValueType);
 		}
 	}
 
 	// Property tags to handle renamed blueprint properties effectively.
-	if (Ar.UE4Ver() >= VER_UE4_PROPERTY_GUID_IN_PROPERTY_TAG)
+	if (Version >= VER_UE4_PROPERTY_GUID_IN_PROPERTY_TAG)
 	{
-		Ar << Tag.HasPropertyGuid;
+		Record << NAMED_ITEM("HasPropertyGuid", Tag.HasPropertyGuid);
 		if (Tag.HasPropertyGuid)
 		{
-			Ar << Tag.PropertyGuid;
+			Record << NAMED_ITEM("PropertyGuid", Tag.PropertyGuid);
 		}
 	}
-
-	return Ar;
 }
 
 // Property serializer.
-void FPropertyTag::SerializeTaggedProperty( FArchive& Ar, UProperty* Property, uint8* Value, uint8* Defaults ) const
+void FPropertyTag::SerializeTaggedProperty(FArchive& Ar, UProperty* Property, uint8* Value, uint8* Defaults) const
 {
-	if (Property->GetClass() == UBoolProperty::StaticClass())
+	SerializeTaggedProperty(FStructuredArchiveFromArchive(Ar).GetSlot(), Property, Value, Defaults);
+}
+
+void FPropertyTag::SerializeTaggedProperty(FStructuredArchive::FSlot Slot, UProperty* Property, uint8* Value, uint8* Defaults) const
+{
+	FArchive& UnderlyingArchive = Slot.GetUnderlyingArchive();
+
+	if (!UnderlyingArchive.IsTextFormat() && Property->GetClass() == UBoolProperty::StaticClass())
 	{
 		UBoolProperty* Bool = (UBoolProperty*)Property;
-		if (Ar.IsLoading())
+		if (UnderlyingArchive.IsLoading())
 		{
 			Bool->SetPropertyValue(Value, BoolVal != 0);
 		}
+
+		Slot.EnterStream();	// Effectively discard
 	}
 	else
 	{
 #if WITH_EDITOR
 		static const FName NAME_SerializeTaggedProperty = FName(TEXT("SerializeTaggedProperty"));
-		FArchive::FScopeAddDebugData P(Ar, NAME_SerializeTaggedProperty);
-		FArchive::FScopeAddDebugData A(Ar, Property->GetFName());
+		FArchive::FScopeAddDebugData P(UnderlyingArchive, NAME_SerializeTaggedProperty);
+		FArchive::FScopeAddDebugData A(UnderlyingArchive, Property->GetFName());
 #endif
-		FSerializedPropertyScope SerializedProperty(Ar, Property);
-		Property->SerializeItem( Ar, Value, Defaults );
+		FSerializedPropertyScope SerializedProperty(UnderlyingArchive, Property);
+
+		FArchiveUObjectFromStructuredArchive Ar(Slot);
+		Property->SerializeItem(Ar, Value, Defaults);
 	}
 }
diff --git a/Engine/Source/Runtime/CoreUObject/Public/UObject/Class.h b/Engine/Source/Runtime/CoreUObject/Public/UObject/Class.h
--- a/Engine/Source/Runtime/CoreUObject/Public/UObject/Class.h
+++ b/Engine/Source/Runtime/CoreUObject/Public/UObject/Class.h
@@ -304,6 +304,7 @@ public:
 
 	/** Serializes list of properties, using property tags to handle mismatches */
 	virtual void SerializeTaggedProperties( FArchive& Ar, uint8* Data, UStruct* DefaultsStruct, uint8* Defaults, const UObject* BreakRecursionIfFullyLoad=nullptr) const;
+	virtual void SerializeTaggedProperties( FStructuredArchive::FSlot Slot, uint8* Data, UStruct* DefaultsStruct, uint8* Defaults, const UObject* BreakRecursionIfFullyLoad = nullptr) const;
 
 	/**
 	 * Initialize a struct over uninitialized memory. This may be done by calling the native constructor or individually initializing properties
diff --git a/Engine/Source/Runtime/CoreUObject/Public/UObject/LazyObjectPtr.h b/Engine/Source/Runtime/CoreUObject/Public/UObject/LazyObjectPtr.h
--- a/Engine/Source/Runtime/CoreUObject/Public/UObject/LazyObjectPtr.h
+++ b/Engine/Source/Runtime/CoreUObject/Public/UObject/LazyObjectPtr.h
@@ -89,6 +89,11 @@ struct COREUOBJECT_API FUniqueObjectGuid
 		return Ar;
 	}
 
+	friend void operator<<(FStructuredArchive::FSlot Slot, FUniqueObjectGuid& ObjectGuid)
+	{
+		Slot << ObjectGuid.Guid;
+	}
+
 	/** Code needed by FLazyPtr internals */
 	static int32 GetCurrentTag()
 	{
@@ -158,7 +163,7 @@ public:
 	}
 	
 	/** Called by UObject::Serialize so that we can save / load the Guid possibly associated with an object */
-	COREUOBJECT_API static void PossiblySerializeObjectGuid(UObject* Object, FArchive& Ar);
+	COREUOBJECT_API static void PossiblySerializeObjectGuid(UObject* Object, FStructuredArchive::FRecord Record);
 
 	/** Called when entering PIE to prepare it for PIE-specific fixups */
 	COREUOBJECT_API static void ResetPIEFixups();
diff --git a/Engine/Source/Runtime/CoreUObject/Public/UObject/Object.h b/Engine/Source/Runtime/CoreUObject/Public/UObject/Object.h
--- a/Engine/Source/Runtime/CoreUObject/Public/UObject/Object.h
+++ b/Engine/Source/Runtime/CoreUObject/Public/UObject/Object.h
@@ -784,6 +784,14 @@ public:
 	 */
 	void SerializeScriptProperties( FArchive& Ar ) const;
 
+	/**
+	 * Serializes the script property data located at Data.  When saving, only saves those properties which differ from the corresponding
+	 * value in the specified 'DiffObject' (usually the object's archetype).
+	 *
+	 * @param	Slot				the archive slot to serialize to
+	 */
+	void SerializeScriptProperties( FStructuredArchive::FSlot Slot ) const;
+
 	/**
 	 * Wrapper function for InitProperties() which handles safely tearing down this object before re-initializing it
 	 * from the specified source object.
diff --git a/Engine/Source/Runtime/CoreUObject/Public/UObject/ObjectMacros.h b/Engine/Source/Runtime/CoreUObject/Public/UObject/ObjectMacros.h
--- a/Engine/Source/Runtime/CoreUObject/Public/UObject/ObjectMacros.h
+++ b/Engine/Source/Runtime/CoreUObject/Public/UObject/ObjectMacros.h
@@ -1358,6 +1358,10 @@ namespace UM
 	friend FArchive &operator<<( FArchive& Ar, TClass*& Res ) \
 	{ \
 		return Ar << (UObject*&)Res; \
+	} \
+	friend void operator<<(FStructuredArchive::FSlot InSlot, TClass*& Res) \
+	{ \
+		InSlot << (UObject*&)Res; \
 	}
 
 #define IMPLEMENT_FARCHIVE_SERIALIZER( TClass ) void TClass::Serialize(FArchive& Ar) { TClass::Serialize(FStructuredArchiveFromArchive(Ar).GetSlot().EnterRecord()); }
diff --git a/Engine/Source/Runtime/CoreUObject/Public/UObject/PropertyTag.h b/Engine/Source/Runtime/CoreUObject/Public/UObject/PropertyTag.h
--- a/Engine/Source/Runtime/CoreUObject/Public/UObject/PropertyTag.h
+++ b/Engine/Source/Runtime/CoreUObject/Public/UObject/PropertyTag.h
@@ -37,8 +37,10 @@ struct FPropertyTag
 	void SetPropertyGuid(const FGuid& InPropertyGuid);
 
 	// Serializer.
-	friend FArchive& operator<<( FArchive& Ar, FPropertyTag& Tag );
+	friend FArchive& operator<<(FArchive& Ar, FPropertyTag& Tag);
+	friend void operator<<(FStructuredArchive::FSlot Slot, FPropertyTag& Tag);
 
 	// Property serializer.
-	void SerializeTaggedProperty( FArchive& Ar, UProperty* Property, uint8* Value, uint8* Defaults ) const;
+	void SerializeTaggedProperty(FArchive& Ar, UProperty* Property, uint8* Value, uint8* Defaults) const;
+	void SerializeTaggedProperty(FStructuredArchive::FSlot Slot, UProperty* Property, uint8* Value, uint8* Defaults) const;
 };
