TBA: Temporarily disable loading of text assets until tagged property serialization path is fixed up

diff --git a/Engine/Source/Runtime/CoreUObject/Private/Misc/PackageName.cpp b/Engine/Source/Runtime/CoreUObject/Private/Misc/PackageName.cpp
--- a/Engine/Source/Runtime/CoreUObject/Private/Misc/PackageName.cpp
+++ b/Engine/Source/Runtime/CoreUObject/Private/Misc/PackageName.cpp
@@ -753,7 +753,7 @@ bool FPackageName::FindPackageFileWithoutExtension(const FString& InPackageFilen
 	};
 
 	// Loop through all known extensions and check if the file exist.
-	const int32 BaseIndex = (InAllowTextFormats && WITH_TEXT_ARCHIVE_SUPPORT) ? 0 : 2;
+	const int32 BaseIndex = (false && InAllowTextFormats && WITH_TEXT_ARCHIVE_SUPPORT) ? 0 : 2;
 	for (int32 ExtensionIndex = BaseIndex; ExtensionIndex < ARRAY_COUNT(PackageExtensions); ++ExtensionIndex)
 	{
 		FString   PackageFilename = InPackageFilename + *PackageExtensions[ExtensionIndex];
