TBA: Added default structured archive implementation of SerializeItem to UProperty, which just calls the FArchive version on an FArchiveUObjectFromStructuredArchive adapter. Implemented structured archive SerializeItem for UArrayProperty

diff --git a/Engine/Source/Runtime/CoreUObject/Private/UObject/PropertyArray.cpp b/Engine/Source/Runtime/CoreUObject/Private/UObject/PropertyArray.cpp
--- a/Engine/Source/Runtime/CoreUObject/Private/UObject/PropertyArray.cpp
+++ b/Engine/Source/Runtime/CoreUObject/Private/UObject/PropertyArray.cpp
@@ -57,20 +57,28 @@ bool UArrayProperty::Identical( const void* A, const void* B, uint32 PortFlags )
 
 	return true;
 }
-void UArrayProperty::SerializeItem( FArchive& Ar, void* Value, void const* Defaults ) const
+
+void UArrayProperty::SerializeItem(FArchive& Ar, void* Value, void const* Defaults) const
+{
+	SerializeItem(FStructuredArchiveFromArchive(Ar).GetSlot(), Value, Defaults);
+}
+
+void UArrayProperty::SerializeItem(FStructuredArchive::FSlot Slot, void* Value, void const* Defaults) const
 {
 	checkSlow(Inner);
+	FArchive& UnderlyingArchive = Slot.GetUnderlyingArchive();
+	FStructuredArchive::FRecord Record = Slot.EnterRecord();
 
 	// Ensure that the Inner itself has been loaded before calling SerializeItem() on it
-	Ar.Preload(Inner);
+	UnderlyingArchive.Preload(Inner);
 
 	FScriptArrayHelper ArrayHelper(this, Value);
 	int32		n		= ArrayHelper.Num();
-	Ar << n;
-	if( Ar.IsLoading() )
+	Record << NAMED_ITEM("Count", n);
+	if( UnderlyingArchive.IsLoading() )
 	{
 		// If using a custom property list, don't empty the array on load. Not all indices may have been serialized, so we need to preserve existing values at those slots.
-		if (Ar.ArUseCustomPropertyList)
+		if (UnderlyingArchive.ArUseCustomPropertyList)
 		{
 			const int32 OldNum = ArrayHelper.Num();
 			if (n > OldNum)
@@ -87,21 +95,21 @@ void UArrayProperty::SerializeItem( FArchive& Ar, void* Value, void const* Defau
 			ArrayHelper.EmptyAndAddValues(n);
 		}
 	}
-	ArrayHelper.CountBytes( Ar );
+	ArrayHelper.CountBytes( UnderlyingArchive );
 
 	// Serialize a PropertyTag for the inner property of this array, allows us to validate the inner struct to see if it has changed
-	FPropertyTag InnerTag(Ar, Inner, 0, (uint8*)Value, (uint8*)Defaults);
-	if (Ar.UE4Ver() >= VER_UE4_INNER_ARRAY_TAG_INFO && InnerTag.Type == NAME_StructProperty)
+	FPropertyTag InnerTag(UnderlyingArchive, Inner, 0, (uint8*)Value, (uint8*)Defaults);
+	if (UnderlyingArchive.UE4Ver() >= VER_UE4_INNER_ARRAY_TAG_INFO && InnerTag.Type == NAME_StructProperty)
 	{
-		if (Ar.IsSaving())
+		if (UnderlyingArchive.IsSaving())
 		{
-			Ar << InnerTag;
+			Record << NAMED_FIELD(InnerTag);
 		}
-		else if (Ar.IsLoading())
+		else if (UnderlyingArchive.IsLoading())
 		{
-			Ar << InnerTag;
+			Record << NAMED_FIELD(InnerTag);
 
-			auto CanSerializeFromStructWithDifferentName = [](const FArchive& InAr, const FPropertyTag& PropertyTag, const UStructProperty* StructProperty)
+			auto CanSerializeFromStructWithDifferentName = [](const FPropertyTag& PropertyTag, const UStructProperty* StructProperty)
 			{
 				return PropertyTag.StructGuid.IsValid()
 					&& StructProperty 
@@ -120,10 +128,10 @@ void UArrayProperty::SerializeItem( FArchive& Ar, void* Value, void const* Defau
 			}
 
 			if (InnerTag.StructName != StructProperty->Struct->GetFName()
-				&& !CanSerializeFromStructWithDifferentName(Ar, InnerTag, StructProperty))
+				&& !CanSerializeFromStructWithDifferentName(InnerTag, StructProperty))
 			{
 				UE_LOG(LogClass, Warning, TEXT("Property %s of %s has a struct type mismatch (tag %s != prop %s) in package:  %s. If that struct got renamed, add an entry to ActiveStructRedirects."),
-					*InnerTag.Name.ToString(), *GetName(), *InnerTag.StructName.ToString(), *CastChecked<UStructProperty>(Inner)->Struct->GetName(), *Ar.GetArchiveName());
+					*InnerTag.Name.ToString(), *GetName(), *InnerTag.StructName.ToString(), *CastChecked<UStructProperty>(Inner)->Struct->GetName(), *UnderlyingArchive.GetArchiveName());
 
 #if WITH_EDITOR
 				// Ensure the structure is initialized
@@ -133,13 +141,16 @@ void UArrayProperty::SerializeItem( FArchive& Ar, void* Value, void const* Defau
 				}
 #endif // WITH_EDITOR
 
-				// Skip the property
-				const int64 StartOfProperty = Ar.Tell();
-				const int64 RemainingSize = InnerTag.Size - (Ar.Tell() - StartOfProperty);
-				uint8 B;
-				for (int64 i = 0; i < RemainingSize; i++)
+				if (!UnderlyingArchive.IsTextFormat())
 				{
-					Ar << B;
+					// Skip the property
+					const int64 StartOfProperty = UnderlyingArchive.Tell();
+					const int64 RemainingSize = InnerTag.Size - (UnderlyingArchive.Tell() - StartOfProperty);
+					uint8 B;
+					for (int64 i = 0; i < RemainingSize; i++)
+					{
+						UnderlyingArchive << B;
+					}
 				}
 				return;
 			}
@@ -147,20 +158,22 @@ void UArrayProperty::SerializeItem( FArchive& Ar, void* Value, void const* Defau
 	}
 
 	// need to know how much data this call to SerializeItem consumes, so mark where we are
-	int64 DataOffset = Ar.Tell();
+	int64 DataOffset = UnderlyingArchive.Tell();
+
+	FStructuredArchive::FStream ValueStream = Record.EnterField(FIELD_NAME_TEXT("Values")).EnterStream();
 
 	// If we're using a custom property list, first serialize any explicit indices
 	int32 i = 0;
 	bool bSerializeRemainingItems = true;
-	bool bUsingCustomPropertyList = Ar.ArUseCustomPropertyList;
-	if (bUsingCustomPropertyList && Ar.ArCustomPropertyList != nullptr)
+	bool bUsingCustomPropertyList = UnderlyingArchive.ArUseCustomPropertyList;
+	if (bUsingCustomPropertyList && UnderlyingArchive.ArCustomPropertyList != nullptr)
 	{
 		// Initially we only serialize indices that are explicitly specified (in order)
 		bSerializeRemainingItems = false;
 
-		const FCustomPropertyListNode* CustomPropertyList = Ar.ArCustomPropertyList;
+		const FCustomPropertyListNode* CustomPropertyList = UnderlyingArchive.ArCustomPropertyList;
 		const FCustomPropertyListNode* PropertyNode = CustomPropertyList;
-		FSerializedPropertyScope SerializedProperty(Ar, Inner, this);
+		FSerializedPropertyScope SerializedProperty(UnderlyingArchive, Inner, this);
 		while (PropertyNode && i < n && !bSerializeRemainingItems)
 		{
 			if (PropertyNode->Property != Inner)
@@ -179,15 +192,15 @@ void UArrayProperty::SerializeItem( FArchive& Ar, void* Value, void const* Defau
 				FCustomPropertyListNode ItemNode = *PropertyNode;
 				ItemNode.ArrayIndex = 0;
 				ItemNode.PropertyListNext = nullptr;
-				Ar.ArCustomPropertyList = &ItemNode;
+				UnderlyingArchive.ArCustomPropertyList = &ItemNode;
 
 				// Serialize the item at this array index
 				i = PropertyNode->ArrayIndex;
-				Inner->SerializeItem(Ar, ArrayHelper.GetRawPtr(i));
+				Inner->SerializeItem(ValueStream.EnterElement(), ArrayHelper.GetRawPtr(i));
 				PropertyNode = PropertyNode->PropertyListNext;
 
 				// Restore the current property list
-				Ar.ArCustomPropertyList = CustomPropertyList;
+				UnderlyingArchive.ArCustomPropertyList = CustomPropertyList;
 			}
 		}
 	}
@@ -195,41 +208,41 @@ void UArrayProperty::SerializeItem( FArchive& Ar, void* Value, void const* Defau
 	if (bSerializeRemainingItems)
 	{
 		// Temporarily suspend the custom property list (as we need these items to be serialized in full)
-		Ar.ArUseCustomPropertyList = false;
+		UnderlyingArchive.ArUseCustomPropertyList = false;
 
 		// Serialize each item until we get to the end of the array
-		FSerializedPropertyScope SerializedProperty(Ar, Inner, this);
+		FSerializedPropertyScope SerializedProperty(UnderlyingArchive, Inner, this);
 		while (i < n)
 		{
 #if WITH_EDITOR
 			static const FName NAME_UArraySerialize = FName(TEXT("UArrayProperty::Serialize"));
 			FName NAME_UArraySerializeCount = FName(NAME_UArraySerialize);
 			NAME_UArraySerializeCount.SetNumber(i);
-			FArchive::FScopeAddDebugData P(Ar, NAME_UArraySerializeCount);
+			FArchive::FScopeAddDebugData P(UnderlyingArchive, NAME_UArraySerializeCount);
 #endif
-			Inner->SerializeItem(Ar, ArrayHelper.GetRawPtr(i++));
+			Inner->SerializeItem(ValueStream.EnterElement(), ArrayHelper.GetRawPtr(i++));
 		}
 
 		// Restore use of the custom property list (if it was previously enabled)
-		Ar.ArUseCustomPropertyList = bUsingCustomPropertyList;
+		UnderlyingArchive.ArUseCustomPropertyList = bUsingCustomPropertyList;
 	}
 
-	if (Ar.UE4Ver() >= VER_UE4_INNER_ARRAY_TAG_INFO && Ar.IsSaving() && InnerTag.Type == NAME_StructProperty)
+	if (UnderlyingArchive.UE4Ver() >= VER_UE4_INNER_ARRAY_TAG_INFO && UnderlyingArchive.IsSaving() && InnerTag.Type == NAME_StructProperty && !UnderlyingArchive.IsTextFormat())
 	{
 		// set the tag's size
-		InnerTag.Size = Ar.Tell() - DataOffset;
+		InnerTag.Size = UnderlyingArchive.Tell() - DataOffset;
 
 		if (InnerTag.Size > 0)
 		{
 			// mark our current location
-			DataOffset = Ar.Tell();
+			DataOffset = UnderlyingArchive.Tell();
 
 			// go back and re-serialize the size now that we know it
-			Ar.Seek(InnerTag.SizeOffset);
-			Ar << InnerTag.Size;
+			UnderlyingArchive.Seek(InnerTag.SizeOffset);
+			UnderlyingArchive << InnerTag.Size;
 
 			// return to the current location
-			Ar.Seek(DataOffset);
+			UnderlyingArchive.Seek(DataOffset);
 		}
 	}
 }
diff --git a/Engine/Source/Runtime/CoreUObject/Private/UObject/PropertyTag.cpp b/Engine/Source/Runtime/CoreUObject/Private/UObject/PropertyTag.cpp
--- a/Engine/Source/Runtime/CoreUObject/Private/UObject/PropertyTag.cpp
+++ b/Engine/Source/Runtime/CoreUObject/Private/UObject/PropertyTag.cpp
@@ -204,7 +204,6 @@ void FPropertyTag::SerializeTaggedProperty(FStructuredArchive::FSlot Slot, UProp
 #endif
 		FSerializedPropertyScope SerializedProperty(UnderlyingArchive, Property);
 
-		FArchiveUObjectFromStructuredArchive Ar(Slot);
-		Property->SerializeItem(Ar, Value, Defaults);
+		Property->SerializeItem(Slot, Value, Defaults);
 	}
 }
diff --git a/Engine/Source/Runtime/CoreUObject/Private/UObject/UnrealType.cpp b/Engine/Source/Runtime/CoreUObject/Private/UObject/UnrealType.cpp
--- a/Engine/Source/Runtime/CoreUObject/Private/UObject/UnrealType.cpp
+++ b/Engine/Source/Runtime/CoreUObject/Private/UObject/UnrealType.cpp
@@ -1,6 +1,7 @@
 // Copyright 1998-2018 Epic Games, Inc. All Rights Reserved.
 
 #include "UObject/UnrealType.h"
+#include "Serialization/ArchiveUObjectFromStructuredArchive.h"
 
 DEFINE_LOG_CATEGORY(LogType);
 
@@ -172,3 +173,9 @@ void FPropertyValueIterator::GetPropertyChain(TArray<const UProperty*>& Property
 		}
 	}
 }
+
+void UProperty::SerializeItem(FStructuredArchive::FSlot Slot, void* Value, void const* Defaults) const
+{
+	FArchiveUObjectFromStructuredArchive Ar(Slot);
+	SerializeItem(Ar, Value, Defaults);
+}
\ No newline at end of file
diff --git a/Engine/Source/Runtime/CoreUObject/Public/UObject/UnrealType.h b/Engine/Source/Runtime/CoreUObject/Public/UObject/UnrealType.h
--- a/Engine/Source/Runtime/CoreUObject/Public/UObject/UnrealType.h
+++ b/Engine/Source/Runtime/CoreUObject/Public/UObject/UnrealType.h
@@ -301,7 +301,8 @@ public:
 		}
 	}
 
-	virtual void SerializeItem( FArchive& Ar, void* Value, void const* Defaults=NULL ) const PURE_VIRTUAL(UProperty::SerializeItem,);
+	virtual void SerializeItem(FArchive& Ar, void* Value, void const* Defaults=NULL ) const PURE_VIRTUAL(UProperty::SerializeItem,);
+	virtual void SerializeItem(FStructuredArchive::FSlot Slot, void* Value, void const* Defaults = NULL) const;
 	virtual bool NetSerializeItem( FArchive& Ar, UPackageMap* Map, void* Data, TArray<uint8> * MetaData = NULL ) const;
 	virtual bool SupportsNetSharedSerialization() const;
 	virtual void ExportTextItem( FString& ValueStr, const void* PropertyValue, const void* DefaultValue, UObject* Parent, int32 PortFlags, UObject* ExportRootScope = NULL ) const PURE_VIRTUAL(UProperty::ExportTextItem,);
@@ -2509,6 +2510,7 @@ public:
 	virtual void LinkInternal(FArchive& Ar) override;
 	virtual bool Identical( const void* A, const void* B, uint32 PortFlags ) const override;
 	virtual void SerializeItem( FArchive& Ar, void* Value, void const* Defaults ) const override;
+	virtual void SerializeItem(FStructuredArchive::FSlot Slot, void* Value, void const* Defaults) const override;
 	virtual bool NetSerializeItem( FArchive& Ar, UPackageMap* Map, void* Data, TArray<uint8> * MetaData = NULL ) const override;
 	virtual void ExportTextItem( FString& ValueStr, const void* PropertyValue, const void* DefaultValue, UObject* Parent, int32 PortFlags, UObject* ExportRootScope ) const override;
 	virtual const TCHAR* ImportText_Internal( const TCHAR* Buffer, void* Data, int32 PortFlags, UObject* OwnerObject, FOutputDevice* ErrorText ) const override;
