TBA: ObjectProperty text serialization

diff --git a/Engine/Source/Runtime/CoreUObject/Private/UObject/PropertyObject.cpp b/Engine/Source/Runtime/CoreUObject/Private/UObject/PropertyObject.cpp
--- a/Engine/Source/Runtime/CoreUObject/Private/UObject/PropertyObject.cpp
+++ b/Engine/Source/Runtime/CoreUObject/Private/UObject/PropertyObject.cpp
@@ -53,18 +53,25 @@ EConvertFromTypeResult UObjectProperty::ConvertFromType(const FPropertyTag& Tag,
 	return EConvertFromTypeResult::UseSerializeItem;
 }
 
-void UObjectProperty::SerializeItem( FArchive& Ar, void* Value, void const* Defaults ) const
+void UObjectProperty::SerializeItem(FArchive& Ar, void* Value, void const* Defaults) const
 {
-	if (Ar.IsObjectReferenceCollector())
+	SerializeItem(FStructuredArchiveFromArchive(Ar).GetSlot(), Value, Defaults);
+}
+
+void UObjectProperty::SerializeItem( FStructuredArchive::FSlot Slot, void* Value, void const* Defaults ) const
+{
+	FArchive& UnderlyingArchive = Slot.GetUnderlyingArchive();
+
+	if (UnderlyingArchive.IsObjectReferenceCollector())
 	{
 		// Serialize in place
 		UObject** ObjectPtr = GetPropertyValuePtr(Value);
-		Ar << (*ObjectPtr);
+		Slot << (*ObjectPtr);
 	}
 	else
 	{
 		UObject* ObjectValue = GetObjectPropertyValue(Value);
-		Ar << ObjectValue;
+		Slot << ObjectValue;
 
 		UObject* CurrentValue = GetObjectPropertyValue(Value);
 		if (ObjectValue != CurrentValue)
diff --git a/Engine/Source/Runtime/CoreUObject/Public/UObject/UnrealType.h b/Engine/Source/Runtime/CoreUObject/Public/UObject/UnrealType.h
--- a/Engine/Source/Runtime/CoreUObject/Public/UObject/UnrealType.h
+++ b/Engine/Source/Runtime/CoreUObject/Public/UObject/UnrealType.h
@@ -2054,7 +2054,8 @@ class COREUOBJECT_API UObjectProperty : public TUObjectPropertyBase<UObject*>
 	// End of UHT interface
 
 	// UProperty interface
-	virtual void SerializeItem( FArchive& Ar, void* Value, void const* Defaults ) const override;
+	virtual void SerializeItem(FArchive& Ar, void* Value, void const* Defaults) const override;
+	virtual void SerializeItem(FStructuredArchive::FSlot Slot, void* Value, void const* Defaults) const override;
 	virtual void EmitReferenceInfo(UClass& OwnerClass, int32 BaseOffset, TArray<const UStructProperty*>& EncounteredStructProps) override;
 	virtual const TCHAR* ImportText_Internal(const TCHAR* Buffer, void* Data, int32 PortFlags, UObject* OwnerObject, FOutputDevice* ErrorText) const override;
 	virtual EConvertFromTypeResult ConvertFromType(const FPropertyTag& Tag, FArchive& Ar, uint8* Data, UStruct* DefaultsStruct) override;
