Restored a modified FWeakObjectPtr::Serialize function to keep backwards compatibility in code I don't have access to.

diff --git a/Engine/Source/Runtime/CoreUObject/Private/UObject/WeakObjectPtr.cpp b/Engine/Source/Runtime/CoreUObject/Private/UObject/WeakObjectPtr.cpp
--- a/Engine/Source/Runtime/CoreUObject/Private/UObject/WeakObjectPtr.cpp
+++ b/Engine/Source/Runtime/CoreUObject/Private/UObject/WeakObjectPtr.cpp
@@ -95,4 +95,9 @@ UObject* FWeakObjectPtr::GetEvenIfUnreachable() const
 		Result = static_cast<UObject*>(ObjectItem->Object);
 	}
 	return Result;
-}
\ No newline at end of file
+}
+
+void FWeakObjectPtr::Serialize(FArchive& Ar)
+{
+	FArchiveUObject::SerializeWeakObjectPtr(Ar, *this);
+}
diff --git a/Engine/Source/Runtime/CoreUObject/Public/UObject/WeakObjectPtr.h b/Engine/Source/Runtime/CoreUObject/Public/UObject/WeakObjectPtr.h
--- a/Engine/Source/Runtime/CoreUObject/Public/UObject/WeakObjectPtr.h
+++ b/Engine/Source/Runtime/CoreUObject/Public/UObject/WeakObjectPtr.h
@@ -143,6 +143,14 @@ public:
 		return uint32(WeakObjectPtr.ObjectIndex ^ WeakObjectPtr.ObjectSerialNumber);
 	}
 
+	/**
+	 * Weak object pointer serialization.  Weak object pointers only have weak references to objects and
+	 * won't serialize the object when gathering references for garbage collection.  So in many cases, you
+	 * don't need to bother serializing weak object pointers.  However, serialization is required if you
+	 * want to load and save your object.
+	 */
+	COREUOBJECT_API void Serialize(FArchive& Ar);
+
 protected:
 
 	FORCEINLINE int32 GetObjectIndex() const
