TBA: Convert UMetaData to structured archive

diff --git a/Engine/Source/Runtime/CoreUObject/Private/UObject/MetaData.cpp b/Engine/Source/Runtime/CoreUObject/Private/UObject/MetaData.cpp
--- a/Engine/Source/Runtime/CoreUObject/Private/UObject/MetaData.cpp
+++ b/Engine/Source/Runtime/CoreUObject/Private/UObject/MetaData.cpp
@@ -113,13 +113,17 @@ IMPLEMENT_CORE_INTRINSIC_CLASS(UMetaData, UObject,
 	}
 );
 
-void UMetaData::Serialize(FArchive& Ar)
+IMPLEMENT_FARCHIVE_SERIALIZER(UMetaData);
+
+void UMetaData::Serialize(FStructuredArchive::FRecord Record)
 {
-	Super::Serialize(Ar);
+	FArchive& UnderlyingArchive = Record.GetUnderlyingArchive();
+
+	Super::Serialize(Record);
 
-	Ar.UsingCustomVersion(FEditorObjectVersion::GUID);
+	UnderlyingArchive.UsingCustomVersion(FEditorObjectVersion::GUID);
 
-	if( Ar.IsSaving() )
+	if( UnderlyingArchive.IsSaving() )
 	{
 		// Remove entries belonging to destructed objects
 		for (TMap< FWeakObjectPtr, TMap<FName, FString> >::TIterator It(ObjectMetaDataMap); It; ++It)
@@ -131,18 +135,18 @@ void UMetaData::Serialize(FArchive& Ar)
 		}
 	}
 	
-	if (!Ar.IsLoading())
+	if (!UnderlyingArchive.IsLoading())
 	{
-		Ar << ObjectMetaDataMap;
-		Ar << RootMetaDataMap;
+		Record << NAMED_FIELD(ObjectMetaDataMap);
+		Record << NAMED_FIELD(RootMetaDataMap);
 	}
 	else
 	{
 		{
 			TMap< FWeakObjectPtr, TMap<FName, FString> > TempMap;
-			Ar << TempMap;
+			Record << NAMED_ITEM("ObjectMetaDataMap", TempMap);
 
-			const bool bLoadFromLinker = (NULL != Ar.GetLinker());
+			const bool bLoadFromLinker = (NULL != UnderlyingArchive.GetLinker());
 			if (bLoadFromLinker && HasAnyFlags(RF_LoadCompleted))
 			{
 				UE_LOG(LogMetaData, Verbose, TEXT("Metadata was already loaded by linker. %s"), *GetFullName());
@@ -157,12 +161,12 @@ void UMetaData::Serialize(FArchive& Ar)
 			}
 		}
 
-		if (Ar.CustomVer(FEditorObjectVersion::GUID) >= FEditorObjectVersion::RootMetaDataSupport)
+		if (UnderlyingArchive.CustomVer(FEditorObjectVersion::GUID) >= FEditorObjectVersion::RootMetaDataSupport)
 		{
 			TMap<FName, FString> TempMap;
-			Ar << TempMap;
+			Record << NAMED_ITEM("RootMetaDataMap", TempMap);
 
-			const bool bLoadFromLinker = (NULL != Ar.GetLinker());
+			const bool bLoadFromLinker = (NULL != UnderlyingArchive.GetLinker());
 			if (bLoadFromLinker && HasAnyFlags(RF_LoadCompleted))
 			{
 				UE_LOG(LogMetaData, Verbose, TEXT("Root metadata was already loaded by linker. %s"), *GetFullName());
diff --git a/Engine/Source/Runtime/CoreUObject/Public/UObject/MetaData.h b/Engine/Source/Runtime/CoreUObject/Public/UObject/MetaData.h
--- a/Engine/Source/Runtime/CoreUObject/Public/UObject/MetaData.h
+++ b/Engine/Source/Runtime/CoreUObject/Public/UObject/MetaData.h
@@ -16,7 +16,7 @@
  */
 class COREUOBJECT_API UMetaData : public UObject
 {
-	DECLARE_CLASS_INTRINSIC(UMetaData, UObject, 0, TEXT("/Script/CoreUObject"))
+	DECLARE_CLASS_INTRINSIC(UMetaData, UObject, CLASS_MatchedSerializers, TEXT("/Script/CoreUObject"))
 
 public:
 	/**
@@ -123,6 +123,7 @@ public:
 
 	// UObject interface
 	virtual void Serialize(FArchive& Ar) override;
+	virtual void Serialize(FStructuredArchive::FRecord Record) override;
 	virtual bool NeedsLoadForEditorGame() const override;
 	virtual bool IsAsset() const override { return false; }
 	// End of UObject interface
