Adding stats for the new GC internal functions

diff --git a/Engine/Source/Runtime/CoreUObject/Private/UObject/GarbageCollection.cpp b/Engine/Source/Runtime/CoreUObject/Private/UObject/GarbageCollection.cpp
--- a/Engine/Source/Runtime/CoreUObject/Private/UObject/GarbageCollection.cpp
+++ b/Engine/Source/Runtime/CoreUObject/Private/UObject/GarbageCollection.cpp
@@ -1338,6 +1338,8 @@ static FAutoConsoleVariableRef CVarFlushStreamingOnGC(
 
 void GatherUnreachableObjects(bool bForceSingleThreaded)
 {
+	DECLARE_SCOPE_CYCLE_COUNTER(TEXT("CollectGarbageInternal.GatherUnreachableObjects"), STAT_CollectGarbageInternal_GatherUnreachableObjects, STATGROUP_GC);
+
 	const double StartTime = FPlatformTime::Seconds();
 
 	GUnreachableObjects.Reset();
@@ -1476,9 +1478,7 @@ void CollectGarbageInternal(EObjectFlags KeepFlags, bool bPerformFullPurge)
 		// Fire post-reachability analysis hooks
 		FCoreUObjectDelegates::PostReachabilityAnalysis.Broadcast();
 
-		{
-			DECLARE_SCOPE_CYCLE_COUNTER(TEXT("CollectGarbageInternal.UnhashUnreachable"), STAT_CollectGarbageInternal_UnhashUnreachable, STATGROUP_GC);			
-
+		{			
 			FGCArrayPool::Get().ClearWeakReferences(bPerformFullPurge);
 
 			GatherUnreachableObjects(bForceSingleThreadedGC);
@@ -1516,6 +1516,8 @@ void CollectGarbageInternal(EObjectFlags KeepFlags, bool bPerformFullPurge)
 
 bool UnhashUnreachableObjects(bool bUseTimeLimit, float TimeLimit)
 {
+	DECLARE_SCOPE_CYCLE_COUNTER(TEXT("UnhashUnreachableObjects"), STAT_UnhashUnreachableObjects, STATGROUP_GC);
+
 	TGuardValue<bool> GuardObjUnhashUnreachableIsInProgress(GObjUnhashUnreachableIsInProgress, true);
 
 	FCoreUObjectDelegates::PreGarbageCollectConditionalBeginDestroy.Broadcast();
