TBA: Removed UProperty::SerializeItem(FArchive, ...) and replaced with UProperty::SerializeItem(FStructuredArchive::FSlot, ...). Fixed up most of them to work properly and added adapters in for any that were non-trivial.

diff --git a/Engine/Source/Developer/HotReload/Private/HotReloadClassReinstancer.cpp b/Engine/Source/Developer/HotReload/Private/HotReloadClassReinstancer.cpp
--- a/Engine/Source/Developer/HotReload/Private/HotReloadClassReinstancer.cpp
+++ b/Engine/Source/Developer/HotReload/Private/HotReloadClassReinstancer.cpp
@@ -508,7 +508,7 @@ void FHotReloadClassReinstancer::UpdateDefaultProperties()
 					// Serialize current value to a byte array as we don't have the previous CDO to compare against, we only have its serialized property data
 					CurrentValueSerializedData.Empty(CurrentValueSerializedData.Num() + CurrentValueSerializedData.GetSlack());
 					FPropertyValueMemoryWriter CurrentValueWriter(CurrentValueSerializedData);
-					PropertyToUpdate.Property->SerializeItem(CurrentValueWriter, InstanceValuePtr);
+					PropertyToUpdate.Property->SerializeItem(FStructuredArchiveFromArchive(CurrentValueWriter).GetSlot(), InstanceValuePtr);
 
 					// Update only when the current value on the instance is identical to the original CDO
 					if (CurrentValueSerializedData.Num() == PropertyToUpdate.OldSerializedSize &&
diff --git a/Engine/Source/Runtime/CoreUObject/Private/UObject/EnumProperty.cpp b/Engine/Source/Runtime/CoreUObject/Private/UObject/EnumProperty.cpp
--- a/Engine/Source/Runtime/CoreUObject/Private/UObject/EnumProperty.cpp
+++ b/Engine/Source/Runtime/CoreUObject/Private/UObject/EnumProperty.cpp
@@ -70,23 +70,25 @@ void UEnumProperty::AddCppProperty(UProperty* Inner)
 	}
 }
 
-void UEnumProperty::SerializeItem( FArchive& Ar, void* Value, void const* Defaults ) const
+void UEnumProperty::SerializeItem(FStructuredArchive::FSlot Slot, void* Value, void const* Defaults) const
 {
+	FArchive& UnderlyingArchive = Slot.GetUnderlyingArchive();
+
 	check(UnderlyingProp);
 
-	if (Enum && Ar.UseToResolveEnumerators())
+	if (Enum && UnderlyingArchive.UseToResolveEnumerators())
 	{
 		int64 IntValue = UnderlyingProp->GetSignedIntPropertyValue(Value);
-		int64 ResolvedIndex = Enum->ResolveEnumerator(Ar, IntValue);
+		int64 ResolvedIndex = Enum->ResolveEnumerator(UnderlyingArchive, IntValue);
 		UnderlyingProp->SetIntPropertyValue(Value, ResolvedIndex);
 		return;
 	}
 
 	// Loading
-	if (Ar.IsLoading())
+	if (UnderlyingArchive.IsLoading())
 	{
 		FName EnumValueName;
-		Ar << EnumValueName;
+		Slot << EnumValueName;
 
 		int64 NewEnumValue = 0;
 
@@ -95,7 +97,7 @@ void UEnumProperty::SerializeItem( FArchive& Ar, void* Value, void const* Defaul
 			// Make sure enum is properly populated
 			if (Enum->HasAnyFlags(RF_NeedLoad))
 			{
-				Ar.Preload(Enum);
+				UnderlyingArchive.Preload(Enum);
 			}
 
 			// There's no guarantee EnumValueName is still present in Enum, in which case Value will be set to the enum's max value.
@@ -114,7 +116,7 @@ void UEnumProperty::SerializeItem( FArchive& Ar, void* Value, void const* Defaul
 		UnderlyingProp->SetIntPropertyValue(Value, NewEnumValue);
 	}
 	// Saving
-	else if (Ar.IsSaving())
+	else if (UnderlyingArchive.IsSaving())
 	{
 		FName EnumValueName;
 		if (Enum)
@@ -127,11 +129,11 @@ void UEnumProperty::SerializeItem( FArchive& Ar, void* Value, void const* Defaul
 			}
 		}
 
-		Ar << EnumValueName;
+		Slot << EnumValueName;
 	}
 	else
 	{
-		UnderlyingProp->SerializeItem(Ar, Value, Defaults);
+		UnderlyingProp->SerializeItem(Slot, Value, Defaults);
 	}
 }
 
diff --git a/Engine/Source/Runtime/CoreUObject/Private/UObject/Property.cpp b/Engine/Source/Runtime/CoreUObject/Private/UObject/Property.cpp
--- a/Engine/Source/Runtime/CoreUObject/Private/UObject/Property.cpp
+++ b/Engine/Source/Runtime/CoreUObject/Private/UObject/Property.cpp
@@ -736,7 +736,7 @@ bool UProperty::ShouldSerializeValue( FArchive& Ar ) const
 //
 bool UProperty::NetSerializeItem( FArchive& Ar, UPackageMap* Map, void* Data, TArray<uint8> * MetaData ) const
 {
-	SerializeItem( Ar, Data, NULL );
+	SerializeItem( FStructuredArchiveFromArchive(Ar).GetSlot(), Data, NULL );
 	return 1;
 }
 
diff --git a/Engine/Source/Runtime/CoreUObject/Private/UObject/PropertyArray.cpp b/Engine/Source/Runtime/CoreUObject/Private/UObject/PropertyArray.cpp
--- a/Engine/Source/Runtime/CoreUObject/Private/UObject/PropertyArray.cpp
+++ b/Engine/Source/Runtime/CoreUObject/Private/UObject/PropertyArray.cpp
@@ -58,11 +58,6 @@ bool UArrayProperty::Identical( const void* A, const void* B, uint32 PortFlags )
 	return true;
 }
 
-void UArrayProperty::SerializeItem(FArchive& Ar, void* Value, void const* Defaults) const
-{
-	SerializeItem(FStructuredArchiveFromArchive(Ar).GetSlot(), Value, Defaults);
-}
-
 void UArrayProperty::SerializeItem(FStructuredArchive::FSlot Slot, void* Value, void const* Defaults) const
 {
 	checkSlow(Inner);
diff --git a/Engine/Source/Runtime/CoreUObject/Private/UObject/PropertyBool.cpp b/Engine/Source/Runtime/CoreUObject/Private/UObject/PropertyBool.cpp
--- a/Engine/Source/Runtime/CoreUObject/Private/UObject/PropertyBool.cpp
+++ b/Engine/Source/Runtime/CoreUObject/Private/UObject/PropertyBool.cpp
@@ -319,12 +319,12 @@ bool UBoolProperty::Identical( const void* A, const void* B, uint32 PortFlags )
 	return ((*ByteValueA ^ (B ? *ByteValueB : 0)) & FieldMask) == 0;
 }
 
-void UBoolProperty::SerializeItem( FArchive& Ar, void* Value, void const* Defaults ) const
+void UBoolProperty::SerializeItem( FStructuredArchive::FSlot Slot, void* Value, void const* Defaults ) const
 {
 	check(FieldSize != 0);
 	uint8* ByteValue = (uint8*)Value + ByteOffset;
 	uint8 B = (*ByteValue & FieldMask) ? 1 : 0;
-	Ar << B;
+	Slot << B;
 	*ByteValue = ((*ByteValue) & ~FieldMask) | (B ? ByteMask : 0);
 }
 
diff --git a/Engine/Source/Runtime/CoreUObject/Private/UObject/PropertyByte.cpp b/Engine/Source/Runtime/CoreUObject/Private/UObject/PropertyByte.cpp
--- a/Engine/Source/Runtime/CoreUObject/Private/UObject/PropertyByte.cpp
+++ b/Engine/Source/Runtime/CoreUObject/Private/UObject/PropertyByte.cpp
@@ -6,6 +6,7 @@
 #include "UObject/PropertyPortFlags.h"
 #include "UObject/UnrealType.h"
 #include "UObject/UObjectThreadContext.h"
+#include "Serialization/ArchiveUObjectFromStructuredArchive.h"
 #include "Algo/Find.h"
 
 /*-----------------------------------------------------------------------------
@@ -18,30 +19,32 @@ void UByteProperty::GetPreloadDependencies(TArray<UObject*>& OutDeps)
 	OutDeps.Add(Enum);
 }
 
-void UByteProperty::SerializeItem( FArchive& Ar, void* Value, void const* Defaults ) const
+void UByteProperty::SerializeItem( FStructuredArchive::FSlot Slot, void* Value, void const* Defaults ) const
 {
-	if(Enum && Ar.UseToResolveEnumerators())
+	FArchive& UnderlyingArchive = Slot.GetUnderlyingArchive();
+
+	if(Enum && UnderlyingArchive.UseToResolveEnumerators())
 	{
-		 const int64 ResolvedIndex = Enum->ResolveEnumerator(Ar, *(uint8*)Value);
+		 const int64 ResolvedIndex = Enum->ResolveEnumerator(UnderlyingArchive, *(uint8*)Value);
 		 *(uint8*)Value = static_cast<uint8>(ResolvedIndex);
 		 return;
 	}
 
 	// Serialize enum values by name unless we're not saving or loading OR for backwards compatibility
-	const bool bUseBinarySerialization = (Enum == NULL) || (!Ar.IsLoading() && !Ar.IsSaving());
+	const bool bUseBinarySerialization = (Enum == NULL) || (!UnderlyingArchive.IsLoading() && !UnderlyingArchive.IsSaving());
 	if( bUseBinarySerialization )
 	{
-		Super::SerializeItem(Ar, Value, Defaults);
+		Super::SerializeItem(Slot, Value, Defaults);
 	}
 	// Loading
-	else if (Ar.IsLoading())
+	else if (UnderlyingArchive.IsLoading())
 	{
 		FName EnumValueName;
-		Ar << EnumValueName;
+		Slot << EnumValueName;
 		// Make sure enum is properly populated
 		if( Enum->HasAnyFlags(RF_NeedLoad) )
 		{
-			Ar.Preload(Enum);
+			UnderlyingArchive.Preload(Enum);
 		}
 
 		// There's no guarantee EnumValueName is still present in Enum, in which case Value will be set to the enum's max value.
@@ -72,7 +75,7 @@ void UByteProperty::SerializeItem( FArchive& Ar, void* Value, void const* Defaul
 		{
 			EnumValueName = NAME_None;
 		}
-		Ar << EnumValueName;
+		Slot << EnumValueName;
 	}
 }
 bool UByteProperty::NetSerializeItem( FArchive& Ar, UPackageMap* Map, void* Data, TArray<uint8> * MetaData ) const
diff --git a/Engine/Source/Runtime/CoreUObject/Private/UObject/PropertyDelegate.cpp b/Engine/Source/Runtime/CoreUObject/Private/UObject/PropertyDelegate.cpp
--- a/Engine/Source/Runtime/CoreUObject/Private/UObject/PropertyDelegate.cpp
+++ b/Engine/Source/Runtime/CoreUObject/Private/UObject/PropertyDelegate.cpp
@@ -6,6 +6,7 @@
 #include "UObject/UnrealType.h"
 #include "UObject/PropertyHelper.h"
 #include "UObject/LinkerPlaceholderFunction.h"
+#include "Serialization/ArchiveUObjectFromStructuredArchive.h"
 
 /*-----------------------------------------------------------------------------
 	UDelegateProperty.
@@ -58,8 +59,9 @@ bool UDelegateProperty::Identical( const void* A, const void* B, uint32 PortFlag
 }
 
 
-void UDelegateProperty::SerializeItem( FArchive& Ar, void* Value, void const* Defaults ) const
+void UDelegateProperty::SerializeItem(FStructuredArchive::FSlot Slot, void* Value, void const* Defaults) const
 {
+	FArchiveUObjectFromStructuredArchive Ar(Slot);
 	Ar << *GetPropertyValuePtr(Value);
 }
 
diff --git a/Engine/Source/Runtime/CoreUObject/Private/UObject/PropertyInterface.cpp b/Engine/Source/Runtime/CoreUObject/Private/UObject/PropertyInterface.cpp
--- a/Engine/Source/Runtime/CoreUObject/Private/UObject/PropertyInterface.cpp
+++ b/Engine/Source/Runtime/CoreUObject/Private/UObject/PropertyInterface.cpp
@@ -110,12 +110,13 @@ bool UInterfaceProperty::Identical( const void* A, const void* B, uint32 PortFla
 	return (InterfaceA->GetObject() == InterfaceB->GetObject() && InterfaceA->GetInterface() == InterfaceB->GetInterface());
 }
 
-void UInterfaceProperty::SerializeItem( FArchive& Ar, void* Value, void const* Defaults ) const
+void UInterfaceProperty::SerializeItem( FStructuredArchive::FSlot Slot, void* Value, void const* Defaults ) const
 {
+	FArchive& UnderlyingArchive = Slot.GetUnderlyingArchive();
 	FScriptInterface* InterfaceValue = (FScriptInterface*)Value;
 
-	Ar << InterfaceValue->GetObjectRef();
-	if ( Ar.IsLoading() || Ar.IsTransacting() || Ar.IsObjectReferenceCollector() )
+	Slot << InterfaceValue->GetObjectRef();
+	if (UnderlyingArchive.IsLoading() || UnderlyingArchive.IsTransacting() || UnderlyingArchive.IsObjectReferenceCollector())
 	{
 		if ( InterfaceValue->GetObject() != NULL )
 		{
diff --git a/Engine/Source/Runtime/CoreUObject/Private/UObject/PropertyLazyObjectPtr.cpp b/Engine/Source/Runtime/CoreUObject/Private/UObject/PropertyLazyObjectPtr.cpp
--- a/Engine/Source/Runtime/CoreUObject/Private/UObject/PropertyLazyObjectPtr.cpp
+++ b/Engine/Source/Runtime/CoreUObject/Private/UObject/PropertyLazyObjectPtr.cpp
@@ -25,18 +25,20 @@ FName ULazyObjectProperty::GetID() const
 	return NAME_LazyObjectProperty;
 }
 
-void ULazyObjectProperty::SerializeItem( FArchive& Ar, void* Value, void const* Defaults ) const
+void ULazyObjectProperty::SerializeItem( FStructuredArchive::FSlot Slot, void* Value, void const* Defaults ) const
 {
+	FArchive& UnderlyingArchive = Slot.GetUnderlyingArchive();
+
 	// We never serialize our reference while the garbage collector is harvesting references
 	// to objects, because we don't want lazy pointers to keep objects from being garbage collected
 
-	if( !Ar.IsObjectReferenceCollector() || Ar.IsModifyingWeakAndStrongReferences() )
+	if( !UnderlyingArchive.IsObjectReferenceCollector() || UnderlyingArchive.IsModifyingWeakAndStrongReferences() )
 	{
 		UObject* ObjectValue = GetObjectPropertyValue(Value);
 
-		Ar << *(FLazyObjectPtr*)Value;
+		Slot << *(FLazyObjectPtr*)Value;
 
-		if ((Ar.IsLoading() || Ar.IsModifyingWeakAndStrongReferences()) && ObjectValue != GetObjectPropertyValue(Value))
+		if ((UnderlyingArchive.IsLoading() || UnderlyingArchive.IsModifyingWeakAndStrongReferences()) && ObjectValue != GetObjectPropertyValue(Value))
 		{
 			CheckValidObject(Value);
 		}
diff --git a/Engine/Source/Runtime/CoreUObject/Private/UObject/PropertyMap.cpp b/Engine/Source/Runtime/CoreUObject/Private/UObject/PropertyMap.cpp
--- a/Engine/Source/Runtime/CoreUObject/Private/UObject/PropertyMap.cpp
+++ b/Engine/Source/Runtime/CoreUObject/Private/UObject/PropertyMap.cpp
@@ -8,6 +8,7 @@
 #include "UObject/LinkerLoad.h"
 #include "UObject/PropertyHelper.h"
 #include "Misc/ScopeExit.h"
+#include "Serialization/ArchiveUObjectFromStructuredArchive.h"
 
 namespace UE4MapProperty_Private
 {
@@ -238,8 +239,10 @@ void UMapProperty::GetPreloadDependencies(TArray<UObject*>& OutDeps)
 	OutDeps.Add(ValueProp);
 }
 
-void UMapProperty::SerializeItem(FArchive& Ar, void* Value, const void* Defaults) const
+void UMapProperty::SerializeItem(FStructuredArchive::FSlot Slot, void* Value, const void* Defaults) const
 {
+	FArchiveUObjectFromStructuredArchive Ar(Slot);
+
 	// Ar related calls in this function must be mirrored in UMapProperty::ConvertFromType
 	checkSlow(KeyProp);
 	checkSlow(ValueProp);
@@ -283,7 +286,7 @@ void UMapProperty::SerializeItem(FArchive& Ar, void* Value, const void* Defaults
 			for (; NumKeysToRemove; --NumKeysToRemove)
 			{
 				// Read key into temporary storage
-				KeyProp->SerializeItem(Ar, TempKeyStorage);
+				KeyProp->SerializeItem(FStructuredArchiveFromArchive(Ar).GetSlot(), TempKeyStorage);
 
 				// If the key is in the map, remove it
 				int32 Found = MapHelper.FindMapIndexWithKey(TempKeyStorage);
@@ -310,7 +313,7 @@ void UMapProperty::SerializeItem(FArchive& Ar, void* Value, const void* Defaults
 			// Read key into temporary storage
 			{
 				FSerializedPropertyScope SerializedProperty(Ar, KeyProp, this);
-				KeyProp->SerializeItem(Ar, TempKeyStorage);
+				KeyProp->SerializeItem(FStructuredArchiveFromArchive(Ar).GetSlot(), TempKeyStorage);
 			}
 			
 			// Add a new default value if the key doesn't currently exist in the map
@@ -328,7 +331,7 @@ void UMapProperty::SerializeItem(FArchive& Ar, void* Value, const void* Defaults
 			// Deserialize value
 			{
 				FSerializedPropertyScope SerializedProperty(Ar, ValueProp, this);
-				ValueProp->SerializeItem(Ar, NextPairPtr + MapLayout.ValueOffset);
+				ValueProp->SerializeItem(FStructuredArchiveFromArchive(Ar).GetSlot(), NextPairPtr + MapLayout.ValueOffset);
 			}
 		}
 
@@ -367,7 +370,7 @@ void UMapProperty::SerializeItem(FArchive& Ar, void* Value, const void* Defaults
 			FSerializedPropertyScope SerializedProperty(Ar, KeyProp, this);
 			for (int32 Index : Indices)
 			{
-				KeyProp->SerializeItem(Ar, DefaultsHelper.GetPairPtr(Index));
+				KeyProp->SerializeItem(FStructuredArchiveFromArchive(Ar).GetSlot(), DefaultsHelper.GetPairPtr(Index));
 			}
 		}
 
@@ -400,11 +403,11 @@ void UMapProperty::SerializeItem(FArchive& Ar, void* Value, const void* Defaults
 
 				{
 					FSerializedPropertyScope SerializedProperty(Ar, KeyProp, this);
-					KeyProp->SerializeItem(Ar, ValuePairPtr);
+					KeyProp->SerializeItem(FStructuredArchiveFromArchive(Ar).GetSlot(), ValuePairPtr);
 				}
 				{
 					FSerializedPropertyScope SerializedProperty(Ar, ValueProp, this);
-					ValueProp->SerializeItem(Ar, ValuePairPtr + MapLayout.ValueOffset);
+					ValueProp->SerializeItem(FStructuredArchiveFromArchive(Ar).GetSlot(), ValuePairPtr + MapLayout.ValueOffset);
 				}
 			}
 		}
@@ -420,11 +423,11 @@ void UMapProperty::SerializeItem(FArchive& Ar, void* Value, const void* Defaults
 
 					{
 						FSerializedPropertyScope SerializedProperty(Ar, KeyProp, this);
-						KeyProp->SerializeItem(Ar, ValuePairPtr);
+						KeyProp->SerializeItem(FStructuredArchiveFromArchive(Ar).GetSlot(), ValuePairPtr);
 					}
 					{
 						FSerializedPropertyScope SerializedProperty(Ar, ValueProp, this);
-						ValueProp->SerializeItem(Ar, ValuePairPtr + MapLayout.ValueOffset);
+						ValueProp->SerializeItem(FStructuredArchiveFromArchive(Ar).GetSlot(), ValuePairPtr + MapLayout.ValueOffset);
 					}
 
 					--Num;
@@ -900,7 +903,7 @@ EConvertFromTypeResult UMapProperty::ConvertFromType(const FPropertyTag& Tag, FA
 		if(CurrentType->GetID() == InTag.Type)
 		{
 			uint8* DestAddress = CurrentType->ContainerPtrToValuePtr<uint8>(InData, InTag.ArrayIndex);
-			CurrentType->SerializeItem(InAr, DestAddress, nullptr);
+			CurrentType->SerializeItem(FStructuredArchiveFromArchive(InAr).GetSlot(), DestAddress, nullptr);
 			return true;
 		}
 		else if( CurrentType->ConvertFromType(InTag, InAr, InData, InDefaultsStruct) == EConvertFromTypeResult::Converted )
diff --git a/Engine/Source/Runtime/CoreUObject/Private/UObject/PropertyMulticastDelegate.cpp b/Engine/Source/Runtime/CoreUObject/Private/UObject/PropertyMulticastDelegate.cpp
--- a/Engine/Source/Runtime/CoreUObject/Private/UObject/PropertyMulticastDelegate.cpp
+++ b/Engine/Source/Runtime/CoreUObject/Private/UObject/PropertyMulticastDelegate.cpp
@@ -6,6 +6,7 @@
 #include "UObject/UnrealType.h"
 #include "UObject/PropertyHelper.h"
 #include "UObject/LinkerPlaceholderFunction.h"
+#include "Serialization/ArchiveUObjectFromStructuredArchive.h"
 
 /*-----------------------------------------------------------------------------
 	UMulticastDelegateProperty.
@@ -109,8 +110,9 @@ bool UMulticastDelegateProperty::Identical( const void* A, const void* B, uint32
 	return true;
 }
 
-void UMulticastDelegateProperty::SerializeItem( FArchive& Ar, void* Value, void const* Defaults ) const
+void UMulticastDelegateProperty::SerializeItem(FStructuredArchive::FSlot Slot, void* Value, void const* Defaults) const
 {
+	FArchiveUObjectFromStructuredArchive Ar(Slot);
 	Ar << *GetPropertyValuePtr(Value);
 }
 
diff --git a/Engine/Source/Runtime/CoreUObject/Private/UObject/PropertyObject.cpp b/Engine/Source/Runtime/CoreUObject/Private/UObject/PropertyObject.cpp
--- a/Engine/Source/Runtime/CoreUObject/Private/UObject/PropertyObject.cpp
+++ b/Engine/Source/Runtime/CoreUObject/Private/UObject/PropertyObject.cpp
@@ -53,11 +53,6 @@ EConvertFromTypeResult UObjectProperty::ConvertFromType(const FPropertyTag& Tag,
 	return EConvertFromTypeResult::UseSerializeItem;
 }
 
-void UObjectProperty::SerializeItem(FArchive& Ar, void* Value, void const* Defaults) const
-{
-	SerializeItem(FStructuredArchiveFromArchive(Ar).GetSlot(), Value, Defaults);
-}
-
 void UObjectProperty::SerializeItem( FStructuredArchive::FSlot Slot, void* Value, void const* Defaults ) const
 {
 	FArchive& UnderlyingArchive = Slot.GetUnderlyingArchive();
diff --git a/Engine/Source/Runtime/CoreUObject/Private/UObject/PropertySet.cpp b/Engine/Source/Runtime/CoreUObject/Private/UObject/PropertySet.cpp
--- a/Engine/Source/Runtime/CoreUObject/Private/UObject/PropertySet.cpp
+++ b/Engine/Source/Runtime/CoreUObject/Private/UObject/PropertySet.cpp
@@ -7,6 +7,7 @@
 #include "UObject/LinkerLoad.h"
 #include "UObject/PropertyHelper.h"
 #include "Misc/ScopeExit.h"
+#include "Serialization/ArchiveUObjectFromStructuredArchive.h"
 
 namespace UE4SetProperty_Private
 {
@@ -212,8 +213,10 @@ void USetProperty::GetPreloadDependencies(TArray<UObject*>& OutDeps)
 	OutDeps.Add(ElementProp);
 }
 
-void USetProperty::SerializeItem(FArchive& Ar, void* Value, const void* Defaults) const
+void USetProperty::SerializeItem(FStructuredArchive::FSlot Slot, void* Value, const void* Defaults) const
 {
+	FArchiveUObjectFromStructuredArchive Ar(Slot);
+
 	// Ar related calls in this function must be mirrored in USetProperty::ConvertFromType
 	checkSlow(ElementProp);
 
@@ -255,7 +258,7 @@ void USetProperty::SerializeItem(FArchive& Ar, void* Value, const void* Defaults
 			for (; NumElementsToRemove; --NumElementsToRemove)
 			{
 				// Read key into temporary storage
-				ElementProp->SerializeItem(Ar, TempElementStorage);
+				ElementProp->SerializeItem(FStructuredArchiveFromArchive(Ar).GetSlot(), TempElementStorage);
 
 				// If the key is in the map, remove it
 				const int32 Found = SetHelper.FindElementIndex(TempElementStorage);
@@ -281,7 +284,7 @@ void USetProperty::SerializeItem(FArchive& Ar, void* Value, const void* Defaults
 		for (; Num; --Num)
 		{
 			// Read key into temporary storage
-			ElementProp->SerializeItem(Ar, TempElementStorage);
+			ElementProp->SerializeItem(FStructuredArchiveFromArchive(Ar).GetSlot(), TempElementStorage);
 
 			// Add a new entry if the element doesn't currently exist in the set
 			if (SetHelper.FindElementIndex(TempElementStorage) == INDEX_NONE)
@@ -329,7 +332,7 @@ void USetProperty::SerializeItem(FArchive& Ar, void* Value, const void* Defaults
 			FSerializedPropertyScope SerializedProperty(Ar, ElementProp, this);
 			for (int32 Index : Indices)
 			{
-				ElementProp->SerializeItem(Ar, DefaultsHelper.GetElementPtr(Index));
+				ElementProp->SerializeItem(FStructuredArchiveFromArchive(Ar).GetSlot(), DefaultsHelper.GetElementPtr(Index));
 			}
 		}
 
@@ -362,7 +365,7 @@ void USetProperty::SerializeItem(FArchive& Ar, void* Value, const void* Defaults
 			{
 				uint8* ElementPtr = SetHelper.GetElementPtrWithoutCheck(Index);
 
-				ElementProp->SerializeItem(Ar, ElementPtr);
+				ElementProp->SerializeItem(FStructuredArchiveFromArchive(Ar).GetSlot(), ElementPtr);
 			}
 		}
 		else
@@ -377,7 +380,7 @@ void USetProperty::SerializeItem(FArchive& Ar, void* Value, const void* Defaults
 				{
 					uint8* ElementPtr = SetHelper.GetElementPtrWithoutCheck(Index);
 
-					ElementProp->SerializeItem(Ar, ElementPtr);
+					ElementProp->SerializeItem(FStructuredArchiveFromArchive(Ar).GetSlot(), ElementPtr);
 
 					--Num;
 				}
diff --git a/Engine/Source/Runtime/CoreUObject/Private/UObject/PropertySoftObjectPtr.cpp b/Engine/Source/Runtime/CoreUObject/Private/UObject/PropertySoftObjectPtr.cpp
--- a/Engine/Source/Runtime/CoreUObject/Private/UObject/PropertySoftObjectPtr.cpp
+++ b/Engine/Source/Runtime/CoreUObject/Private/UObject/PropertySoftObjectPtr.cpp
@@ -41,17 +41,19 @@ bool USoftObjectProperty::Identical( const void* A, const void* B, uint32 PortFl
 	return ObjectA.GetUniqueID() == ObjectB.GetUniqueID();
 }
 
-void USoftObjectProperty::SerializeItem( FArchive& Ar, void* Value, void const* Defaults ) const
+void USoftObjectProperty::SerializeItem( FStructuredArchive::FSlot Slot, void* Value, void const* Defaults ) const
 {
+	FArchive& UnderlyingArchive = Slot.GetUnderlyingArchive();
+
 	// We never serialize our reference while the garbage collector is harvesting references
 	// to objects, because we don't want soft object pointers to keep objects from being garbage collected
 	// Allow persistent archives so they can keep track of string references. (e.g. FArchiveSaveTagImports)
-	if( !Ar.IsObjectReferenceCollector() || Ar.IsModifyingWeakAndStrongReferences() || Ar.IsPersistent() )
+	if( !UnderlyingArchive.IsObjectReferenceCollector() || UnderlyingArchive.IsModifyingWeakAndStrongReferences() || UnderlyingArchive.IsPersistent() )
 	{
 		FSoftObjectPtr OldValue = *(FSoftObjectPtr*)Value;
-		Ar << *(FSoftObjectPtr*)Value;
+		Slot << *(FSoftObjectPtr*)Value;
 
-		if (Ar.IsLoading() || Ar.IsModifyingWeakAndStrongReferences()) 
+		if (UnderlyingArchive.IsLoading() || UnderlyingArchive.IsModifyingWeakAndStrongReferences()) 
 		{
 			if (OldValue.GetUniqueID() != ((FSoftObjectPtr*)Value)->GetUniqueID())
 			{
diff --git a/Engine/Source/Runtime/CoreUObject/Private/UObject/PropertyStruct.cpp b/Engine/Source/Runtime/CoreUObject/Private/UObject/PropertyStruct.cpp
--- a/Engine/Source/Runtime/CoreUObject/Private/UObject/PropertyStruct.cpp
+++ b/Engine/Source/Runtime/CoreUObject/Private/UObject/PropertyStruct.cpp
@@ -7,6 +7,7 @@
 #include "UObject/UnrealType.h"
 #include "UObject/PropertyHelper.h"
 #include "UObject/LinkerPlaceholderBase.h"
+#include "Serialization/ArchiveUObjectFromStructuredArchive.h"
 
 static inline void PreloadInnerStructMembers(UStructProperty* StructProperty)
 {
@@ -105,7 +106,7 @@ uint32 UStructProperty::GetValueTypeHashInternal(const void* Src) const
 	return Struct->GetStructTypeHash(Src);
 }
 
-void UStructProperty::SerializeItem( FArchive& Ar, void* Value, void const* Defaults ) const
+void UStructProperty::SerializeItem(FStructuredArchive::FSlot Slot, void* Value, void const* Defaults) const
 {
 	check(Struct);
 
@@ -113,6 +114,7 @@ void UStructProperty::SerializeItem( FArchive& Ar, void* Value, void const* Defa
 	FScopedPlaceholderPropertyTracker ImportPropertyTracker(this);
 #endif
 
+	FArchiveUObjectFromStructuredArchive Ar(Slot);
 	Struct->SerializeItem(Ar, Value, Defaults);
 }
 
diff --git a/Engine/Source/Runtime/CoreUObject/Private/UObject/PropertyWeakObjectPtr.cpp b/Engine/Source/Runtime/CoreUObject/Private/UObject/PropertyWeakObjectPtr.cpp
--- a/Engine/Source/Runtime/CoreUObject/Private/UObject/PropertyWeakObjectPtr.cpp
+++ b/Engine/Source/Runtime/CoreUObject/Private/UObject/PropertyWeakObjectPtr.cpp
@@ -33,11 +33,13 @@ FString UWeakObjectProperty::GetCPPMacroType( FString& ExtendedTypeText ) const
 	return TEXT("WEAKOBJECT");
 }
 
-void UWeakObjectProperty::SerializeItem( FArchive& Ar, void* Value, void const* Defaults ) const
+void UWeakObjectProperty::SerializeItem( FStructuredArchive::FSlot Slot, void* Value, void const* Defaults ) const
 {
+	FArchive& UnderlyingArchive = Slot.GetUnderlyingArchive();
+
 	UObject* ObjectValue = GetObjectPropertyValue(Value);
-	Ar << *(FWeakObjectPtr*)Value;
-	if ((Ar.IsLoading() || Ar.IsModifyingWeakAndStrongReferences()) && ObjectValue != GetObjectPropertyValue(Value))
+	Slot << *(FWeakObjectPtr*)Value;
+	if ((UnderlyingArchive.IsLoading() || UnderlyingArchive.IsModifyingWeakAndStrongReferences()) && ObjectValue != GetObjectPropertyValue(Value))
 	{
 		CheckValidObject(Value);
 	}
diff --git a/Engine/Source/Runtime/CoreUObject/Private/UObject/TextProperty.cpp b/Engine/Source/Runtime/CoreUObject/Private/UObject/TextProperty.cpp
--- a/Engine/Source/Runtime/CoreUObject/Private/UObject/TextProperty.cpp
+++ b/Engine/Source/Runtime/CoreUObject/Private/UObject/TextProperty.cpp
@@ -7,6 +7,7 @@
 #include "Internationalization/TextNamespaceUtil.h"
 #include "Internationalization/TextPackageNamespaceUtil.h"
 #include "Internationalization/StringTableRegistry.h"
+#include "Serialization/ArchiveUObjectFromStructuredArchive.h"
 
 EConvertFromTypeResult UTextProperty::ConvertFromType(const FPropertyTag& Tag, FArchive& Ar, uint8* Data, UStruct* DefaultsStruct)
 {
@@ -81,9 +82,11 @@ bool UTextProperty::Identical( const void* A, const void* B, uint32 PortFlags )
 	return FTextInspector::GetDisplayString(ValueA).IsEmpty();
 }
 
-void UTextProperty::SerializeItem( FArchive& Ar, void* Value, void const* Defaults ) const
+void UTextProperty::SerializeItem(FStructuredArchive::FSlot Slot, void* Value, void const* Defaults) const
 {
 	TCppType* TextPtr = GetPropertyValuePtr(Value);
+
+	FArchiveUObjectFromStructuredArchive Ar(Slot);
 	Ar << *TextPtr;
 }
 
diff --git a/Engine/Source/Runtime/CoreUObject/Private/UObject/UnrealType.cpp b/Engine/Source/Runtime/CoreUObject/Private/UObject/UnrealType.cpp
--- a/Engine/Source/Runtime/CoreUObject/Private/UObject/UnrealType.cpp
+++ b/Engine/Source/Runtime/CoreUObject/Private/UObject/UnrealType.cpp
@@ -172,10 +172,4 @@ void FPropertyValueIterator::GetPropertyChain(TArray<const UProperty*>& Property
 			Property = Cast<UProperty>(Property->GetOuter());
 		}
 	}
-}
-
-void UProperty::SerializeItem(FStructuredArchive::FSlot Slot, void* Value, void const* Defaults) const
-{
-	FArchiveUObjectFromStructuredArchive Ar(Slot);
-	SerializeItem(Ar, Value, Defaults);
 }
\ No newline at end of file
diff --git a/Engine/Source/Runtime/CoreUObject/Public/UObject/EnumProperty.h b/Engine/Source/Runtime/CoreUObject/Public/UObject/EnumProperty.h
--- a/Engine/Source/Runtime/CoreUObject/Public/UObject/EnumProperty.h
+++ b/Engine/Source/Runtime/CoreUObject/Public/UObject/EnumProperty.h
@@ -38,7 +38,7 @@ public:
 	virtual FString GetCPPTypeForwardDeclaration() const override;
 	virtual void LinkInternal(FArchive& Ar) override;
 	virtual bool Identical( const void* A, const void* B, uint32 PortFlags ) const override;
-	virtual void SerializeItem( FArchive& Ar, void* Value, void const* Defaults ) const override;
+	virtual void SerializeItem(FStructuredArchive::FSlot Slot, void* Value, void const* Defaults) const override;
 	virtual bool NetSerializeItem( FArchive& Ar, UPackageMap* Map, void* Data, TArray<uint8> * MetaData = NULL ) const override;
 	virtual void ExportTextItem( FString& ValueStr, const void* PropertyValue, const void* DefaultValue, UObject* Parent, int32 PortFlags, UObject* ExportRootScope ) const override;
 	virtual const TCHAR* ImportText_Internal( const TCHAR* Buffer, void* Data, int32 PortFlags, UObject* OwnerObject, FOutputDevice* ErrorText ) const override;
diff --git a/Engine/Source/Runtime/CoreUObject/Public/UObject/FastReferenceCollector.h b/Engine/Source/Runtime/CoreUObject/Public/UObject/FastReferenceCollector.h
--- a/Engine/Source/Runtime/CoreUObject/Public/UObject/FastReferenceCollector.h
+++ b/Engine/Source/Runtime/CoreUObject/Public/UObject/FastReferenceCollector.h
@@ -779,7 +779,7 @@ private:
 						void*         Map = StackEntryData + ReferenceInfo.Offset;
 						UMapProperty* MapProperty = (UMapProperty*)TokenStream->ReadPointer(TokenStreamIndex);
 						TokenReturnCount = ReferenceInfo.ReturnCount;
-						MapProperty->SerializeItem(ReferenceCollector.GetVerySlowReferenceCollectorArchive(), Map, nullptr);
+						MapProperty->SerializeItem(FStructuredArchiveFromArchive(ReferenceCollector.GetVerySlowReferenceCollectorArchive()).GetSlot(), Map, nullptr);
 					}
 					break;
 					case GCRT_AddTSetReferencedObjects:
@@ -787,7 +787,7 @@ private:
 						void*         Set = StackEntryData + ReferenceInfo.Offset;
 						USetProperty* SetProperty = (USetProperty*)TokenStream->ReadPointer(TokenStreamIndex);
 						TokenReturnCount = ReferenceInfo.ReturnCount;
-						SetProperty->SerializeItem(ReferenceCollector.GetVerySlowReferenceCollectorArchive(), Set, nullptr);
+						SetProperty->SerializeItem(FStructuredArchiveFromArchive(ReferenceCollector.GetVerySlowReferenceCollectorArchive()).GetSlot(), Set, nullptr);
 					}
 					break;
 					case GCRT_EndOfPointer:
diff --git a/Engine/Source/Runtime/CoreUObject/Public/UObject/TextProperty.h b/Engine/Source/Runtime/CoreUObject/Public/UObject/TextProperty.h
--- a/Engine/Source/Runtime/CoreUObject/Public/UObject/TextProperty.h
+++ b/Engine/Source/Runtime/CoreUObject/Public/UObject/TextProperty.h
@@ -31,7 +31,7 @@ public:
 	// UProperty interface
 	virtual EConvertFromTypeResult ConvertFromType(const FPropertyTag& Tag, FArchive& Ar, uint8* Data, UStruct* DefaultsStruct) override;
 	virtual bool Identical( const void* A, const void* B, uint32 PortFlags ) const override;
-	virtual void SerializeItem( FArchive& Ar, void* Value, void const* Defaults ) const override;
+	virtual void SerializeItem(FStructuredArchive::FSlot Slot, void* Value, void const* Defaults) const override;
 	virtual void ExportTextItem( FString& ValueStr, const void* PropertyValue, const void* DefaultValue, UObject* Parent, int32 PortFlags, UObject* ExportRootScope ) const override;
 	virtual const TCHAR* ImportText_Internal( const TCHAR* Buffer, void* Data, int32 PortFlags, UObject* OwnerObject, FOutputDevice* ErrorText ) const override;
 	virtual FString GetCPPTypeForwardDeclaration() const override;
diff --git a/Engine/Source/Runtime/CoreUObject/Public/UObject/UnrealType.h b/Engine/Source/Runtime/CoreUObject/Public/UObject/UnrealType.h
--- a/Engine/Source/Runtime/CoreUObject/Public/UObject/UnrealType.h
+++ b/Engine/Source/Runtime/CoreUObject/Public/UObject/UnrealType.h
@@ -272,7 +272,7 @@ public:
 			{
 				// Keep setting the property in case something inside of SerializeItem changes it
 				FSerializedPropertyScope SerializedProperty(Ar, this);
-				SerializeItem( Ar, ContainerPtrToValuePtr<void>(Data, Idx) );
+				SerializeItem( FStructuredArchiveFromArchive(Ar).GetSlot(), ContainerPtrToValuePtr<void>(Data, Idx) );
 			}
 		}
 	}
@@ -295,14 +295,13 @@ public:
 				if ( !Identical(Target, Default, Ar.GetPortFlags()) )
 				{
 					FSerializedPropertyScope SerializedProperty(Ar, this);
-					SerializeItem( Ar, Target, Default );
+					SerializeItem( FStructuredArchiveFromArchive(Ar).GetSlot(), Target, Default );
 				}
 			}
 		}
 	}
 
-	virtual void SerializeItem(FArchive& Ar, void* Value, void const* Defaults=NULL ) const PURE_VIRTUAL(UProperty::SerializeItem,);
-	virtual void SerializeItem(FStructuredArchive::FSlot Slot, void* Value, void const* Defaults = NULL) const;
+	virtual void SerializeItem(FStructuredArchive::FSlot Slot, void* Value, void const* Defaults = NULL) const PURE_VIRTUAL(UProperty::SerializeItem, );
 	virtual bool NetSerializeItem( FArchive& Ar, UPackageMap* Map, void* Data, TArray<uint8> * MetaData = NULL ) const;
 	virtual bool SupportsNetSharedSerialization() const;
 	virtual void ExportTextItem( FString& ValueStr, const void* PropertyValue, const void* DefaultValue, UObject* Parent, int32 PortFlags, UObject* ExportRootScope = NULL ) const PURE_VIRTUAL(UProperty::ExportTextItem,);
@@ -1109,9 +1108,9 @@ public:
 	{
 		return TTypeFundamentals::GetPropertyValue(A) == TTypeFundamentals::GetOptionalPropertyValue(B);
 	}
-	virtual void SerializeItem( FArchive& Ar, void* Value, void const* Defaults ) const override
+	virtual void SerializeItem(FStructuredArchive::FSlot Slot, void* Value, void const* Defaults) const override
 	{
-		Ar << *TTypeFundamentals::GetPropertyValuePtr(Value);
+		Slot << *TTypeFundamentals::GetPropertyValuePtr(Value);
 	}
 	// End of UProperty interface
 
@@ -1465,7 +1464,7 @@ class COREUOBJECT_API UByteProperty : public TProperty_Numeric<uint8>
 	// End of UHT interface
 
 	// UProperty interface.
-	virtual void SerializeItem( FArchive& Ar, void* Value, void const* Defaults ) const override;
+	virtual void SerializeItem(FStructuredArchive::FSlot Slot, void* Value, void const* Defaults) const override;
 	virtual bool NetSerializeItem( FArchive& Ar, UPackageMap* Map, void* Data, TArray<uint8> * MetaData = NULL ) const override;
 	virtual void ExportTextItem( FString& ValueStr, const void* PropertyValue, const void* DefaultValue, UObject* Parent, int32 PortFlags, UObject* ExportRootScope ) const override;
 	virtual const TCHAR* ImportText_Internal( const TCHAR* Buffer, void* Data, int32 PortFlags, UObject* Parent, FOutputDevice* ErrorText ) const override;
@@ -1767,7 +1766,7 @@ public:
 	// UProperty interface.
 	virtual void LinkInternal(FArchive& Ar) override;
 	virtual bool Identical( const void* A, const void* B, uint32 PortFlags ) const override;
-	virtual void SerializeItem( FArchive& Ar, void* Value, void const* Defaults ) const override;
+	virtual void SerializeItem( FStructuredArchive::FSlot Slot, void* Value, void const* Defaults) const override;
 	virtual bool NetSerializeItem( FArchive& Ar, UPackageMap* Map, void* Data, TArray<uint8> * MetaData = NULL ) const override;
 	virtual void ExportTextItem( FString& ValueStr, const void* PropertyValue, const void* DefaultValue, UObject* Parent, int32 PortFlags, UObject* ExportRootScope ) const override;
 	virtual const TCHAR* ImportText_Internal( const TCHAR* Buffer, void* Data, int32 PortFlags, UObject* Parent, FOutputDevice* ErrorText ) const override;
@@ -2054,7 +2053,6 @@ class COREUOBJECT_API UObjectProperty : public TUObjectPropertyBase<UObject*>
 	// End of UHT interface
 
 	// UProperty interface
-	virtual void SerializeItem(FArchive& Ar, void* Value, void const* Defaults) const override;
 	virtual void SerializeItem(FStructuredArchive::FSlot Slot, void* Value, void const* Defaults) const override;
 	virtual void EmitReferenceInfo(UClass& OwnerClass, int32 BaseOffset, TArray<const UStructProperty*>& EncounteredStructProps) override;
 	virtual const TCHAR* ImportText_Internal(const TCHAR* Buffer, void* Data, int32 PortFlags, UObject* OwnerObject, FOutputDevice* ErrorText) const override;
@@ -2096,7 +2094,7 @@ class COREUOBJECT_API UWeakObjectProperty : public TUObjectPropertyBase<FWeakObj
 	// End of UHT interface
 
 	// UProperty interface
-	virtual void SerializeItem( FArchive& Ar, void* Value, void const* Defaults ) const override;
+	virtual void SerializeItem(FStructuredArchive::FSlot Slot, void* Value, void const* Defaults) const override;
 	// End of UProperty interface
 
 	// UObjectProperty interface
@@ -2130,7 +2128,7 @@ class COREUOBJECT_API ULazyObjectProperty : public TUObjectPropertyBase<FLazyObj
 	// UProperty interface
 	virtual FName GetID() const override;
 	virtual bool Identical( const void* A, const void* B, uint32 PortFlags ) const override;
-	virtual void SerializeItem( FArchive& Ar, void* Value, void const* Defaults ) const override;
+	virtual void SerializeItem( FStructuredArchive::FSlot Slot, void* Value, void const* Defaults ) const override;
 	// End of UProperty interface
 
 	// UObjectProperty interface
@@ -2166,7 +2164,7 @@ class COREUOBJECT_API USoftObjectProperty : public TUObjectPropertyBase<FSoftObj
 	// UProperty interface
 	virtual FName GetID() const override;
 	virtual bool Identical( const void* A, const void* B, uint32 PortFlags ) const override;
-	virtual void SerializeItem( FArchive& Ar, void* Value, void const* Defaults ) const override;
+	virtual void SerializeItem( FStructuredArchive::FSlot Slot, void* Value, void const* Defaults) const override;
 	virtual void ExportTextItem( FString& ValueStr, const void* PropertyValue, const void* DefaultValue, UObject* Parent, int32 PortFlags, UObject* ExportRootScope ) const override;
 	virtual const TCHAR* ImportText_Internal( const TCHAR* Buffer, void* Data, int32 PortFlags, UObject* OwnerObject, FOutputDevice* ErrorText ) const override;
 	virtual EConvertFromTypeResult ConvertFromType(const FPropertyTag& Tag, FArchive& Ar, uint8* Data, UStruct* DefaultsStruct) override;
@@ -2354,7 +2352,7 @@ public:
 	// UProperty interface
 	virtual void LinkInternal(FArchive& Ar) override;
 	virtual bool Identical( const void* A, const void* B, uint32 PortFlags ) const override;
-	virtual void SerializeItem( FArchive& Ar, void* Value, void const* Defaults ) const override;
+	virtual void SerializeItem( FStructuredArchive::FSlot Slot, void* Value, void const* Defaults) const override;
 	virtual bool NetSerializeItem( FArchive& Ar, UPackageMap* Map, void* Data, TArray<uint8> * MetaData = NULL ) const override;
 	virtual void ExportTextItem( FString& ValueStr, const void* PropertyValue, const void* DefaultValue, UObject* Parent, int32 PortFlags, UObject* ExportRootScope ) const override;
 	virtual const TCHAR* ImportText_Internal( const TCHAR* Buffer, void* Data, int32 PortFlags, UObject* OwnerObject, FOutputDevice* ErrorText ) const override;
@@ -2510,7 +2508,6 @@ public:
 	virtual FString GetCPPTypeForwardDeclaration() const override;
 	virtual void LinkInternal(FArchive& Ar) override;
 	virtual bool Identical( const void* A, const void* B, uint32 PortFlags ) const override;
-	virtual void SerializeItem( FArchive& Ar, void* Value, void const* Defaults ) const override;
 	virtual void SerializeItem(FStructuredArchive::FSlot Slot, void* Value, void const* Defaults) const override;
 	virtual bool NetSerializeItem( FArchive& Ar, UPackageMap* Map, void* Data, TArray<uint8> * MetaData = NULL ) const override;
 	virtual void ExportTextItem( FString& ValueStr, const void* PropertyValue, const void* DefaultValue, UObject* Parent, int32 PortFlags, UObject* ExportRootScope ) const override;
@@ -2564,7 +2561,7 @@ public:
 	virtual FString GetCPPTypeForwardDeclaration() const override;
 	virtual void LinkInternal(FArchive& Ar) override;
 	virtual bool Identical(const void* A, const void* B, uint32 PortFlags) const override;
-	virtual void SerializeItem(FArchive& Ar, void* Value, void const* Defaults) const override;
+	virtual void SerializeItem(FStructuredArchive::FSlot Slot, void* Value, void const* Defaults) const override;
 	virtual bool NetSerializeItem(FArchive& Ar, UPackageMap* Map, void* Data, TArray<uint8> * MetaData = NULL) const override;
 	virtual void ExportTextItem(FString& ValueStr, const void* PropertyValue, const void* DefaultValue, UObject* Parent, int32 PortFlags, UObject* ExportRootScope) const override;
 	virtual const TCHAR* ImportText_Internal(const TCHAR* Buffer, void* Data, int32 PortFlags, UObject* OwnerObject, FOutputDevice* ErrorText) const override;
@@ -2616,7 +2613,7 @@ public:
 	virtual FString GetCPPTypeForwardDeclaration() const override;
 	virtual void LinkInternal(FArchive& Ar) override;
 	virtual bool Identical(const void* A, const void* B, uint32 PortFlags) const override;
-	virtual void SerializeItem(FArchive& Ar, void* Value, void const* Defaults) const override;
+	virtual void SerializeItem(FStructuredArchive::FSlot Slot, void* Value, void const* Defaults) const override;
 	virtual bool NetSerializeItem(FArchive& Ar, UPackageMap* Map, void* Data, TArray<uint8> * MetaData = NULL) const override;
 	virtual void ExportTextItem(FString& ValueStr, const void* PropertyValue, const void* DefaultValue, UObject* Parent, int32 PortFlags, UObject* ExportRootScope) const override;
 	virtual const TCHAR* ImportText_Internal(const TCHAR* Buffer, void* Data, int32 PortFlags, UObject* OwnerObject, FOutputDevice* ErrorText) const override;
@@ -4046,7 +4043,7 @@ public:
 	virtual FString GetCPPTypeForwardDeclaration() const override;
 	virtual void LinkInternal(FArchive& Ar) override;
 	virtual bool Identical( const void* A, const void* B, uint32 PortFlags ) const override;
-	virtual void SerializeItem( FArchive& Ar, void* Value, void const* Defaults ) const override;
+	virtual void SerializeItem(FStructuredArchive::FSlot Slot, void* Value, void const* Defaults) const override;
 	virtual bool NetSerializeItem( FArchive& Ar, UPackageMap* Map, void* Data, TArray<uint8> * MetaData = NULL ) const override;
 	virtual bool SupportsNetSharedSerialization() const override;
 	virtual void ExportTextItem( FString& ValueStr, const void* PropertyValue, const void* DefaultValue, UObject* Parent, int32 PortFlags, UObject* ExportRootScope ) const override;
@@ -4130,7 +4127,7 @@ public:
 	virtual FString GetCPPType( FString* ExtendedTypeText, uint32 CPPExportFlags ) const override;
 	virtual FString GetCPPTypeForwardDeclaration() const override;
 	virtual bool Identical( const void* A, const void* B, uint32 PortFlags ) const override;
-	virtual void SerializeItem( FArchive& Ar, void* Value, void const* Defaults ) const override;
+	virtual void SerializeItem(FStructuredArchive::FSlot Slot, void* Value, void const* Defaults) const override;
 	virtual bool NetSerializeItem( FArchive& Ar, UPackageMap* Map, void* Data, TArray<uint8> * MetaData = NULL ) const override;
 	virtual void ExportTextItem( FString& ValueStr, const void* PropertyValue, const void* DefaultValue, UObject* Parent, int32 PortFlags, UObject* ExportRootScope ) const override;
 	virtual const TCHAR* ImportText_Internal( const TCHAR* Buffer, void* Data, int32 PortFlags, UObject* OwnerObject, FOutputDevice* ErrorText ) const override;
@@ -4183,7 +4180,7 @@ public:
 	virtual FString GetCPPType( FString* ExtendedTypeText, uint32 CPPExportFlags ) const override;
 	virtual FString GetCPPTypeForwardDeclaration() const override;
 	virtual bool Identical( const void* A, const void* B, uint32 PortFlags ) const override;
-	virtual void SerializeItem( FArchive& Ar, void* Value, void const* Defaults ) const override;
+	virtual void SerializeItem(FStructuredArchive::FSlot Slot, void* Value, void const* Defaults) const override;
 	virtual bool NetSerializeItem( FArchive& Ar, UPackageMap* Map, void* Data, TArray<uint8> * MetaData = NULL ) const override;
 	virtual void ExportTextItem( FString& ValueStr, const void* PropertyValue, const void* DefaultValue, UObject* Parent, int32 PortFlags, UObject* ExportRootScope ) const override;
 	virtual const TCHAR* ImportText_Internal( const TCHAR* Buffer, void* Data, int32 PortFlags, UObject* OwnerObject, FOutputDevice* ErrorText ) const override;
