TBA: Temporary workaround for softobjectptr and lazyobjectptr uproperties not serialization anything when they know the archive is a reference collector. They should always be serializing their pointers and letting the underlying archive itself ignore them.

diff --git a/Engine/Source/Runtime/CoreUObject/Private/UObject/PropertyLazyObjectPtr.cpp b/Engine/Source/Runtime/CoreUObject/Private/UObject/PropertyLazyObjectPtr.cpp
--- a/Engine/Source/Runtime/CoreUObject/Private/UObject/PropertyLazyObjectPtr.cpp
+++ b/Engine/Source/Runtime/CoreUObject/Private/UObject/PropertyLazyObjectPtr.cpp
@@ -43,6 +43,12 @@ void ULazyObjectProperty::SerializeItem( FStructuredArchive::FSlot Slot, void* V
 			CheckValidObject(Value);
 		}
 	}
+	else
+	{
+		// TODO: This isn't correct, but it keeps binary serialization happy. We should ALWAYS be serializing the pointer
+		// to the archive in this function, and allowing the underlying archive to ignore it if necessary
+		Slot.EnterStream();
+	}
 }
 
 bool ULazyObjectProperty::Identical( const void* A, const void* B, uint32 PortFlags ) const
diff --git a/Engine/Source/Runtime/CoreUObject/Private/UObject/PropertySoftObjectPtr.cpp b/Engine/Source/Runtime/CoreUObject/Private/UObject/PropertySoftObjectPtr.cpp
--- a/Engine/Source/Runtime/CoreUObject/Private/UObject/PropertySoftObjectPtr.cpp
+++ b/Engine/Source/Runtime/CoreUObject/Private/UObject/PropertySoftObjectPtr.cpp
@@ -61,6 +61,12 @@ void USoftObjectProperty::SerializeItem( FStructuredArchive::FSlot Slot, void* V
 			}
 		}
 	}
+	else
+	{
+		// TODO: This isn't correct, but it keeps binary serialization happy. We should ALWAYS be serializing the pointer
+		// to the archive in this function, and allowing the underlying archive to ignore it if necessary
+		Slot.EnterStream();
+	}
 }
 
 void USoftObjectProperty::ExportTextItem( FString& ValueStr, const void* PropertyValue, const void* DefaultValue, UObject* Parent, int32 PortFlags, UObject* ExportRootScope ) const
