Clustered objects from clusters that are no longer reachable will now be marked as unreachable immediately when gathering unreachable objects

diff --git a/Engine/Source/Runtime/CoreUObject/Private/UObject/GarbageCollection.cpp b/Engine/Source/Runtime/CoreUObject/Private/UObject/GarbageCollection.cpp
--- a/Engine/Source/Runtime/CoreUObject/Private/UObject/GarbageCollection.cpp
+++ b/Engine/Source/Runtime/CoreUObject/Private/UObject/GarbageCollection.cpp
@@ -1054,7 +1054,7 @@ void IncrementalPurgeGarbage( bool bUseTimeLimit, float TimeLimit )
 	// Depending on platform FPlatformTime::Seconds might take a noticeable amount of time if called thousands of times so we avoid
 	// enforcing the time limit too often, especially as neither Destroy nor actual deletion should take significant
 	// amounts of time.
-	const int32	TimeLimitEnforcementGranularityForDestroy	= 10;	
+	const int32	TimeLimitEnforcementGranularityForDestroy	= 10;
 	const int32	TimeLimitEnforcementGranularityForDeletion	= 100;
 
 	if (GUnrechableObjectIndex < GUnreachableObjects.Num())
@@ -1303,7 +1303,7 @@ bool IsIncrementalPurgePending()
 
 // Allow parallel GC to be overridden to single threaded via console command.
 #if !PLATFORM_MAC || !WITH_EDITORONLY_DATA
-	static int32 GAllowParallelGC = 1;
+static int32 GAllowParallelGC = 1;
 #else
 	static int32 GAllowParallelGC = 0;
 #endif
@@ -1349,14 +1349,18 @@ void GatherUnreachableObjects(bool bForceSingleThreaded)
 	int32 NumThreads = FMath::Max(1, FTaskGraphInterface::Get().GetNumWorkerThreads());
 	int32 NumberOfObjectsPerThread = (MaxNumberOfObjects / NumThreads) + 1;
 
+	TArray<FUObjectItem*> ClusterItemsToDestroy;
+	int32 ClusterObjects = 0;
+
 	// Iterate over all objects. Note that we iterate over the UObjectArray and usually check only internal flags which
 	// are part of the array so we don't suffer from cache misses as much as we would if we were to check ObjectFlags.
-	ParallelFor(NumThreads, [NumberOfObjectsPerThread, NumThreads, MaxNumberOfObjects](int32 ThreadIndex)
+	ParallelFor(NumThreads, [&ClusterItemsToDestroy, NumberOfObjectsPerThread, NumThreads, MaxNumberOfObjects](int32 ThreadIndex)
 	{
 		int32 FirstObjectIndex = ThreadIndex * NumberOfObjectsPerThread + GUObjectArray.GetFirstGCIndex();
 		int32 NumObjects = (ThreadIndex < (NumThreads - 1)) ? NumberOfObjectsPerThread : (MaxNumberOfObjects - (NumThreads - 1) * NumberOfObjectsPerThread);
 		int32 LastObjectIndex = FMath::Min(GUObjectArray.GetObjectArrayNum() - 1, FirstObjectIndex + NumObjects - 1);
 		TArray<FUObjectItem*> ThisThreadUnreachableObjects;
+		TArray<FUObjectItem*> ThisThreadClusterItemsToDestroy;
 
 		for (int32 ObjectIndex = FirstObjectIndex; ObjectIndex <= LastObjectIndex; ++ObjectIndex)
 		{
@@ -1364,16 +1368,53 @@ void GatherUnreachableObjects(bool bForceSingleThreaded)
 			if (ObjectItem->IsUnreachable())
 			{
 				ThisThreadUnreachableObjects.Add(ObjectItem);
+				if (ObjectItem->HasAnyFlags(EInternalObjectFlags::ClusterRoot))
+				{
+					// We can't mark cluster objects as unreachable here as they may be currently being processed on another thread
+					ThisThreadClusterItemsToDestroy.Add(ObjectItem);
+				}
 			}
 		}
 		if (ThisThreadUnreachableObjects.Num())
 		{
 			FScopeLock UnreachableObjectsLock(&GUnreachableObjectsCritical);
 			GUnreachableObjects.Append(ThisThreadUnreachableObjects);
+			ClusterItemsToDestroy.Append(ThisThreadClusterItemsToDestroy);
 		}
 	}, bForceSingleThreaded);
 
-	UE_LOG(LogGarbage, Log, TEXT("%f ms for Gather Unreachable Objects (%d objects collected)"), (FPlatformTime::Seconds() - StartTime) * 1000, GUnreachableObjects.Num());
+	{
+		// @todo: if GUObjectClusters.FreeCluster() was thread safe we could do this in parallel too
+		for (FUObjectItem* ClusterRootItem : ClusterItemsToDestroy)
+		{
+#if UE_GCCLUSTER_VERBOSE_LOGGING
+			UE_LOG(LogGarbage, Log, TEXT("Destroying cluster (%d) %s"), ClusterRootItem->GetClusterIndex(), *static_cast<UObject*>(ClusterRootItem->Object)->GetFullName());
+#endif
+			ClusterRootItem->ClearFlags(EInternalObjectFlags::ClusterRoot);
+
+			const int32 ClusterIndex = ClusterRootItem->GetClusterIndex();
+			FUObjectCluster& Cluster = GUObjectClusters[ClusterIndex];
+			for (int32 ClusterObjectIndex : Cluster.Objects)
+			{
+				FUObjectItem* ClusterObjectItem = GUObjectArray.IndexToObjectUnsafeForGC(ClusterObjectIndex);
+				ClusterObjectItem->SetOwnerIndex(0);
+
+				if (!ClusterObjectItem->HasAnyFlags(EInternalObjectFlags::ReachableInCluster))
+				{
+					ClusterObjectItem->SetFlags(EInternalObjectFlags::Unreachable);
+					ClusterObjects++;
+					GUnreachableObjects.Add(ClusterObjectItem);
+				}
+			}
+			GUObjectClusters.FreeCluster(ClusterIndex);
+		}
+	}
+
+	UE_LOG(LogGarbage, Log, TEXT("%f ms for Gather Unreachable Objects (%d objects collected including %d cluster objects from %d clusters)"),
+		(FPlatformTime::Seconds() - StartTime) * 1000,
+		GUnreachableObjects.Num(),
+		ClusterObjects,
+		ClusterItemsToDestroy.Num());
 }
 
 /** 
@@ -1478,7 +1519,7 @@ void CollectGarbageInternal(EObjectFlags KeepFlags, bool bPerformFullPurge)
 		// Fire post-reachability analysis hooks
 		FCoreUObjectDelegates::PostReachabilityAnalysis.Broadcast();
 
-		{			
+		{
 			FGCArrayPool::Get().ClearWeakReferences(bPerformFullPurge);
 
 			GatherUnreachableObjects(bForceSingleThreadedGC);
@@ -1501,12 +1542,13 @@ void CollectGarbageInternal(EObjectFlags KeepFlags, bool bPerformFullPurge)
 		{
 			IncrementalPurgeGarbage(false);
 		}
-	}
 
-	// Destroy all pending delete linkers
-	DeleteLoaders();
+		// Destroy all pending delete linkers
+		DeleteLoaders();
 
-	FMemory::Trim();
+		// Trim allocator memory
+		FMemory::Trim();
+	}
 
 	// Route callbacks to verify GC assumptions
 	FCoreUObjectDelegates::GetPostGarbageCollect().Broadcast();
@@ -1525,74 +1567,35 @@ bool UnhashUnreachableObjects(bool bUseTimeLimit, float TimeLimit)
 	// Unhash all unreachable objects.
 	const double StartTime = FPlatformTime::Seconds();
 	const int32 TimeLimitEnforcementGranularityForBeginDestroy = 10;
-	int32 ClustersRemoved = 0;
 	int32 Items = 0;
-	int32 ClusterItems = 0;
-	int32 TimePollCounter = 0;	
+	int32 TimePollCounter = 0;
 
 	while (GUnrechableObjectIndex < GUnreachableObjects.Num())
 	{
 		//@todo UE4 - A prefetch was removed here. Re-add it. It wasn't right anyway, since it was ten items ahead and the consoles on have 8 prefetch slots
 
 		FUObjectItem* ObjectItem = GUnreachableObjects[GUnrechableObjectIndex++];
-
-		Items++;
-#if UE_GCCLUSTER_VERBOSE_LOGGING
-		UObject* Object = (UObject*)ObjectItem->Object;
-#endif
-		if ((ObjectItem->GetFlags() & EInternalObjectFlags::ClusterRoot) == EInternalObjectFlags::ClusterRoot)
-		{
-#if UE_GCCLUSTER_VERBOSE_LOGGING
-			UE_LOG(LogGarbage, Log, TEXT("Destroying cluster (%d) %s"), ObjectItem->GetClusterIndex(), *Object->GetFullName());
-#endif
-			// Nuke the entire cluster
-			ObjectItem->ClearFlags(EInternalObjectFlags::ClusterRoot);
-			FUObjectCluster& Cluster = GUObjectClusters[ObjectItem->GetClusterIndex()];
-			for (int32 ClusterObjectIndex : Cluster.Objects)
-			{
-				FUObjectItem* ClusterObjectItem = GUObjectArray.IndexToObjectUnsafeForGC(ClusterObjectIndex);
-				ClusterObjectItem->SetOwnerIndex(0);
-
-				if (!ClusterObjectItem->HasAnyFlags(EInternalObjectFlags::ReachableInCluster))
-				{
-					ClusterObjectItem->SetFlags(EInternalObjectFlags::Unreachable);
-					UObject* ClusterObject = (UObject*)ClusterObjectItem->Object;
-					{
-						FScopedCBDProfile Profile(ClusterObject);
-						ClusterObject->ConditionalBeginDestroy();
-						ClusterItems++;
-						TimePollCounter++;
-					}
-				}
-			}
-			GUObjectClusters.FreeCluster(ObjectItem->GetClusterIndex());
-			ClustersRemoved++;
-		}
-
-		// Begin the object's asynchronous destruction.
-#if !UE_GCCLUSTER_VERBOSE_LOGGING
-		UObject* Object = (UObject*)ObjectItem->Object;
-#endif
 		{
+			UObject* Object = static_cast<UObject*>(ObjectItem->Object);
 			FScopedCBDProfile Profile(Object);
+			// Begin the object's asynchronous destruction.
 			Object->ConditionalBeginDestroy();
-			ObjectItem->ClearFlags(EInternalObjectFlags::HadReferenceKilled);
 		}
+		ObjectItem->ClearFlags(EInternalObjectFlags::HadReferenceKilled);
+		Items++;
 
 		const bool bPollTimeLimit = ((TimePollCounter++) % TimeLimitEnforcementGranularityForBeginDestroy == 0);
 		if (bUseTimeLimit && bPollTimeLimit && ((FPlatformTime::Seconds() - StartTime) > TimeLimit))
 		{
 			break;
 		}
-	}	
+	}
 
-	UE_LOG(LogGarbage, Log, TEXT("%f ms for %sunhashing unreachable objects. Clusters removed: %d. Items %d (%d/%d) Cluster Items %d"), 
-		(FPlatformTime::Seconds() - StartTime) * 1000, 
+	UE_LOG(LogGarbage, Log, TEXT("%f ms for %sunhashing unreachable objects. Items %d (%d/%d)"),
+		(FPlatformTime::Seconds() - StartTime) * 1000,
 		bUseTimeLimit ? TEXT("incrementally ") : TEXT(""),
-		ClustersRemoved,
-		Items, 
-		GUnrechableObjectIndex, GUnreachableObjects.Num(),
-		ClusterItems);
+		Items,
+		GUnrechableObjectIndex, GUnreachableObjects.Num());
 
 	FCoreUObjectDelegates::PostGarbageCollectConditionalBeginDestroy.Broadcast();
 
