Added the ability to disable incremental BeginDestroy in ini/project settings

diff --git a/Engine/Config/BaseEngine.ini b/Engine/Config/BaseEngine.ini
--- a/Engine/Config/BaseEngine.ini
+++ b/Engine/Config/BaseEngine.ini
@@ -693,6 +693,7 @@ gc.AllowParallelGC=True
 ; pick a fractional number to keep phase shifting and avoid collisions
 gc.TimeBetweenPurgingPendingKillObjects=61.1
 gc.MaxObjectsInEditor=8388607
+gc.IncrementalBeginDestroyEnabled=True
 gc.CreateGCClusters=True
 gc.MergeGCClusters=False
 gc.MinGCClusterSize=5
diff --git a/Engine/Source/Runtime/CoreUObject/Private/UObject/GarbageCollection.cpp b/Engine/Source/Runtime/CoreUObject/Private/UObject/GarbageCollection.cpp
--- a/Engine/Source/Runtime/CoreUObject/Private/UObject/GarbageCollection.cpp
+++ b/Engine/Source/Runtime/CoreUObject/Private/UObject/GarbageCollection.cpp
@@ -158,6 +158,13 @@ static FAutoConsoleVariableRef CVarCheckForIllegalMarkPendingKill(
 	ECVF_Default
 );
 
+static int32 GIncrementalBeginDestroyEnabled = 1;
+static FAutoConsoleVariableRef CIncrementalBeginDestroyEnabled(
+	TEXT("gc.IncrementalBeginDestroyEnabled"),
+	GIncrementalBeginDestroyEnabled,
+	TEXT("If true, the engine will destroy objects incrementally using time limit each frame"),
+	ECVF_Default
+);
 
 #if PERF_DETAILED_PER_CLASS_GC_STATS
 /** Map from a UClass' FName to the number of objects that were purged during the last purge phase of this class.	*/
@@ -1524,7 +1531,7 @@ void CollectGarbageInternal(EObjectFlags KeepFlags, bool bPerformFullPurge)
 
 			GatherUnreachableObjects(bForceSingleThreadedGC);
 
-			if (bPerformFullPurge)
+			if (bPerformFullPurge || !GIncrementalBeginDestroyEnabled)
 			{
 				UnhashUnreachableObjects(/**bUseTimeLimit = */ false);
 				FScopedCBDProfile::DumpProfile();
diff --git a/Engine/Source/Runtime/Engine/Classes/Engine/CoreSettings.h b/Engine/Source/Runtime/Engine/Classes/Engine/CoreSettings.h
--- a/Engine/Source/Runtime/Engine/Classes/Engine/CoreSettings.h
+++ b/Engine/Source/Runtime/Engine/Classes/Engine/CoreSettings.h
@@ -170,6 +170,11 @@ protected:
 		ToolTip = "If enabled, garbage collection will use multiple threads."))
 	uint32 AllowParallelGC : 1;
 
+	UPROPERTY(EditAnywhere, config, Category = Optimization, meta = (
+		ConsoleVariable = "gc.IncrementalBeginDestroyEnabled", DisplayName = "Incremental BeginDestroy Enabled",
+		ToolTip = "If true, the engine will destroy objects incrementally using time limit each frame."))
+	uint32 IncrementalBeginDestroyEnabled : 1;
+
 	UPROPERTY(EditAnywhere, config, Category = Optimization, meta = (
 		ConsoleVariable = "gc.CreateGCClusters", DisplayName = "Create Garbage Collector UObject Clusters",
 		ToolTip = "If true, the engine will attempt to create clusters of objects for better garbage collection performance."))
diff --git a/Engine/Source/Runtime/Engine/Private/CoreSettings.cpp b/Engine/Source/Runtime/Engine/Private/CoreSettings.cpp
--- a/Engine/Source/Runtime/Engine/Private/CoreSettings.cpp
+++ b/Engine/Source/Runtime/Engine/Private/CoreSettings.cpp
@@ -168,12 +168,13 @@ UGarbageCollectionSettings::UGarbageCollectionSettings()
 	TimeBetweenPurgingPendingKillObjects = 60.0f;
 	FlushStreamingOnGC = false;
 	AllowParallelGC = true;
+	IncrementalBeginDestroyEnabled = true;
 	NumRetriesBeforeForcingGC = 0;
 	MaxObjectsNotConsideredByGC = 0;
 	SizeOfPermanentObjectPool = 0;
 	MaxObjectsInEditor = 12 * 1024 * 1024;
-	MaxObjectsInGame = 2 * 1024 * 1024;
-	CreateGCClusters = true;
+	MaxObjectsInGame = 2 * 1024 * 1024;	
+	CreateGCClusters = true;	
 	MergeGCClusters = false;
 	MinGCClusterSize = 5;
 	ActorClusteringEnabled = true;
