TBA: Updated USetProperty to proper structured archive usage

diff --git a/Engine/Source/Runtime/CoreUObject/Private/UObject/PropertySet.cpp b/Engine/Source/Runtime/CoreUObject/Private/UObject/PropertySet.cpp
--- a/Engine/Source/Runtime/CoreUObject/Private/UObject/PropertySet.cpp
+++ b/Engine/Source/Runtime/CoreUObject/Private/UObject/PropertySet.cpp
@@ -215,17 +215,18 @@ void USetProperty::GetPreloadDependencies(TArray<UObject*>& OutDeps)
 
 void USetProperty::SerializeItem(FStructuredArchive::FSlot Slot, void* Value, const void* Defaults) const
 {
-	FArchiveUObjectFromStructuredArchive Ar(Slot);
+	FArchive& UnderlyingArchive = Slot.GetUnderlyingArchive();
+	FStructuredArchive::FRecord Record = Slot.EnterRecord();
 
 	// Ar related calls in this function must be mirrored in USetProperty::ConvertFromType
 	checkSlow(ElementProp);
 
 	// Ensure that the element property has been loaded before calling SerializeItem() on it
-	Ar.Preload(ElementProp);
+	UnderlyingArchive.Preload(ElementProp);
 
 	FScriptSetHelper SetHelper(this, Value);
 
-	if (Ar.IsLoading())
+	if (UnderlyingArchive.IsLoading())
 	{
 		if (Defaults)
 		{
@@ -248,17 +249,18 @@ void USetProperty::SerializeItem(FStructuredArchive::FSlot Slot, void* Value, co
 
 		// Delete any explicitly-removed elements
 		int32 NumElementsToRemove = 0;
-		Ar << NumElementsToRemove;
+		FStructuredArchive::FArray ElementsToRemoveArray = Record.EnterArray(FIELD_NAME_TEXT("ElementsToRemove"), NumElementsToRemove);
+
 		if (NumElementsToRemove)
 		{
 			TempElementStorage = (uint8*)FMemory::Malloc(SetLayout.Size);
 			ElementProp->InitializeValue(TempElementStorage);
 
-			FSerializedPropertyScope SerializedProperty(Ar, ElementProp, this);
+			FSerializedPropertyScope SerializedProperty(UnderlyingArchive, ElementProp, this);
 			for (; NumElementsToRemove; --NumElementsToRemove)
 			{
 				// Read key into temporary storage
-				ElementProp->SerializeItem(FStructuredArchiveFromArchive(Ar).GetSlot(), TempElementStorage);
+				ElementProp->SerializeItem(ElementsToRemoveArray.EnterElement(), TempElementStorage);
 
 				// If the key is in the map, remove it
 				const int32 Found = SetHelper.FindElementIndex(TempElementStorage);
@@ -270,7 +272,7 @@ void USetProperty::SerializeItem(FStructuredArchive::FSlot Slot, void* Value, co
 		}
 
 		int32 Num = 0;
-		Ar << Num;
+		FStructuredArchive::FArray ElementsArray = Record.EnterArray(FIELD_NAME_TEXT("Elements"), Num);
 
 		// Allocate temporary key space if we haven't allocated it already above
 		if (Num != 0 && !TempElementStorage)
@@ -279,12 +281,12 @@ void USetProperty::SerializeItem(FStructuredArchive::FSlot Slot, void* Value, co
 			ElementProp->InitializeValue(TempElementStorage);
 		}
 
-		FSerializedPropertyScope SerializedProperty(Ar, ElementProp, this);
+		FSerializedPropertyScope SerializedProperty(UnderlyingArchive, ElementProp, this);
 		// Read remaining items into container
 		for (; Num; --Num)
 		{
 			// Read key into temporary storage
-			ElementProp->SerializeItem(FStructuredArchiveFromArchive(Ar).GetSlot(), TempElementStorage);
+			ElementProp->SerializeItem(ElementsArray.EnterElement(), TempElementStorage);
 
 			// Add a new entry if the element doesn't currently exist in the set
 			if (SetHelper.FindElementIndex(TempElementStorage) == INDEX_NONE)
@@ -327,12 +329,13 @@ void USetProperty::SerializeItem(FStructuredArchive::FSlot Slot, void* Value, co
 
 		// Write out the removed elements
 		int32 RemovedElementsNum = Indices.Num();
-		Ar << RemovedElementsNum;
+		FStructuredArchive::FArray RemovedElementsArray = Record.EnterArray(FIELD_NAME_TEXT("ElementsToRemove"), RemovedElementsNum);
+		
 		{
-			FSerializedPropertyScope SerializedProperty(Ar, ElementProp, this);
+			FSerializedPropertyScope SerializedProperty(UnderlyingArchive, ElementProp, this);
 			for (int32 Index : Indices)
 			{
-				ElementProp->SerializeItem(FStructuredArchiveFromArchive(Ar).GetSlot(), DefaultsHelper.GetElementPtr(Index));
+				ElementProp->SerializeItem(RemovedElementsArray.EnterElement(), DefaultsHelper.GetElementPtr(Index));
 			}
 		}
 
@@ -358,29 +361,29 @@ void USetProperty::SerializeItem(FStructuredArchive::FSlot Slot, void* Value, co
 
 			// Write out differences from defaults
 			int32 Num = Indices.Num();
-			Ar << Num;
+			FStructuredArchive::FArray ElementsArray = Record.EnterArray(FIELD_NAME_TEXT("Elements"), Num);
 
-			FSerializedPropertyScope SerializedProperty(Ar, ElementProp, this);
+			FSerializedPropertyScope SerializedProperty(UnderlyingArchive, ElementProp, this);
 			for (int32 Index : Indices)
 			{
 				uint8* ElementPtr = SetHelper.GetElementPtrWithoutCheck(Index);
 
-				ElementProp->SerializeItem(FStructuredArchiveFromArchive(Ar).GetSlot(), ElementPtr);
+				ElementProp->SerializeItem(ElementsArray.EnterElement(), ElementPtr);
 			}
 		}
 		else
 		{
 			int32 Num = SetHelper.Num();
-			Ar << Num;
+			FStructuredArchive::FArray ElementsArray = Record.EnterArray(FIELD_NAME_TEXT("Elements"), Num);
 
-			FSerializedPropertyScope SerializedProperty(Ar, ElementProp, this);
+			FSerializedPropertyScope SerializedProperty(UnderlyingArchive, ElementProp, this);
 			for (int32 Index = 0; Num; ++Index)
 			{
 				if (SetHelper.IsValidIndex(Index))
 				{
 					uint8* ElementPtr = SetHelper.GetElementPtrWithoutCheck(Index);
 
-					ElementProp->SerializeItem(FStructuredArchiveFromArchive(Ar).GetSlot(), ElementPtr);
+					ElementProp->SerializeItem(ElementsArray.EnterElement(), ElementPtr);
 
 					--Num;
 				}
