TBA: Convert UMapProperty to full structured archive

diff --git a/Engine/Source/Runtime/CoreUObject/Private/UObject/PropertyMap.cpp b/Engine/Source/Runtime/CoreUObject/Private/UObject/PropertyMap.cpp
--- a/Engine/Source/Runtime/CoreUObject/Private/UObject/PropertyMap.cpp
+++ b/Engine/Source/Runtime/CoreUObject/Private/UObject/PropertyMap.cpp
@@ -241,19 +241,20 @@ void UMapProperty::GetPreloadDependencies(TArray<UObject*>& OutDeps)
 
 void UMapProperty::SerializeItem(FStructuredArchive::FSlot Slot, void* Value, const void* Defaults) const
 {
-	FArchiveUObjectFromStructuredArchive Ar(Slot);
+	FArchive& UnderlyingArchive = Slot.GetUnderlyingArchive();
+	FStructuredArchive::FRecord Record = Slot.EnterRecord();
 
 	// Ar related calls in this function must be mirrored in UMapProperty::ConvertFromType
 	checkSlow(KeyProp);
 	checkSlow(ValueProp);
 
 	// Ensure that the key/value properties have been loaded before calling SerializeItem() on them
-	Ar.Preload(KeyProp);
-	Ar.Preload(ValueProp);
+	UnderlyingArchive.Preload(KeyProp);
+	UnderlyingArchive.Preload(ValueProp);
 
 	FScriptMapHelper MapHelper(this, Value);
 
-	if (Ar.IsLoading())
+	if (UnderlyingArchive.IsLoading())
 	{
 		if (Defaults)
 		{
@@ -276,17 +277,17 @@ void UMapProperty::SerializeItem(FStructuredArchive::FSlot Slot, void* Value, co
 
 		// Delete any explicitly-removed keys
 		int32 NumKeysToRemove = 0;
-		Ar << NumKeysToRemove;
+		FStructuredArchive::FArray KeysToRemoveArray = Record.EnterArray(FIELD_NAME_TEXT("KeysToRemove"), NumKeysToRemove);
 		if (NumKeysToRemove)
 		{
 			TempKeyStorage = (uint8*)FMemory::Malloc(MapLayout.SetLayout.Size);
 			KeyProp->InitializeValue(TempKeyStorage);
 
-			FSerializedPropertyScope SerializedProperty(Ar, KeyProp, this);
+			FSerializedPropertyScope SerializedProperty(UnderlyingArchive, KeyProp, this);
 			for (; NumKeysToRemove; --NumKeysToRemove)
 			{
 				// Read key into temporary storage
-				KeyProp->SerializeItem(FStructuredArchiveFromArchive(Ar).GetSlot(), TempKeyStorage);
+				KeyProp->SerializeItem(KeysToRemoveArray.EnterElement(), TempKeyStorage);
 
 				// If the key is in the map, remove it
 				int32 Found = MapHelper.FindMapIndexWithKey(TempKeyStorage);
@@ -298,7 +299,7 @@ void UMapProperty::SerializeItem(FStructuredArchive::FSlot Slot, void* Value, co
 		}
 
 		int32 NumEntries = 0;
-		Ar << NumEntries;
+		FStructuredArchive::FArray EntriesArray = Record.EnterArray(FIELD_NAME_TEXT("Entries"), NumEntries);
 
 		// Allocate temporary key space if we haven't allocated it already above
 		if (NumEntries != 0 && !TempKeyStorage)
@@ -310,10 +311,12 @@ void UMapProperty::SerializeItem(FStructuredArchive::FSlot Slot, void* Value, co
 		// Read remaining items into container
 		for (; NumEntries; --NumEntries)
 		{
+			FStructuredArchive::FRecord EntryRecord = EntriesArray.EnterElement().EnterRecord();
+
 			// Read key into temporary storage
 			{
-				FSerializedPropertyScope SerializedProperty(Ar, KeyProp, this);
-				KeyProp->SerializeItem(FStructuredArchiveFromArchive(Ar).GetSlot(), TempKeyStorage);
+				FSerializedPropertyScope SerializedProperty(UnderlyingArchive, KeyProp, this);
+				KeyProp->SerializeItem(EntryRecord.EnterField(FIELD_NAME_TEXT("Key")), TempKeyStorage);
 			}
 			
 			// Add a new default value if the key doesn't currently exist in the map
@@ -330,8 +333,8 @@ void UMapProperty::SerializeItem(FStructuredArchive::FSlot Slot, void* Value, co
 
 			// Deserialize value
 			{
-				FSerializedPropertyScope SerializedProperty(Ar, ValueProp, this);
-				ValueProp->SerializeItem(FStructuredArchiveFromArchive(Ar).GetSlot(), NextPairPtr + MapLayout.ValueOffset);
+				FSerializedPropertyScope SerializedProperty(UnderlyingArchive, ValueProp, this);
+				ValueProp->SerializeItem(EntryRecord.EnterField(FIELD_NAME_TEXT("Value")), NextPairPtr + MapLayout.ValueOffset);
 			}
 		}
 
@@ -365,12 +368,12 @@ void UMapProperty::SerializeItem(FStructuredArchive::FSlot Slot, void* Value, co
 
 		// Write out the missing keys
 		int32 MissingKeysNum = Indices.Num();
-		Ar << MissingKeysNum;
+		FStructuredArchive::FArray KeysToRemoveArray = Record.EnterArray(FIELD_NAME_TEXT("KeysToRemove"), MissingKeysNum);
 		{
-			FSerializedPropertyScope SerializedProperty(Ar, KeyProp, this);
+			FSerializedPropertyScope SerializedProperty(UnderlyingArchive, KeyProp, this);
 			for (int32 Index : Indices)
 			{
-				KeyProp->SerializeItem(FStructuredArchiveFromArchive(Ar).GetSlot(), DefaultsHelper.GetPairPtr(Index));
+				KeyProp->SerializeItem(KeysToRemoveArray.EnterElement(), DefaultsHelper.GetPairPtr(Index));
 			}
 		}
 
@@ -396,38 +399,42 @@ void UMapProperty::SerializeItem(FStructuredArchive::FSlot Slot, void* Value, co
 
 			// Write out differences from defaults
 			int32 Num = Indices.Num();
-			Ar << Num;
+			FStructuredArchive::FArray EntriesArray = Record.EnterArray(FIELD_NAME_TEXT("Entries"), Num);
 			for (int32 Index : Indices)
 			{
 				uint8* ValuePairPtr = MapHelper.GetPairPtrWithoutCheck(Index);
+				FStructuredArchive::FRecord EntryRecord = EntriesArray.EnterElement().EnterRecord();
 
 				{
-					FSerializedPropertyScope SerializedProperty(Ar, KeyProp, this);
-					KeyProp->SerializeItem(FStructuredArchiveFromArchive(Ar).GetSlot(), ValuePairPtr);
+					FSerializedPropertyScope SerializedProperty(UnderlyingArchive, KeyProp, this);
+					KeyProp->SerializeItem(EntryRecord.EnterField(FIELD_NAME_TEXT("Key")), ValuePairPtr);
 				}
 				{
-					FSerializedPropertyScope SerializedProperty(Ar, ValueProp, this);
-					ValueProp->SerializeItem(FStructuredArchiveFromArchive(Ar).GetSlot(), ValuePairPtr + MapLayout.ValueOffset);
+					FSerializedPropertyScope SerializedProperty(UnderlyingArchive, ValueProp, this);
+					ValueProp->SerializeItem(EntryRecord.EnterField(FIELD_NAME_TEXT("Value")), ValuePairPtr + MapLayout.ValueOffset);
 				}
 			}
 		}
 		else
 		{
 			int32 Num = MapHelper.Num();
-			Ar << Num;
+			FStructuredArchive::FArray EntriesArray = Record.EnterArray(FIELD_NAME_TEXT("Entries"), Num);
+
 			for (int32 Index = 0; Num; ++Index)
 			{
 				if (MapHelper.IsValidIndex(Index))
 				{
+					FStructuredArchive::FRecord EntryRecord = EntriesArray.EnterElement().EnterRecord();
+
 					uint8* ValuePairPtr = MapHelper.GetPairPtrWithoutCheck(Index);
 
 					{
-						FSerializedPropertyScope SerializedProperty(Ar, KeyProp, this);
-						KeyProp->SerializeItem(FStructuredArchiveFromArchive(Ar).GetSlot(), ValuePairPtr);
+						FSerializedPropertyScope SerializedProperty(UnderlyingArchive, KeyProp, this);
+						KeyProp->SerializeItem(EntryRecord.EnterField(FIELD_NAME_TEXT("Key")), ValuePairPtr);
 					}
 					{
-						FSerializedPropertyScope SerializedProperty(Ar, ValueProp, this);
-						ValueProp->SerializeItem(FStructuredArchiveFromArchive(Ar).GetSlot(), ValuePairPtr + MapLayout.ValueOffset);
+						FSerializedPropertyScope SerializedProperty(UnderlyingArchive, ValueProp, this);
+						ValueProp->SerializeItem(EntryRecord.EnterField(FIELD_NAME_TEXT("Value")), ValuePairPtr + MapLayout.ValueOffset);
 					}
 
 					--Num;
