Remove out of date assert
#rnx

diff --git a/Engine/Source/Runtime/CoreUObject/Private/UObject/LinkerLoad.cpp b/Engine/Source/Runtime/CoreUObject/Private/UObject/LinkerLoad.cpp
--- a/Engine/Source/Runtime/CoreUObject/Private/UObject/LinkerLoad.cpp
+++ b/Engine/Source/Runtime/CoreUObject/Private/UObject/LinkerLoad.cpp
@@ -2936,13 +2936,7 @@ UClass* FLinkerLoad::GetExportLoadClass(int32 Index)
 		VerifyImport(Export.ClassIndex.ToImport());
 	}
 
-	UClass* ExportClass = (UClass*)IndexToObject(Export.ClassIndex);
-#if USE_DEFERRED_DEPENDENCY_CHECK_VERIFICATION_TESTS
-	check(!Export.ClassIndex.IsImport() || !(LoadFlags & LOAD_DeferDependencyLoads) || 
-		(ExportClass && ExportClass->HasAnyClassFlags(CLASS_Native)) || (Cast<ULinkerPlaceholderClass>(ExportClass) != nullptr));
-#endif // USE_DEFERRED_DEPENDENCY_CHECK_VERIFICATION_TESTS
-
-	return ExportClass;
+	return (UClass*)IndexToObject(Export.ClassIndex);
 }
 
 #if WITH_EDITORONLY_DATA
