Nested components from a native class should be initialized from the blueprint template and not recycled from the constructor created instance.

diff --git a/Engine/Source/Runtime/CoreUObject/Private/UObject/CoreNative.cpp b/Engine/Source/Runtime/CoreUObject/Private/UObject/CoreNative.cpp
--- a/Engine/Source/Runtime/CoreUObject/Private/UObject/CoreNative.cpp
+++ b/Engine/Source/Runtime/CoreUObject/Private/UObject/CoreNative.cpp
@@ -194,34 +194,12 @@ UObject* FObjectInstancingGraph::GetInstancedSubobject( UObject* SourceSubobject
 								checkf(SubobjectOuter, TEXT("No corresponding destination object found for '%s' while attempting to instance component '%s'"), *SourceSubobject->GetOuter()->GetFullName(), *SourceSubobject->GetFullName());
 							}
 
-							FName SubobjectName = SourceSubobject->GetFName();
+							const FName SubobjectName = SourceSubobject->GetFName();
 
-							// final archetype archetype will be the archetype of the template
-							UObject* FinalSubobjectArchetype = CurrentValue->GetArchetype();
-
-							// Don't search for the existing subobjects on Blueprint-generated classes. What we'll find is a subobject
-							// created by the constructor which may not have all of its fields initialized to the correct value (which
-							// should be coming from a blueprint).
-							// NOTE: Since this function is called ONLY for Blueprint-generated classes, we may as well delete this 'if'.
-							if (!SubobjectOuter->GetClass()->HasAnyClassFlags(CLASS_CompiledFromBlueprint))
-							{
-								InstancedSubobject = StaticFindObjectFast(nullptr, SubobjectOuter, SubobjectName);
-							}
-
-							if (InstancedSubobject && IsCreatingArchetype())
-							{
-								// since we are updating an archetype, this needs to reconstruct as that is the mechanism used to copy properties
-								// it will destroy the existing object and overwrite it
-								InstancedSubobject = nullptr;
-							}
-
-							if (!InstancedSubobject)
-							{
-								// finally, create the component instance
-								InstancedSubobject = StaticConstructObject_Internal(SourceSubobject->GetClass(), SubobjectOuter,
-									SubobjectName, SubobjectOuter->GetMaskedFlags(RF_PropagateToSubObjects), EInternalObjectFlags::None, SourceSubobject,
-									true, this);
-							}
+							// finally, create the component instance
+							InstancedSubobject = StaticConstructObject_Internal(SourceSubobject->GetClass(), SubobjectOuter,
+								SubobjectName, SubobjectOuter->GetMaskedFlags(RF_PropagateToSubObjects), EInternalObjectFlags::None, SourceSubobject,
+								true, this);
 						}
 					}
 				}
diff --git a/Engine/Source/Runtime/CoreUObject/Private/UObject/UObjectGlobals.cpp b/Engine/Source/Runtime/CoreUObject/Private/UObject/UObjectGlobals.cpp
--- a/Engine/Source/Runtime/CoreUObject/Private/UObject/UObjectGlobals.cpp
+++ b/Engine/Source/Runtime/CoreUObject/Private/UObject/UObjectGlobals.cpp
@@ -3138,7 +3138,7 @@ UObject* StaticConstructObject_Internal
 	const bool bIsNativeFromCDO = bIsNativeClass &&
 		(
 			!InTemplate || 
-			(InName != NAME_None && (bAssumeTemplateIsArchetype || InTemplate == UObject::GetArchetypeFromRequiredInfo(InClass, InOuter, InName, InFlags)))
+			(InName != NAME_None && (bAssumeTemplateIsArchetype || (!InTemplate->IsInBlueprint() && InTemplate == UObject::GetArchetypeFromRequiredInfo(InClass, InOuter, InName, InFlags))))
 		);
 #if WITH_HOT_RELOAD
 	// Do not recycle subobjects when performing hot-reload as they may contain old property values.
