Redo some of the CoreRedirects logic to support passing in a query type with multiple separate type flags, which the comments said was supported but which was not

diff --git a/Engine/Source/Runtime/CoreUObject/Private/UObject/CoreRedirects.cpp b/Engine/Source/Runtime/CoreUObject/Private/UObject/CoreRedirects.cpp
--- a/Engine/Source/Runtime/CoreUObject/Private/UObject/CoreRedirects.cpp
+++ b/Engine/Source/Runtime/CoreUObject/Private/UObject/CoreRedirects.cpp
@@ -493,81 +493,89 @@ bool FCoreRedirects::GetMatchingRedirects(ECoreRedirectFlags Type, const FCoreRe
 {
 	// Look for all redirects that match the given names and flags
 	bool bFound = false;
-	FRedirectNameMap* RedirectNameMap = RedirectTypeMap.Find(Type);
-	if (RedirectNameMap)
+
+	// Always search for package redirects and substring matches
+	Type |= ECoreRedirectFlags::Type_Package;
+	Type |= ECoreRedirectFlags::Option_MatchSubstring;
+
+	// Determine list of maps to look over, need to handle being passed multiple types in a bit mask
+	for (const TPair<ECoreRedirectFlags, FRedirectNameMap>& Pair : RedirectTypeMap)
 	{
-		TArray<FCoreRedirect>* RedirectsForName = RedirectNameMap->RedirectMap.Find(OldObjectName.GetSearchKey(Type));
+		ECoreRedirectFlags PairType = Pair.Key;
 
-		if (RedirectsForName)
+		// For type, check it includes the matching type
+		const bool bTypeMatches = !!((PairType | Type) & ECoreRedirectFlags::Type_AllMask);
+		// For options, instance/remove must match exactly, but allow both substring and non substring
+		const bool bOptionsMatch = (PairType & ECoreRedirectFlags::Option_ExactMatchMask) == (Type & ECoreRedirectFlags::Option_ExactMatchMask);
+
+		// We need to check all maps that match the type mask
+		if (bTypeMatches && bOptionsMatch)
 		{
-			for (FCoreRedirect& CheckRedirect : *RedirectsForName)
+			const TArray<FCoreRedirect>* RedirectsForName = Pair.Value.RedirectMap.Find(OldObjectName.GetSearchKey(PairType));
+
+			if (RedirectsForName)
 			{
-				if (CheckRedirect.Matches(Type, OldObjectName))
+				for (const FCoreRedirect& CheckRedirect : *RedirectsForName)
 				{
-					bFound = true;
-					FoundRedirects.Add(&CheckRedirect);
+					if (CheckRedirect.Matches(PairType, OldObjectName))
+					{
+						bFound = true;
+						FoundRedirects.Add(&CheckRedirect);
+					}
 				}
 			}
 		}
 	}
 
-	// Add Package redirects now as well
-	if (!(Type & ECoreRedirectFlags::Type_Package))
-	{
-		bFound |= GetMatchingRedirects(ECoreRedirectFlags::Type_Package, OldObjectName, FoundRedirects);
-	}
-
-	// Add substring matches as well, these can be slow
-	if (!(Type & ECoreRedirectFlags::Option_MatchSubstring))
-	{
-		bFound |= GetMatchingRedirects(Type | ECoreRedirectFlags::Option_MatchSubstring, OldObjectName, FoundRedirects);
-	}
-
 	return bFound;
 }
 
 bool FCoreRedirects::FindPreviousNames(ECoreRedirectFlags Type, const FCoreRedirectObjectName& NewObjectName, TArray<FCoreRedirectObjectName>& PreviousNames)
 {
+	// Look for reverse direction redirects
 	bool bFound = false;
 
-	FRedirectNameMap* RedirectNameMap = RedirectTypeMap.Find(Type);
-	if (RedirectNameMap)
+	// Always search for package redirects and substring matches
+	Type |= ECoreRedirectFlags::Type_Package;
+	Type |= ECoreRedirectFlags::Option_MatchSubstring;
+
+	// Determine list of maps to look over, need to handle being passed multiple types in a bit mask
+	for (const TPair<ECoreRedirectFlags, FRedirectNameMap>& Pair : RedirectTypeMap)
 	{
-		for (TPair<FName, TArray<FCoreRedirect> >& Pair : RedirectNameMap->RedirectMap)
+		ECoreRedirectFlags PairType = Pair.Key;
+
+		// For type, check it includes the matching type
+		const bool bTypeMatches = !!((PairType | Type) & ECoreRedirectFlags::Type_AllMask);
+		// For options, instance/remove must match exactly, but allow both substring and non substring
+		const bool bOptionsMatch = (PairType & ECoreRedirectFlags::Option_ExactMatchMask) == (Type & ECoreRedirectFlags::Option_ExactMatchMask);
+
+		// We need to check all maps that match the type mask
+		if (bTypeMatches && bOptionsMatch)
 		{
-			for (FCoreRedirect& Redirect : Pair.Value)
+			for (const TPair<FName, TArray<FCoreRedirect> >& RedirectPair : Pair.Value.RedirectMap)
 			{
-				if (Redirect.NewName.Matches(NewObjectName, Redirect.IsSubstringMatch()))
+				for (const FCoreRedirect& Redirect : RedirectPair.Value)
 				{
-					// Construct a reverse redirect
-					FCoreRedirect ReverseRedirect = FCoreRedirect(Redirect);
-					ReverseRedirect.OldName = Redirect.NewName;
-					ReverseRedirect.NewName = Redirect.OldName;
+					if (Redirect.NewName.Matches(NewObjectName, Redirect.IsSubstringMatch()))
+					{
+						// Construct a reverse redirect
+						FCoreRedirect ReverseRedirect = FCoreRedirect(Redirect);
+						ReverseRedirect.OldName = Redirect.NewName;
+						ReverseRedirect.NewName = Redirect.OldName;
 
-					FCoreRedirectObjectName OldName = ReverseRedirect.RedirectName(NewObjectName);
+						FCoreRedirectObjectName OldName = ReverseRedirect.RedirectName(NewObjectName);
 
-					if (OldName != NewObjectName)
-					{
-						bFound = true;
-						PreviousNames.AddUnique(OldName);
+						if (OldName != NewObjectName)
+						{
+							bFound = true;
+							PreviousNames.AddUnique(OldName);
+						}
 					}
 				}
 			}
 		}
 	}
 
-	// Add Package redirects now as well
-	if (!(Type & ECoreRedirectFlags::Type_Package))
-	{
-		bFound |= FindPreviousNames(ECoreRedirectFlags::Type_Package, NewObjectName, PreviousNames);
-	}
-
-	// Add substring matches as well, these can be slow
-	if (!(Type & ECoreRedirectFlags::Option_MatchSubstring))
-	{
-		bFound |= FindPreviousNames(Type | ECoreRedirectFlags::Option_MatchSubstring, NewObjectName, PreviousNames);
-	}
-
 	return bFound;
 }
 
diff --git a/Engine/Source/Runtime/CoreUObject/Private/UObject/SoftObjectPath.cpp b/Engine/Source/Runtime/CoreUObject/Private/UObject/SoftObjectPath.cpp
--- a/Engine/Source/Runtime/CoreUObject/Private/UObject/SoftObjectPath.cpp
+++ b/Engine/Source/Runtime/CoreUObject/Private/UObject/SoftObjectPath.cpp
@@ -444,13 +444,11 @@ bool FSoftObjectPath::FixupCoreRedirects()
 {
 	FString OldString = ToString();
 	FCoreRedirectObjectName OldName = FCoreRedirectObjectName(OldString);
-	FCoreRedirectObjectName NewName = FCoreRedirects::GetRedirectedName(ECoreRedirectFlags::Type_Object, OldName);
 
-	// This also might be a class
-	if (OldName == NewName && OldString.StartsWith(TEXT("/Script/")))
-	{
-		NewName = FCoreRedirects::GetRedirectedName(ECoreRedirectFlags::Type_Class, OldName);
-	}
+	// Always try the object redirect, this will pick up any package redirects as well
+	// For things that look like native objects, try all types as we don't know which it would be
+	const bool bIsNative = OldString.StartsWith(TEXT("/Script/"));
+	FCoreRedirectObjectName NewName = FCoreRedirects::GetRedirectedName(bIsNative ? ECoreRedirectFlags::Type_AllMask : ECoreRedirectFlags::Type_Object, OldName);
 
 	if (OldName != NewName)
 	{
diff --git a/Engine/Source/Runtime/CoreUObject/Public/UObject/CoreRedirects.h b/Engine/Source/Runtime/CoreUObject/Public/UObject/CoreRedirects.h
--- a/Engine/Source/Runtime/CoreUObject/Public/UObject/CoreRedirects.h
+++ b/Engine/Source/Runtime/CoreUObject/Public/UObject/CoreRedirects.h
@@ -12,7 +12,7 @@
 /** 
  * Flags describing the type and properties of this redirect
  */
-enum class ECoreRedirectFlags : int32
+enum class ECoreRedirectFlags : uint32
 {
 	None = 0,
 
@@ -24,11 +24,14 @@ enum class ECoreRedirectFlags : int32
 	Type_Function =			0x00000010, // UFunction
 	Type_Property =			0x00000020, // UProperty
 	Type_Package =			0x00000040, // UPackage
+	Type_AllMask =			0x0000FFFF, // Bit mask of all possible types
 
 	// Option flags, specify rules for this redirect
 	Option_InstanceOnly =	0x00010000, // Only redirect instances of this type, not the type itself
 	Option_Removed =		0x00020000, // This type was explicitly removed, new name isn't valid
 	Option_MatchSubstring = 0x00040000, // Does a slow substring match
+	Option_AllMask =		0xFFFF0000, // Bit mask of all possible options
+	Option_ExactMatchMask = Option_InstanceOnly | Option_Removed // These options must match exactly, substring is checked both ways
 };
 ENUM_CLASS_FLAGS(ECoreRedirectFlags);
 
