[NO NON-PII COMMIT MESSAGE]

diff --git a/Engine/Source/Runtime/CoreUObject/Private/UObject/CoreRedirects.cpp b/Engine/Source/Runtime/CoreUObject/Private/UObject/CoreRedirects.cpp
--- a/Engine/Source/Runtime/CoreUObject/Private/UObject/CoreRedirects.cpp
+++ b/Engine/Source/Runtime/CoreUObject/Private/UObject/CoreRedirects.cpp
@@ -504,7 +504,7 @@ bool FCoreRedirects::GetMatchingRedirects(ECoreRedirectFlags Type, const FCoreRe
 		ECoreRedirectFlags PairType = Pair.Key;
 
 		// For type, check it includes the matching type
-		const bool bTypeMatches = !!((PairType | Type) & ECoreRedirectFlags::Type_AllMask);
+		const bool bTypeMatches = !!((PairType & Type) & ECoreRedirectFlags::Type_AllMask);
 		// For options, instance/remove must match exactly, but allow both substring and non substring
 		const bool bOptionsMatch = (PairType & ECoreRedirectFlags::Option_ExactMatchMask) == (Type & ECoreRedirectFlags::Option_ExactMatchMask);
 
@@ -545,7 +545,7 @@ bool FCoreRedirects::FindPreviousNames(ECoreRedirectFlags Type, const FCoreRedir
 		ECoreRedirectFlags PairType = Pair.Key;
 
 		// For type, check it includes the matching type
-		const bool bTypeMatches = !!((PairType | Type) & ECoreRedirectFlags::Type_AllMask);
+		const bool bTypeMatches = !!((PairType & Type) & ECoreRedirectFlags::Type_AllMask);
 		// For options, instance/remove must match exactly, but allow both substring and non substring
 		const bool bOptionsMatch = (PairType & ECoreRedirectFlags::Option_ExactMatchMask) == (Type & ECoreRedirectFlags::Option_ExactMatchMask);
 
diff --git a/Engine/Source/Runtime/CoreUObject/Private/UObject/SoftObjectPath.cpp b/Engine/Source/Runtime/CoreUObject/Private/UObject/SoftObjectPath.cpp
--- a/Engine/Source/Runtime/CoreUObject/Private/UObject/SoftObjectPath.cpp
+++ b/Engine/Source/Runtime/CoreUObject/Private/UObject/SoftObjectPath.cpp
@@ -452,8 +452,14 @@ bool FSoftObjectPath::FixupCoreRedirects()
 
 	if (OldName != NewName)
 	{
-		SetPath(NewName.ToString());
-		return true;
+		// Only do the fixup if the old object isn't in memory, this avoids false positives
+		UObject* FoundOldObject = FindObject<UObject>(nullptr, *OldString);
+
+		if (!FoundOldObject)
+		{
+			SetPath(NewName.ToString());
+			return true;
+		}
 	}
 
 	return false;
