Guard FScopeCycleCounterUObject in execCallMathFunction with PER_FUNCTION_SCRIPT_STATS, FBlueprintEventTimer moved into cpp because access to thread singleton in header file is unsafe (fixes TOTAL_OVERHEAD_SCRIPT_STATS behavior), only track blueprint time in gamethread as combining timers across threads is not meaningful

diff --git a/Engine/Source/Runtime/CoreUObject/Private/UObject/ScriptCore.cpp b/Engine/Source/Runtime/CoreUObject/Private/UObject/ScriptCore.cpp
--- a/Engine/Source/Runtime/CoreUObject/Private/UObject/ScriptCore.cpp
+++ b/Engine/Source/Runtime/CoreUObject/Private/UObject/ScriptCore.cpp
@@ -697,7 +697,9 @@ DEFINE_FUNCTION(UObject::execCallMathFunction)
 	UObject* NewContext = Function->GetOuterUClass()->GetDefaultObject(false);
 	checkSlow(NewContext);
 	{
+#if PER_FUNCTION_SCRIPT_STATS
 		FScopeCycleCounterUObject FunctionScope(Function);
+#endif // PER_FUNCTION_SCRIPT_STATS
 
 		// CurrentNativeFunction is used so far only by FLuaContext::InvokeScriptFunction
 		// TGuardValue<UFunction*> NativeFuncGuard(Stack.CurrentNativeFunction, Function);
@@ -1183,6 +1185,82 @@ UFunction* UObject::FindFunctionChecked( FName InName ) const
 	return Result;
 }
 
+#if TOTAL_OVERHEAD_SCRIPT_STATS
+void FBlueprintEventTimer::FPausableScopeTimer::Start()
+{
+	FPausableScopeTimer*& ActiveTimer = FThreadedTimerManager::Get().ActiveTimer;
+
+	double CurrentTime = FPlatformTime::Seconds();
+	if (ActiveTimer)
+	{
+		ActiveTimer->Pause(CurrentTime);
+	}
+
+	PreviouslyActiveTimer = ActiveTimer;
+	StartTime = CurrentTime;
+	TotalTime = 0.0;
+
+	ActiveTimer = this;
+}
+
+double FBlueprintEventTimer::FPausableScopeTimer::Stop()
+{
+	if (PreviouslyActiveTimer)
+	{
+		PreviouslyActiveTimer->Resume();
+	}
+	FThreadedTimerManager::Get().ActiveTimer = PreviouslyActiveTimer;
+	return TotalTime + (FPlatformTime::Seconds() - StartTime);
+}
+
+FBlueprintEventTimer::FScopedVMTimer::FScopedVMTimer()
+	: Timer()
+	, VMParent(nullptr)
+{
+	if (IsInGameThread())
+	{
+		FScopedVMTimer*& ActiveVMTimer = FThreadedTimerManager::Get().ActiveVMScope;
+		VMParent = ActiveVMTimer;
+
+		ActiveVMTimer = this;
+		Timer.Start();
+	}
+}
+
+FBlueprintEventTimer::FScopedVMTimer::~FScopedVMTimer()
+{
+	if (IsInGameThread())
+	{
+		INC_FLOAT_STAT_BY(STAT_ScriptVmTime_Total, Timer.Stop() * 1000.0);
+		FThreadedTimerManager::Get().ActiveVMScope = VMParent;
+	}
+}
+
+FBlueprintEventTimer::FScopedNativeTimer::FScopedNativeTimer()
+	: Timer()
+{
+	if (IsInGameThread())
+	{
+		Timer.Start();
+	}
+}
+
+FBlueprintEventTimer::FScopedNativeTimer::~FScopedNativeTimer()
+{
+	if (IsInGameThread())
+	{
+		if (FThreadedTimerManager::Get().ActiveVMScope)
+		{
+			if (IsInGameThread())
+			{
+				INC_FLOAT_STAT_BY(STAT_ScriptNativeTime_Total, Timer.Stop()* 1000.0);
+			}
+		}
+	}
+}
+
+#endif
+
 void UObject::ProcessEvent( UFunction* Function, void* Parms )
 {
 	checkf(!IsUnreachable(),TEXT("%s  Function: '%s'"), *GetFullName(), *Function->GetPathName());
@@ -1244,7 +1322,7 @@ void UObject::ProcessEvent( UFunction* Function, void* Parms )
 	FBlueprintExceptionTracker& BlueprintExceptionTracker = FBlueprintExceptionTracker::Get();
 	BlueprintExceptionTracker.ScriptEntryTag++;
 
-	CONDITIONAL_SCOPE_CYCLE_COUNTER(STAT_BlueprintTime, BlueprintExceptionTracker.ScriptEntryTag == 1);
+	CONDITIONAL_SCOPE_CYCLE_COUNTER(STAT_BlueprintTime, IsInGameThread() && BlueprintExceptionTracker.ScriptEntryTag == 1);
 #endif
 
 #if UE_BLUEPRINT_EVENTGRAPH_FASTCALLS
diff --git a/Engine/Source/Runtime/CoreUObject/Public/UObject/Script.h b/Engine/Source/Runtime/CoreUObject/Public/UObject/Script.h
--- a/Engine/Source/Runtime/CoreUObject/Public/UObject/Script.h
+++ b/Engine/Source/Runtime/CoreUObject/Public/UObject/Script.h
@@ -496,35 +496,17 @@ private:
 
 		struct FPausableScopeTimer
 		{
-			FPausableScopeTimer() 
-			{ 
-				FPausableScopeTimer*& ActiveTimer = FThreadedTimerManager::Get().ActiveTimer;
-
-				double CurrentTime = FPlatformTime::Seconds();
-				if (ActiveTimer)
-				{
-					ActiveTimer->Pause(CurrentTime);
-				}
-
-				PreviouslyActiveTimer = ActiveTimer;
-				StartTime = CurrentTime;
-				TotalTime = 0.0;
-
-				ActiveTimer = this;
-			}
-
-			~FPausableScopeTimer()
+			FPausableScopeTimer()
+				: PreviouslyActiveTimer(nullptr)
+				, TotalTime(0.0)
+				, StartTime(0.0)
 			{
-				if (PreviouslyActiveTimer)
-				{
-					PreviouslyActiveTimer->Resume();
-				}
-				FThreadedTimerManager::Get().ActiveTimer = PreviouslyActiveTimer;
 			}
 
+			void Start();
 			void Pause(double CurrentTime) { TotalTime += CurrentTime - StartTime; }
 			void Resume() { StartTime = FPlatformTime::Seconds();  }
-			double Stop() { return TotalTime + (FPlatformTime::Seconds() - StartTime); }
+			double Stop();
 
 		private:
 			FPausableScopeTimer* PreviouslyActiveTimer;
@@ -534,40 +516,17 @@ private:
 
 		struct FScopedVMTimer
 		{
-			FScopedVMTimer()
-				: Timer()
-				, VMParent(nullptr)
-			{
-				FScopedVMTimer*& ActiveVMTimer = FThreadedTimerManager::Get().ActiveVMScope;
-				VMParent = ActiveVMTimer;
-
-				ActiveVMTimer = this;
-			}
+			COREUOBJECT_API FScopedVMTimer();
+			COREUOBJECT_API ~FScopedVMTimer();
 
-			~FScopedVMTimer()
-			{
-				INC_FLOAT_STAT_BY(STAT_ScriptVmTime_Total, Timer.Stop() * 1000.0);
-				FThreadedTimerManager::Get().ActiveVMScope = VMParent;
-			}
 			FPausableScopeTimer Timer;
 			FScopedVMTimer* VMParent;
 		};
 
 		struct FScopedNativeTimer
 		{
-			FScopedNativeTimer()
-				: Timer()
-			{}
-
-			~FScopedNativeTimer()
-			{
-				// only track native time when in a VM scope, RPC time
-				// can be tracked by the online system or whatever is making RPCs:
-				if (FThreadedTimerManager::Get().ActiveVMScope)
-				{
-					INC_FLOAT_STAT_BY(STAT_ScriptNativeTime_Total, Timer.Stop()* 1000.0);
-				}
-			}
+			COREUOBJECT_API FScopedNativeTimer();
+			COREUOBJECT_API ~FScopedNativeTimer();
 
 			FPausableScopeTimer Timer;
 		};
