Interpreter Optimization 1 of N: Remove ensure from interpreter hot path

diff --git a/Engine/Source/Runtime/CoreUObject/Private/UObject/Class.cpp b/Engine/Source/Runtime/CoreUObject/Private/UObject/Class.cpp
--- a/Engine/Source/Runtime/CoreUObject/Private/UObject/Class.cpp
+++ b/Engine/Source/Runtime/CoreUObject/Private/UObject/Class.cpp
@@ -4718,7 +4718,7 @@ void UFunction::Invoke(UObject* Obj, FFrame& Stack, RESULT_DECL)
 {
 	checkSlow(Func);
 
-	UClass* OuterClass = GetOuterUClass();
+	UClass* OuterClass = (UClass*)GetOuter();
 	if (OuterClass->IsChildOf(UInterface::StaticClass()))
 	{
 		Obj = (UObject*)Obj->GetInterfaceAddress(OuterClass);
diff --git a/Engine/Source/Runtime/CoreUObject/Private/UObject/ScriptCore.cpp b/Engine/Source/Runtime/CoreUObject/Private/UObject/ScriptCore.cpp
--- a/Engine/Source/Runtime/CoreUObject/Private/UObject/ScriptCore.cpp
+++ b/Engine/Source/Runtime/CoreUObject/Private/UObject/ScriptCore.cpp
@@ -694,7 +694,7 @@ DEFINE_FUNCTION(UObject::execCallMathFunction)
 	UFunction* Function = (UFunction*)Stack.ReadObject();
 	checkSlow(Function);
 	checkSlow(Function->FunctionFlags & FUNC_Native);
-	UObject* NewContext = Function->GetOuterUClass()->GetDefaultObject(false);
+	UObject* NewContext = ((UClass*)Function->GetOuter())->ClassDefaultObject;
 	checkSlow(NewContext);
 	{
 #if PER_FUNCTION_SCRIPT_STATS
