Fix automated test failure (CoreRedirects).

#rnx

diff --git a/Engine/Source/Runtime/CoreUObject/Private/UObject/CoreRedirects.cpp b/Engine/Source/Runtime/CoreUObject/Private/UObject/CoreRedirects.cpp
--- a/Engine/Source/Runtime/CoreUObject/Private/UObject/CoreRedirects.cpp
+++ b/Engine/Source/Runtime/CoreUObject/Private/UObject/CoreRedirects.cpp
@@ -613,7 +613,7 @@ bool FCoreRedirects::RunTests()
 	new (NewRedirects) FCoreRedirect(ECoreRedirectFlags::Type_Property, TEXT("/game/Package.Class.OtherProperty"), TEXT("/game/Package.Class.OtherProperty2"));
 	new (NewRedirects) FCoreRedirect(ECoreRedirectFlags::Type_Class, TEXT("Class"), TEXT("Class2"));
 	new (NewRedirects) FCoreRedirect(ECoreRedirectFlags::Type_Class, TEXT("/game/Package.Class"), TEXT("Class3"));
-	new (NewRedirects) FCoreRedirect(ECoreRedirectFlags::Type_Class | ECoreRedirectFlags::Option_InstanceOnly, TEXT("/game/Package.Class"), TEXT("ClassInstance"));
+	new (NewRedirects) FCoreRedirect(ECoreRedirectFlags::Type_Class | ECoreRedirectFlags::Option_InstanceOnly, TEXT("/game/Package.Class"), TEXT("/game/Package2.ClassInstance"));
 	new (NewRedirects) FCoreRedirect(ECoreRedirectFlags::Type_Package, TEXT("/game/Package"), TEXT("/game/Package2"));
 	new (NewRedirects) FCoreRedirect(ECoreRedirectFlags::Type_Package | ECoreRedirectFlags::Option_MatchSubstring, TEXT("/oldgame"), TEXT("/newgame"));
 	new (NewRedirects) FCoreRedirect(ECoreRedirectFlags::Type_Package | ECoreRedirectFlags::Option_Removed, TEXT("/game/RemovedPackage"), TEXT("/game/RemovedPackage"));
