Fix bug where CoreRedirects was incorrectly skipping the implicit package redirects when checking for a class instanceonly redirect. The automated tests picked up on this so restore that test to the right data

diff --git a/Engine/Source/Runtime/CoreUObject/Private/UObject/CoreRedirects.cpp b/Engine/Source/Runtime/CoreUObject/Private/UObject/CoreRedirects.cpp
--- a/Engine/Source/Runtime/CoreUObject/Private/UObject/CoreRedirects.cpp
+++ b/Engine/Source/Runtime/CoreUObject/Private/UObject/CoreRedirects.cpp
@@ -505,8 +505,11 @@ bool FCoreRedirects::GetMatchingRedirects(ECoreRedirectFlags Type, const FCoreRe
 
 		// For type, check it includes the matching type
 		const bool bTypeMatches = !!((PairType & Type) & ECoreRedirectFlags::Type_AllMask);
-		// For options, instance/remove must match exactly, but allow both substring and non substring
-		const bool bOptionsMatch = (PairType & ECoreRedirectFlags::Option_ExactMatchMask) == (Type & ECoreRedirectFlags::Option_ExactMatchMask);
+
+		// For options, instance/remove must match exactly but allow both substring and non-substring matches
+		// But for packages ignore instance so the implicit package redirects get picked up
+		const ECoreRedirectFlags TypeOptions = !(PairType & ECoreRedirectFlags::Type_Package) ? (Type & ECoreRedirectFlags::Option_ExactMatchMask) : (Type & ECoreRedirectFlags::Option_Removed);
+		const bool bOptionsMatch = (PairType & ECoreRedirectFlags::Option_ExactMatchMask) == TypeOptions;
 
 		// We need to check all maps that match the type mask
 		if (bTypeMatches && bOptionsMatch)
@@ -546,8 +549,11 @@ bool FCoreRedirects::FindPreviousNames(ECoreRedirectFlags Type, const FCoreRedir
 
 		// For type, check it includes the matching type
 		const bool bTypeMatches = !!((PairType & Type) & ECoreRedirectFlags::Type_AllMask);
-		// For options, instance/remove must match exactly, but allow both substring and non substring
-		const bool bOptionsMatch = (PairType & ECoreRedirectFlags::Option_ExactMatchMask) == (Type & ECoreRedirectFlags::Option_ExactMatchMask);
+
+		// For options, instance/remove must match exactly but allow both substring and non-substring matches
+		// But for packages ignore instance so the implicit package redirects get picked up
+		const ECoreRedirectFlags TypeOptions = !(PairType & ECoreRedirectFlags::Type_Package) ? (Type & ECoreRedirectFlags::Option_ExactMatchMask) : (Type & ECoreRedirectFlags::Option_Removed);
+		const bool bOptionsMatch = (PairType & ECoreRedirectFlags::Option_ExactMatchMask) == TypeOptions;
 
 		// We need to check all maps that match the type mask
 		if (bTypeMatches && bOptionsMatch)
@@ -613,7 +619,7 @@ bool FCoreRedirects::RunTests()
 	new (NewRedirects) FCoreRedirect(ECoreRedirectFlags::Type_Property, TEXT("/game/Package.Class.OtherProperty"), TEXT("/game/Package.Class.OtherProperty2"));
 	new (NewRedirects) FCoreRedirect(ECoreRedirectFlags::Type_Class, TEXT("Class"), TEXT("Class2"));
 	new (NewRedirects) FCoreRedirect(ECoreRedirectFlags::Type_Class, TEXT("/game/Package.Class"), TEXT("Class3"));
-	new (NewRedirects) FCoreRedirect(ECoreRedirectFlags::Type_Class | ECoreRedirectFlags::Option_InstanceOnly, TEXT("/game/Package.Class"), TEXT("/game/Package2.ClassInstance"));
+	new (NewRedirects) FCoreRedirect(ECoreRedirectFlags::Type_Class | ECoreRedirectFlags::Option_InstanceOnly, TEXT("/game/Package.Class"), TEXT("/game/Package.ClassInstance"));
 	new (NewRedirects) FCoreRedirect(ECoreRedirectFlags::Type_Package, TEXT("/game/Package"), TEXT("/game/Package2"));
 	new (NewRedirects) FCoreRedirect(ECoreRedirectFlags::Type_Package | ECoreRedirectFlags::Option_MatchSubstring, TEXT("/oldgame"), TEXT("/newgame"));
 	new (NewRedirects) FCoreRedirect(ECoreRedirectFlags::Type_Package | ECoreRedirectFlags::Option_Removed, TEXT("/game/RemovedPackage"), TEXT("/game/RemovedPackage"));
