Interpreter Optimization 4 of N: Avoid branching when looking up owning UClass from UFunction, begin search for ubergraph frame using the functions uclass

diff --git a/Engine/Source/Runtime/CoreUObject/Private/UObject/ScriptCore.cpp b/Engine/Source/Runtime/CoreUObject/Private/UObject/ScriptCore.cpp
--- a/Engine/Source/Runtime/CoreUObject/Private/UObject/ScriptCore.cpp
+++ b/Engine/Source/Runtime/CoreUObject/Private/UObject/ScriptCore.cpp
@@ -694,7 +694,7 @@ DEFINE_FUNCTION(UObject::execCallMathFunction)
 	UFunction* Function = (UFunction*)Stack.ReadObject();
 	checkSlow(Function);
 	checkSlow(Function->FunctionFlags & FUNC_Native);
-	UObject* NewContext = ((UClass*)Function->GetOuter())->ClassDefaultObject;
+	UObject* NewContext = Function->GetOuterUClassUnchecked()->ClassDefaultObject;
 	checkSlow(NewContext);
 	{
 #if PER_FUNCTION_SCRIPT_STATS
@@ -783,7 +783,7 @@ void UObject::CallFunction( FFrame& Stack, RESULT_DECL, UFunction* Function )
 	{
 		uint8* Frame = NULL;
 #if USE_UBER_GRAPH_PERSISTENT_FRAME
-		Frame = GetClass()->GetPersistentUberGraphFrame(this, Function);
+		Frame = Function->GetOuterUClassUnchecked()->GetPersistentUberGraphFrame(this, Function);
 #endif
 		const bool bUsePersistentFrame = (NULL != Frame);
 		if (!bUsePersistentFrame)
@@ -1361,7 +1361,7 @@ void UObject::ProcessEvent( UFunction* Function, void* Parms )
 	{
 		uint8* Frame = NULL;
 #if USE_UBER_GRAPH_PERSISTENT_FRAME
-		Frame = GetClass()->GetPersistentUberGraphFrame(this, Function);
+		Frame = Function->GetOuterUClassUnchecked()->GetPersistentUberGraphFrame(this, Function);
 #endif
 		const bool bUsePersistentFrame = (NULL != Frame);
 		if (!bUsePersistentFrame)
@@ -1897,12 +1897,13 @@ DEFINE_FUNCTION(UObject::execLetValueOnPersistentFrame)
 	Stack.MostRecentProperty = NULL;
 	Stack.MostRecentPropertyAddress = NULL;
 
-	auto DestProperty = Stack.ReadProperty();
+	UProperty* DestProperty = Stack.ReadProperty();
 	checkSlow(DestProperty);
-	auto UberGraphFunction = CastChecked<UFunction>(DestProperty->GetOwnerStruct());
-	auto FrameBase = Stack.Object->GetClass()->GetPersistentUberGraphFrame(Stack.Object, UberGraphFunction);
+	UFunction* UberGraphFunction = CastChecked<UFunction>(DestProperty->GetOwnerStruct());
+	checkSlow(Stack.Object->GetClass()->IsChildOf(UberGraphFunction->GetOuterUClassUnchecked()));
+	uint8* FrameBase = UberGraphFunction->GetOuterUClassUnchecked()->GetPersistentUberGraphFrame(Stack.Object, UberGraphFunction);
 	checkSlow(FrameBase);
-	auto DestAddress = DestProperty->ContainerPtrToValuePtr<uint8>(FrameBase);
+	uint8* DestAddress = DestProperty->ContainerPtrToValuePtr<uint8>(FrameBase);
 
 	Stack.Step(Stack.Object, DestAddress);
 #else
diff --git a/Engine/Source/Runtime/CoreUObject/Public/UObject/Class.h b/Engine/Source/Runtime/CoreUObject/Public/UObject/Class.h
--- a/Engine/Source/Runtime/CoreUObject/Public/UObject/Class.h
+++ b/Engine/Source/Runtime/CoreUObject/Public/UObject/Class.h
@@ -1397,6 +1397,13 @@ public:
 	/** Returns the return value property if there is one, or null */
 	UProperty* GetReturnProperty() const;
 
+	/** Returns the owning UClass* without branching */
+	FORCEINLINE UClass* GetOuterUClassUnchecked() const
+	{
+		// declaration order mandates reinterpret_cast:
+		return reinterpret_cast<UClass*>(GetOuter());
+	}
+
 	/**
 	 * Used to safely check whether the passed in flag is set.
 	 *
