Interpreter Optimization 6 of N: Remove redundant lookups of FBlueprintExceptionHandler

diff --git a/Engine/Source/Runtime/CoreUObject/Private/UObject/ScriptCore.cpp b/Engine/Source/Runtime/CoreUObject/Private/UObject/ScriptCore.cpp
--- a/Engine/Source/Runtime/CoreUObject/Private/UObject/ScriptCore.cpp
+++ b/Engine/Source/Runtime/CoreUObject/Private/UObject/ScriptCore.cpp
@@ -993,14 +993,15 @@ DEFINE_FUNCTION(UObject::ProcessInternal)
 		MS_ALIGN(16) uint8 Buffer[MAX_SIMPLE_RETURN_VALUE_SIZE] GCC_ALIGN(16);
 
 #if DO_BLUEPRINT_GUARD
-		if(FBlueprintExceptionTracker::Get().bRanaway)
+		FBlueprintExceptionTracker& BpET = FBlueprintExceptionTracker::Get();
+		if(BpET.bRanaway)
 		{
 			// If we have a return property, return a zeroed value in it, to try and save execution as much as possible
 			UProperty* ReturnProp = (Function)->GetReturnProperty();
 			ClearReturnValue(ReturnProp, RESULT_PARAM);
 			return;
 		}
-		else if (++FBlueprintExceptionTracker::Get().Recurse == RECURSE_LIMIT)
+		else if (++BpET.Recurse == RECURSE_LIMIT)
 		{
 			// If we have a return property, return a zeroed value in it, to try and save execution as much as possible
 			UProperty* ReturnProp = (Function)->GetReturnProperty();
@@ -1018,7 +1019,7 @@ DEFINE_FUNCTION(UObject::ProcessInternal)
 
 			// This flag prevents repeated warnings of infinite loop, script exception handler 
 			// is expected to have terminated execution appropriately:
-			FBlueprintExceptionTracker::Get().bRanaway = true;
+			BpET.bRanaway = true;
 
 			return;
 		}
@@ -1027,7 +1028,7 @@ DEFINE_FUNCTION(UObject::ProcessInternal)
 		while (*Stack.Code != EX_Return)
 		{
 #if DO_BLUEPRINT_GUARD
-			if( FBlueprintExceptionTracker::Get().Runaway > GMaximumScriptLoopIterations )
+			if(BpET.Runaway > GMaximumScriptLoopIterations )
 			{
 				// If we have a return property, return a zeroed value in it, to try and save execution as much as possible
 				UProperty* ReturnProp = (Function)->GetReturnProperty();
@@ -1044,7 +1045,7 @@ DEFINE_FUNCTION(UObject::ProcessInternal)
 
 				// Need to reset Runaway counter BEFORE throwing script exception, because the exception causes a modal dialog,
 				// and other scripts running will then erroneously think they are also "runaway".
-				FBlueprintExceptionTracker::Get().Runaway = 0;
+				BpET.Runaway = 0;
 
 				FBlueprintCoreDelegates::ThrowScriptException(P_THIS, Stack, RunawayLoopExceptionInfo);
 				return;
@@ -1067,7 +1068,7 @@ DEFINE_FUNCTION(UObject::ProcessInternal)
 		}
 
 #if DO_BLUEPRINT_GUARD
-		--FBlueprintExceptionTracker::Get().Recurse;
+		--BpET.Recurse;
 #endif
 	}
 	else
