[NO NON-PII COMMIT MESSAGE]

diff --git a/Engine/Source/Runtime/CoreUObject/Private/UObject/CoreRedirects.cpp b/Engine/Source/Runtime/CoreUObject/Private/UObject/CoreRedirects.cpp
--- a/Engine/Source/Runtime/CoreUObject/Private/UObject/CoreRedirects.cpp
+++ b/Engine/Source/Runtime/CoreUObject/Private/UObject/CoreRedirects.cpp
@@ -489,38 +489,40 @@ const TMap<FString, FString>* FCoreRedirects::GetValueRedirects(ECoreRedirectFla
 	return nullptr;
 }
 
-bool FCoreRedirects::GetMatchingRedirects(ECoreRedirectFlags Type, const FCoreRedirectObjectName& OldObjectName, TArray<const FCoreRedirect*>& FoundRedirects)
+static FORCEINLINE bool CheckRedirectFlagsMatch(ECoreRedirectFlags FlagsA, ECoreRedirectFlags FlagsB)
+{
+	// For type, check it includes the matching type
+	const bool bTypeMatches = !!((FlagsA & FlagsB) & ECoreRedirectFlags::Type_AllMask);
+
+	// For options, instance/remove must match exactly but allow both substring and non-substring matches
+	const bool bOptionsMatch = (FlagsA & ECoreRedirectFlags::Option_ExactMatchMask) == (FlagsB & ECoreRedirectFlags::Option_ExactMatchMask);
+
+	return bTypeMatches && bOptionsMatch;
+}
+
+bool FCoreRedirects::GetMatchingRedirects(ECoreRedirectFlags SearchFlags, const FCoreRedirectObjectName& OldObjectName, TArray<const FCoreRedirect*>& FoundRedirects)
 {
 	// Look for all redirects that match the given names and flags
 	bool bFound = false;
-
-	// Always search for package redirects and substring matches
-	Type |= ECoreRedirectFlags::Type_Package;
-	Type |= ECoreRedirectFlags::Option_MatchSubstring;
+	
+	// If we're not explicitly searching for packages or looking for removed things, add the implicit package redirects
+	const bool bSearchPackageRedirects = !(SearchFlags & ECoreRedirectFlags::Type_Package) && !(SearchFlags & ECoreRedirectFlags::Option_Removed);
 
 	// Determine list of maps to look over, need to handle being passed multiple types in a bit mask
 	for (const TPair<ECoreRedirectFlags, FRedirectNameMap>& Pair : RedirectTypeMap)
 	{
-		ECoreRedirectFlags PairType = Pair.Key;
-
-		// For type, check it includes the matching type
-		const bool bTypeMatches = !!((PairType & Type) & ECoreRedirectFlags::Type_AllMask);
-
-		// For options, instance/remove must match exactly but allow both substring and non-substring matches
-		// But for packages ignore instance so the implicit package redirects get picked up
-		const ECoreRedirectFlags TypeOptions = !(PairType & ECoreRedirectFlags::Type_Package) ? (Type & ECoreRedirectFlags::Option_ExactMatchMask) : (Type & ECoreRedirectFlags::Option_Removed);
-		const bool bOptionsMatch = (PairType & ECoreRedirectFlags::Option_ExactMatchMask) == TypeOptions;
+		ECoreRedirectFlags PairFlags = Pair.Key;
 
-		// We need to check all maps that match the type mask
-		if (bTypeMatches && bOptionsMatch)
+		// We need to check all maps that match the search or package flags
+		if (CheckRedirectFlagsMatch(PairFlags, SearchFlags) || (bSearchPackageRedirects && CheckRedirectFlagsMatch(PairFlags, ECoreRedirectFlags::Type_Package)))
 		{
-			const TArray<FCoreRedirect>* RedirectsForName = Pair.Value.RedirectMap.Find(OldObjectName.GetSearchKey(PairType));
+			const TArray<FCoreRedirect>* RedirectsForName = Pair.Value.RedirectMap.Find(OldObjectName.GetSearchKey(PairFlags));
 
 			if (RedirectsForName)
 			{
 				for (const FCoreRedirect& CheckRedirect : *RedirectsForName)
 				{
-					if (CheckRedirect.Matches(PairType, OldObjectName))
+					if (CheckRedirect.Matches(PairFlags, OldObjectName))
 					{
 						bFound = true;
 						FoundRedirects.Add(&CheckRedirect);
@@ -533,30 +535,21 @@ bool FCoreRedirects::GetMatchingRedirects(ECoreRedirectFlags Type, const FCoreRe
 	return bFound;
 }
 
-bool FCoreRedirects::FindPreviousNames(ECoreRedirectFlags Type, const FCoreRedirectObjectName& NewObjectName, TArray<FCoreRedirectObjectName>& PreviousNames)
+bool FCoreRedirects::FindPreviousNames(ECoreRedirectFlags SearchFlags, const FCoreRedirectObjectName& NewObjectName, TArray<FCoreRedirectObjectName>& PreviousNames)
 {
 	// Look for reverse direction redirects
 	bool bFound = false;
 
-	// Always search for package redirects and substring matches
-	Type |= ECoreRedirectFlags::Type_Package;
-	Type |= ECoreRedirectFlags::Option_MatchSubstring;
+	// If we're not explicitly searching for packages or looking for removed things, add the implicit package redirects
+	const bool bSearchPackageRedirects = !(SearchFlags & ECoreRedirectFlags::Type_Package) && !(SearchFlags & ECoreRedirectFlags::Option_Removed);
 
-	// Determine list of maps to look over, need to handle being passed multiple types in a bit mask
+	// Determine list of maps to look over, need to handle being passed multiple Flagss in a bit mask
 	for (const TPair<ECoreRedirectFlags, FRedirectNameMap>& Pair : RedirectTypeMap)
 	{
-		ECoreRedirectFlags PairType = Pair.Key;
-
-		// For type, check it includes the matching type
-		const bool bTypeMatches = !!((PairType & Type) & ECoreRedirectFlags::Type_AllMask);
-
-		// For options, instance/remove must match exactly but allow both substring and non-substring matches
-		// But for packages ignore instance so the implicit package redirects get picked up
-		const ECoreRedirectFlags TypeOptions = !(PairType & ECoreRedirectFlags::Type_Package) ? (Type & ECoreRedirectFlags::Option_ExactMatchMask) : (Type & ECoreRedirectFlags::Option_Removed);
-		const bool bOptionsMatch = (PairType & ECoreRedirectFlags::Option_ExactMatchMask) == TypeOptions;
+		ECoreRedirectFlags PairFlags = Pair.Key;
 
-		// We need to check all maps that match the type mask
-		if (bTypeMatches && bOptionsMatch)
+		// We need to check all maps that match the search or package flags
+		if (CheckRedirectFlagsMatch(PairFlags, SearchFlags) || (bSearchPackageRedirects && CheckRedirectFlagsMatch(PairFlags, ECoreRedirectFlags::Type_Package)))
 		{
 			for (const TPair<FName, TArray<FCoreRedirect> >& RedirectPair : Pair.Value.RedirectMap)
 			{
