[NO NON-PII COMMIT MESSAGE]

diff --git a/Engine/Source/Runtime/CoreUObject/Private/Serialization/AsyncLoading.cpp b/Engine/Source/Runtime/CoreUObject/Private/Serialization/AsyncLoading.cpp
--- a/Engine/Source/Runtime/CoreUObject/Private/Serialization/AsyncLoading.cpp
+++ b/Engine/Source/Runtime/CoreUObject/Private/Serialization/AsyncLoading.cpp
@@ -4745,7 +4745,11 @@ EAsyncPackageState::Type FAsyncLoadingThread::ProcessLoadedPackages(bool bUseTim
 
 #if WITH_EDITOR
 		// In editor builds, call the asset load callback. This happens in both editor and standalone to match EndLoad
-		for (const FWeakObjectPtr& WeakAsset : LoadedAssets)
+		TArray<FWeakObjectPtr> TempLoadedAssets = LoadedAssets;
+		LoadedAssets.Reset();
+
+		// Make a copy because LoadedAssets could be modified by one of the OnAssetLoaded callbacks
+		for (const FWeakObjectPtr& WeakAsset : TempLoadedAssets)
 		{
 			// It may have been unloaded/marked pending kill since being added, ignore those cases
 			if (UObject* LoadedAsset = WeakAsset.Get())
@@ -4753,7 +4757,6 @@ EAsyncPackageState::Type FAsyncLoadingThread::ProcessLoadedPackages(bool bUseTim
 				FCoreUObjectDelegates::OnAssetLoaded.Broadcast(LoadedAsset);
 			}
 		}
-		LoadedAssets.Reset();
 #endif
 
 		// We're not done until all packages have been deleted
