Interpreter Optimization 7 of N: CPF_ZeroConstructor assigned for struct properties that that are User Defined Structs

diff --git a/Engine/Source/Runtime/CoreUObject/Private/UObject/Class.cpp b/Engine/Source/Runtime/CoreUObject/Private/UObject/Class.cpp
--- a/Engine/Source/Runtime/CoreUObject/Private/UObject/Class.cpp
+++ b/Engine/Source/Runtime/CoreUObject/Private/UObject/Class.cpp
@@ -1254,7 +1254,8 @@ void UStruct::Serialize( FArchive& Ar )
 		//       introduce unwanted dependency loads at this time)
 		ScriptLoadHelper.LoadStructWithScript(this, Ar, bAllowDeferredScriptSerialization);
 
-		if (!dynamic_cast<UClass*>(this) && !(Ar.GetPortFlags() & PPF_Duplicate)) // classes are linked in the UClass serializer, which just called me
+		// UClass and UScriptStruct are linked in the their respective ::Serialize implementation:
+		if (!dynamic_cast<UClass*>(this) && !dynamic_cast<UScriptStruct*>(this) && !(Ar.GetPortFlags() & PPF_Duplicate))
 		{
 			// Link the properties.
 			Link(Ar, true);
@@ -2013,10 +2014,9 @@ void UScriptStruct::Serialize( FArchive& Ar )
 	// serialize the struct's flags
 	Ar << (uint32&)StructFlags;
 
-	if (Ar.IsLoading())
+	if (Ar.IsLoading() && !(Ar.GetPortFlags() & PPF_Duplicate))
 	{
-		ClearCppStructOps(); // we want to be sure to do this from scratch
-		PrepareCppStructOps();
+		Link(Ar, true);
 	}
 }
 
diff --git a/Engine/Source/Runtime/CoreUObject/Private/UObject/PropertyStruct.cpp b/Engine/Source/Runtime/CoreUObject/Private/UObject/PropertyStruct.cpp
--- a/Engine/Source/Runtime/CoreUObject/Private/UObject/PropertyStruct.cpp
+++ b/Engine/Source/Runtime/CoreUObject/Private/UObject/PropertyStruct.cpp
@@ -39,6 +39,10 @@ UStructProperty::UStructProperty(ECppProperty, int32 InOffset, EPropertyFlags In
 	, Struct(InStruct)
 {
 	ElementSize = Struct->PropertiesSize;
+	if (Struct->StructFlags & STRUCT_ZeroConstructor)
+	{
+		SetPropertyFlags(CPF_ZeroConstructor);
+	}
 }
 
 UStructProperty::UStructProperty( const FObjectInitializer& ObjectInitializer, ECppProperty, int32 InOffset, EPropertyFlags InFlags, UScriptStruct* InStruct )
@@ -46,6 +50,10 @@ UStructProperty::UStructProperty( const FObjectInitializer& ObjectInitializer, E
 	,	Struct( InStruct )
 {
 	ElementSize = Struct->PropertiesSize;
+	if (Struct->StructFlags & STRUCT_ZeroConstructor)
+	{
+		SetPropertyFlags(CPF_ZeroConstructor);
+	}
 }
 
 int32 UStructProperty::GetMinAlignment() const
@@ -83,6 +91,11 @@ void UStructProperty::LinkInternal(FArchive& Ar)
 		// User Defined structs won't have UScriptStruct::ICppStructOps. Setting their flags here.
 		PropertyFlags |= CPF_HasGetValueTypeHash;
 	}
+
+	if (Struct->StructFlags & STRUCT_ZeroConstructor)
+	{
+		SetPropertyFlags(CPF_ZeroConstructor);
+	}
 }
 
 bool UStructProperty::Identical( const void* A, const void* B, uint32 PortFlags ) const
