Also don't clear on cached weak pointer when saving via ArchiveUObject

diff --git a/Engine/Source/Runtime/CoreUObject/Private/Serialization/ArchiveUObject.cpp b/Engine/Source/Runtime/CoreUObject/Private/Serialization/ArchiveUObject.cpp
--- a/Engine/Source/Runtime/CoreUObject/Private/Serialization/ArchiveUObject.cpp
+++ b/Engine/Source/Runtime/CoreUObject/Private/Serialization/ArchiveUObject.cpp
@@ -62,8 +62,11 @@ FArchive& FArchiveUObject::SerializeSoftObjectPtr(FArchive& Ar, FSoftObjectPtr&
 {
 	if (Ar.IsSaving() || Ar.IsLoading())
 	{
-		// Reset before serializing to clear the internal weak pointer. 
-		Value.ResetWeakPtr();
+		if (Ar.IsLoading())
+		{
+			// Reset before serializing to clear the internal weak pointer. 
+			Value.ResetWeakPtr();
+		}
 		Ar << Value.GetUniqueID();
 	}
 	else if (!Ar.IsObjectReferenceCollector() || Ar.IsModifyingWeakAndStrongReferences())
@@ -130,6 +133,23 @@ FArchive& FObjectAndNameAsStringProxyArchive::operator<<(FWeakObjectPtr& Obj)
 	return *this;
 }
 
+FArchive& FObjectAndNameAsStringProxyArchive::operator<<(FSoftObjectPtr& Value)
+{
+	if (IsLoading())
+	{
+		// Reset before serializing to clear the internal weak pointer. 
+		Value.ResetWeakPtr();
+	}
+	*this << Value.GetUniqueID();
+	return *this;
+}
+
+FArchive& FObjectAndNameAsStringProxyArchive::operator<<(FSoftObjectPath& Value)
+{
+	Value.SerializePath(*this);
+	return *this;
+}
+
 void FSerializedPropertyScope::PushProperty()
 {
 	if (Property)
diff --git a/Engine/Source/Runtime/CoreUObject/Public/Serialization/ObjectAndNameAsStringProxyArchive.h b/Engine/Source/Runtime/CoreUObject/Public/Serialization/ObjectAndNameAsStringProxyArchive.h
--- a/Engine/Source/Runtime/CoreUObject/Public/Serialization/ObjectAndNameAsStringProxyArchive.h
+++ b/Engine/Source/Runtime/CoreUObject/Public/Serialization/ObjectAndNameAsStringProxyArchive.h
@@ -36,5 +36,7 @@ struct FObjectAndNameAsStringProxyArchive : public FNameAsStringProxyArchive
 
 	COREUOBJECT_API virtual FArchive& operator<<(UObject*& Obj) override;
 	COREUOBJECT_API virtual FArchive& operator<<(FWeakObjectPtr& Obj) override;
+	COREUOBJECT_API virtual FArchive& operator<<(FSoftObjectPtr& Value) override;
+	COREUOBJECT_API virtual FArchive& operator<<(FSoftObjectPath& Value) override;
 };
 
