Made the NoExport FDateTime and FTimespan struct Tick data actual properties

diff --git a/Engine/Source/Runtime/CoreUObject/Public/UObject/NoExportTypes.h b/Engine/Source/Runtime/CoreUObject/Public/UObject/NoExportTypes.h
--- a/Engine/Source/Runtime/CoreUObject/Public/UObject/NoExportTypes.h
+++ b/Engine/Source/Runtime/CoreUObject/Public/UObject/NoExportTypes.h
@@ -944,6 +944,7 @@ public:
 USTRUCT(immutable, noexport, BlueprintType, meta=(HasNativeMake="Engine.KismetMathLibrary.MakeDateTime", HasNativeBreak="Engine.KismetMathLibrary.BreakDateTime"))
 struct FDateTime
 {
+	UPROPERTY()
 	int64 Ticks;
 };
 
@@ -1021,6 +1022,7 @@ struct FTimecode
 USTRUCT(immutable, noexport, BlueprintType, meta=(HasNativeMake="Engine.KismetMathLibrary.MakeTimespan", HasNativeBreak="Engine.KismetMathLibrary.BreakTimespan"))
 struct FTimespan
 {
+	UPROPERTY()
 	int64 Ticks;
 };
 
