Property Localization - Only gather text from collection properties who's inner types can contain text.

diff --git a/Engine/Source/Runtime/CoreUObject/Private/Serialization/PropertyLocalizationDataGathering.cpp b/Engine/Source/Runtime/CoreUObject/Private/Serialization/PropertyLocalizationDataGathering.cpp
--- a/Engine/Source/Runtime/CoreUObject/Private/Serialization/PropertyLocalizationDataGathering.cpp
+++ b/Engine/Source/Runtime/CoreUObject/Private/Serialization/PropertyLocalizationDataGathering.cpp
@@ -211,6 +211,13 @@ void FPropertyLocalizationDataGatherer::GatherLocalizationDataFromChildTextPrope
 
 	const EPropertyLocalizationGathererTextFlags ChildPropertyGatherTextFlags = GatherTextFlags | (Property->HasAnyPropertyFlags(CPF_EditorOnly) ? EPropertyLocalizationGathererTextFlags::ForceEditorOnly : EPropertyLocalizationGathererTextFlags::None);
 
+	auto CanGatherFromInnerProperty = [](UProperty* InnerProperty)
+	{
+		return InnerProperty->IsA<UTextProperty>() ||
+			InnerProperty->IsA<UStructProperty>() ||
+			InnerProperty->IsA<UObjectPropertyBase>();
+	};
+
 	// Handles both native, fixed-size arrays and plain old non-array properties.
 	const bool IsFixedSizeArray = Property->ArrayDim > 1;
 	for(int32 i = 0; i < Property->ArrayDim; ++i)
@@ -242,6 +249,12 @@ void FPropertyLocalizationDataGatherer::GatherLocalizationDataFromChildTextPrope
 		// Property is a DYNAMIC array property.
 		else if (ArrayProperty)
 		{
+			// Skip arrays of types which can't contain any text.
+			if (CanGatherFromInnerProperty(ArrayProperty->Inner) == false)
+			{
+				continue;
+			}
+
 			// Iterate over all elements of the array.
 			FScriptArrayHelper ScriptArrayHelper(ArrayProperty, ElementValueAddress);
 			const int32 ElementCount = ScriptArrayHelper.Num();
@@ -253,6 +266,12 @@ void FPropertyLocalizationDataGatherer::GatherLocalizationDataFromChildTextPrope
 		// Property is a map property.
 		else if (MapProperty)
 		{
+			// Skip maps of types which can't contain any text.
+			if (CanGatherFromInnerProperty(MapProperty->KeyProp) == false && CanGatherFromInnerProperty(MapProperty->ValueProp) == false)
+			{
+				continue;
+			}
+
 			// Iterate over all elements of the map.
 			FScriptMapHelper ScriptMapHelper(MapProperty, ElementValueAddress);
 			const int32 ElementCount = ScriptMapHelper.Num();
@@ -264,14 +283,26 @@ void FPropertyLocalizationDataGatherer::GatherLocalizationDataFromChildTextPrope
 				}
 
 				const uint8* MapPairPtr = ScriptMapHelper.GetPairPtr(j);
-				GatherLocalizationDataFromChildTextProperties(PathToElement + FString::Printf(TEXT("(%d - Key)"), ElementIndex), MapProperty->KeyProp, MapPairPtr + MapProperty->MapLayout.KeyOffset, nullptr, ChildPropertyGatherTextFlags);
-				GatherLocalizationDataFromChildTextProperties(PathToElement + FString::Printf(TEXT("(%d - Value)"), ElementIndex), MapProperty->ValueProp, MapPairPtr + MapProperty->MapLayout.ValueOffset, nullptr, ChildPropertyGatherTextFlags);
+				if (CanGatherFromInnerProperty(MapProperty->KeyProp))
+				{
+					GatherLocalizationDataFromChildTextProperties(PathToElement + FString::Printf(TEXT("(%d - Key)"), ElementIndex), MapProperty->KeyProp, MapPairPtr + MapProperty->MapLayout.KeyOffset, nullptr, ChildPropertyGatherTextFlags);
+				}
+				if (CanGatherFromInnerProperty(MapProperty->ValueProp))
+				{
+					GatherLocalizationDataFromChildTextProperties(PathToElement + FString::Printf(TEXT("(%d - Value)"), ElementIndex), MapProperty->ValueProp, MapPairPtr + MapProperty->MapLayout.ValueOffset, nullptr, ChildPropertyGatherTextFlags);
+				}
 				++ElementIndex;
 			}
 		}
 		// Property is a set property.
 		else if (SetProperty)
 		{
+			// Skip sets of types which can't contain any text.
+			if (CanGatherFromInnerProperty(SetProperty->ElementProp) == false)
+			{
+				continue;
+			}
+
 			// Iterate over all elements of the Set.
 			FScriptSetHelper ScriptSetHelper(SetProperty, ElementValueAddress);
 			const int32 ElementCount = ScriptSetHelper.Num();
