Remove dumb returns
#rnx

diff --git a/Engine/Source/Runtime/CoreUObject/Private/UObject/Class.cpp b/Engine/Source/Runtime/CoreUObject/Private/UObject/Class.cpp
--- a/Engine/Source/Runtime/CoreUObject/Private/UObject/Class.cpp
+++ b/Engine/Source/Runtime/CoreUObject/Private/UObject/Class.cpp
@@ -442,7 +442,7 @@ void UField::SetMetaData(const TCHAR* Key, const TCHAR* InValue)
 	UPackage* Package = GetOutermost();
 	check(Package);
 
-	return Package->GetMetaData()->SetValue(this, Key, InValue);
+	Package->GetMetaData()->SetValue(this, Key, InValue);
 }
 
 void UField::SetMetaData(const FName& Key, const TCHAR* InValue)
@@ -450,7 +450,7 @@ void UField::SetMetaData(const FName& Key, const TCHAR* InValue)
 	UPackage* Package = GetOutermost();
 	check(Package);
 
-	return Package->GetMetaData()->SetValue(this, Key, InValue);
+	Package->GetMetaData()->SetValue(this, Key, InValue);
 }
 
 UClass* UField::GetClassMetaData(const TCHAR* Key) const
