Implemented NetSerialize for FTimespan. This allows such variables to be replicated.

diff --git a/Engine/Source/Runtime/Core/Private/Misc/Timespan.cpp b/Engine/Source/Runtime/Core/Private/Misc/Timespan.cpp
--- a/Engine/Source/Runtime/Core/Private/Misc/Timespan.cpp
+++ b/Engine/Source/Runtime/Core/Private/Misc/Timespan.cpp
@@ -51,6 +51,14 @@ bool FTimespan::Serialize(FArchive& Ar)
 }
 
 
+bool FTimespan::NetSerialize(FArchive& Ar, class UPackageMap* Map, bool& bOutSuccess)
+{
+	Ar << *this;
+	bOutSuccess = true;
+	return true;
+}
+
+
 FString FTimespan::ToString() const
 {
 	if (GetDays() == 0)
diff --git a/Engine/Source/Runtime/Core/Public/Misc/Timespan.h b/Engine/Source/Runtime/Core/Public/Misc/Timespan.h
--- a/Engine/Source/Runtime/Core/Public/Misc/Timespan.h
+++ b/Engine/Source/Runtime/Core/Public/Misc/Timespan.h
@@ -547,6 +547,11 @@ public:
 	 */
 	CORE_API bool Serialize(FArchive& Ar);
 
+	/**
+	* Serializes this date and time for network transmission
+	*/
+	CORE_API bool NetSerialize(FArchive& Ar, class UPackageMap* Map, bool& bOutSuccess);
+
 	/**
 	 * Return the string representation of this time span using a default format.
 	 *
diff --git a/Engine/Source/Runtime/CoreUObject/Private/UObject/Property.cpp b/Engine/Source/Runtime/CoreUObject/Private/UObject/Property.cpp
--- a/Engine/Source/Runtime/CoreUObject/Private/UObject/Property.cpp
+++ b/Engine/Source/Runtime/CoreUObject/Private/UObject/Property.cpp
@@ -272,6 +272,8 @@ struct TStructOpsTypeTraits<FTimespan> : public TStructOpsTypeTraitsBase2<FTimes
 		WithExportTextItem = true,
 		WithImportTextItem = true,
 		WithSerializer = true,
+		WithNetSerializer = true,
+		WithNetSharedSerialization = true,
 		WithZeroConstructor = true,
 		WithIdenticalViaEquality = true,
 	};
