Copy of Dev-Core CL 4366352

Fix a crash in structured archive when using enum resolver archives

diff --git a/Engine/Source/Runtime/CoreUObject/Private/UObject/EnumProperty.cpp b/Engine/Source/Runtime/CoreUObject/Private/UObject/EnumProperty.cpp
--- a/Engine/Source/Runtime/CoreUObject/Private/UObject/EnumProperty.cpp
+++ b/Engine/Source/Runtime/CoreUObject/Private/UObject/EnumProperty.cpp
@@ -78,6 +78,7 @@ void UEnumProperty::SerializeItem(FStructuredArchive::FSlot Slot, void* Value, v
 
 	if (Enum && UnderlyingArchive.UseToResolveEnumerators())
 	{
+		Slot.EnterStream();
 		int64 IntValue = UnderlyingProp->GetSignedIntPropertyValue(Value);
 		int64 ResolvedIndex = Enum->ResolveEnumerator(UnderlyingArchive, IntValue);
 		UnderlyingProp->SetIntPropertyValue(Value, ResolvedIndex);
diff --git a/Engine/Source/Runtime/CoreUObject/Private/UObject/PropertyByte.cpp b/Engine/Source/Runtime/CoreUObject/Private/UObject/PropertyByte.cpp
--- a/Engine/Source/Runtime/CoreUObject/Private/UObject/PropertyByte.cpp
+++ b/Engine/Source/Runtime/CoreUObject/Private/UObject/PropertyByte.cpp
@@ -25,6 +25,7 @@ void UByteProperty::SerializeItem( FStructuredArchive::FSlot Slot, void* Value,
 
 	if(Enum && UnderlyingArchive.UseToResolveEnumerators())
 	{
+		Slot.EnterStream();
 		 const int64 ResolvedIndex = Enum->ResolveEnumerator(UnderlyingArchive, *(uint8*)Value);
 		 *(uint8*)Value = static_cast<uint8>(ResolvedIndex);
 		 return;
