Merging //UE4/Dev-Main to Dev-Build (//UE4/Dev-Build)

#rnx

diff --git a/Engine/Source/Editor/Blutility/Classes/EditorAutomationActor.h b/Engine/Source/Editor/Blutility/Classes/EditorAutomationActor.h
--- a/Engine/Source/Editor/Blutility/Classes/EditorAutomationActor.h
+++ b/Engine/Source/Editor/Blutility/Classes/EditorAutomationActor.h
@@ -1,4 +1,4 @@
-// Copyright 1998-2018 Epic Games, Inc. All Rights Reserved.
+// Copyright 1998-2019 Epic Games, Inc. All Rights Reserved.
 
 /**
  * Base class of any editor-only automation test actors
diff --git a/Engine/Source/Editor/Blutility/Classes/EditorAutomationActorComponent.h b/Engine/Source/Editor/Blutility/Classes/EditorAutomationActorComponent.h
--- a/Engine/Source/Editor/Blutility/Classes/EditorAutomationActorComponent.h
+++ b/Engine/Source/Editor/Blutility/Classes/EditorAutomationActorComponent.h
@@ -1,4 +1,4 @@
-// Copyright 1998-2018 Epic Games, Inc. All Rights Reserved.
+// Copyright 1998-2019 Epic Games, Inc. All Rights Reserved.
 
 /**
  * Base class of any editor-only automation test actor components
diff --git a/Engine/Source/Editor/Blutility/Classes/EditorAutomationObject.h b/Engine/Source/Editor/Blutility/Classes/EditorAutomationObject.h
--- a/Engine/Source/Editor/Blutility/Classes/EditorAutomationObject.h
+++ b/Engine/Source/Editor/Blutility/Classes/EditorAutomationObject.h
@@ -1,4 +1,4 @@
-// Copyright 1998-2018 Epic Games, Inc. All Rights Reserved.
+// Copyright 1998-2019 Epic Games, Inc. All Rights Reserved.
 
 /**
  * Base class of any editor-only automation test objects
diff --git a/Engine/Source/Editor/Blutility/Classes/EditorUtilityWidget.h b/Engine/Source/Editor/Blutility/Classes/EditorUtilityWidget.h
--- a/Engine/Source/Editor/Blutility/Classes/EditorUtilityWidget.h
+++ b/Engine/Source/Editor/Blutility/Classes/EditorUtilityWidget.h
@@ -1,4 +1,4 @@
-// Copyright 1998-2018 Epic Games, Inc. All Rights Reserved.
+// Copyright 1998-2019 Epic Games, Inc. All Rights Reserved.
 
 /**
  * Widget for editor utilities
diff --git a/Engine/Source/Editor/Blutility/Classes/EditorUtilityWidgetBlueprint.h b/Engine/Source/Editor/Blutility/Classes/EditorUtilityWidgetBlueprint.h
--- a/Engine/Source/Editor/Blutility/Classes/EditorUtilityWidgetBlueprint.h
+++ b/Engine/Source/Editor/Blutility/Classes/EditorUtilityWidgetBlueprint.h
@@ -1,4 +1,4 @@
-// Copyright 1998-2018 Epic Games, Inc. All Rights Reserved.
+// Copyright 1998-2019 Epic Games, Inc. All Rights Reserved.
 
 /**
  * Widget for editor utilities
diff --git a/Engine/Source/Editor/Blutility/Classes/EditorUtilityWidgetBlueprintFactory.h b/Engine/Source/Editor/Blutility/Classes/EditorUtilityWidgetBlueprintFactory.h
--- a/Engine/Source/Editor/Blutility/Classes/EditorUtilityWidgetBlueprintFactory.h
+++ b/Engine/Source/Editor/Blutility/Classes/EditorUtilityWidgetBlueprintFactory.h
@@ -1,4 +1,4 @@
-// Copyright 1998-2018 Epic Games, Inc. All Rights Reserved.
+// Copyright 1998-2019 Epic Games, Inc. All Rights Reserved.
 
 #pragma once
 
diff --git a/Engine/Source/Editor/Blutility/Private/AssetTypeActions_EditorUtilityWidgetBlueprint.cpp b/Engine/Source/Editor/Blutility/Private/AssetTypeActions_EditorUtilityWidgetBlueprint.cpp
--- a/Engine/Source/Editor/Blutility/Private/AssetTypeActions_EditorUtilityWidgetBlueprint.cpp
+++ b/Engine/Source/Editor/Blutility/Private/AssetTypeActions_EditorUtilityWidgetBlueprint.cpp
@@ -1,4 +1,4 @@
-// Copyright 1998-2018 Epic Games, Inc. All Rights Reserved.
+// Copyright 1998-2019 Epic Games, Inc. All Rights Reserved.
 
 #include "AssetTypeActions_EditorUtilityWidgetBlueprint.h"
 #include "Framework/MultiBox/MultiBoxBuilder.h"
diff --git a/Engine/Source/Editor/Blutility/Private/AssetTypeActions_EditorUtilityWidgetBlueprint.h b/Engine/Source/Editor/Blutility/Private/AssetTypeActions_EditorUtilityWidgetBlueprint.h
--- a/Engine/Source/Editor/Blutility/Private/AssetTypeActions_EditorUtilityWidgetBlueprint.h
+++ b/Engine/Source/Editor/Blutility/Private/AssetTypeActions_EditorUtilityWidgetBlueprint.h
@@ -1,4 +1,4 @@
-// Copyright 1998-2018 Epic Games, Inc. All Rights Reserved.
+// Copyright 1998-2019 Epic Games, Inc. All Rights Reserved.
 
 #pragma once
 
diff --git a/Engine/Source/Editor/Blutility/Private/EditorAutomationActor.cpp b/Engine/Source/Editor/Blutility/Private/EditorAutomationActor.cpp
--- a/Engine/Source/Editor/Blutility/Private/EditorAutomationActor.cpp
+++ b/Engine/Source/Editor/Blutility/Private/EditorAutomationActor.cpp
@@ -1,4 +1,4 @@
-// Copyright 1998-2018 Epic Games, Inc. All Rights Reserved.
+// Copyright 1998-2019 Epic Games, Inc. All Rights Reserved.
 
 #include "EditorAutomationActor.h"
 
diff --git a/Engine/Source/Editor/Blutility/Private/EditorAutomationActorComponent.cpp b/Engine/Source/Editor/Blutility/Private/EditorAutomationActorComponent.cpp
--- a/Engine/Source/Editor/Blutility/Private/EditorAutomationActorComponent.cpp
+++ b/Engine/Source/Editor/Blutility/Private/EditorAutomationActorComponent.cpp
@@ -1,4 +1,4 @@
-// Copyright 1998-2018 Epic Games, Inc. All Rights Reserved.
+// Copyright 1998-2019 Epic Games, Inc. All Rights Reserved.
 
 #include "EditorAutomationActorComponent.h"
 
diff --git a/Engine/Source/Editor/Blutility/Private/EditorAutomationObject.cpp b/Engine/Source/Editor/Blutility/Private/EditorAutomationObject.cpp
--- a/Engine/Source/Editor/Blutility/Private/EditorAutomationObject.cpp
+++ b/Engine/Source/Editor/Blutility/Private/EditorAutomationObject.cpp
@@ -1,4 +1,4 @@
-// Copyright 1998-2018 Epic Games, Inc. All Rights Reserved.
+// Copyright 1998-2019 Epic Games, Inc. All Rights Reserved.
 
 #include "EditorAutomationObject.h"
 
diff --git a/Engine/Source/Editor/Blutility/Private/EditorUtilityContext.cpp b/Engine/Source/Editor/Blutility/Private/EditorUtilityContext.cpp
--- a/Engine/Source/Editor/Blutility/Private/EditorUtilityContext.cpp
+++ b/Engine/Source/Editor/Blutility/Private/EditorUtilityContext.cpp
@@ -1,4 +1,4 @@
-// Copyright 1998-2018 Epic Games, Inc. All Rights Reserved.
+// Copyright 1998-2019 Epic Games, Inc. All Rights Reserved.
 
 #include "EditorUtilityContext.h"
 #include "Engine/Selection.h"
diff --git a/Engine/Source/Editor/Blutility/Private/EditorUtilityContext.h b/Engine/Source/Editor/Blutility/Private/EditorUtilityContext.h
--- a/Engine/Source/Editor/Blutility/Private/EditorUtilityContext.h
+++ b/Engine/Source/Editor/Blutility/Private/EditorUtilityContext.h
@@ -1,4 +1,4 @@
-// Copyright 1998-2018 Epic Games, Inc. All Rights Reserved.
+// Copyright 1998-2019 Epic Games, Inc. All Rights Reserved.
 
 #pragma once
 
diff --git a/Engine/Source/Editor/Blutility/Private/EditorUtilityLibrary.cpp b/Engine/Source/Editor/Blutility/Private/EditorUtilityLibrary.cpp
--- a/Engine/Source/Editor/Blutility/Private/EditorUtilityLibrary.cpp
+++ b/Engine/Source/Editor/Blutility/Private/EditorUtilityLibrary.cpp
@@ -1,4 +1,4 @@
-// Copyright 1998-2018 Epic Games, Inc. All Rights Reserved.
+// Copyright 1998-2019 Epic Games, Inc. All Rights Reserved.
 
 #include "EditorUtilityLibrary.h"
 #include "Engine/Selection.h"
diff --git a/Engine/Source/Editor/Blutility/Private/EditorUtilityWidget.cpp b/Engine/Source/Editor/Blutility/Private/EditorUtilityWidget.cpp
--- a/Engine/Source/Editor/Blutility/Private/EditorUtilityWidget.cpp
+++ b/Engine/Source/Editor/Blutility/Private/EditorUtilityWidget.cpp
@@ -1,4 +1,4 @@
-// Copyright 1998-2018 Epic Games, Inc. All Rights Reserved.
+// Copyright 1998-2019 Epic Games, Inc. All Rights Reserved.
 
 #include "EditorUtilityWidget.h"
 #include "Modules/ModuleManager.h"
diff --git a/Engine/Source/Editor/Blutility/Private/EditorUtilityWidgetBlueprint.cpp b/Engine/Source/Editor/Blutility/Private/EditorUtilityWidgetBlueprint.cpp
--- a/Engine/Source/Editor/Blutility/Private/EditorUtilityWidgetBlueprint.cpp
+++ b/Engine/Source/Editor/Blutility/Private/EditorUtilityWidgetBlueprint.cpp
@@ -1,4 +1,4 @@
-// Copyright 1998-2018 Epic Games, Inc. All Rights Reserved.
+// Copyright 1998-2019 Epic Games, Inc. All Rights Reserved.
 
 #include "EditorUtilityWidgetBlueprint.h"
 #include "WidgetBlueprint.h"
diff --git a/Engine/Source/Editor/Blutility/Private/EditorUtilityWidgetBlueprintFactory.cpp b/Engine/Source/Editor/Blutility/Private/EditorUtilityWidgetBlueprintFactory.cpp
--- a/Engine/Source/Editor/Blutility/Private/EditorUtilityWidgetBlueprintFactory.cpp
+++ b/Engine/Source/Editor/Blutility/Private/EditorUtilityWidgetBlueprintFactory.cpp
@@ -1,4 +1,4 @@
-// Copyright 1998-2018 Epic Games, Inc. All Rights Reserved.
+// Copyright 1998-2019 Epic Games, Inc. All Rights Reserved.
 
 #include "EditorUtilityWidgetBlueprintFactory.h"
 #include "Misc/MessageDialog.h"
diff --git a/Engine/Source/Editor/Blutility/Public/EditorUtilityLibrary.h b/Engine/Source/Editor/Blutility/Public/EditorUtilityLibrary.h
--- a/Engine/Source/Editor/Blutility/Public/EditorUtilityLibrary.h
+++ b/Engine/Source/Editor/Blutility/Public/EditorUtilityLibrary.h
@@ -1,4 +1,4 @@
-// Copyright 1998-2018 Epic Games, Inc. All Rights Reserved.
+// Copyright 1998-2019 Epic Games, Inc. All Rights Reserved.
 
 #pragma once
 
diff --git a/Engine/Source/Editor/MaterialEditor/Private/SMaterialParametersOverviewWidget.cpp b/Engine/Source/Editor/MaterialEditor/Private/SMaterialParametersOverviewWidget.cpp
--- a/Engine/Source/Editor/MaterialEditor/Private/SMaterialParametersOverviewWidget.cpp
+++ b/Engine/Source/Editor/MaterialEditor/Private/SMaterialParametersOverviewWidget.cpp
@@ -1,4 +1,4 @@
-// Copyright 1998-2018 Epic Games, Inc. All Rights Reserved.
+// Copyright 1998-2019 Epic Games, Inc. All Rights Reserved.
 
 #include "SMaterialParametersOverviewWidget.h"
 #include "MaterialEditor/DEditorFontParameterValue.h"
diff --git a/Engine/Source/Editor/MaterialEditor/Private/SMaterialParametersOverviewWidget.h b/Engine/Source/Editor/MaterialEditor/Private/SMaterialParametersOverviewWidget.h
--- a/Engine/Source/Editor/MaterialEditor/Private/SMaterialParametersOverviewWidget.h
+++ b/Engine/Source/Editor/MaterialEditor/Private/SMaterialParametersOverviewWidget.h
@@ -1,4 +1,4 @@
-// Copyright 1998-2018 Epic Games, Inc. All Rights Reserved.
+// Copyright 1998-2019 Epic Games, Inc. All Rights Reserved.
 
 #pragma once
 
diff --git a/Engine/Source/Editor/UMGEditor/Private/Widgets/SGraphNodeCreateWidget.cpp b/Engine/Source/Editor/UMGEditor/Private/Widgets/SGraphNodeCreateWidget.cpp
--- a/Engine/Source/Editor/UMGEditor/Private/Widgets/SGraphNodeCreateWidget.cpp
+++ b/Engine/Source/Editor/UMGEditor/Private/Widgets/SGraphNodeCreateWidget.cpp
@@ -1,4 +1,4 @@
-// Copyright 1998-2018 Epic Games, Inc. All Rights Reserved.
+// Copyright 1998-2019 Epic Games, Inc. All Rights Reserved.
 
 #include "Widgets/SGraphNodeCreateWidget.h"
 #include "Modules/ModuleManager.h"
diff --git a/Engine/Source/Editor/UMGEditor/Private/Widgets/SGraphNodeCreateWidget.h b/Engine/Source/Editor/UMGEditor/Private/Widgets/SGraphNodeCreateWidget.h
--- a/Engine/Source/Editor/UMGEditor/Private/Widgets/SGraphNodeCreateWidget.h
+++ b/Engine/Source/Editor/UMGEditor/Private/Widgets/SGraphNodeCreateWidget.h
@@ -1,4 +1,4 @@
-// Copyright 1998-2018 Epic Games, Inc. All Rights Reserved.
+// Copyright 1998-2019 Epic Games, Inc. All Rights Reserved.
 
 #pragma once
 
diff --git a/Engine/Source/Runtime/CoreUObject/Private/Internationalization/StabilizeLocalizationKeys.cpp b/Engine/Source/Runtime/CoreUObject/Private/Internationalization/StabilizeLocalizationKeys.cpp
--- a/Engine/Source/Runtime/CoreUObject/Private/Internationalization/StabilizeLocalizationKeys.cpp
+++ b/Engine/Source/Runtime/CoreUObject/Private/Internationalization/StabilizeLocalizationKeys.cpp
@@ -1,4 +1,4 @@
-// Copyright 1998-2018 Epic Games, Inc. All Rights Reserved.
+// Copyright 1998-2019 Epic Games, Inc. All Rights Reserved.
 
 #include "Internationalization/StabilizeLocalizationKeys.h"
 #include "UObject/TextProperty.h"
diff --git a/Engine/Source/Runtime/CoreUObject/Public/Internationalization/StabilizeLocalizationKeys.h b/Engine/Source/Runtime/CoreUObject/Public/Internationalization/StabilizeLocalizationKeys.h
--- a/Engine/Source/Runtime/CoreUObject/Public/Internationalization/StabilizeLocalizationKeys.h
+++ b/Engine/Source/Runtime/CoreUObject/Public/Internationalization/StabilizeLocalizationKeys.h
@@ -1,4 +1,4 @@
-// Copyright 1998-2018 Epic Games, Inc. All Rights Reserved.
+// Copyright 1998-2019 Epic Games, Inc. All Rights Reserved.
 
 #pragma once
 
diff --git a/Engine/Source/Runtime/CoreUObject/Public/Serialization/BulkDataReader.h b/Engine/Source/Runtime/CoreUObject/Public/Serialization/BulkDataReader.h
--- a/Engine/Source/Runtime/CoreUObject/Public/Serialization/BulkDataReader.h
+++ b/Engine/Source/Runtime/CoreUObject/Public/Serialization/BulkDataReader.h
@@ -1,4 +1,4 @@
-// Copyright 1998-2018 Epic Games, Inc. All Rights Reserved.
+// Copyright 1998-2019 Epic Games, Inc. All Rights Reserved.
 
 #pragma once
 
diff --git a/Engine/Source/Runtime/CoreUObject/Public/Serialization/BulkDataWriter.h b/Engine/Source/Runtime/CoreUObject/Public/Serialization/BulkDataWriter.h
--- a/Engine/Source/Runtime/CoreUObject/Public/Serialization/BulkDataWriter.h
+++ b/Engine/Source/Runtime/CoreUObject/Public/Serialization/BulkDataWriter.h
@@ -1,4 +1,4 @@
-// Copyright 1998-2018 Epic Games, Inc. All Rights Reserved.
+// Copyright 1998-2019 Epic Games, Inc. All Rights Reserved.
 
 #pragma once
 
