Updated assert in PRIVATE_PatchNewObjectIntoExport in light of stricter CurrentLoadContext, fixes crash when compiling blueprints in preview releases

diff --git a/Engine/Source/Runtime/CoreUObject/Private/UObject/LinkerLoad.cpp b/Engine/Source/Runtime/CoreUObject/Private/UObject/LinkerLoad.cpp
--- a/Engine/Source/Runtime/CoreUObject/Private/UObject/LinkerLoad.cpp
+++ b/Engine/Source/Runtime/CoreUObject/Private/UObject/LinkerLoad.cpp
@@ -554,8 +554,10 @@ void FLinkerLoad::PRIVATE_PatchNewObjectIntoExport(UObject* OldObject, UObject*
 		ObjExport.Object = NewObject;
 
 		// If the object was in the ObjLoaded queue (exported, but not yet serialized), swap out for our new object
-		check(OldObjectLinker->GetSerializeContext());
-		OldObjectLinker->GetSerializeContext()->PRIVATE_PatchNewObjectIntoExport(OldObject, NewObject);
+		if(OldObjectLinker->GetSerializeContext())
+		{
+			OldObjectLinker->GetSerializeContext()->PRIVATE_PatchNewObjectIntoExport(OldObject, NewObject);
+		}
 	}
 }
 
