DuplicatedDataReader and DuplicatedDataWriter will now store UObject serialize context internally so that they don't crash in case of an error

diff --git a/Engine/Source/Runtime/CoreUObject/Public/Serialization/DuplicatedDataReader.h b/Engine/Source/Runtime/CoreUObject/Public/Serialization/DuplicatedDataReader.h
--- a/Engine/Source/Runtime/CoreUObject/Public/Serialization/DuplicatedDataReader.h
+++ b/Engine/Source/Runtime/CoreUObject/Public/Serialization/DuplicatedDataReader.h
@@ -23,6 +23,9 @@ private:
 	const FLargeMemoryData&					ObjectData;
 	int64									Offset;
 
+	/** Context for duplication */
+	TRefCountPtr<FUObjectSerializeContext> DuplicateContext;
+
 	//~ Begin FArchive Interface.
 
 	virtual FArchive& operator<<(FName& N);
@@ -67,6 +70,16 @@ public:
 		return ObjectData.GetSize();
 	}
 
+	virtual void SetSerializeContext(FUObjectSerializeContext* InLoadContext) override
+	{
+		DuplicateContext = InLoadContext;
+	}
+
+	virtual FUObjectSerializeContext* GetSerializeContext() override
+	{
+		return DuplicateContext;
+	}
+
 	/**
 	 * Constructor
 	 * 
diff --git a/Engine/Source/Runtime/CoreUObject/Public/Serialization/DuplicatedDataWriter.h b/Engine/Source/Runtime/CoreUObject/Public/Serialization/DuplicatedDataWriter.h
--- a/Engine/Source/Runtime/CoreUObject/Public/Serialization/DuplicatedDataWriter.h
+++ b/Engine/Source/Runtime/CoreUObject/Public/Serialization/DuplicatedDataWriter.h
@@ -34,6 +34,9 @@ private:
 	 */
 	struct FObjectInstancingGraph*			InstanceGraph;
 
+	/** Context for duplication */
+	TRefCountPtr<FUObjectSerializeContext> DuplicateContext;
+
 	//~ Begin FArchive Interface.
 
 	virtual FArchive& operator<<(FName& N);
@@ -80,6 +83,16 @@ public:
 		return ObjectData.GetSize();
 	}
 
+	virtual void SetSerializeContext(FUObjectSerializeContext* InLoadContext) override
+	{
+		DuplicateContext = InLoadContext;
+	}
+
+	virtual FUObjectSerializeContext* GetSerializeContext() override
+	{
+		return DuplicateContext;
+	}
+
 	TArray<UObject*>	UnserializedObjects;
 
 	/**
