Fixing incremental build compile errors introduced in CL #5223889

diff --git a/Engine/Source/Runtime/CoreUObject/Public/Serialization/DuplicatedDataReader.h b/Engine/Source/Runtime/CoreUObject/Public/Serialization/DuplicatedDataReader.h
--- a/Engine/Source/Runtime/CoreUObject/Public/Serialization/DuplicatedDataReader.h
+++ b/Engine/Source/Runtime/CoreUObject/Public/Serialization/DuplicatedDataReader.h
@@ -7,6 +7,8 @@
 #include "UObject/UObjectAnnotation.h"
 #include "Serialization/DuplicatedObject.h"
 #include "Serialization/LargeMemoryData.h"
+#include "Templates/RefCounting.h"
+#include "UObject/UObjectThreadContext.h"
 
 /*----------------------------------------------------------------------------
 	FDuplicateDataReader.
diff --git a/Engine/Source/Runtime/CoreUObject/Public/Serialization/DuplicatedDataWriter.h b/Engine/Source/Runtime/CoreUObject/Public/Serialization/DuplicatedDataWriter.h
--- a/Engine/Source/Runtime/CoreUObject/Public/Serialization/DuplicatedDataWriter.h
+++ b/Engine/Source/Runtime/CoreUObject/Public/Serialization/DuplicatedDataWriter.h
@@ -7,6 +7,8 @@
 #include "UObject/UObjectAnnotation.h"
 #include "Serialization/DuplicatedObject.h"
 #include "Serialization/LargeMemoryData.h"
+#include "Templates/RefCounting.h"
+#include "UObject/UObjectThreadContext.h"
 
 struct FObjectInstancingGraph;
 
