Add DECLARE_CAST_BY_FLAG for all properties with cast flags that weren't already being declared.
#rnx

diff --git a/Engine/Source/Runtime/CoreUObject/Public/Templates/Casts.h b/Engine/Source/Runtime/CoreUObject/Public/Templates/Casts.h
--- a/Engine/Source/Runtime/CoreUObject/Public/Templates/Casts.h
+++ b/Engine/Source/Runtime/CoreUObject/Public/Templates/Casts.h
@@ -344,6 +344,15 @@ DECLARE_CAST_BY_FLAG(UBlueprint)						\
 DECLARE_CAST_BY_FLAG(UDelegateFunction)					\
 DECLARE_CAST_BY_FLAG(UStaticMeshComponent)				\
 DECLARE_CAST_BY_FLAG(UEnumProperty)						\
+DECLARE_CAST_BY_FLAG(UNumericProperty)					\
+DECLARE_CAST_BY_FLAG(UInt8Property)						\
+DECLARE_CAST_BY_FLAG(UInt16Property)					\
+DECLARE_CAST_BY_FLAG(UInt64Property)					\
+DECLARE_CAST_BY_FLAG(UUInt16Property)					\
+DECLARE_CAST_BY_FLAG(UUInt32Property)					\
+DECLARE_CAST_BY_FLAG(UUInt64Property)					\
+DECLARE_CAST_BY_FLAG(UMapProperty)						\
+DECLARE_CAST_BY_FLAG(USetProperty)						\
 FINISH_DECLARING_CAST_FLAGS		// This is here to hopefully remind people to include the "\" in all declarations above, especially when copy/pasting the final line.
 
 // Now actually declare the flags
