Revert to old behavior of using a vcall to get the ubergraphframe due to misuse of event graph by animblueprintcompiler

diff --git a/Engine/Plugins/Experimental/ControlRig/Source/ControlRig/Private/ControlRigBlueprintGeneratedClass.cpp b/Engine/Plugins/Experimental/ControlRig/Source/ControlRig/Private/ControlRigBlueprintGeneratedClass.cpp
--- a/Engine/Plugins/Experimental/ControlRig/Source/ControlRig/Private/ControlRigBlueprintGeneratedClass.cpp
+++ b/Engine/Plugins/Experimental/ControlRig/Source/ControlRig/Private/ControlRigBlueprintGeneratedClass.cpp
@@ -37,3 +37,15 @@ void UControlRigBlueprintGeneratedClass::PurgeClass(bool bRecompilingOnLoad)
 {
 	Super::PurgeClass(bRecompilingOnLoad);
 }
+
+uint8* UControlRigBlueprintGeneratedClass::GetPersistentUberGraphFrame(UObject* Obj, UFunction* FuncToCheck) const
+{
+	if(!IsInGameThread())
+	{
+		// we cant use the persistent frame if we are executing in parallel (as we could potentially thunk to BP)
+		return nullptr;
+	}
+
+	return Super::GetPersistentUberGraphFrame(Obj, FuncToCheck);
+}
+
diff --git a/Engine/Plugins/Experimental/ControlRig/Source/ControlRig/Public/ControlRigBlueprintGeneratedClass.h b/Engine/Plugins/Experimental/ControlRig/Source/ControlRig/Public/ControlRigBlueprintGeneratedClass.h
--- a/Engine/Plugins/Experimental/ControlRig/Source/ControlRig/Public/ControlRigBlueprintGeneratedClass.h
+++ b/Engine/Plugins/Experimental/ControlRig/Source/ControlRig/Public/ControlRigBlueprintGeneratedClass.h
@@ -20,6 +20,7 @@ public:
 
 	// UClass interface
 	virtual void PurgeClass(bool bRecompilingOnLoad) override;
+	virtual uint8* GetPersistentUberGraphFrame(UObject* Obj, UFunction* FuncToCheck) const override;
 
 public:
 #if WITH_EDITORONLY_DATA
diff --git a/Engine/Source/Editor/UnrealEd/Private/Kismet2/KismetDebugUtilities.cpp b/Engine/Source/Editor/UnrealEd/Private/Kismet2/KismetDebugUtilities.cpp
--- a/Engine/Source/Editor/UnrealEd/Private/Kismet2/KismetDebugUtilities.cpp
+++ b/Engine/Source/Editor/UnrealEd/Private/Kismet2/KismetDebugUtilities.cpp
@@ -1209,10 +1209,10 @@ FKismetDebugUtilities::EWatchTextResult FKismetDebugUtilities::FindDebuggingData
 			UFunction* OuterFunction = Cast<UFunction>(Property->GetOuter());
 			if (!PropertyBase && OuterFunction)
 			{
-				UBlueprintGeneratedClass* BPGC = Cast<UBlueprintGeneratedClass>(OuterFunction->GetOuter());
+				UBlueprintGeneratedClass* BPGC = Cast<UBlueprintGeneratedClass>(Blueprint->GeneratedClass);
 				if (BPGC && ActiveObject->IsA(BPGC))
 				{
-					PropertyBase = GetPersistentUberGraphFrame( OuterFunction, ActiveObject );
+					PropertyBase = BPGC->GetPersistentUberGraphFrame(ActiveObject, OuterFunction);
 				}
 			}
 #endif // USE_UBER_GRAPH_PERSISTENT_FRAME
diff --git a/Engine/Source/Runtime/CoreUObject/Private/UObject/Class.cpp b/Engine/Source/Runtime/CoreUObject/Private/UObject/Class.cpp
--- a/Engine/Source/Runtime/CoreUObject/Private/UObject/Class.cpp
+++ b/Engine/Source/Runtime/CoreUObject/Private/UObject/Class.cpp
@@ -3948,10 +3948,6 @@ void UClass::PurgeClass(bool bRecompilingOnLoad)
 	}
 #endif
 
-#if USE_UBER_GRAPH_PERSISTENT_FRAME
-	UberGraphFramePointerProperty = NULL;
-#endif//USE_UBER_GRAPH_PERSISTENT_FRAME
-
 	ClassDefaultObject = NULL;
 
 	Interfaces.Empty();
@@ -4032,9 +4028,6 @@ UClass::UClass(const FObjectInitializer& ObjectInitializer)
 ,	ClassCastFlags(CASTCLASS_None)
 ,	ClassWithin( UObject::StaticClass() )
 ,	ClassGeneratedBy(nullptr)
-#if USE_UBER_GRAPH_PERSISTENT_FRAME
-,	UberGraphFramePointerProperty(nullptr)
-#endif // USE_UBER_GRAPH_PERSISTENT_FRAME
 ,	ClassDefaultObject(nullptr)
 {
 	// If you add properties here, please update the other constructors and PurgeClass()
@@ -4053,9 +4046,6 @@ UClass::UClass(const FObjectInitializer& ObjectInitializer, UClass* InBaseClass)
 ,	ClassCastFlags(CASTCLASS_None)
 ,	ClassWithin(UObject::StaticClass())
 ,	ClassGeneratedBy(nullptr)
-#if USE_UBER_GRAPH_PERSISTENT_FRAME
-,	UberGraphFramePointerProperty(nullptr)
-#endif // USE_UBER_GRAPH_PERSISTENT_FRAME
 ,	ClassDefaultObject(nullptr)
 {
 	// If you add properties here, please update the other constructors and PurgeClass()
diff --git a/Engine/Source/Runtime/CoreUObject/Private/UObject/ScriptCore.cpp b/Engine/Source/Runtime/CoreUObject/Private/UObject/ScriptCore.cpp
--- a/Engine/Source/Runtime/CoreUObject/Private/UObject/ScriptCore.cpp
+++ b/Engine/Source/Runtime/CoreUObject/Private/UObject/ScriptCore.cpp
@@ -50,22 +50,6 @@ static FAutoConsoleVariableRef CVarVerboseScriptStats(
 	ECVF_Default
 );
 
-#if USE_UBER_GRAPH_PERSISTENT_FRAME
-// Mirror definition of FPointerToUberGraphFrame, it is a UStruct and
-// we cannot easily generate its reflection data here in CoreUObject. The
-// builtins pattern we use for FVector, FQuat etc cannot be used because
-// our only member is a raw pointer and it cannot be a UProperty. This
-// creates difficulty in determining the correct size for the UStruct
-struct FPointerToUberGraphFrameCoreUObject
-{
-	uint8* RawPointer;
-
-#if VALIDATE_UBER_GRAPH_PERSISTENT_FRAME
-	uint32 UberGraphFunctionKey;
-#endif//VALIDATE_UBER_GRAPH_PERSISTENT_FRAME
-};
-#endif //USE_UBER_GRAPH_PERSISTENT_FRAME
-
 /*-----------------------------------------------------------------------------
 	Globals.
 -----------------------------------------------------------------------------*/
@@ -325,33 +309,6 @@ FString UnicodeToCPPIdentifier(const FString& InName, bool bDeprecated, const TC
 	return bDeprecated ? Ret + TEXT("_DEPRECATED") : Ret;
 }
 
-#if USE_UBER_GRAPH_PERSISTENT_FRAME
-/** Returns memory used to store temporary data on an instance, used by blueprints */
-static uint8* GetPersistentUberGraphFrameUnchecked(const UFunction* ForFn, UObject* Obj)
-{
-	const UClass* FromClass = ForFn->GetOuterUClassUnchecked();
-	checkSlow(ForFn->HasAnyFunctionFlags(FUNC_UbergraphFunction));
-	checkSlow(Obj->IsA(FromClass));
-	checkSlow(FromClass->UberGraphFramePointerProperty);
-	FPointerToUberGraphFrameCoreUObject* PointerToUberGraphFrame =
-		FromClass->UberGraphFramePointerProperty->ContainerPtrToValuePtr<FPointerToUberGraphFrameCoreUObject>(
-			(void*)Obj
-		);
-	checkSlow(PointerToUberGraphFrame);
-	checkSlow(PointerToUberGraphFrame->RawPointer);
-	return PointerToUberGraphFrame->RawPointer;
-}
-
-uint8* GetPersistentUberGraphFrame(const UFunction* ForFn, UObject* Obj)
-{
-	if (ForFn->HasAnyFunctionFlags(FUNC_UbergraphFunction))
-	{
-		return GetPersistentUberGraphFrameUnchecked(ForFn, Obj);
-	}
-	return nullptr;
-}
-#endif //USE_UBER_GRAPH_PERSISTENT_FRAME
-
 /*-----------------------------------------------------------------------------
 	FFrame implementation.
 -----------------------------------------------------------------------------*/
@@ -745,7 +702,7 @@ void ProcessScriptFunction(UObject* Context, UFunction* Function, FFrame& Stack,
 	uint8* FrameMemory = nullptr;
 	FFrame NewStack(Context, Function, nullptr, &Stack, Function->Children);
 #if USE_UBER_GRAPH_PERSISTENT_FRAME
-	FrameMemory = GetPersistentUberGraphFrame(Function, Context);
+	FrameMemory = Function->GetOuterUClassUnchecked()->GetPersistentUberGraphFrame(Context, Function);
 #endif
 	bool bUsePersistentFrame = (nullptr != FrameMemory);
 	if (!bUsePersistentFrame)
@@ -1430,7 +1387,10 @@ void UObject::ProcessEvent( UFunction* Function, void* Parms )
 	{
 		uint8* Frame = NULL;
 #if USE_UBER_GRAPH_PERSISTENT_FRAME
-		Frame = GetPersistentUberGraphFrame(Function, this);
+		if (Function->HasAnyFunctionFlags(FUNC_UbergraphFunction))
+		{
+			Frame = Function->GetOuterUClassUnchecked()->GetPersistentUberGraphFrame(this, Function);
+		}
 #endif
 		const bool bUsePersistentFrame = (NULL != Frame);
 		if (!bUsePersistentFrame)
@@ -1964,7 +1924,7 @@ DEFINE_FUNCTION(UObject::execLetValueOnPersistentFrame)
 	checkSlow(DestProperty);
 	UFunction* UberGraphFunction = CastChecked<UFunction>(DestProperty->GetOwnerStruct());
 	checkSlow(Stack.Object->GetClass()->IsChildOf(UberGraphFunction->GetOuterUClassUnchecked()));
-	uint8* FrameBase = GetPersistentUberGraphFrameUnchecked(UberGraphFunction, Stack.Object);
+	uint8* FrameBase = UberGraphFunction->GetOuterUClassUnchecked()->GetPersistentUberGraphFrame(Stack.Object, UberGraphFunction);
 	checkSlow(FrameBase);
 	uint8* DestAddress = DestProperty->ContainerPtrToValuePtr<uint8>(FrameBase);
 
diff --git a/Engine/Source/Runtime/CoreUObject/Public/UObject/Class.h b/Engine/Source/Runtime/CoreUObject/Public/UObject/Class.h
--- a/Engine/Source/Runtime/CoreUObject/Public/UObject/Class.h
+++ b/Engine/Source/Runtime/CoreUObject/Public/UObject/Class.h
@@ -2203,16 +2203,6 @@ public:
 	/** This is the blueprint that caused the generation of this class, or null if it is a native compiled-in class */
 	UObject* ClassGeneratedBy;
 
-#if USE_UBER_GRAPH_PERSISTENT_FRAME
-	/**
-	 * Property that points to the ubergraph frame, this is a blueprint specific structure that has been hoisted
-	 * to UClass so that the interpreter (ScriptCore.cpp) can access it efficiently. The uber graph frame is a struct
-	 * owned by a UObject but allocated separately that has a layout that corresponds to a specific UFunction (the
-	 * UberGraphFunction) in a blueprint.
-	 */
-	UStructProperty* UberGraphFramePointerProperty;
-#endif //USE_UBER_GRAPH_PERSISTENT_FRAME
-
 #if WITH_EDITOR
 	/**
 	 * Conditionally recompiles the class after loading, in case any dependencies were also newly loaded
@@ -2506,6 +2496,12 @@ public:
 	 */
 	FName GetDefaultObjectName();
 
+	/** Returns memory used to store temporary data on an instance, used by blueprints */
+	virtual uint8* GetPersistentUberGraphFrame(UObject* Obj, UFunction* FuncToCheck) const
+	{
+		return nullptr;
+	}
+
 	/** Creates memory to store temporary data */
 	virtual void CreatePersistentUberGraphFrame(UObject* Obj, bool bCreateOnlyIfEmpty = false, bool bSkipSuperClass = false, UClass* OldClass = nullptr) const
 	{
diff --git a/Engine/Source/Runtime/CoreUObject/Public/UObject/Script.h b/Engine/Source/Runtime/CoreUObject/Public/UObject/Script.h
--- a/Engine/Source/Runtime/CoreUObject/Public/UObject/Script.h
+++ b/Engine/Source/Runtime/CoreUObject/Public/UObject/Script.h
@@ -552,13 +552,3 @@ COREUOBJECT_API FString ToValidCPPIdentifierChars(TCHAR Char);
 */
 COREUOBJECT_API FString UnicodeToCPPIdentifier(const FString& InName, bool bDeprecated, const TCHAR* Prefix);
 
-/**
-	@param ForFn Function to search for a corresponding uber graph frame
-	@param Obj owning the uber graph frame, Obj must be an instance based on ForFn->GetOuterUClass()
-	@return A pointer to the start of the objects persistent uber graph
-		frame if ForFn is an ubergraph function, else nullptr. The uber graph
-		frame is a struct that has a matching layout to a blueprint's UberGraphFunction's
-		parameters. The Uber Graph frame is an optimization because it avoids a large
-		allocation and corresponding initialization when calling the UberGraphFunction
-*/
-COREUOBJECT_API uint8* GetPersistentUberGraphFrame(const UFunction* ForFn, UObject* Obj);
\ No newline at end of file
diff --git a/Engine/Source/Runtime/Engine/Classes/Animation/AnimBlueprintGeneratedClass.h b/Engine/Source/Runtime/Engine/Classes/Animation/AnimBlueprintGeneratedClass.h
--- a/Engine/Source/Runtime/Engine/Classes/Animation/AnimBlueprintGeneratedClass.h
+++ b/Engine/Source/Runtime/Engine/Classes/Animation/AnimBlueprintGeneratedClass.h
@@ -350,6 +350,7 @@ public:
 
 	// UClass interface
 	virtual void PurgeClass(bool bRecompilingOnLoad) override;
+	virtual uint8* GetPersistentUberGraphFrame(UObject* Obj, UFunction* FuncToCheck) const override;
 	virtual void PostLoadDefaultObject(UObject* Object) override;
 	// End of UClass interface
 };
diff --git a/Engine/Source/Runtime/Engine/Classes/Engine/BlueprintGeneratedClass.h b/Engine/Source/Runtime/Engine/Classes/Engine/BlueprintGeneratedClass.h
--- a/Engine/Source/Runtime/Engine/Classes/Engine/BlueprintGeneratedClass.h
+++ b/Engine/Source/Runtime/Engine/Classes/Engine/BlueprintGeneratedClass.h
@@ -77,18 +77,13 @@ public:
 	}
 };
 
-//////////////////////////////////////////////////////////////////////////
-// WARNING: Following struct layout definition repeated in ScriptCore.cpp as
-// FPointerToUberGraphFrameCoreUObject to work around reflection generation issues:
+
 USTRUCT()
 struct FPointerToUberGraphFrame
 {
 	GENERATED_USTRUCT_BODY()
 
 public:
-	//////////////////////////////////////////////////////////////////////////
-	// WARNING: This struct layout definition repeated in ScriptCore.cpp as
-	// FPointerToUberGraphFrameCoreUObject to work around reflection generation issues:
 	uint8* RawPointer;
 #if VALIDATE_UBER_GRAPH_PERSISTENT_FRAME
 	uint32 UberGraphFunctionKey;
@@ -109,9 +104,7 @@ public:
 		check(!RawPointer);
 	}
 };
-// WARNING: Preceding struct layout definition repeated in ScriptCore.cpp as
-// FPointerToUberGraphFrameCoreUObject to work around reflection generation issues!
-//////////////////////////////////////////////////////////////////////////
+
 
 template<>
 struct TStructOpsTypeTraits<FPointerToUberGraphFrame> : public TStructOpsTypeTraitsBase2<FPointerToUberGraphFrame>
@@ -641,6 +634,9 @@ public:
 	UPROPERTY()
 	class UInheritableComponentHandler* InheritableComponentHandler;
 
+	UPROPERTY()
+	UStructProperty* UberGraphFramePointerProperty;
+
 	UPROPERTY()
 	UFunction* UberGraphFunction;
 
@@ -709,6 +705,7 @@ public:
 	virtual void SerializeDefaultObject(UObject* Object, FArchive& Ar) override;
 	virtual void PostLoadDefaultObject(UObject* Object) override;
 	virtual bool IsFunctionImplementedInBlueprint(FName InFunctionName) const override;
+	virtual uint8* GetPersistentUberGraphFrame(UObject* Obj, UFunction* FuncToCheck) const override;
 	virtual void CreatePersistentUberGraphFrame(UObject* Obj, bool bCreateOnlyIfEmpty = false, bool bSkipSuperClass = false, UClass* OldClass = nullptr) const override;
 	virtual void DestroyPersistentUberGraphFrame(UObject* Obj, bool bSkipSuperClass = false) const override;
 	virtual void Link(FArchive& Ar, bool bRelinkExistingProperties) override;
diff --git a/Engine/Source/Runtime/Engine/Private/Animation/AnimBlueprintGeneratedClass.cpp b/Engine/Source/Runtime/Engine/Private/Animation/AnimBlueprintGeneratedClass.cpp
--- a/Engine/Source/Runtime/Engine/Private/Animation/AnimBlueprintGeneratedClass.cpp
+++ b/Engine/Source/Runtime/Engine/Private/Animation/AnimBlueprintGeneratedClass.cpp
@@ -334,6 +334,17 @@ void UAnimBlueprintGeneratedClass::PurgeClass(bool bRecompilingOnLoad)
 	BakedStateMachines.Empty();
 }
 
+uint8* UAnimBlueprintGeneratedClass::GetPersistentUberGraphFrame(UObject* Obj, UFunction* FuncToCheck) const
+{
+	if(!IsInGameThread())
+	{
+		// we cant use the persistent frame if we are executing in parallel (as we could potentially thunk to BP)
+		return nullptr;
+	}
+
+	return Super::GetPersistentUberGraphFrame(Obj, FuncToCheck);
+}
+
 void UAnimBlueprintGeneratedClass::PostLoadDefaultObject(UObject* Object)
 {
 	Super::PostLoadDefaultObject(Object);
diff --git a/Engine/Source/Runtime/Engine/Private/BlueprintGeneratedClass.cpp b/Engine/Source/Runtime/Engine/Private/BlueprintGeneratedClass.cpp
--- a/Engine/Source/Runtime/Engine/Private/BlueprintGeneratedClass.cpp
+++ b/Engine/Source/Runtime/Engine/Private/BlueprintGeneratedClass.cpp
@@ -1220,6 +1220,23 @@ void UBlueprintGeneratedClass::CheckAndApplyComponentTemplateOverrides(UObject*
 	}
 }
 
+uint8* UBlueprintGeneratedClass::GetPersistentUberGraphFrame(UObject* Obj, UFunction* FuncToCheck) const
+{
+	if (Obj && UsePersistentUberGraphFrame() && UberGraphFramePointerProperty && UberGraphFunction)
+	{
+		if (UberGraphFunction == FuncToCheck)
+		{
+			FPointerToUberGraphFrame* PointerToUberGraphFrame = UberGraphFramePointerProperty->ContainerPtrToValuePtr<FPointerToUberGraphFrame>(Obj);
+			checkSlow(PointerToUberGraphFrame);
+			ensure(PointerToUberGraphFrame->RawPointer);
+			return PointerToUberGraphFrame->RawPointer;
+		}
+	}
+	UClass* ParentClass = GetSuperClass();
+	checkSlow(ParentClass);
+	return ParentClass->GetPersistentUberGraphFrame(Obj, FuncToCheck);
+}
+
 void UBlueprintGeneratedClass::CreatePersistentUberGraphFrame(UObject* Obj, bool bCreateOnlyIfEmpty, bool bSkipSuperClass, UClass* OldClass) const
 {
 #if USE_UBER_GRAPH_PERSISTENT_FRAME
@@ -1481,6 +1498,7 @@ void UBlueprintGeneratedClass::PurgeClass(bool bRecompilingOnLoad)
 {
 	Super::PurgeClass(bRecompilingOnLoad);
 
+	UberGraphFramePointerProperty = NULL;
 	UberGraphFunction = NULL;
 #if VALIDATE_UBER_GRAPH_PERSISTENT_FRAME
 	UberGraphFunctionKey = 0;
