Dynamic Multicast delegates can now specify that their storage be stored in a sparse allocation for delegates that are expected to be bound rarely. The memory usage of an unbound sparse delegate is 1 byte instead of 16 bytes (for the empty array), however, a bound sparse delegate will have access overhead to retrieve the invocation list from the sparse object annotation.
The definition of the delegate type must specify the class and property name of the delegate that will use it and the type cannot be used outside of that context or just on the stack.
UMulticastDelegateProperty is now abstract and has 2 subclasses (UMulticastDelegateInlineProperty and UMulticastDelegateSparseProperty). The SparseProperty SignatureFunction will be a USparseDelegateFunction.
The sparse delegates do not work correctly with the python bindings at this point, but this will be visited soon.

diff --git a/Engine/Config/BaseEngine.ini b/Engine/Config/BaseEngine.ini
--- a/Engine/Config/BaseEngine.ini
+++ b/Engine/Config/BaseEngine.ini
@@ -652,6 +652,8 @@ bUseStreamingPause=false
 +FunctionRedirects=(OldName="SkeletalMeshComponent.SetAnimInstanceClass",NewName="SkeletalMeshComponent.SetAnimClass")
 +FunctionRedirects=(OldName="SkeletalMeshComponent.K2_SetAnimInstanceClass",NewName="SkeletalMeshComponent.SetAnimClass")
 
++ClassRedirects=(OldName="/Script/CoreUObject.MulticastDelegateProperty",NewName="/Script/CoreUObject.MulticastInlineDelegateProperty")
+
 [CoreUObject.Metadata]
 ; Note: UnrealHeaderTool should be rerun after making changes to MetadataRedirects to catch any keys specified in class headers
 +MetadataRedirects=(OldKey="K2Protected", NewKey="BlueprintProtected")
diff --git a/Engine/Plugins/Experimental/PythonScriptPlugin/Source/PythonScriptPlugin/Private/PyConversion.cpp b/Engine/Plugins/Experimental/PythonScriptPlugin/Source/PythonScriptPlugin/Private/PyConversion.cpp
--- a/Engine/Plugins/Experimental/PythonScriptPlugin/Source/PythonScriptPlugin/Private/PyConversion.cpp
+++ b/Engine/Plugins/Experimental/PythonScriptPlugin/Source/PythonScriptPlugin/Private/PyConversion.cpp
@@ -953,7 +953,7 @@ FPyConversionResult NativizeProperty_Direct(PyObject* PyObj, const UProperty* Pr
 			{
 				EmitPropertyChangeNotifications(InChangeOwner, [&]()
 				{
-					CastProp->SetPropertyValue(ValueAddr, *PyDelegate->DelegateInstance);
+					CastProp->SetMulticastDelegate(ValueAddr, *PyDelegate->DelegateInstance);
 				});
 			}
 		}
@@ -1113,13 +1113,21 @@ FPyConversionResult PythonizeProperty_Direct(const UProperty* Prop, const void*
 		return FPyConversionResult::Success();
 	}
 
-	if (auto* CastProp = Cast<UMulticastDelegateProperty>(Prop))
+	if (auto* CastProp = Cast<UMulticastInlineDelegateProperty>(Prop))
 	{
-		const FMulticastScriptDelegate* Value = CastProp->GetPropertyValuePtr(ValueAddr);
+		const FMulticastScriptDelegate* Value = CastProp->GetMulticastDelegate(ValueAddr);
 		OutPyObj = (PyObject*)FPyWrapperMulticastDelegateFactory::Get().CreateInstance(CastProp->SignatureFunction, (FMulticastScriptDelegate*)Value, OwnerContext, ConversionMethod);
 		return FPyConversionResult::Success();
 	}
 
+	// TODO: We're just not handling sparse delegates at this time
+	/*if (auto* CastProp = Cast<UMulticastSparseDelegateProperty>(Prop))
+	{
+		const FMulticastScriptDelegate* Value = CastProp->GetMulticastDelegate(ValueAddr);
+		OutPyObj = (PyObject*)FPyWrapperMulticastDelegateFactory::Get().CreateInstance(CastProp->SignatureFunction, (FMulticastScriptDelegate*)Value, OwnerContext, ConversionMethod);
+		return FPyConversionResult::Success();
+	}*/
+
 	if (auto* CastProp = Cast<UArrayProperty>(Prop))
 	{
 		OutPyObj = (PyObject*)FPyWrapperArrayFactory::Get().CreateInstance((void*)ValueAddr, CastProp, OwnerContext, ConversionMethod);
diff --git a/Engine/Plugins/Experimental/PythonScriptPlugin/Source/PythonScriptPlugin/Private/PyReferenceCollector.cpp b/Engine/Plugins/Experimental/PythonScriptPlugin/Source/PythonScriptPlugin/Private/PyReferenceCollector.cpp
--- a/Engine/Plugins/Experimental/PythonScriptPlugin/Source/PythonScriptPlugin/Private/PyReferenceCollector.cpp
+++ b/Engine/Plugins/Experimental/PythonScriptPlugin/Source/PythonScriptPlugin/Private/PyReferenceCollector.cpp
@@ -78,7 +78,7 @@ void FPyReferenceCollector::AddReferencedObjectsFromDelegate(FReferenceCollector
 	}
 }
 
-void FPyReferenceCollector::AddReferencedObjectsFromMulticastDelegate(FReferenceCollector& InCollector, FMulticastScriptDelegate& InDelegate)
+void FPyReferenceCollector::AddReferencedObjectsFromMulticastDelegate(FReferenceCollector& InCollector, const FMulticastScriptDelegate& InDelegate)
 {
 	// Keep the delegate objects alive if they're using a Python proxy instance
 	// We have to use the EvenIfUnreachable variant here as the objects are speculatively marked as unreachable during GC
@@ -192,8 +192,10 @@ void FPyReferenceCollector::AddReferencedObjectsFromPropertyInternal(FReferenceC
 		{
 			for (int32 ArrIndex = 0; ArrIndex < InProp->ArrayDim; ++ArrIndex)
 			{
-				FMulticastScriptDelegate* Value = CastProp->GetPropertyValuePtr(CastProp->ContainerPtrToValuePtr<void>(InBaseAddr, ArrIndex));
-				AddReferencedObjectsFromMulticastDelegate(InCollector, *Value);
+				if (const FMulticastScriptDelegate* Value = CastProp->GetMulticastDelegate(CastProp->ContainerPtrToValuePtr<void>(InBaseAddr, ArrIndex)))
+				{
+					AddReferencedObjectsFromMulticastDelegate(InCollector, *Value);
+				}
 			}
 		}
 		return;
diff --git a/Engine/Plugins/Experimental/PythonScriptPlugin/Source/PythonScriptPlugin/Private/PyReferenceCollector.h b/Engine/Plugins/Experimental/PythonScriptPlugin/Source/PythonScriptPlugin/Private/PyReferenceCollector.h
--- a/Engine/Plugins/Experimental/PythonScriptPlugin/Source/PythonScriptPlugin/Private/PyReferenceCollector.h
+++ b/Engine/Plugins/Experimental/PythonScriptPlugin/Source/PythonScriptPlugin/Private/PyReferenceCollector.h
@@ -51,7 +51,7 @@ public:
 	static void AddReferencedObjectsFromDelegate(FReferenceCollector& InCollector, FScriptDelegate& InDelegate);
 
 	/** Utility function to ARO a multicast delegate instance */
-	static void AddReferencedObjectsFromMulticastDelegate(FReferenceCollector& InCollector, FMulticastScriptDelegate& InDelegate);
+	static void AddReferencedObjectsFromMulticastDelegate(FReferenceCollector& InCollector, const FMulticastScriptDelegate& InDelegate);
 
 	/** Utility function to ARO all properties on a struct instance */
 	static void AddReferencedObjectsFromStruct(FReferenceCollector& InCollector, const UStruct* InStruct, void* InStructAddr, const EPyReferenceCollectorFlags InFlags = EPyReferenceCollectorFlags::IncludeAll);
diff --git a/Engine/Source/Editor/BlueprintGraph/Classes/K2Node_ActorBoundEvent.h b/Engine/Source/Editor/BlueprintGraph/Classes/K2Node_ActorBoundEvent.h
--- a/Engine/Source/Editor/BlueprintGraph/Classes/K2Node_ActorBoundEvent.h
+++ b/Engine/Source/Editor/BlueprintGraph/Classes/K2Node_ActorBoundEvent.h
@@ -66,9 +66,6 @@ class UK2Node_ActorBoundEvent : public UK2Node_Event
 	BLUEPRINTGRAPH_API UMulticastDelegateProperty* GetTargetDelegateProperty() const;
 	BLUEPRINTGRAPH_API UMulticastDelegateProperty* GetTargetDelegatePropertyFromSkel() const;
 
-	/** Returns the delegate that this event is bound in, for actors bound in level script */
-	BLUEPRINTGRAPH_API FMulticastScriptDelegate* GetTargetDelegate() const;
-
 private:
 	/** Constructing FText strings can be costly, so we cache the node's title */
 	FNodeTextCache CachedNodeTitle;
diff --git a/Engine/Source/Editor/BlueprintGraph/Private/K2Node_ActorBoundEvent.cpp b/Engine/Source/Editor/BlueprintGraph/Private/K2Node_ActorBoundEvent.cpp
--- a/Engine/Source/Editor/BlueprintGraph/Private/K2Node_ActorBoundEvent.cpp
+++ b/Engine/Source/Editor/BlueprintGraph/Private/K2Node_ActorBoundEvent.cpp
@@ -92,9 +92,9 @@ void UK2Node_ActorBoundEvent::DestroyNode()
 				Delegate.BindUFunction(LSA, CustomFunctionName);
 
 				// Attempt to remove it from the target's MC delegate
-				if (FMulticastScriptDelegate* TargetDelegate = GetTargetDelegate())
+				if (UMulticastDelegateProperty* TargetDelegate = GetTargetDelegateProperty())
 				{
-					TargetDelegate->Remove(Delegate);
+					TargetDelegate->RemoveDelegate(Delegate, EventOwner);
 				}
 			}
 		}
@@ -240,20 +240,6 @@ UMulticastDelegateProperty* UK2Node_ActorBoundEvent::GetTargetDelegatePropertyFr
 	return FindField<UMulticastDelegateProperty>(FBlueprintEditorUtils::GetMostUpToDateClass(DelegateOwnerClass), DelegatePropertyName);
 }
 
-FMulticastScriptDelegate* UK2Node_ActorBoundEvent::GetTargetDelegate() const 
-{
-	if( EventOwner )
-	{
-		UMulticastDelegateProperty* TargetDelegateProp = GetTargetDelegateProperty();
-		if (TargetDelegateProp && ensure(EventOwner->IsA(DelegateOwnerClass)))
-		{
-			return TargetDelegateProp->GetPropertyValuePtr_InContainer(EventOwner);
-		}
-	}
-
-	return NULL;
-}
-
 bool UK2Node_ActorBoundEvent::IsUsedByAuthorityOnlyDelegate() const
 {
 	const UMulticastDelegateProperty* TargetDelegateProp = FindField<UMulticastDelegateProperty>(DelegateOwnerClass, DelegatePropertyName);
diff --git a/Engine/Source/Editor/KismetCompiler/Private/KismetCompilerMisc.cpp b/Engine/Source/Editor/KismetCompiler/Private/KismetCompilerMisc.cpp
--- a/Engine/Source/Editor/KismetCompiler/Private/KismetCompilerMisc.cpp
+++ b/Engine/Source/Editor/KismetCompiler/Private/KismetCompilerMisc.cpp
@@ -1153,7 +1153,7 @@ UProperty* FKismetCompilerUtilities::CreatePropertyOnScope(UStruct* Scope, const
 	else if (Type.PinCategory == UEdGraphSchema_K2::PC_MCDelegate)
 	{
 		UFunction* const SignatureFunction = FMemberReference::ResolveSimpleMemberReference<UFunction>(Type.PinSubCategoryMemberReference);
-		UMulticastDelegateProperty* NewPropertyDelegate = NewObject<UMulticastDelegateProperty>(PropertyScope, ValidatedPropertyName, ObjectFlags);
+		UMulticastDelegateProperty* NewPropertyDelegate = NewObject<UMulticastInlineDelegateProperty>(PropertyScope, ValidatedPropertyName, ObjectFlags);
 		NewPropertyDelegate->SignatureFunction = SignatureFunction;
 		NewProperty = NewPropertyDelegate;
 	}
diff --git a/Engine/Source/Editor/PropertyEditor/Private/PropertyNode.cpp b/Engine/Source/Editor/PropertyEditor/Private/PropertyNode.cpp
--- a/Engine/Source/Editor/PropertyEditor/Private/PropertyNode.cpp
+++ b/Engine/Source/Editor/PropertyEditor/Private/PropertyNode.cpp
@@ -1513,14 +1513,15 @@ private:
 
 		if ( MulticastDelegateProp != NULL )
 		{
-			FMulticastScriptDelegate* MulticastDelegateValue = MulticastDelegateProp->GetPropertyValuePtr(PropertyValueAddress);
-
-			TArray<UObject*> AllObjects = MulticastDelegateValue->GetAllObjects();
-			for( TArray<UObject*>::TConstIterator CurObjectIt( AllObjects ); CurObjectIt; ++CurObjectIt )
+			if (const FMulticastScriptDelegate* MulticastDelegateValue = MulticastDelegateProp->GetMulticastDelegate(PropertyValueAddress))
 			{
-				if ((*CurObjectIt)->IsDefaultSubobject())
+				TArray<UObject*> AllObjects = MulticastDelegateValue->GetAllObjects();
+				for (TArray<UObject*>::TConstIterator CurObjectIt(AllObjects); CurObjectIt; ++CurObjectIt)
 				{
-					Components.AddUnique((*CurObjectIt));
+					if ((*CurObjectIt)->IsDefaultSubobject())
+					{
+						Components.AddUnique((*CurObjectIt));
+					}
 				}
 			}
 
diff --git a/Engine/Source/Editor/UnrealEd/Private/Kismet2/BlueprintEditorUtils.cpp b/Engine/Source/Editor/UnrealEd/Private/Kismet2/BlueprintEditorUtils.cpp
--- a/Engine/Source/Editor/UnrealEd/Private/Kismet2/BlueprintEditorUtils.cpp
+++ b/Engine/Source/Editor/UnrealEd/Private/Kismet2/BlueprintEditorUtils.cpp
@@ -7352,13 +7352,13 @@ void FBlueprintEditorUtils::FixLevelScriptActorBindings(ALevelScriptActor* Level
 				if( LevelScriptActor->FindFunction(TargetFunction) )
 				{
 					// Grab the MC delegate we need to add to
-					FMulticastScriptDelegate* TargetDelegate = EventNode->GetTargetDelegate();
+					UMulticastDelegateProperty* TargetDelegate = EventNode->GetTargetDelegateProperty();
 					if( TargetDelegate != nullptr )
 					{
 						// Create the delegate, and add it if it doesn't already exist
 						FScriptDelegate Delegate;
 						Delegate.BindUFunction(LevelScriptActor, TargetFunction);
-						TargetDelegate->AddUnique(Delegate);
+						TargetDelegate->AddDelegate(MoveTemp(Delegate), EventNode->EventOwner);
 					}
 				}
 			}
@@ -9179,4 +9179,3 @@ void FBlueprintEditorUtils::BuildComponentInstancingData(UActorComponent* Compon
 }
 
 #undef LOCTEXT_NAMESPACE
-
diff --git a/Engine/Source/Programs/UnrealHeaderTool/Private/CodeGenerator.cpp b/Engine/Source/Programs/UnrealHeaderTool/Private/CodeGenerator.cpp
--- a/Engine/Source/Programs/UnrealHeaderTool/Private/CodeGenerator.cpp
+++ b/Engine/Source/Programs/UnrealHeaderTool/Private/CodeGenerator.cpp
@@ -656,7 +656,7 @@ private:
 		{
 			TypeStr = TEXT("UClass");
 		}
-		else if (InType->GetClass() == UFunction::StaticClass() || InType->GetClass() == UDelegateFunction::StaticClass())
+		else if (InType->GetClass() == UFunction::StaticClass() || InType->GetClass() == UDelegateFunction::StaticClass() || InType->GetClass() == USparseDelegateFunction::StaticClass())
 		{
 			TypeStr = TEXT("UFunction");
 		}
@@ -1480,12 +1480,13 @@ void FNativeClassHeaderGenerator::PropertyNew(FOutputDevice& DeclOut, FOutputDev
 		DeclOut.Logf(TEXT("%sstatic const UE4CodeGen_Private::FMulticastDelegatePropertyParams %s;\r\n"), DeclSpaces, *NameWithoutScope);
 
 		Out.Logf(
-			TEXT("%sconst UE4CodeGen_Private::FMulticastDelegatePropertyParams %s = { %s, %s, (EPropertyFlags)0x%016llx, UE4CodeGen_Private::EPropertyGenFlags::MulticastDelegate, %s, %s, %s, %s, %s };%s\r\n"),
+			TEXT("%sconst UE4CodeGen_Private::FMulticastDelegatePropertyParams %s = { %s, %s, (EPropertyFlags)0x%016llx, UE4CodeGen_Private::EPropertyGenFlags::%sMulticastDelegate, %s, %s, %s, %s, %s };%s\r\n"),
 			Spaces,
 			Name,
 			*PropName,
 			*PropNotifyFunc,
 			PropFlags,
+			(TypedProp->IsA<UMulticastInlineDelegateProperty>() ? TEXT("Inline") : TEXT("Sparse")),
 			UPropertyObjectFlags,
 			*ArrayDim,
 			OffsetStr,
@@ -2304,7 +2305,7 @@ void FNativeClassHeaderGenerator::ExportFunction(FOutputDevice& Out, const FUnre
 		StructureSize = TEXT("0");
 	}
 
-	const TCHAR* UFunctionType = bIsDelegate ? TEXT("UDelegateFunction") : TEXT("UFunction");
+	USparseDelegateFunction* SparseDelegateFunction = Cast<USparseDelegateFunction>(Function);
 	const TCHAR* UFunctionObjectFlags = FClass::IsOwnedByDynamicType(Function) ? TEXT("RF_Public|RF_Transient") : TEXT("RF_Public|RF_Transient|RF_MarkAsNative");
 
 	TTuple<FString, FString> PropertyRange = OutputProperties(CurrentFunctionText, StaticDefinitions, *FString::Printf(TEXT("%s::"), *StaticsStructName), Props, TEXT("\t\t"), TEXT("\t"));
@@ -2318,11 +2319,13 @@ void FNativeClassHeaderGenerator::ExportFunction(FOutputDevice& Out, const FUnre
 	CurrentFunctionText.Log(TEXT("\t\tstatic const UE4CodeGen_Private::FFunctionParams FuncParams;\r\n"));
 
 	StaticDefinitions.Logf(
-		TEXT("\tconst UE4CodeGen_Private::FFunctionParams %s::FuncParams = { (UObject*(*)())%s, %s, %s, %s, %s, %s, %s, (EFunctionFlags)0x%08X, %d, %d, %s };\r\n"),
+		TEXT("\tconst UE4CodeGen_Private::FFunctionParams %s::FuncParams = { (UObject*(*)())%s, %s, %s, %s, %s, %s, %s, %s, %s, (EFunctionFlags)0x%08X, %d, %d, %s };\r\n"),
 		*StaticsStructName,
 		*OuterFunc,
 		*GetSingletonNameFuncAddr(SuperFunction),
 		*CreateUTF8LiteralString(FNativeClassHeaderGenerator::GetOverriddenName(Function)),
+		(SparseDelegateFunction ? *CreateUTF8LiteralString(SparseDelegateFunction->OwningClassName.ToString()) : TEXT("nullptr")),
+		(SparseDelegateFunction ? *CreateUTF8LiteralString(SparseDelegateFunction->DelegateName.ToString()) : TEXT("nullptr")),
 		*StructureSize,
 		*PropertyRange.Get<0>(),
 		*PropertyRange.Get<1>(),
diff --git a/Engine/Source/Programs/UnrealHeaderTool/Private/HeaderParser.cpp b/Engine/Source/Programs/UnrealHeaderTool/Private/HeaderParser.cpp
--- a/Engine/Source/Programs/UnrealHeaderTool/Private/HeaderParser.cpp
+++ b/Engine/Source/Programs/UnrealHeaderTool/Private/HeaderParser.cpp
@@ -891,7 +891,15 @@ namespace
 
 			case CPT_MulticastDelegate:
 			{
-				UMulticastDelegateProperty* Result = new (EC_InternalUseOnlyConstructor, Scope, Name, ObjectFlags) UMulticastDelegateProperty(FObjectInitializer());
+				UMulticastDelegateProperty* Result;
+				if (VarProperty.Function->IsA<USparseDelegateFunction>())
+				{
+					Result = new (EC_InternalUseOnlyConstructor, Scope, Name, ObjectFlags) UMulticastSparseDelegateProperty(FObjectInitializer());
+				}
+				else
+				{
+					Result = new (EC_InternalUseOnlyConstructor, Scope, Name, ObjectFlags) UMulticastInlineDelegateProperty(FObjectInitializer());
+				}
 				return Result;
 			}
 
@@ -3898,6 +3906,7 @@ void FHeaderParser::GetVarType(
 		{
 			FError::Throwf(TEXT("Nested containers are not supported.") );
 		}
+		// TODO: Prevent sparse delegate types from being used in a container
 
 		if (VarProperty.MetaData.Find(TEXT("NativeConst")))
 		{
@@ -3943,6 +3952,7 @@ void FHeaderParser::GetVarType(
 		{
 			FError::Throwf(TEXT("Nested containers are not supported.") );
 		}
+		// TODO: Prevent sparse delegate types from being used in a container
 
 		if (MapKeyType.Type == CPT_Interface)
 		{
@@ -3965,6 +3975,7 @@ void FHeaderParser::GetVarType(
 		{
 			FError::Throwf(TEXT("Nested containers are not supported.") );
 		}
+		// TODO: Prevent sparse delegate types from being used in a container
 
 		EPropertyFlags InnerFlags = (MapKeyType.PropertyFlags | VarProperty.PropertyFlags) & (CPF_ContainsInstancedReference | CPF_InstancedReference); // propagate these to the map value, we will fix them later
 		VarType.PropertyFlags = InnerFlags;
@@ -4006,6 +4017,7 @@ void FHeaderParser::GetVarType(
 		{
 			FError::Throwf(TEXT("Nested containers are not supported.") );
 		}
+		// TODO: Prevent sparse delegate types from being used in a container
 
 		if (VarProperty.Type == CPT_Interface)
 		{
@@ -4164,6 +4176,7 @@ void FHeaderParser::GetVarType(
 
 			VarProperty = FPropertyBase(InFunction->HasAnyFunctionFlags(FUNC_MulticastDelegate) ? CPT_MulticastDelegate : CPT_Delegate);
 			VarProperty.DelegateName = *InIdentifierStripped;
+			VarProperty.Function = InFunction;
 
 			if (!(Disallow & CPF_InstancedReference))
 			{
@@ -6241,13 +6254,15 @@ UDelegateFunction* FHeaderParser::CompileDelegateDeclaration(FClasses& AllClasse
 	const bool bHasReturnValue = DelegateMacro.Contains(TEXT("_RetVal"));
 	const bool bDeclaredConst  = DelegateMacro.Contains(TEXT("_Const"));
 	const bool bIsMulticast    = DelegateMacro.Contains(TEXT("_MULTICAST"));
+	const bool bIsSparse       = DelegateMacro.Contains(TEXT("_SPARSE"));
 
 	// Determine the parameter count
 	const FString* FoundParamCount = DelegateParameterCountStrings.FindByPredicate([&](const FString& Str){ return DelegateMacro.Contains(Str); });
 
 	// Try reconstructing the string to make sure it matches our expectations
-	FString ExpectedOriginalString = FString::Printf(TEXT("DECLARE_DYNAMIC%s_DELEGATE%s%s%s"),
+	FString ExpectedOriginalString = FString::Printf(TEXT("DECLARE_DYNAMIC%s%s_DELEGATE%s%s%s"),
 		bIsMulticast ? TEXT("_MULTICAST") : TEXT(""),
+		bIsSparse ? TEXT("_SPARSE") : TEXT(""),
 		bHasReturnValue ? TEXT("_RetVal") : TEXT(""),
 		FoundParamCount ? **FoundParamCount : TEXT(""),
 		bDeclaredConst ? TEXT("_Const") : TEXT(""));
@@ -6316,8 +6331,8 @@ UDelegateFunction* FHeaderParser::CompileDelegateDeclaration(FClasses& AllClasse
 		FCString::Strcpy( FuncInfo.Function.Identifier, *Name );
 	}
 
-	UDelegateFunction* DelegateSignatureFunction = CreateDelegateFunction(FuncInfo);
-
+	UDelegateFunction* DelegateSignatureFunction = (bIsSparse ? CreateDelegateFunction<USparseDelegateFunction>(FuncInfo) : CreateDelegateFunction<UDelegateFunction>(FuncInfo));
+	
 	FClassMetaData* ClassMetaData = GScriptHelper.AddClassData(DelegateSignatureFunction, CurrentSrcFile);
 
 	DelegateSignatureFunction->FunctionFlags |= FuncInfo.FunctionFlags;
@@ -6337,6 +6352,29 @@ UDelegateFunction* FHeaderParser::CompileDelegateDeclaration(FClasses& AllClasse
 		DelegateSignatureFunction->FunctionFlags |= FUNC_Const;
 	}
 
+	if (bIsSparse)
+	{
+		FToken OwningClass;
+
+		RequireSymbol(TEXT(","), TEXT("Delegate Declaration"));
+
+		if (!GetIdentifier(OwningClass))
+		{
+			FError::Throwf(TEXT("Missing OwningClass specifier."));
+		}
+		RequireSymbol(TEXT(","), TEXT("Delegate Declaration"));
+		
+		FToken DelegateName;
+		if (!GetIdentifier(DelegateName))
+		{
+			FError::Throwf(TEXT("Missing Delegate Name."));
+		}
+
+		USparseDelegateFunction* SDF = CastChecked<USparseDelegateFunction>(DelegateSignatureFunction);
+		SDF->OwningClassName = *GetClassNameWithoutPrefix(OwningClass.TokenName.ToString());
+		SDF->DelegateName = DelegateName.Identifier;
+	}
+
 	// Get parameter list.
 	if (FoundParamCount)
 	{
@@ -9493,10 +9531,11 @@ UFunction* FHeaderParser::CreateFunction(const FFuncInfo &FuncInfo) const
 	return CreateFunctionImpl<UFunction>(FuncInfo, GetCurrentClass(), GetCurrentScope());
 }
 
+template<class T>
 UDelegateFunction* FHeaderParser::CreateDelegateFunction(const FFuncInfo &FuncInfo) const
 {
 	FFileScope* CurrentFileScope = GetCurrentFileScope();
 	FUnrealSourceFile* LocSourceFile = CurrentFileScope ? CurrentFileScope->GetSourceFile() : nullptr;
 	UObject* CurrentPackage = LocSourceFile ? LocSourceFile->GetPackage() : nullptr;
-	return CreateFunctionImpl<UDelegateFunction>(FuncInfo, IsInAClass() ? (UObject*)GetCurrentClass() : CurrentPackage, GetCurrentScope());
+	return CreateFunctionImpl<T>(FuncInfo, IsInAClass() ? (UObject*)GetCurrentClass() : CurrentPackage, GetCurrentScope());
 }
diff --git a/Engine/Source/Programs/UnrealHeaderTool/Private/HeaderParser.h b/Engine/Source/Programs/UnrealHeaderTool/Private/HeaderParser.h
--- a/Engine/Source/Programs/UnrealHeaderTool/Private/HeaderParser.h
+++ b/Engine/Source/Programs/UnrealHeaderTool/Private/HeaderParser.h
@@ -534,6 +534,7 @@ protected:
 	/**
 	 * Create new delegate function object based on given info structure.
 	 */
+	template<typename T>
 	UDelegateFunction* CreateDelegateFunction(const FFuncInfo &FuncInfo) const;	
 
 	UClass* CompileClassDeclaration(FClasses& AllClasses);
diff --git a/Engine/Source/Programs/UnrealHeaderTool/Private/Scope.cpp b/Engine/Source/Programs/UnrealHeaderTool/Private/Scope.cpp
--- a/Engine/Source/Programs/UnrealHeaderTool/Private/Scope.cpp
+++ b/Engine/Source/Programs/UnrealHeaderTool/Private/Scope.cpp
@@ -49,7 +49,7 @@ void DispatchType(UField* Type, TArray<UEnum*> &Enums, TArray<UScriptStruct*> &S
 		UScriptStruct* Struct = (UScriptStruct*)Type;
 		Structs.Add(Struct);
 	}
-	else if (TypeClass == UDelegateFunction::StaticClass())
+	else if (TypeClass == UDelegateFunction::StaticClass() || TypeClass == USparseDelegateFunction::StaticClass())
 	{
 		bool bAdded = false;
 		UDelegateFunction* Function = (UDelegateFunction*)Type;
diff --git a/Engine/Source/Runtime/Core/Public/UObject/ScriptDelegates.h b/Engine/Source/Runtime/Core/Public/UObject/ScriptDelegates.h
--- a/Engine/Source/Runtime/Core/Public/UObject/ScriptDelegates.h
+++ b/Engine/Source/Runtime/Core/Public/UObject/ScriptDelegates.h
@@ -506,7 +506,7 @@ public:
 	 * For advanced uses only -- you should never need call this function in normal circumstances.
  	 * @return	List of objects bound to this delegate
 	 */
-	TArray< UObject* > GetAllObjects()
+	TArray< UObject* > GetAllObjects() const
 	{
 		TArray< UObject* > OutputList;
 		for( typename FInvocationList::TIterator CurDelegate( InvocationList ); CurDelegate; ++CurDelegate )
@@ -525,7 +525,7 @@ public:
 	 * For advanced uses only -- you should never need call this function in normal circumstances.
  	 * @return	List of objects bound to this delegate
 	 */
-	TArray< UObject* > GetAllObjectsEvenIfUnreachable()
+	TArray< UObject* > GetAllObjectsEvenIfUnreachable() const
 	{
 		TArray< UObject* > OutputList;
 		for( typename FInvocationList::TIterator CurDelegate( InvocationList ); CurDelegate; ++CurDelegate )
@@ -609,14 +609,18 @@ protected:
 		});
 	}
 
+public:
+	typedef TArray< TScriptDelegate<TWeakPtr> > FInvocationList;
+
 protected:
 
 	/** Ordered list functions to invoke when the Broadcast function is called */
-	typedef TArray< TScriptDelegate<TWeakPtr> > FInvocationList;
 	mutable FInvocationList InvocationList;		// Mutable so that we can housekeep list even with 'const' broadcasts
 
 	// Declare ourselves as a friend of UMulticastDelegateProperty so that it can access our function list
 	friend class UMulticastDelegateProperty;
+	friend class UMulticastInlineDelegateProperty;
+	friend class UMulticastSparseDelegateProperty;
 
 	// 
 	friend class FCallDelegateHelper;
diff --git a/Engine/Source/Runtime/CoreUObject/Private/UObject/Class.cpp b/Engine/Source/Runtime/CoreUObject/Private/UObject/Class.cpp
--- a/Engine/Source/Runtime/CoreUObject/Private/UObject/Class.cpp
+++ b/Engine/Source/Runtime/CoreUObject/Private/UObject/Class.cpp
@@ -3074,6 +3074,13 @@ UObject* UClass::CreateDefaultObject()
 				check(ClassDefaultObject);
 				// Blueprint CDOs have their properties always initialized.
 				const bool bShouldInitializeProperties = !HasAnyClassFlags(CLASS_Native | CLASS_Intrinsic);
+				// Register the offsets of any sparse delegates this class introduces with the sparse delegate storage
+				for (TFieldIterator<UMulticastSparseDelegateProperty> SparseDelegateIt(this, EFieldIteratorFlags::ExcludeSuper, EFieldIteratorFlags::ExcludeDeprecated); SparseDelegateIt; ++SparseDelegateIt)
+				{
+					const FSparseDelegate& SparseDelegate = SparseDelegateIt->GetPropertyValue_InContainer(ClassDefaultObject);
+					USparseDelegateFunction* SparseDelegateFunction = CastChecked<USparseDelegateFunction>(SparseDelegateIt->SignatureFunction);
+					FSparseDelegateStorage::RegisterDelegateOffset(ClassDefaultObject, SparseDelegateFunction->DelegateName, (size_t)&SparseDelegate - (size_t)ClassDefaultObject);
+				}
 				(*ClassConstructor)(FObjectInitializer(ClassDefaultObject, ParentDefaultObject, false, bShouldInitializeProperties));
 				if (bDoNotify)
 				{
@@ -5152,6 +5159,31 @@ IMPLEMENT_CORE_INTRINSIC_CLASS(UDelegateFunction, UFunction,
 	}
 );
 
+USparseDelegateFunction::USparseDelegateFunction(const FObjectInitializer& ObjectInitializer, UFunction* InSuperFunction, EFunctionFlags InFunctionFlags, SIZE_T ParamsSize)
+	: UDelegateFunction(ObjectInitializer, InSuperFunction, InFunctionFlags, ParamsSize)
+{
+
+}
+
+USparseDelegateFunction::USparseDelegateFunction(UFunction* InSuperFunction, EFunctionFlags InFunctionFlags, SIZE_T ParamsSize)
+	: UDelegateFunction(InSuperFunction, InFunctionFlags, ParamsSize)
+{
+
+}
+
+void USparseDelegateFunction::Serialize(FArchive& Ar)
+{
+	Super::Serialize(Ar);
+
+	Ar << OwningClassName;
+	Ar << DelegateName;
+}
+
+IMPLEMENT_CORE_INTRINSIC_CLASS(USparseDelegateFunction, UDelegateFunction,
+	{
+	}
+);
+
 /*-----------------------------------------------------------------------------
 UDynamicClass constructors.
 -----------------------------------------------------------------------------*/
diff --git a/Engine/Source/Runtime/CoreUObject/Private/UObject/Property.cpp b/Engine/Source/Runtime/CoreUObject/Private/UObject/Property.cpp
--- a/Engine/Source/Runtime/CoreUObject/Private/UObject/Property.cpp
+++ b/Engine/Source/Runtime/CoreUObject/Private/UObject/Property.cpp
@@ -1056,7 +1056,7 @@ const TCHAR* UProperty::ImportSingleProperty( const TCHAR* Str, void* DestData,
 		}
 
 		UArrayProperty* const ArrayProperty = ExactCast<UArrayProperty>(Property);
-		UMulticastDelegateProperty* const MulticastDelegateProperty = ExactCast<UMulticastDelegateProperty>(Property);
+		UMulticastDelegateProperty* const MulticastDelegateProperty = Cast<UMulticastDelegateProperty>(Property);
 		if( MulticastDelegateProperty != NULL && ArrayOp != ADO_None )
 		{
 			// Allow Add(), Remove() and Empty() on multi-cast delegates
@@ -1073,7 +1073,7 @@ const TCHAR* UProperty::ImportSingleProperty( const TCHAR* Str, void* DestData,
 				if( ArrayOp == ADO_Empty )
 				{
 					// Clear out the delegate
-					MulticastDelegateProperty->GetPropertyValuePtr_InContainer(DestData)->Clear();
+					MulticastDelegateProperty->ClearDelegate(SubobjectOuter, Property->ContainerPtrToValuePtr<void>(DestData));
 				}
 				else
 				{
@@ -1421,4 +1421,4 @@ UProperty* UStruct::FindPropertyByName(FName InName) const
 	}
 
 	return NULL;
-}
+}
\ No newline at end of file
diff --git a/Engine/Source/Runtime/CoreUObject/Private/UObject/PropertyMulticastDelegate.cpp b/Engine/Source/Runtime/CoreUObject/Private/UObject/PropertyMulticastDelegate.cpp
--- a/Engine/Source/Runtime/CoreUObject/Private/UObject/PropertyMulticastDelegate.cpp
+++ b/Engine/Source/Runtime/CoreUObject/Private/UObject/PropertyMulticastDelegate.cpp
@@ -12,18 +12,17 @@
 	UMulticastDelegateProperty.
 -----------------------------------------------------------------------------*/
 
+FMulticastScriptDelegate::FInvocationList UMulticastDelegateProperty::EmptyList;
+
 void UMulticastDelegateProperty::InstanceSubobjects(void* Data, void const* DefaultData, UObject* Owner, FObjectInstancingGraph* InstanceGraph )
 {
 	if (DefaultData)
 	{
 		for( int32 i=0; i<ArrayDim; i++ )
 		{
-			FMulticastScriptDelegate& DestDelegate = ((FMulticastScriptDelegate*)Data)[i];
-			FMulticastScriptDelegate& DefaultDelegate = ((FMulticastScriptDelegate*)DefaultData)[i];
-
 			// Fix up references to the class default object (if necessary)
-			FMulticastScriptDelegate::FInvocationList::TIterator CurInvocation( DestDelegate.InvocationList );
-			FMulticastScriptDelegate::FInvocationList::TIterator DefaultInvocation( DefaultDelegate.InvocationList );
+			FMulticastScriptDelegate::FInvocationList::TIterator CurInvocation(GetInvocationList((uint8*)Data + i));
+			FMulticastScriptDelegate::FInvocationList::TIterator DefaultInvocation(GetInvocationList((uint8*)DefaultData + i));
 			for(; CurInvocation && DefaultInvocation; ++CurInvocation, ++DefaultInvocation )
 			{
 				FScriptDelegate& DestDelegateInvocation = *CurInvocation;
@@ -55,9 +54,7 @@ void UMulticastDelegateProperty::InstanceSubobjects(void* Data, void const* Defa
 	{
 		for( int32 i=0; i<ArrayDim; i++ )
 		{
-			FMulticastScriptDelegate& DestDelegate = ((FMulticastScriptDelegate*)Data)[i];
-
-			for( FMulticastScriptDelegate::FInvocationList::TIterator CurInvocation( DestDelegate.InvocationList ); CurInvocation; ++CurInvocation )
+			for( FMulticastScriptDelegate::FInvocationList::TIterator CurInvocation(GetInvocationList((uint8*)Data + i)); CurInvocation; ++CurInvocation )
 			{
 				FScriptDelegate& DestDelegateInvocation = *CurInvocation;
 				UObject* CurrentUObject = DestDelegateInvocation.GetUObject();
@@ -74,18 +71,10 @@ void UMulticastDelegateProperty::InstanceSubobjects(void* Data, void const* Defa
 
 bool UMulticastDelegateProperty::Identical( const void* A, const void* B, uint32 PortFlags ) const
 {
-	FMulticastScriptDelegate* DA = (FMulticastScriptDelegate*)A;
-	FMulticastScriptDelegate* DB = (FMulticastScriptDelegate*)B;
-	
-	if (!DB)
-	{
-		return DA->InvocationList.Num() == 0;
-	}
+	const FMulticastScriptDelegate::FInvocationList& ListA = GetInvocationList(A);
+	const FMulticastScriptDelegate::FInvocationList& ListB = GetInvocationList(B);
 
-	const FMulticastScriptDelegate::FInvocationList& ListA = DA->InvocationList;
-	const FMulticastScriptDelegate::FInvocationList& ListB = DB->InvocationList;
-
-	int32 ListASize = ListA.Num();
+	const int32 ListASize = ListA.Num();
 	if (ListASize != ListB.Num())
 	{
 		return false;
@@ -110,12 +99,6 @@ bool UMulticastDelegateProperty::Identical( const void* A, const void* B, uint32
 	return true;
 }
 
-void UMulticastDelegateProperty::SerializeItem(FStructuredArchive::FSlot Slot, void* Value, void const* Defaults) const
-{
-	FArchiveUObjectFromStructuredArchive Ar(Slot);
-	Ar << *GetPropertyValuePtr(Value);
-}
-
 bool UMulticastDelegateProperty::NetSerializeItem( FArchive& Ar, UPackageMap* Map, void* Data, TArray<uint8> * MetaData ) const
 {
 	// Do not allow replication of delegates, as there is no way to make this secure (it allows the execution of any function in any object, on the remote client/server)
@@ -184,27 +167,26 @@ void UMulticastDelegateProperty::ExportTextItem( FString& ValueStr, const void*
 		return;
 	}
 
-	const FMulticastScriptDelegate* MulticastDelegate = (const FMulticastScriptDelegate*)( PropertyValue );
-	check( MulticastDelegate != NULL );
+	const FMulticastScriptDelegate::FInvocationList& InvocationList = GetInvocationList(PropertyValue);
 
 	// Start delegate array with open paren
 	ValueStr += TEXT( "(" );
 
 	bool bIsFirstFunction = true;
-	for( FMulticastScriptDelegate::FInvocationList::TConstIterator CurInvocation( MulticastDelegate->InvocationList ); CurInvocation; ++CurInvocation )
+	for (FMulticastScriptDelegate::FInvocationList::TConstIterator CurInvocation(InvocationList); CurInvocation; ++CurInvocation)
 	{
-		if( CurInvocation->IsBound() )
+		if (CurInvocation->IsBound())
 		{
-			if( !bIsFirstFunction )
+			if (!bIsFirstFunction)
 			{
-				ValueStr += TEXT( "," );
+				ValueStr += TEXT(",");
 			}
 			bIsFirstFunction = false;
 
 			bool bDelegateHasValue = CurInvocation->GetFunctionName() != NAME_None;
-			ValueStr += FString::Printf( TEXT("%s.%s"),
+			ValueStr += FString::Printf(TEXT("%s.%s"),
 				CurInvocation->GetUObject() != NULL ? *CurInvocation->GetUObject()->GetName() : TEXT("(null)"),
-				*CurInvocation->GetFunctionName().ToString() );
+				*CurInvocation->GetFunctionName().ToString());
 		}
 	}
 
@@ -212,8 +194,7 @@ void UMulticastDelegateProperty::ExportTextItem( FString& ValueStr, const void*
 	ValueStr += TEXT( ")" );
 }
 
-
-const TCHAR* UMulticastDelegateProperty::ImportText_Internal( const TCHAR* Buffer, void* PropertyValue, int32 PortFlags, UObject* Parent, FOutputDevice* ErrorText ) const
+const TCHAR* UMulticastDelegateProperty::ImportDelegateFromText( FMulticastScriptDelegate& MulticastDelegate, const TCHAR* Buffer, void* PropertyValue, int32 PortFlags, UObject* Parent, FOutputDevice* ErrorText ) const
 {
 	// Multi-cast delegates always expect an opening parenthesis when using assignment syntax, so that
 	// users don't accidentally blow away already-bound delegates in DefaultProperties.  This also helps
@@ -223,8 +204,6 @@ const TCHAR* UMulticastDelegateProperty::ImportText_Internal( const TCHAR* Buffe
 		return NULL;
 	}
 
-	FMulticastScriptDelegate& MulticastDelegate = (*(FMulticastScriptDelegate*)PropertyValue);
-
 	// Clear the existing delegate
 	MulticastDelegate.Clear();
 
@@ -273,8 +252,6 @@ const TCHAR* UMulticastDelegateProperty::ImportText_Add( const TCHAR* Buffer, vo
 		return NULL;
 	}
 
-	FMulticastScriptDelegate& MulticastDelegate = (*(FMulticastScriptDelegate*)PropertyValue);
-
 	// Parse the delegate
 	FScriptDelegate ImportedDelegate;
 	Buffer = DelegatePropertyTools::ImportDelegateFromText( ImportedDelegate, SignatureFunction, Buffer, Parent, ErrorText );
@@ -284,7 +261,7 @@ const TCHAR* UMulticastDelegateProperty::ImportText_Add( const TCHAR* Buffer, vo
 	}
 
 	// Add this delegate to our multicast delegate's invocation list
-	MulticastDelegate.Add( ImportedDelegate );
+	AddDelegate(MoveTemp(ImportedDelegate), Parent, PropertyValue);
 
 	SkipWhitespace(Buffer);
 
@@ -299,8 +276,6 @@ const TCHAR* UMulticastDelegateProperty::ImportText_Remove( const TCHAR* Buffer,
 		return NULL;
 	}
 
-	FMulticastScriptDelegate& MulticastDelegate = (*(FMulticastScriptDelegate*)PropertyValue);
-
 	// Parse the delegate
 	FScriptDelegate ImportedDelegate;
 	Buffer = DelegatePropertyTools::ImportDelegateFromText( ImportedDelegate, SignatureFunction, Buffer, Parent, ErrorText );
@@ -310,7 +285,7 @@ const TCHAR* UMulticastDelegateProperty::ImportText_Remove( const TCHAR* Buffer,
 	}
 
 	// Remove this delegate to our multicast delegate's invocation list
-	MulticastDelegate.Remove( ImportedDelegate );
+	RemoveDelegate(ImportedDelegate, Parent, PropertyValue);
 
 	SkipWhitespace(Buffer);
 
@@ -351,9 +326,224 @@ bool UMulticastDelegateProperty::SameType(const UProperty* Other) const
 	return Super::SameType(Other) && (SignatureFunction == ((UMulticastDelegateProperty*)Other)->SignatureFunction);
 }
 
+EConvertFromTypeResult UMulticastDelegateProperty::ConvertFromType(const FPropertyTag& Tag, FStructuredArchive::FSlot Slot, uint8* Data, UStruct* DefaultsStruct)
+{
+	// Multicast delegate properties are serialization compatible
+	if (Tag.Type == NAME_MulticastDelegateProperty || Tag.Type == UMulticastInlineDelegateProperty::StaticClass()->GetFName() || Tag.Type == UMulticastSparseDelegateProperty::StaticClass()->GetFName())
+	{
+		uint8* DestAddress = ContainerPtrToValuePtr<uint8>(Data, Tag.ArrayIndex);
+		SerializeItem(Slot, DestAddress, nullptr);
+
+		return EConvertFromTypeResult::Converted;
+	}
+
+	return EConvertFromTypeResult::UseSerializeItem;
+}
+
+
 IMPLEMENT_CORE_INTRINSIC_CLASS(UMulticastDelegateProperty, UProperty,
 	{
 		Class->EmitObjectReference(STRUCT_OFFSET(UMulticastDelegateProperty, SignatureFunction), TEXT("SignatureFunction"));
 	}
 );
 
+const FMulticastScriptDelegate* UMulticastInlineDelegateProperty::GetMulticastDelegate(const void* PropertyValue) const
+{
+	return (const FMulticastScriptDelegate*)PropertyValue;
+}
+
+void UMulticastInlineDelegateProperty::SetMulticastDelegate(void* PropertyValue, FMulticastScriptDelegate ScriptDelegate) const
+{
+	*(FMulticastScriptDelegate*)PropertyValue = MoveTemp(ScriptDelegate);
+}
+
+FMulticastScriptDelegate::FInvocationList& UMulticastInlineDelegateProperty::GetInvocationList(const void* PropertyValue) const
+{
+	return (PropertyValue ? ((FMulticastScriptDelegate*)PropertyValue)->InvocationList : EmptyList);
+}
+
+void UMulticastInlineDelegateProperty::SerializeItem(FStructuredArchive::FSlot Slot, void* Value, void const* Defaults) const
+{
+	FArchiveUObjectFromStructuredArchive Ar(Slot);
+	Ar << *GetPropertyValuePtr(Value);
+}
+
+const TCHAR* UMulticastInlineDelegateProperty::ImportText_Internal(const TCHAR* Buffer, void* PropertyValue, int32 PortFlags, UObject* Parent, FOutputDevice* ErrorText) const
+{
+	FMulticastScriptDelegate& MulticastDelegate = (*(FMulticastScriptDelegate*)PropertyValue);
+	return ImportDelegateFromText(MulticastDelegate, Buffer, PropertyValue, PortFlags, Parent, ErrorText);
+}
+
+void UMulticastInlineDelegateProperty::AddDelegate(FScriptDelegate ScriptDelegate, UObject* Parent, void* PropertyValue) const
+{
+	FMulticastScriptDelegate& MulticastDelegate = (*(FMulticastScriptDelegate*)PropertyValue);
+
+	// Remove this delegate to our multicast delegate's invocation list
+	MulticastDelegate.AddUnique(MoveTemp(ScriptDelegate));
+}
+
+void UMulticastInlineDelegateProperty::RemoveDelegate(const FScriptDelegate& ScriptDelegate, UObject* Parent, void* PropertyValue) const
+{
+	FMulticastScriptDelegate& MulticastDelegate = (*(FMulticastScriptDelegate*)PropertyValue);
+
+	// Remove this delegate to our multicast delegate's invocation list
+	MulticastDelegate.Remove(ScriptDelegate);
+}
+
+void UMulticastInlineDelegateProperty::ClearDelegate(UObject* Parent, void* PropertyValue) const
+{
+	FMulticastScriptDelegate& MulticastDelegate = (*(FMulticastScriptDelegate*)PropertyValue);
+	MulticastDelegate.Clear();
+}
+
+IMPLEMENT_CORE_INTRINSIC_CLASS(UMulticastInlineDelegateProperty, UMulticastDelegateProperty,
+	{
+	}
+);
+
+const FMulticastScriptDelegate* UMulticastSparseDelegateProperty::GetMulticastDelegate(const void* PropertyValue) const
+{
+	const FSparseDelegate* SparseDelegate = (const FSparseDelegate*)PropertyValue;
+	if (SparseDelegate->IsBound())
+	{
+		USparseDelegateFunction* SparseDelegateFunc = CastChecked<USparseDelegateFunction>(SignatureFunction);
+		UObject* OwningObject = FSparseDelegateStorage::ResolveSparseOwner(*SparseDelegate, SparseDelegateFunc->OwningClassName, SparseDelegateFunc->DelegateName);
+		return FSparseDelegateStorage::GetMulticastDelegate(OwningObject, SparseDelegateFunc->DelegateName);
+	}
+
+	return nullptr;
+}
+
+void UMulticastSparseDelegateProperty::SetMulticastDelegate(void* PropertyValue, FMulticastScriptDelegate ScriptDelegate) const
+{
+	FSparseDelegate& SparseDelegate = *(FSparseDelegate*)PropertyValue;
+
+	USparseDelegateFunction* SparseDelegateFunc = CastChecked<USparseDelegateFunction>(SignatureFunction);
+	UObject* OwningObject = FSparseDelegateStorage::ResolveSparseOwner(SparseDelegate, SparseDelegateFunc->OwningClassName, SparseDelegateFunc->DelegateName);
+	FMulticastScriptDelegate& Delegate = FSparseDelegateStorage::GetOrCreateMulticastDelegate(OwningObject, SparseDelegateFunc->DelegateName);
+
+	Delegate = MoveTemp(ScriptDelegate);
+
+	SparseDelegate.bIsBound = Delegate.IsBound();
+	if (!SparseDelegate.bIsBound)
+	{
+		FSparseDelegateStorage::Clear(OwningObject, SparseDelegateFunc->DelegateName);
+	}
+}
+
+
+FMulticastScriptDelegate::FInvocationList& UMulticastSparseDelegateProperty::GetInvocationList(const void* PropertyValue) const
+{
+	if (FSparseDelegate* SparseDelegate = (FSparseDelegate*)PropertyValue)
+	{
+		if (SparseDelegate->IsBound())
+		{
+			USparseDelegateFunction* SparseDelegateFunc = CastChecked<USparseDelegateFunction>(SignatureFunction);
+			UObject* OwningObject = FSparseDelegateStorage::ResolveSparseOwner(*SparseDelegate, SparseDelegateFunc->OwningClassName, SparseDelegateFunc->DelegateName);
+			if (FMulticastScriptDelegate* Delegate = FSparseDelegateStorage::GetMulticastDelegate(OwningObject, SparseDelegateFunc->DelegateName))
+			{
+				return Delegate->InvocationList;
+			}
+		}
+	}
+	return EmptyList;
+}
+
+void UMulticastSparseDelegateProperty::SerializeItem(FStructuredArchive::FSlot Slot, void* Value, void const* Defaults) const
+{
+	FArchiveUObjectFromStructuredArchive Ar(Slot);
+
+	FSparseDelegate& SparseDelegate = *(FSparseDelegate*)Value;
+
+	if (Ar.IsLoading())
+	{
+		FMulticastScriptDelegate Delegate;
+		Ar << Delegate;
+
+		if (Delegate.IsBound())
+		{
+			USparseDelegateFunction* SparseDelegateFunc = CastChecked<USparseDelegateFunction>(SignatureFunction);
+			UObject* OwningObject = FSparseDelegateStorage::ResolveSparseOwner(SparseDelegate, SparseDelegateFunc->OwningClassName, SparseDelegateFunc->DelegateName);
+			FMulticastScriptDelegate& BoundDelegate = FSparseDelegateStorage::GetOrCreateMulticastDelegate(OwningObject, SparseDelegateFunc->DelegateName);
+			BoundDelegate = MoveTemp(Delegate);
+			SparseDelegate.bIsBound = true;
+		}
+		else if (SparseDelegate.bIsBound)
+		{
+			USparseDelegateFunction* SparseDelegateFunc = CastChecked<USparseDelegateFunction>(SignatureFunction);
+			UObject* OwningObject = FSparseDelegateStorage::ResolveSparseOwner(SparseDelegate, SparseDelegateFunc->OwningClassName, SparseDelegateFunc->DelegateName);
+			FSparseDelegateStorage::Clear(OwningObject, SparseDelegateFunc->DelegateName);
+			SparseDelegate.bIsBound = false;
+		}
+	} 
+	else
+	{
+		if (SparseDelegate.IsBound())
+		{
+			USparseDelegateFunction* SparseDelegateFunc = CastChecked<USparseDelegateFunction>(SignatureFunction);
+			UObject* OwningObject = FSparseDelegateStorage::ResolveSparseOwner(SparseDelegate, SparseDelegateFunc->OwningClassName, SparseDelegateFunc->DelegateName);
+			if (FMulticastScriptDelegate* Delegate = FSparseDelegateStorage::GetMulticastDelegate(OwningObject, SparseDelegateFunc->DelegateName))
+			{
+				Ar << *Delegate;
+			}
+			else
+			{
+				Ar << EmptyList;
+			}
+		}
+		else
+		{
+			Ar << EmptyList;
+		}
+	}
+}
+
+const TCHAR* UMulticastSparseDelegateProperty::ImportText_Internal(const TCHAR* Buffer, void* PropertyValue, int32 PortFlags, UObject* Parent, FOutputDevice* ErrorText) const
+{
+	FSparseDelegate& SparseDelegate = *(FSparseDelegate*)PropertyValue;
+
+	USparseDelegateFunction* SparseDelegateFunc = CastChecked<USparseDelegateFunction>(SignatureFunction);
+	FMulticastScriptDelegate& Delegate = FSparseDelegateStorage::GetOrCreateMulticastDelegate(Parent, SparseDelegateFunc->DelegateName);
+
+	const TCHAR* Result = ImportDelegateFromText(Delegate, Buffer, PropertyValue, PortFlags, Parent, ErrorText);
+
+	if (Result)
+	{
+		SparseDelegate.bIsBound = Delegate.IsBound();
+		if (!SparseDelegate.bIsBound)
+		{
+			FSparseDelegateStorage::Clear(Parent, SparseDelegateFunc->DelegateName);
+		}
+	}
+
+	return Result;
+}
+
+
+void UMulticastSparseDelegateProperty::AddDelegate(FScriptDelegate ScriptDelegate, UObject* Parent, void* PropertyValue) const
+{
+	USparseDelegateFunction* SparseDelegateFunc = CastChecked<USparseDelegateFunction>(SignatureFunction);
+	FSparseDelegate& SparseDelegate = *(FSparseDelegate*)PropertyValue;
+	SparseDelegate.__Internal_AddUnique(Parent, SparseDelegateFunc->DelegateName, MoveTemp(ScriptDelegate));
+}
+
+
+void UMulticastSparseDelegateProperty::RemoveDelegate(const FScriptDelegate& ScriptDelegate, UObject* Parent, void* PropertyValue) const
+{
+	USparseDelegateFunction* SparseDelegateFunc = CastChecked<USparseDelegateFunction>(SignatureFunction);
+	FSparseDelegate& SparseDelegate = *(FSparseDelegate*)PropertyValue;
+	SparseDelegate.__Internal_Remove(Parent, SparseDelegateFunc->DelegateName, ScriptDelegate);
+}
+
+void UMulticastSparseDelegateProperty::ClearDelegate(UObject* Parent, void* PropertyValue) const
+{
+	USparseDelegateFunction* SparseDelegateFunc = CastChecked<USparseDelegateFunction>(SignatureFunction);
+	FSparseDelegate& SparseDelegate = *(FSparseDelegate*)PropertyValue;
+	SparseDelegate.__Internal_Clear(Parent, SparseDelegateFunc->DelegateName);
+}
+
+IMPLEMENT_CORE_INTRINSIC_CLASS(UMulticastSparseDelegateProperty, UMulticastDelegateProperty,
+	{
+	}
+);
+
diff --git a/Engine/Source/Runtime/CoreUObject/Private/UObject/ScriptCore.cpp b/Engine/Source/Runtime/CoreUObject/Private/UObject/ScriptCore.cpp
--- a/Engine/Source/Runtime/CoreUObject/Private/UObject/ScriptCore.cpp
+++ b/Engine/Source/Runtime/CoreUObject/Private/UObject/ScriptCore.cpp
@@ -2249,13 +2249,14 @@ DEFINE_FUNCTION(UObject::execLetMulticastDelegate)
 	Stack.MostRecentProperty = NULL;
 	Stack.Step( Stack.Object, NULL ); // Variable.
 
-	FMulticastScriptDelegate* DelegateAddr = (FMulticastScriptDelegate*)Stack.MostRecentPropertyAddress;
+	UMulticastDelegateProperty* DelegateProp = CastChecked<UMulticastDelegateProperty>(Stack.MostRecentProperty, ECastCheckedType::NullAllowed);
+	void* DelegateAddr = Stack.MostRecentPropertyAddress;
 	FMulticastScriptDelegate Delegate;
 	Stack.Step( Stack.Object, &Delegate );
 
-	if (DelegateAddr != NULL)
+	if (DelegateProp && DelegateAddr)
 	{
-		*DelegateAddr = Delegate;
+		DelegateProp->SetMulticastDelegate(DelegateAddr, MoveTemp(Delegate));
 	}
 }
 IMPLEMENT_VM_FUNCTION( EX_LetMulticastDelegate, execLetMulticastDelegate );
@@ -2451,7 +2452,8 @@ public:
 		Stack.MostRecentPropertyAddress = NULL;
 		Stack.MostRecentProperty = NULL;
 		Stack.Step( Stack.Object, NULL );
-		const FMulticastScriptDelegate* DelegateAddr = (FMulticastScriptDelegate*)Stack.MostRecentPropertyAddress;
+		UMulticastDelegateProperty* DelegateProp = CastChecked<UMulticastDelegateProperty>(Stack.MostRecentProperty, ECastCheckedType::NullAllowed);
+		const FMulticastScriptDelegate* DelegateAddr = (DelegateProp ? DelegateProp->GetMulticastDelegate(Stack.MostRecentPropertyAddress) : nullptr);
 
 		//Fill parameters
 		uint8* Parameters = (uint8*)FMemory_Alloca(SignatureFunction->ParmsSize);
@@ -2506,13 +2508,15 @@ DEFINE_FUNCTION(UObject::execAddMulticastDelegate)
 	Stack.MostRecentProperty = NULL;
 	Stack.Step( Stack.Object, NULL ); // Variable.
 
-	FMulticastScriptDelegate* DelegateAddr = (FMulticastScriptDelegate*)Stack.MostRecentPropertyAddress;
+	UMulticastDelegateProperty* DelegateProp = CastChecked<UMulticastDelegateProperty>(Stack.MostRecentProperty, ECastCheckedType::NullAllowed);
+	void* DelegateAddr = Stack.MostRecentPropertyAddress;
+
 	FScriptDelegate Delegate;
 	Stack.Step( Stack.Object, &Delegate );
 
-	if (DelegateAddr != NULL)
+	if (DelegateProp && DelegateAddr)
 	{
-		DelegateAddr->AddUnique(Delegate);
+		DelegateProp->AddDelegate(MoveTemp(Delegate), Stack.Object, DelegateAddr);
 	}
 }
 IMPLEMENT_VM_FUNCTION( EX_AddMulticastDelegate, execAddMulticastDelegate );
@@ -2524,13 +2528,15 @@ DEFINE_FUNCTION(UObject::execRemoveMulticastDelegate)
 	Stack.MostRecentProperty = NULL;
 	Stack.Step( Stack.Object, NULL ); // Variable.
 
-	FMulticastScriptDelegate* DelegateAddr = (FMulticastScriptDelegate*)Stack.MostRecentPropertyAddress;
+	UMulticastDelegateProperty* DelegateProp = CastChecked<UMulticastDelegateProperty>(Stack.MostRecentProperty, ECastCheckedType::NullAllowed);
+	void* DelegateAddr = Stack.MostRecentPropertyAddress;
+
 	FScriptDelegate Delegate;
 	Stack.Step( Stack.Object, &Delegate );
 
-	if (DelegateAddr != NULL)
+	if (DelegateProp && DelegateAddr)
 	{
-		DelegateAddr->Remove(Delegate);
+		DelegateProp->RemoveDelegate(Delegate, Stack.Object, DelegateAddr);
 	}
 }
 IMPLEMENT_VM_FUNCTION( EX_RemoveMulticastDelegate, execRemoveMulticastDelegate );
@@ -2542,10 +2548,12 @@ DEFINE_FUNCTION(UObject::execClearMulticastDelegate)
 	Stack.MostRecentProperty = NULL;
 	Stack.Step( Stack.Object, NULL );
 
-	FMulticastScriptDelegate* DelegateAddr = (FMulticastScriptDelegate*)Stack.MostRecentPropertyAddress;
-	if (DelegateAddr != NULL)
+	UMulticastDelegateProperty* DelegateProp = CastChecked<UMulticastDelegateProperty>(Stack.MostRecentProperty, ECastCheckedType::NullAllowed);
+	void* DelegateAddr = Stack.MostRecentPropertyAddress;
+
+	if (DelegateProp && DelegateAddr)
 	{
-		DelegateAddr->Clear();
+		DelegateProp->ClearDelegate(Stack.Object, DelegateAddr);
 	}
 }
 IMPLEMENT_VM_FUNCTION( EX_ClearMulticastDelegate, execClearMulticastDelegate );
diff --git a/Engine/Source/Runtime/CoreUObject/Private/UObject/SparseDelegate.cpp b/Engine/Source/Runtime/CoreUObject/Private/UObject/SparseDelegate.cpp
new file mode 100644
--- /dev/null
+++ b/Engine/Source/Runtime/CoreUObject/Private/UObject/SparseDelegate.cpp
@@ -0,0 +1,372 @@
+// Copyright 1998-2019 Epic Games, Inc. All Rights Reserved.
+
+#include "UObject/SparseDelegate.h"
+#include "UObject/Object.h"
+#include "UObject/Class.h"
+#include "UObject/Package.h"
+#include "HAL/IConsoleManager.h"
+
+FUObjectAnnotationSparse<FSparseDelegateStorage::FSparseDelegateMap, true> FSparseDelegateStorage::SparseDelegates;
+TMap<TPair<FName,FName>, size_t> FSparseDelegateStorage::SparseDelegateObjectOffsets;
+
+void FSparseDelegateStorage::RegisterDelegateOffset(const UObject* OwningObject, const FName DelegateName, const size_t DelegateOffsetToOwner)
+{
+	check(OwningObject);
+	SparseDelegateObjectOffsets.Add(TPair<FName,FName>(OwningObject->GetClass()->GetFName(), DelegateName), DelegateOffsetToOwner);
+}
+
+FSparseDelegate* FSparseDelegateStorage::ResolveSparseDelegate(const UObject* OwningObject, const FName DelegateName)
+{
+	const UClass* OwningClass = OwningObject->GetClass();
+	while (OwningClass)
+	{
+		if (!OwningClass->HasAnyClassFlags(CLASS_CompiledFromBlueprint))
+		{
+			if (size_t* DelegateOffset = SparseDelegateObjectOffsets.Find(TPair<FName, FName>(OwningClass->GetFName(), DelegateName)))
+			{
+				return reinterpret_cast<FSparseDelegate*>((uint8*)OwningObject + *DelegateOffset);
+			}
+		}
+
+		OwningClass = OwningClass->GetSuperClass();
+	}
+	check(false);
+	return nullptr;
+}
+
+UObject* FSparseDelegateStorage::ResolveSparseOwner(const FSparseDelegate& SparseDelegate, const FName OwningClassName, const FName DelegateName)
+{
+	if (size_t* DelegateOffset = SparseDelegateObjectOffsets.Find(TPair<FName, FName>(OwningClassName, DelegateName)))
+	{
+		return reinterpret_cast<UObject*>((uint8*)&SparseDelegate - *DelegateOffset);
+	}
+	check(false);
+	return nullptr;
+}
+
+FMulticastScriptDelegate* FSparseDelegateStorage::GetMulticastDelegate(const UObject* DelegateOwner, const FName DelegateName)
+{
+	FSparseDelegateMap DelegateMap = SparseDelegates.GetAnnotation(DelegateOwner);
+	if (TSharedPtr<FMulticastScriptDelegate>* MulticastDelegatePtr = DelegateMap.Find(DelegateName))
+	{
+		return MulticastDelegatePtr->Get();
+	}
+
+	return nullptr;
+}
+
+FMulticastScriptDelegate& FSparseDelegateStorage::GetOrCreateMulticastDelegate(const UObject* DelegateOwner, const FName DelegateName)
+{
+	FSparseDelegateMap DelegateMap = SparseDelegates.GetAnnotation(DelegateOwner);
+	TSharedPtr<FMulticastScriptDelegate>& MulticastDelegate = DelegateMap.FindOrAdd(DelegateName);
+
+	if (!MulticastDelegate.IsValid())
+	{
+		MulticastDelegate = MakeShared<FMulticastScriptDelegate>();
+		SparseDelegates.AddAnnotation(DelegateOwner, MoveTemp(DelegateMap));
+	}
+
+	return *MulticastDelegate;
+}
+
+bool FSparseDelegateStorage::Add(const UObject* DelegateOwner, const FName DelegateName, FScriptDelegate Delegate)
+{
+	bool bDelegateWasBound = false;
+	if (Delegate.IsBound())
+	{
+		FSparseDelegateMap DelegateMap = SparseDelegates.GetAnnotation(DelegateOwner);
+		TSharedPtr<FMulticastScriptDelegate>& MulticastDelegate = DelegateMap.FindOrAdd(DelegateName);
+
+		if (!MulticastDelegate.IsValid())
+		{
+			MulticastDelegate = MakeShared<FMulticastScriptDelegate>();
+		}
+
+		MulticastDelegate->Add(MoveTemp(Delegate));
+		SparseDelegates.AddAnnotation(DelegateOwner, MoveTemp(DelegateMap));
+		bDelegateWasBound = true;
+	}
+	return bDelegateWasBound;
+}
+
+bool FSparseDelegateStorage::AddUnique(const UObject* DelegateOwner, const FName DelegateName, FScriptDelegate Delegate)
+{
+	bool bDelegateWasBound = false;
+	if (Delegate.IsBound())
+	{
+		FSparseDelegateMap DelegateMap = SparseDelegates.GetAnnotation(DelegateOwner);
+		TSharedPtr<FMulticastScriptDelegate>& MulticastDelegate = DelegateMap.FindOrAdd(DelegateName);
+
+		if (!MulticastDelegate.IsValid())
+		{
+			MulticastDelegate = MakeShared<FMulticastScriptDelegate>();
+		}
+
+		MulticastDelegate->AddUnique(MoveTemp(Delegate));
+		SparseDelegates.AddAnnotation(DelegateOwner, MoveTemp(DelegateMap));
+		bDelegateWasBound = true;
+	}
+	return bDelegateWasBound;
+}
+
+bool FSparseDelegateStorage::Contains(const UObject* DelegateOwner, const FName DelegateName, const FScriptDelegate& Delegate)
+{
+	bool bContainsDelegate = false;
+	FSparseDelegateMap DelegateMap = SparseDelegates.GetAnnotation(DelegateOwner);
+	if (TSharedPtr<FMulticastScriptDelegate>* MulticastDelegatePtr = DelegateMap.Find(DelegateName))
+	{
+		if (FMulticastScriptDelegate* MulticastDelegate = MulticastDelegatePtr->Get())
+		{
+			bContainsDelegate = MulticastDelegate->Contains(Delegate);
+		}
+	}
+
+	return bContainsDelegate;
+}
+
+bool FSparseDelegateStorage::Contains(const UObject* DelegateOwner, const FName DelegateName, const UObject* InObject, FName InFunctionName)
+{
+	bool bContainsDelegate = false;
+	FSparseDelegateMap DelegateMap = SparseDelegates.GetAnnotation(DelegateOwner);
+	if (TSharedPtr<FMulticastScriptDelegate>* MulticastDelegatePtr = DelegateMap.Find(DelegateName))
+	{
+		if (FMulticastScriptDelegate* MulticastDelegate = MulticastDelegatePtr->Get())
+		{
+			bContainsDelegate = MulticastDelegate->Contains(InObject, InFunctionName);
+		}
+	}
+	return bContainsDelegate;
+}
+
+bool FSparseDelegateStorage::Remove(const UObject* DelegateOwner, const FName DelegateName, const FScriptDelegate& Delegate)
+{
+	bool bSparseDelegateBound = false;
+	FSparseDelegateMap DelegateMap = SparseDelegates.GetAnnotation(DelegateOwner);
+	if (TSharedPtr<FMulticastScriptDelegate>* MulticastDelegatePtr = DelegateMap.Find(DelegateName))
+	{
+		if (FMulticastScriptDelegate* MulticastDelegate = MulticastDelegatePtr->Get())
+		{
+			MulticastDelegate->Remove(Delegate);
+			bSparseDelegateBound = MulticastDelegate->IsBound();
+			if (!bSparseDelegateBound)
+			{
+				DelegateMap.Remove(DelegateName);
+			}
+		}
+		else
+		{
+			DelegateMap.Remove(DelegateName);
+		}
+	}
+	if (DelegateMap.IsDefault())
+	{
+		SparseDelegates.RemoveAnnotation(DelegateOwner);
+	}
+	else
+	{
+		SparseDelegates.AddAnnotation(DelegateOwner, MoveTemp(DelegateMap));
+	}
+
+	return bSparseDelegateBound;
+}
+
+bool FSparseDelegateStorage::Remove(const UObject* DelegateOwner, const FName DelegateName, const UObject* InObject, FName InFunctionName)
+{
+	bool bSparseDelegateBound = false;
+	FSparseDelegateMap DelegateMap = SparseDelegates.GetAnnotation(DelegateOwner);
+	if (TSharedPtr<FMulticastScriptDelegate>* MulticastDelegatePtr = DelegateMap.Find(DelegateName))
+	{
+		if (FMulticastScriptDelegate* MulticastDelegate = MulticastDelegatePtr->Get())
+		{
+			MulticastDelegate->Remove(InObject, InFunctionName);
+			bSparseDelegateBound = MulticastDelegate->IsBound();
+			if (!bSparseDelegateBound)
+			{
+				DelegateMap.Remove(DelegateName);
+			}
+		}
+		else
+		{
+			DelegateMap.Remove(DelegateName);
+		}
+	}
+	if (DelegateMap.IsDefault())
+	{
+		SparseDelegates.RemoveAnnotation(DelegateOwner);
+	}
+	else
+	{
+		SparseDelegates.AddAnnotation(DelegateOwner, MoveTemp(DelegateMap));
+	}
+
+	return bSparseDelegateBound;
+}
+
+bool FSparseDelegateStorage::RemoveAll(const UObject* DelegateOwner, const FName DelegateName, const UObject* UserObject)
+{
+	bool bSparseDelegateBound = false;
+	FSparseDelegateMap DelegateMap = SparseDelegates.GetAnnotation(DelegateOwner);
+	if (TSharedPtr<FMulticastScriptDelegate>* MulticastDelegatePtr = DelegateMap.Find(DelegateName))
+	{
+		if (FMulticastScriptDelegate* MulticastDelegate = MulticastDelegatePtr->Get())
+		{
+			MulticastDelegate->RemoveAll(UserObject);
+			bSparseDelegateBound = MulticastDelegate->IsBound();
+			if (!bSparseDelegateBound)
+			{
+				DelegateMap.Remove(DelegateName);
+			}
+		}
+		else
+		{
+			DelegateMap.Remove(DelegateName);
+		}
+	}
+	if (DelegateMap.IsDefault())
+	{
+		SparseDelegates.RemoveAnnotation(DelegateOwner);
+	}
+	else
+	{
+		SparseDelegates.AddAnnotation(DelegateOwner, MoveTemp(DelegateMap));
+	}
+	return bSparseDelegateBound;
+}
+
+void FSparseDelegateStorage::Clear(const UObject* DelegateOwner, const FName DelegateName)
+{
+	FSparseDelegateMap DelegateMap = SparseDelegates.GetAnnotation(DelegateOwner);
+	if (TSharedPtr<FMulticastScriptDelegate>* MulticastDelegatePtr = DelegateMap.Find(DelegateName))
+	{
+		if (FMulticastScriptDelegate* MulticastDelegate = MulticastDelegatePtr->Get())
+		{
+			MulticastDelegate->Clear();
+		}
+		DelegateMap.Remove(DelegateName);
+	}
+	if (DelegateMap.IsDefault())
+	{
+		SparseDelegates.RemoveAnnotation(DelegateOwner);
+	}
+	else
+	{
+		SparseDelegates.AddAnnotation(DelegateOwner, MoveTemp(DelegateMap));
+	}
+}
+
+FAutoConsoleCommand SparseDelegateReportCommand(
+	TEXT("SparseDelegateReport"),
+	TEXT("Outputs a report of what sparse delegates are bound. SparseDelegateReport [name=<ObjectName>] [delegate=<DelegateName>] [class=<ClassName>] -details"),
+	FConsoleCommandWithWorldArgsAndOutputDeviceDelegate::CreateStatic(FSparseDelegateStorage::SparseDelegateReport)
+);
+
+void FSparseDelegateStorage::SparseDelegateReport(const TArray<FString>& Args, UWorld*, FOutputDevice& Ar)
+{
+	UClass* ObjectType = nullptr;
+	FName ObjectName = NAME_None;
+	FName DelegateName = NAME_None;
+	bool bSummary = true;
+	bool bArgumentError = false;
+
+	for (const FString& Arg : Args)
+	{
+		if (Arg == TEXT("-details"))
+		{
+			bSummary = false;
+		}
+		else if (Arg.StartsWith(TEXT("name=")))
+		{
+			FString ObjName;
+			Arg.Split(TEXT("="), nullptr, &ObjName, ESearchCase::CaseSensitive);
+			ObjectName = FName(*ObjName, EFindName::FNAME_Find);
+			if (ObjectName.IsNone())
+			{
+				Ar.Logf(ELogVerbosity::Warning, TEXT("Invalid object name"));
+				bArgumentError = true;
+			}
+		}
+		else if (Arg.StartsWith(TEXT("delegate=")))
+		{
+			FString DelegateNameStr;
+			Arg.Split(TEXT("="), nullptr, &DelegateNameStr, ESearchCase::CaseSensitive);
+			DelegateName = FName(*DelegateNameStr, EFindName::FNAME_Find);
+			if (DelegateName.IsNone())
+			{
+				Ar.Logf(ELogVerbosity::Warning, TEXT("Invalid delegate name"));
+				bArgumentError = true;
+			}
+		}
+		else if (Arg.StartsWith(TEXT("class=")))
+		{
+			FString ClassName;
+			Arg.Split(TEXT("="), nullptr, &ClassName, ESearchCase::CaseSensitive);
+			ObjectType = FindObject<UClass>(ANY_PACKAGE, *ClassName);
+			if (ObjectType == nullptr)
+			{
+				Ar.Logf(ELogVerbosity::Warning, TEXT("No class of specified name found."));
+				bArgumentError = true;
+			}
+		}
+	}
+
+	if (bArgumentError)
+	{
+		return;
+	}
+
+	TArray<FString> Details;
+	uint32 BoundObjects = 0;
+	uint32 BoundDelegates = 0;
+	for (const TPair<const UObjectBase*, FSparseDelegateMap>& ObjectAnnotation : SparseDelegates.GetAnnotationMap())
+	{
+		const UObject* Object = static_cast<const UObject*>(ObjectAnnotation.Key);
+		if (ObjectName.IsNone() || Object->GetFName() == ObjectName)
+		{
+			if (ObjectType == nullptr || Object->IsA(ObjectType))
+			{
+				if (!bSummary)
+				{
+					Details.Add(FString::Printf(TEXT("%s"), *Object->GetPathName()));
+				}
+				++BoundObjects;
+				if (DelegateName.IsNone())
+				{
+					BoundDelegates += ObjectAnnotation.Value.Num();
+					if (!bSummary)
+					{
+						for (const TPair<FName, TSharedPtr<FMulticastScriptDelegate>>& BoundDelegate : ObjectAnnotation.Value)
+						{
+							Details.Add(FString::Printf(TEXT("   %s"), *BoundDelegate.Key.ToString()));
+						}
+					}
+				}
+				else if (ObjectAnnotation.Value.Contains(DelegateName))
+				{
+					++BoundDelegates;
+				}
+			}
+		}
+	}
+	FString SummaryString = TEXT("Bound Sparse Delegates");
+	if (ObjectType)
+	{
+		SummaryString += FString::Printf(TEXT(" - Class=%s"), *ObjectType->GetName());
+	}
+	if (!ObjectName.IsNone())
+	{
+		SummaryString += FString::Printf(TEXT(" - Name=%s"), *ObjectName.ToString());
+	}
+	if (!DelegateName.IsNone())
+	{
+		SummaryString += FString::Printf(TEXT(" - Delegate=%s"), *DelegateName.ToString());
+	}
+
+	Ar.Logf(ELogVerbosity::Log, TEXT("%s"), *SummaryString);
+	Ar.Logf(ELogVerbosity::Log, TEXT("Objects: %d"), BoundObjects);
+	Ar.Logf(ELogVerbosity::Log, TEXT("Delegates: %d"), BoundDelegates);
+	Ar.Logf(ELogVerbosity::Log, TEXT("------------------------------------------------------------------------"));
+	for (const FString& Detail : Details)
+	{
+		Ar.Logf(ELogVerbosity::Log, TEXT("%s"), *Detail);
+	}
+}
\ No newline at end of file
diff --git a/Engine/Source/Runtime/CoreUObject/Private/UObject/UObjectGlobals.cpp b/Engine/Source/Runtime/CoreUObject/Private/UObject/UObjectGlobals.cpp
--- a/Engine/Source/Runtime/CoreUObject/Private/UObject/UObjectGlobals.cpp
+++ b/Engine/Source/Runtime/CoreUObject/Private/UObject/UObjectGlobals.cpp
@@ -4215,10 +4215,10 @@ namespace UE4CodeGen_Private
 			}
 			break;
 
-			case EPropertyGenFlags::MulticastDelegate:
+			case EPropertyGenFlags::InlineMulticastDelegate:
 			{
 				const FMulticastDelegatePropertyParams* Prop = (const FMulticastDelegatePropertyParams*)PropBase;
-				NewProp = new (EC_InternalUseOnlyConstructor, Outer, UTF8_TO_TCHAR(Prop->NameUTF8), Prop->ObjectFlags) UMulticastDelegateProperty(FObjectInitializer(), EC_CppProperty, Prop->Offset, Prop->PropertyFlags, Prop->SignatureFunctionFunc ? Prop->SignatureFunctionFunc() : nullptr);
+				NewProp = new (EC_InternalUseOnlyConstructor, Outer, UTF8_TO_TCHAR(Prop->NameUTF8), Prop->ObjectFlags) UMulticastInlineDelegateProperty(FObjectInitializer(), EC_CppProperty, Prop->Offset, Prop->PropertyFlags, Prop->SignatureFunctionFunc ? Prop->SignatureFunctionFunc() : nullptr);
 
 #if WITH_METADATA
 				MetaDataArray = Prop->MetaDataArray;
@@ -4227,6 +4227,18 @@ namespace UE4CodeGen_Private
 			}
 			break;
 
+			case EPropertyGenFlags::SparseMulticastDelegate:
+			{
+				const FMulticastDelegatePropertyParams* Prop = (const FMulticastDelegatePropertyParams*)PropBase;
+				NewProp = new (EC_InternalUseOnlyConstructor, Outer, UTF8_TO_TCHAR(Prop->NameUTF8), Prop->ObjectFlags) UMulticastSparseDelegateProperty(FObjectInitializer(), EC_CppProperty, Prop->Offset, Prop->PropertyFlags, Prop->SignatureFunctionFunc ? Prop->SignatureFunctionFunc() : nullptr);
+
+#if WITH_METADATA
+				MetaDataArray = Prop->MetaDataArray;
+				NumMetaData = Prop->NumMetaData;
+#endif
+			}
+			break;
+
 			case EPropertyGenFlags::Text:
 			{
 				const FTextPropertyParams* Prop = (const FTextPropertyParams*)PropBase;
@@ -4318,12 +4330,27 @@ namespace UE4CodeGen_Private
 		UFunction* NewFunction;
 		if (Params.FunctionFlags & FUNC_Delegate)
 		{
-			NewFunction = new (EC_InternalUseOnlyConstructor, Outer, UTF8_TO_TCHAR(Params.NameUTF8), Params.ObjectFlags) UDelegateFunction(
-				FObjectInitializer(),
-				Super,
-				Params.FunctionFlags,
-				Params.StructureSize
-			);
+			if (Params.OwningClassName == nullptr)
+			{
+				NewFunction = new (EC_InternalUseOnlyConstructor, Outer, UTF8_TO_TCHAR(Params.NameUTF8), Params.ObjectFlags) UDelegateFunction(
+					FObjectInitializer(),
+					Super,
+					Params.FunctionFlags,
+					Params.StructureSize
+				);
+			}
+			else
+			{
+				USparseDelegateFunction* NewSparseFunction = new (EC_InternalUseOnlyConstructor, Outer, UTF8_TO_TCHAR(Params.NameUTF8), Params.ObjectFlags) USparseDelegateFunction(
+					FObjectInitializer(),
+					Super,
+					Params.FunctionFlags,
+					Params.StructureSize
+				);
+				NewSparseFunction->OwningClassName = FName(Params.OwningClassName);
+				NewSparseFunction->DelegateName = FName(Params.DelegateName);
+				NewFunction = NewSparseFunction;
+			}
 		}
 		else
 		{
diff --git a/Engine/Source/Runtime/CoreUObject/Public/Templates/Casts.h b/Engine/Source/Runtime/CoreUObject/Public/Templates/Casts.h
--- a/Engine/Source/Runtime/CoreUObject/Public/Templates/Casts.h
+++ b/Engine/Source/Runtime/CoreUObject/Public/Templates/Casts.h
@@ -353,6 +353,9 @@ DECLARE_CAST_BY_FLAG(UUInt32Property)					\
 DECLARE_CAST_BY_FLAG(UUInt64Property)					\
 DECLARE_CAST_BY_FLAG(UMapProperty)						\
 DECLARE_CAST_BY_FLAG(USetProperty)						\
+DECLARE_CAST_BY_FLAG(USparseDelegateFunction)			\
+DECLARE_CAST_BY_FLAG(UMulticastInlineDelegateProperty)	\
+DECLARE_CAST_BY_FLAG(UMulticastSparseDelegateProperty)	\
 FINISH_DECLARING_CAST_FLAGS		// This is here to hopefully remind people to include the "\" in all declarations above, especially when copy/pasting the final line.
 
 // Now actually declare the flags
diff --git a/Engine/Source/Runtime/CoreUObject/Public/UObject/Class.h b/Engine/Source/Runtime/CoreUObject/Public/UObject/Class.h
--- a/Engine/Source/Runtime/CoreUObject/Public/UObject/Class.h
+++ b/Engine/Source/Runtime/CoreUObject/Public/UObject/Class.h
@@ -1555,6 +1555,23 @@ public:
 	explicit UDelegateFunction(UFunction* InSuperFunction, EFunctionFlags InFunctionFlags = FUNC_None, SIZE_T ParamsSize = 0);
 };
 
+//
+// Function definition used by sparse dynamic delegate declarations
+//
+class COREUOBJECT_API USparseDelegateFunction : public UDelegateFunction
+{
+	DECLARE_CASTED_CLASS_INTRINSIC(USparseDelegateFunction, UDelegateFunction, 0, TEXT("/Script/CoreUObject"), CASTCLASS_USparseDelegateFunction)
+	DECLARE_WITHIN(UObject)
+public:
+	explicit USparseDelegateFunction(const FObjectInitializer& ObjectInitializer, UFunction* InSuperFunction, EFunctionFlags InFunctionFlags = FUNC_None, SIZE_T ParamsSize = 0);
+	explicit USparseDelegateFunction(UFunction* InSuperFunction, EFunctionFlags InFunctionFlags = FUNC_None, SIZE_T ParamsSize = 0);
+
+	virtual void Serialize(FArchive& Ar) override;
+
+	FName OwningClassName;
+	FName DelegateName;
+};
+
 /*-----------------------------------------------------------------------------
 	UEnum.
 -----------------------------------------------------------------------------*/
diff --git a/Engine/Source/Runtime/CoreUObject/Public/UObject/ObjectMacros.h b/Engine/Source/Runtime/CoreUObject/Public/UObject/ObjectMacros.h
--- a/Engine/Source/Runtime/CoreUObject/Public/UObject/ObjectMacros.h
+++ b/Engine/Source/Runtime/CoreUObject/Public/UObject/ObjectMacros.h
@@ -293,55 +293,58 @@ enum EClassCastFlags : uint64
 {
 	CASTCLASS_None = 0x0000000000000000,
 
-	CASTCLASS_UField						= 0x0000000000000001,
-	CASTCLASS_UInt8Property					= 0x0000000000000002,
-	CASTCLASS_UEnum							= 0x0000000000000004,
-	CASTCLASS_UStruct						= 0x0000000000000008,
-	CASTCLASS_UScriptStruct					= 0x0000000000000010,
-	CASTCLASS_UClass						= 0x0000000000000020,
-	CASTCLASS_UByteProperty					= 0x0000000000000040,
-	CASTCLASS_UIntProperty					= 0x0000000000000080,
-	CASTCLASS_UFloatProperty				= 0x0000000000000100,
-	CASTCLASS_UUInt64Property				= 0x0000000000000200,
-	CASTCLASS_UClassProperty				= 0x0000000000000400,
-	CASTCLASS_UUInt32Property				= 0x0000000000000800,
-	CASTCLASS_UInterfaceProperty			= 0x0000000000001000,
-	CASTCLASS_UNameProperty					= 0x0000000000002000,
-	CASTCLASS_UStrProperty					= 0x0000000000004000,
-	CASTCLASS_UProperty						= 0x0000000000008000,
-	CASTCLASS_UObjectProperty				= 0x0000000000010000,
-	CASTCLASS_UBoolProperty					= 0x0000000000020000,
-	CASTCLASS_UUInt16Property				= 0x0000000000040000,
-	CASTCLASS_UFunction						= 0x0000000000080000,
-	CASTCLASS_UStructProperty				= 0x0000000000100000,
-	CASTCLASS_UArrayProperty				= 0x0000000000200000,
-	CASTCLASS_UInt64Property				= 0x0000000000400000,
-	CASTCLASS_UDelegateProperty				= 0x0000000000800000,
-	CASTCLASS_UNumericProperty				= 0x0000000001000000,
-	CASTCLASS_UMulticastDelegateProperty	= 0x0000000002000000,
-	CASTCLASS_UObjectPropertyBase			= 0x0000000004000000,
-	CASTCLASS_UWeakObjectProperty			= 0x0000000008000000,
-	CASTCLASS_ULazyObjectProperty			= 0x0000000010000000,
-	CASTCLASS_USoftObjectProperty			= 0x0000000020000000,
-	CASTCLASS_UTextProperty					= 0x0000000040000000,
-	CASTCLASS_UInt16Property				= 0x0000000080000000,
-	CASTCLASS_UDoubleProperty				= 0x0000000100000000,
-	CASTCLASS_USoftClassProperty			= 0x0000000200000000,
-	CASTCLASS_UPackage						= 0x0000000400000000,
-	CASTCLASS_ULevel						= 0x0000000800000000,
-	CASTCLASS_AActor						= 0x0000001000000000,
-	CASTCLASS_APlayerController				= 0x0000002000000000,
-	CASTCLASS_APawn							= 0x0000004000000000,
-	CASTCLASS_USceneComponent				= 0x0000008000000000,
-	CASTCLASS_UPrimitiveComponent			= 0x0000010000000000,
-	CASTCLASS_USkinnedMeshComponent			= 0x0000020000000000,
-	CASTCLASS_USkeletalMeshComponent		= 0x0000040000000000,
-	CASTCLASS_UBlueprint					= 0x0000080000000000,
-	CASTCLASS_UDelegateFunction				= 0x0000100000000000,
-	CASTCLASS_UStaticMeshComponent			= 0x0000200000000000,
-	CASTCLASS_UMapProperty					= 0x0000400000000000,
-	CASTCLASS_USetProperty					= 0x0000800000000000,
-	CASTCLASS_UEnumProperty					= 0x0001000000000000,
+	CASTCLASS_UField							= 0x0000000000000001,
+	CASTCLASS_UInt8Property						= 0x0000000000000002,
+	CASTCLASS_UEnum								= 0x0000000000000004,
+	CASTCLASS_UStruct							= 0x0000000000000008,
+	CASTCLASS_UScriptStruct						= 0x0000000000000010,
+	CASTCLASS_UClass							= 0x0000000000000020,
+	CASTCLASS_UByteProperty						= 0x0000000000000040,
+	CASTCLASS_UIntProperty						= 0x0000000000000080,
+	CASTCLASS_UFloatProperty					= 0x0000000000000100,
+	CASTCLASS_UUInt64Property					= 0x0000000000000200,
+	CASTCLASS_UClassProperty					= 0x0000000000000400,
+	CASTCLASS_UUInt32Property					= 0x0000000000000800,
+	CASTCLASS_UInterfaceProperty				= 0x0000000000001000,
+	CASTCLASS_UNameProperty						= 0x0000000000002000,
+	CASTCLASS_UStrProperty						= 0x0000000000004000,
+	CASTCLASS_UProperty							= 0x0000000000008000,
+	CASTCLASS_UObjectProperty					= 0x0000000000010000,
+	CASTCLASS_UBoolProperty						= 0x0000000000020000,
+	CASTCLASS_UUInt16Property					= 0x0000000000040000,
+	CASTCLASS_UFunction							= 0x0000000000080000,
+	CASTCLASS_UStructProperty					= 0x0000000000100000,
+	CASTCLASS_UArrayProperty					= 0x0000000000200000,
+	CASTCLASS_UInt64Property					= 0x0000000000400000,
+	CASTCLASS_UDelegateProperty					= 0x0000000000800000,
+	CASTCLASS_UNumericProperty					= 0x0000000001000000,
+	CASTCLASS_UMulticastDelegateProperty		= 0x0000000002000000,
+	CASTCLASS_UObjectPropertyBase				= 0x0000000004000000,
+	CASTCLASS_UWeakObjectProperty				= 0x0000000008000000,
+	CASTCLASS_ULazyObjectProperty				= 0x0000000010000000,
+	CASTCLASS_USoftObjectProperty				= 0x0000000020000000,
+	CASTCLASS_UTextProperty						= 0x0000000040000000,
+	CASTCLASS_UInt16Property					= 0x0000000080000000,
+	CASTCLASS_UDoubleProperty					= 0x0000000100000000,
+	CASTCLASS_USoftClassProperty				= 0x0000000200000000,
+	CASTCLASS_UPackage							= 0x0000000400000000,
+	CASTCLASS_ULevel							= 0x0000000800000000,
+	CASTCLASS_AActor							= 0x0000001000000000,
+	CASTCLASS_APlayerController					= 0x0000002000000000,
+	CASTCLASS_APawn								= 0x0000004000000000,
+	CASTCLASS_USceneComponent					= 0x0000008000000000,
+	CASTCLASS_UPrimitiveComponent				= 0x0000010000000000,
+	CASTCLASS_USkinnedMeshComponent				= 0x0000020000000000,
+	CASTCLASS_USkeletalMeshComponent			= 0x0000040000000000,
+	CASTCLASS_UBlueprint						= 0x0000080000000000,
+	CASTCLASS_UDelegateFunction					= 0x0000100000000000,
+	CASTCLASS_UStaticMeshComponent				= 0x0000200000000000,
+	CASTCLASS_UMapProperty						= 0x0000400000000000,
+	CASTCLASS_USetProperty						= 0x0000800000000000,
+	CASTCLASS_UEnumProperty						= 0x0001000000000000,
+	CASTCLASS_USparseDelegateFunction			= 0x0002000000000000,
+	CASTCLASS_UMulticastInlineDelegateProperty	= 0x0004000000000000,
+	CASTCLASS_UMulticastSparseDelegateProperty	= 0x0008000000000000,
 };
 
 #define CASTCLASS_AllFlags ((EClassCastFlags)0xFFFFFFFFFFFFFFFF)
diff --git a/Engine/Source/Runtime/CoreUObject/Public/UObject/SparseDelegate.h b/Engine/Source/Runtime/CoreUObject/Public/UObject/SparseDelegate.h
new file mode 100644
--- /dev/null
+++ b/Engine/Source/Runtime/CoreUObject/Public/UObject/SparseDelegate.h
@@ -0,0 +1,404 @@
+// Copyright 1998-2019 Epic Games, Inc. All Rights Reserved.
+
+#pragma once
+
+#include "CoreMinimal.h"
+#include "UObject/UObjectAnnotation.h"
+#include "UObject/WeakObjectPtr.h"
+#include "Delegates/Delegate.h"
+
+struct FSparseDelegate;
+
+/**
+*  Sparse delegates can be used for infrequently bound delegates so that the object uses only 
+*  1 byte of storage instead of having the full overhead of the delegate invocation list.
+*  The cost to invoke, add, remove, etc. from the delegate is higher than using the delegate
+*  directly and thus the memory savings benefit should be traded off against the frequency
+*  with which you would expect the delegate to be bound.
+*/
+
+
+/**
+* Helper class for handling sparse delegate bindings
+*/
+struct COREUOBJECT_API FSparseDelegateStorage
+{
+public:
+
+	/** Registers the sparse delegate so that the offset can be determined. */
+	static void RegisterDelegateOffset(const UObject* OwningObject, FName DelegateName, size_t OffsetToOwner);
+
+	/** Binds a sparse delegate to the owner. Returns if the delegate was successfully bound. */
+	static bool Add(const UObject* DelegateOwner, const FName DelegateName, FScriptDelegate Delegate);
+
+	/** Binds a sparse delegate to the owner, verifying first that the delegate is not already bound. Returns if the delegate was successfully bound. */
+	static bool AddUnique(const UObject* DelegateOwner, const FName DelegateName, FScriptDelegate Delegate);
+
+	/** Returns whether a sparse delegate is bound to the owner. */
+	static bool Contains(const UObject* DelegateOwner, const FName DelegateName, const FScriptDelegate& Delegate);
+	static bool Contains(const UObject* DelegateOwner, const FName DelegateName, const UObject* InObject, FName InFunctionName);
+
+	/** Removes a delegate binding from the owner's sparse delegate storage. Returns true if there are still bindings to the delegate. */
+	static bool Remove(const UObject* DelegateOwner, const FName DelegateName, const FScriptDelegate& Delegate);
+	static bool Remove(const UObject* DelegateOwner, const FName DelegateName, const UObject* InObject, FName InFunctionName);
+
+	/** Removes all sparse delegate binding from the owner for a given object. Returns true if there are still bindings to the delegate. */
+	static bool RemoveAll(const UObject* DelegateOwner, const FName DelegateName, const UObject* UserObject);
+
+	/** Clear all of the named sparse delegate bindings from the owner. */
+	static void Clear(const UObject* DelegateOwner, const FName DelegateName);
+
+	/** Acquires the actual Multicast Delegate from the annotation if any delegates are bound to it. Will be null if no entry exists in the annotation for this object/delegatename. */
+	static FMulticastScriptDelegate* GetMulticastDelegate(const UObject* DelegateOwner, const FName DelegateName);
+
+	/** Acquires the actual Multicast Delegate from the annotation if any delegates are bound to it. Will be created if no entry exists in the annotation for this object/delegatename. */
+	static FMulticastScriptDelegate& GetOrCreateMulticastDelegate(const UObject* DelegateOwner, const FName DelegateName);
+
+	/** Using the registry of sparse delegates recover the FSparseDelegate address from the UObject and name. */
+	static FSparseDelegate* ResolveSparseDelegate(const UObject* OwningObject, FName DelegateName);
+
+	/** Using the registry of sparse delegates recover the UObject owner from the FSparseDelegate address owning class and delegate names. */
+	static UObject* ResolveSparseOwner(const FSparseDelegate& SparseDelegate, const FName OwningClassName, const FName DelegateName);
+
+	/** Outputs a report about which delegates are bound. */
+	static void SparseDelegateReport(const TArray<FString>&, UWorld*, FOutputDevice&);
+
+private:
+
+	struct FSparseDelegateMap : public TMap<FName, TSharedPtr<FMulticastScriptDelegate>>
+	{
+		bool IsDefault() const { return Num() == 0; }
+	};
+
+	static FUObjectAnnotationSparse<FSparseDelegateMap, true> SparseDelegates;
+
+	/** Sparse delegate offsets are indexed by ActorClass/DelegateName pair */
+	static TMap<TPair<FName, FName>, size_t> SparseDelegateObjectOffsets;
+};
+
+struct FSparseDelegate
+{
+public:
+	FSparseDelegate()
+		: bIsBound(false)
+	{
+	}
+
+	/**
+	* Checks to see if any functions are bound to this multi-cast delegate
+	*
+	* @return	True if any functions are bound
+	*/
+	bool IsBound() const
+	{
+		return bIsBound;
+	}
+
+	/**
+	* Adds a function delegate to this multi-cast delegate's invocation list if a delegate with the same signature
+	* doesn't already exist in the invocation list
+	*
+	* @param	DelegateOwner	UObject that owns the resolved sparse delegate
+	* @param	DelegateName	Name of the resolved sparse delegate
+	* @param	InDelegate		Delegate to bind to the sparse delegate
+	* 
+	* NOTE:  Only call this function from blueprint sparse delegate infrastructure on a resolved generic FScriptDelegate pointer.
+	*        Generally from C++ you should call AddUnique() directly.
+	*/
+	void __Internal_AddUnique(const UObject* DelegateOwner, FName DelegateName, FScriptDelegate InDelegate)
+	{
+		bIsBound |= FSparseDelegateStorage::AddUnique(DelegateOwner, DelegateName, MoveTemp(InDelegate));
+	}
+
+	/**
+	* Removes a function from this multi-cast delegate's invocation list (performance is O(N)).  Note that the
+	* order of the delegates may not be preserved!
+	*
+	* @param	DelegateOwner	UObject that owns the resolved sparse delegate
+	* @param	DelegateName	Name of the resolved sparse delegate
+	* @param	InDelegate		Delegate to remove from the sparse delegate
+	*
+	* NOTE:  Only call this function from blueprint sparse delegate infrastructure on a resolved generic FScriptDelegate pointer.
+	*        Generally from C++ you should call Remove() directly.
+	*/
+	void __Internal_Remove(const UObject* DelegateOwner, FName DelegateName, const FScriptDelegate& InDelegate)
+	{
+		if (bIsBound)
+		{
+			bIsBound = FSparseDelegateStorage::Remove(DelegateOwner, DelegateName, InDelegate);
+		}
+	}
+
+	/**
+	* Removes all functions from this delegate's invocation list
+	*
+	* @param	DelegateOwner	UObject that owns the resolved sparse delegate
+	* @param	DelegateName	Name of the resolved sparse delegate
+	*
+	* NOTE:  Only call this function from blueprint sparse delegate infrastructure on a resolved generic FScriptDelegate pointer.
+	*        Generally from C++ you should call Clear() directly.
+	*/
+	void __Internal_Clear(const UObject* DelegateOwner, FName DelegateName)
+	{
+		if (bIsBound)
+		{
+			FSparseDelegateStorage::Clear(DelegateOwner, DelegateName);
+			bIsBound = false;
+		}
+	}
+
+protected:
+
+	friend class UMulticastSparseDelegateProperty;
+	bool bIsBound;
+};
+
+template <typename MulticastDelegate, typename OwningClass, const char* DelegateName, typename DelegateOffsetClass>
+struct TSparseDynamicDelegate : public FSparseDelegate
+{
+public:
+	typedef typename MulticastDelegate::FDelegate FDelegate; 
+
+protected:
+	FName GetDelegateName() const
+	{
+		static const FName DelegateFName(DelegateName);
+		return DelegateFName;
+	}
+
+private:
+	UObject* GetDelegateOwner() const
+	{
+		const size_t OffsetToOwner = DelegateOffsetClass::template GetDelegateOffset<OwningClass>();
+		check(OffsetToOwner);
+		UObject* DelegateOwner = reinterpret_cast<UObject*>((uint8*)this - OffsetToOwner);
+		check(DelegateOwner->IsValidLowLevelFast(false));
+		return DelegateOwner;
+	}
+
+public:
+	/** Returns the multicast delegate if any delegates are bound to the sparse delegate */
+	MulticastDelegate* Get() const
+	{
+		return (MulticastDelegate*)(bIsBound ? FSparseDelegateStorage::GetMulticastDelegate(GetDelegateOwner(), GetDelegateName()) : nullptr);
+	}
+
+	/**
+	* Checks whether a function delegate is already a member of this multi-cast delegate's invocation list
+	*
+	* @param	InDelegate	Delegate to check
+	* @return	True if the delegate is already in the list.
+	*/
+	bool Contains(const FScriptDelegate& InDelegate) const
+	{
+		return (bIsBound ? FSparseDelegateStorage::Contains(GetDelegateOwner(), GetDelegateName(), InDelegate) : false);
+	}
+
+	/**
+	* Checks whether a function delegate is already a member of this multi-cast delegate's invocation list
+	*
+	* @param	InObject		Object of the delegate to check
+	* @param	InFunctionName	Function name of the delegate to check
+	* @return	True if the delegate is already in the list.
+	*/
+	bool Contains(const UObject* InObject, FName InFunctionName) const
+	{
+		return (bIsBound ? FSparseDelegateStorage::Contains(GetDelegateOwner(), GetDelegateName(), InObject, InFunctionName) : false);
+	}
+
+	/**
+	* Adds a function delegate to this multi-cast delegate's invocation list
+	*
+	* @param	InDelegate	Delegate to add
+	*/
+	void Add(FScriptDelegate InDelegate)
+	{
+		bIsBound |= FSparseDelegateStorage::Add(GetDelegateOwner(), GetDelegateName(), MoveTemp(InDelegate));
+	}
+
+	/**
+	* Adds a function delegate to this multi-cast delegate's invocation list if a delegate with the same signature
+	* doesn't already exist in the invocation list
+	*
+	* @param	InDelegate	Delegate to add
+	*/
+	void AddUnique(FScriptDelegate InDelegate)
+	{
+		FSparseDelegate::__Internal_AddUnique(GetDelegateOwner(), GetDelegateName(), MoveTemp(InDelegate));
+	}
+
+	/**
+	* Removes a function from this multi-cast delegate's invocation list (performance is O(N)).  Note that the
+	* order of the delegates may not be preserved!
+	*
+	* @param	InDelegate	Delegate to remove
+	*/
+	void Remove(const FScriptDelegate& InDelegate)
+	{
+		FSparseDelegate::__Internal_Remove(GetDelegateOwner(), GetDelegateName(), InDelegate);
+	}
+
+	/**
+	* Removes a function from this multi-cast delegate's invocation list (performance is O(N)).  Note that the
+	* order of the delegates may not be preserved!
+	*
+	* @param	InObject		Object of the delegate to remove
+	* @param	InFunctionName	Function name of the delegate to remove
+	*/
+	void Remove(const UObject* InObject, FName InFunctionName)
+	{
+		if (bIsBound)
+		{
+			bIsBound = FSparseDelegateStorage::Remove(GetDelegateOwner(), GetDelegateName(), InObject, InFunctionName);
+		}
+	}
+
+	/**
+	* Removes all delegate bindings from this multicast delegate's
+	* invocation list that are bound to the specified object.
+	*
+	* This method also compacts the invocation list.
+	*
+	* @param InObject The object to remove bindings for.
+	*/
+	void RemoveAll(const UObject* Object)
+	{
+		if (bIsBound)
+		{
+			bIsBound = FSparseDelegateStorage::RemoveAll(GetDelegateOwner(), GetDelegateName(), Object);
+		}
+	}
+
+	/**
+	* Removes all functions from this delegate's invocation list
+	*/
+	void Clear()
+	{
+		FSparseDelegate::__Internal_Clear(GetDelegateOwner(), GetDelegateName());
+	}
+
+	template<typename... ParamTypes>
+	void Broadcast(ParamTypes... Params)
+	{
+		if (auto MCDelegate = Get())
+		{
+			MCDelegate->Broadcast(Params...);
+		}
+	}
+
+	/**
+	* Tests if a UObject instance and a UObject method address pair are already bound to this multi-cast delegate.
+	*
+	* @param	InUserObject		UObject instance
+	* @param	InMethodPtr			Member function address pointer
+	* @param	InFunctionName		Name of member function, without class name
+	* @return	True if the instance/method is already bound.
+	*
+	* NOTE:  Do not call this function directly.  Instead, call IsAlreadyBound() which is a macro proxy function that
+	*        automatically sets the function name string for the caller.
+	*/
+	template< class UserClass >
+	bool __Internal_IsAlreadyBound( UserClass* InUserObject, typename FDelegate::template TMethodPtrResolver< UserClass >::FMethodPtr InMethodPtr, FName InFunctionName ) const
+	{
+		check( InUserObject != nullptr && InMethodPtr != nullptr );
+
+		// NOTE: We're not actually using the incoming method pointer or calling it.  We simply require it for type-safety reasons.
+
+		return this->Contains( InUserObject, InFunctionName );
+	}
+
+	/**
+	* Binds a UObject instance and a UObject method address to this multi-cast delegate.
+	*
+	* @param	InUserObject		UObject instance
+	* @param	InMethodPtr			Member function address pointer
+	* @param	InFunctionName		Name of member function, without class name
+	*
+	* NOTE:  Do not call this function directly.  Instead, call AddDynamic() which is a macro proxy function that
+	*        automatically sets the function name string for the caller.
+	*/
+	template< class UserClass >
+	void __Internal_AddDynamic( UserClass* InUserObject, typename FDelegate::template TMethodPtrResolver< UserClass >::FMethodPtr InMethodPtr, FName InFunctionName )
+	{
+		check( InUserObject != nullptr && InMethodPtr != nullptr );
+
+		// NOTE: We're not actually storing the incoming method pointer or calling it.  We simply require it for type-safety reasons.
+
+		FDelegate NewDelegate;
+		NewDelegate.__Internal_BindDynamic( InUserObject, InMethodPtr, InFunctionName );
+
+		this->Add( NewDelegate );
+	}
+
+	/**
+	* Binds a UObject instance and a UObject method address to this multi-cast delegate, but only if it hasn't been bound before.
+	*
+	* @param	InUserObject		UObject instance
+	* @param	InMethodPtr			Member function address pointer
+	* @param	InFunctionName		Name of member function, without class name
+	*
+	* NOTE:  Do not call this function directly.  Instead, call AddUniqueDynamic() which is a macro proxy function that
+	*        automatically sets the function name string for the caller.
+	*/
+	template< class UserClass >
+	void __Internal_AddUniqueDynamic( UserClass* InUserObject, typename FDelegate::template TMethodPtrResolver< UserClass >::FMethodPtr InMethodPtr, FName InFunctionName )
+	{
+		check( InUserObject != nullptr && InMethodPtr != nullptr );
+
+		// NOTE: We're not actually storing the incoming method pointer or calling it.  We simply require it for type-safety reasons.
+
+		FDelegate NewDelegate;
+		NewDelegate.__Internal_BindDynamic( InUserObject, InMethodPtr, InFunctionName );
+
+		this->AddUnique( NewDelegate );
+	}
+
+	/**
+	* Unbinds a UObject instance and a UObject method address from this multi-cast delegate.
+	*
+	* @param	InUserObject		UObject instance
+	* @param	InMethodPtr			Member function address pointer
+	* @param	InFunctionName		Name of member function, without class name
+	*
+	* NOTE:  Do not call this function directly.  Instead, call RemoveDynamic() which is a macro proxy function that
+	*        automatically sets the function name string for the caller.
+	*/
+	template< class UserClass >
+	void __Internal_RemoveDynamic( UserClass* InUserObject, typename FDelegate::template TMethodPtrResolver< UserClass >::FMethodPtr InMethodPtr, FName InFunctionName )
+	{
+		check( InUserObject != nullptr && InMethodPtr != nullptr );
+
+		// NOTE: We're not actually storing the incoming method pointer or calling it.  We simply require it for type-safety reasons.
+
+		this->Remove( InUserObject, InFunctionName );
+	}
+
+private:
+	// Hide internal functions that never need to be called on the derived classes
+	void __Internal_AddUnique(const UObject*, FName, FScriptDelegate);
+	void __Internal_Remove(const UObject*, FName, const FScriptDelegate&);
+	void __Internal_Clear(const UObject*, FName);
+};
+
+#define FUNC_DECLARE_DYNAMIC_MULTICAST_SPARSE_DELEGATE(SparseDelegateClass, OwningClass, DelegateName, FuncParamList, FuncParamPassThru, ...) \
+	FUNC_DECLARE_DYNAMIC_MULTICAST_DELEGATE(FWeakObjectPtr, SparseDelegateClass##_MCSignature, SparseDelegateClass##_DelegateWrapper, FUNC_CONCAT(FuncParamList), FUNC_CONCAT(FuncParamPassThru), __VA_ARGS__) \
+	namespace { char OwningClass##Sparse##DelegateName##Str[] = #DelegateName; } \
+	struct SparseDelegateClass##OffsetGetter \
+	{ \
+		template<typename T> \
+		static size_t GetDelegateOffset() { return offsetof(T, DelegateName); } \
+	}; \
+	struct SparseDelegateClass : public TSparseDynamicDelegate<SparseDelegateClass##_MCSignature, OwningClass, OwningClass##Sparse##DelegateName##Str, SparseDelegateClass##OffsetGetter> \
+	{ \
+	};
+
+#define DECLARE_DYNAMIC_MULTICAST_SPARSE_DELEGATE( SparseDelegateClass, OwningClass, DelegateName ) BODY_MACRO_COMBINE(CURRENT_FILE_ID,_,__LINE__,_DELEGATE) FUNC_DECLARE_DYNAMIC_MULTICAST_SPARSE_DELEGATE( SparseDelegateClass, OwningClass, DelegateName, , FUNC_CONCAT( *this ), void )
+#define DECLARE_DYNAMIC_MULTICAST_SPARSE_DELEGATE_OneParam( SparseDelegateClass, OwningClass, DelegateName, Param1Type, Param1Name ) BODY_MACRO_COMBINE(CURRENT_FILE_ID,_,__LINE__,_DELEGATE) FUNC_DECLARE_DYNAMIC_MULTICAST_SPARSE_DELEGATE( SparseDelegateClass, OwningClass, DelegateName, FUNC_CONCAT( Param1Type InParam1 ), FUNC_CONCAT( *this, InParam1 ), void, Param1Type )
+#define DECLARE_DYNAMIC_MULTICAST_SPARSE_DELEGATE_TwoParams( SparseDelegateClass, OwningClass, DelegateName, Param1Type, Param1Name, Param2Type, Param2Name ) BODY_MACRO_COMBINE(CURRENT_FILE_ID,_,__LINE__,_DELEGATE) FUNC_DECLARE_DYNAMIC_MULTICAST_SPARSE_DELEGATE( SparseDelegateClass, OwningClass, DelegateName, FUNC_CONCAT( Param1Type InParam1, Param2Type InParam2 ), FUNC_CONCAT( *this, InParam1, InParam2 ), void, Param1Type, Param2Type )
+#define DECLARE_DYNAMIC_MULTICAST_SPARSE_DELEGATE_ThreeParams( SparseDelegateClass, OwningClass, DelegateName, Param1Type, Param1Name, Param2Type, Param2Name, Param3Type, Param3Name ) BODY_MACRO_COMBINE(CURRENT_FILE_ID,_,__LINE__,_DELEGATE) FUNC_DECLARE_DYNAMIC_MULTICAST_SPARSE_DELEGATE( SparseDelegateClass, OwningClass, DelegateName, FUNC_CONCAT( Param1Type InParam1, Param2Type InParam2, Param3Type InParam3 ), FUNC_CONCAT( *this, InParam1, InParam2, InParam3 ), void, Param1Type, Param2Type, Param3Type )
+#define DECLARE_DYNAMIC_MULTICAST_SPARSE_DELEGATE_FourParams( SparseDelegateClass, OwningClass, DelegateName, Param1Type, Param1Name, Param2Type, Param2Name, Param3Type, Param3Name, Param4Type, Param4Name ) BODY_MACRO_COMBINE(CURRENT_FILE_ID,_,__LINE__,_DELEGATE) FUNC_DECLARE_DYNAMIC_MULTICAST_SPARSE_DELEGATE( SparseDelegateClass, OwningClass, DelegateName, FUNC_CONCAT( Param1Type InParam1, Param2Type InParam2, Param3Type InParam3, Param4Type InParam4 ), FUNC_CONCAT( *this, InParam1, InParam2, InParam3, InParam4 ), void, Param1Type, Param2Type, Param3Type, Param4Type )
+#define DECLARE_DYNAMIC_MULTICAST_SPARSE_DELEGATE_FiveParams( SparseDelegateClass, OwningClass, DelegateName, Param1Type, Param1Name, Param2Type, Param2Name, Param3Type, Param3Name, Param4Type, Param4Name, Param5Type, Param5Name ) BODY_MACRO_COMBINE(CURRENT_FILE_ID,_,__LINE__,_DELEGATE) FUNC_DECLARE_DYNAMIC_MULTICAST_SPARSE_DELEGATE( SparseDelegateClass, OwningClass, DelegateName, FUNC_CONCAT( Param1Type InParam1, Param2Type InParam2, Param3Type InParam3, Param4Type InParam4, Param5Type InParam5 ), FUNC_CONCAT( *this, InParam1, InParam2, InParam3, InParam4, InParam5 ), void, Param1Type, Param2Type, Param3Type, Param4Type, Param5Type )
+#define DECLARE_DYNAMIC_MULTICAST_SPARSE_DELEGATE_SixParams( SparseDelegateClass, OwningClass, DelegateName, Param1Type, Param1Name, Param2Type, Param2Name, Param3Type, Param3Name, Param4Type, Param4Name, Param5Type, Param5Name, Param6Type, Param6Name ) BODY_MACRO_COMBINE(CURRENT_FILE_ID,_,__LINE__,_DELEGATE) FUNC_DECLARE_DYNAMIC_MULTICAST_SPARSE_DELEGATE( SparseDelegateClass, OwningClass, DelegateName, FUNC_CONCAT( Param1Type InParam1, Param2Type InParam2, Param3Type InParam3, Param4Type InParam4, Param5Type InParam5, Param6Type InParam6 ), FUNC_CONCAT( *this, InParam1, InParam2, InParam3, InParam4, InParam5, InParam6 ), void, Param1Type, Param2Type, Param3Type, Param4Type, Param5Type, Param6Type )
+#define DECLARE_DYNAMIC_MULTICAST_SPARSE_DELEGATE_SevenParams( SparseDelegateClass, OwningClass, DelegateName, Param1Type, Param1Name, Param2Type, Param2Name, Param3Type, Param3Name, Param4Type, Param4Name, Param5Type, Param5Name, Param6Type, Param6Name, Param7Type, Param7Name ) BODY_MACRO_COMBINE(CURRENT_FILE_ID,_,__LINE__,_DELEGATE) FUNC_DECLARE_DYNAMIC_MULTICAST_SPARSE_DELEGATE( SparseDelegateClass, OwningClass, DelegateName, FUNC_CONCAT( Param1Type InParam1, Param2Type InParam2, Param3Type InParam3, Param4Type InParam4, Param5Type InParam5, Param6Type InParam6, Param7Type InParam7 ), FUNC_CONCAT( *this, InParam1, InParam2, InParam3, InParam4, InParam5, InParam6, InParam7 ), void, Param1Type, Param2Type, Param3Type, Param4Type, Param5Type, Param6Type, Param7Type )
+#define DECLARE_DYNAMIC_MULTICAST_SPARSE_DELEGATE_EightParams( SparseDelegateClass, OwningClass, DelegateName, Param1Type, Param1Name, Param2Type, Param2Name, Param3Type, Param3Name, Param4Type, Param4Name, Param5Type, Param5Name, Param6Type, Param6Name, Param7Type, Param7Name, Param8Type, Param8Name ) BODY_MACRO_COMBINE(CURRENT_FILE_ID,_,__LINE__,_DELEGATE) FUNC_DECLARE_DYNAMIC_MULTICAST_SPARSE_DELEGATE( SparseDelegateClass, OwningClass, DelegateName, FUNC_CONCAT( Param1Type InParam1, Param2Type InParam2, Param3Type InParam3, Param4Type InParam4, Param5Type InParam5, Param6Type InParam6, Param7Type InParam7, Param8Type InParam8 ), FUNC_CONCAT( *this, InParam1, InParam2, InParam3, InParam4, InParam5, InParam6, InParam7, InParam8 ), void, Param1Type, Param2Type, Param3Type, Param4Type, Param5Type, Param6Type, Param7Type, Param8Type )
+#define DECLARE_DYNAMIC_MULTICAST_SPARSE_DELEGATE_NineParams( SparseDelegateClass, OwningClass, DelegateName, Param1Type, Param1Name, Param2Type, Param2Name, Param3Type, Param3Name, Param4Type, Param4Name, Param5Type, Param5Name, Param6Type, Param6Name, Param7Type, Param7Name, Param8Type, Param8Name, Param9Type, Param9Name ) BODY_MACRO_COMBINE(CURRENT_FILE_ID,_,__LINE__,_DELEGATE) FUNC_DECLARE_DYNAMIC_MULTICAST_SPARSE_DELEGATE( SparseDelegateClass, OwningClass, DelegateName, FUNC_CONCAT( Param1Type InParam1, Param2Type InParam2, Param3Type InParam3, Param4Type InParam4, Param5Type InParam5, Param6Type InParam6, Param7Type InParam7, Param8Type InParam8, Param9Type InParam9 ), FUNC_CONCAT( *this, InParam1, InParam2, InParam3, InParam4, InParam5, InParam6, InParam7, InParam8, InParam9 ), void, Param1Type, Param2Type, Param3Type, Param4Type, Param5Type, Param6Type, Param7Type, Param8Type, Param9Type )
diff --git a/Engine/Source/Runtime/CoreUObject/Public/UObject/UObjectGlobals.h b/Engine/Source/Runtime/CoreUObject/Public/UObject/UObjectGlobals.h
--- a/Engine/Source/Runtime/CoreUObject/Public/UObject/UObjectGlobals.h
+++ b/Engine/Source/Runtime/CoreUObject/Public/UObject/UObjectGlobals.h
@@ -2232,9 +2232,10 @@ namespace UE4CodeGen_Private
 		Set               = 0x18,
 		Struct            = 0x19,
 		Delegate          = 0x1A,
-		MulticastDelegate = 0x1B,
-		Text              = 0x1C,
-		Enum              = 0x1D,
+		InlineMulticastDelegate = 0x1B,
+		SparseMulticastDelegate = 0x1C,
+		Text              = 0x1D,
+		Enum              = 0x1E,
 
 		// Property-specific flags
 		NativeBool        = 0x20
@@ -2492,6 +2493,8 @@ namespace UE4CodeGen_Private
 		UObject*                          (*OuterFunc)();
 		UFunction*                        (*SuperFunc)();
 		const char*                         NameUTF8;
+		const char*                         OwningClassName;
+		const char*                         DelegateName;
 		SIZE_T                              StructureSize;
 		const FPropertyParamsBase* const*   PropertyArray;
 		int32                               NumProperties;
diff --git a/Engine/Source/Runtime/CoreUObject/Public/UObject/UnrealType.h b/Engine/Source/Runtime/CoreUObject/Public/UObject/UnrealType.h
--- a/Engine/Source/Runtime/CoreUObject/Public/UObject/UnrealType.h
+++ b/Engine/Source/Runtime/CoreUObject/Public/UObject/UnrealType.h
@@ -14,6 +14,7 @@
 #include "UObject/WeakObjectPtr.h"
 #include "UObject/CoreNetTypes.h"
 #include "UObject/ScriptInterface.h"
+#include "UObject/SparseDelegate.h"
 #include "Templates/Casts.h"
 #include "Templates/IsFloatingPoint.h"
 #include "Templates/IsIntegral.h"
@@ -4186,31 +4187,20 @@ public:
 -----------------------------------------------------------------------------*/
 
 /**
- * Describes a pointer to a function bound to an Object.
+ * Describes a list of functions bound to an Object.
  */
-// need to break this out a different type so that the DECLARE_CASTED_CLASS_INTRINSIC macro can digest the comma
-typedef TProperty<FMulticastScriptDelegate, UProperty> UMulticastDelegateProperty_Super;
-
-class COREUOBJECT_API UMulticastDelegateProperty : public UMulticastDelegateProperty_Super
+class COREUOBJECT_API UMulticastDelegateProperty : public UProperty
 {
-	DECLARE_CASTED_CLASS_INTRINSIC(UMulticastDelegateProperty, UMulticastDelegateProperty_Super, 0, TEXT("/Script/CoreUObject"), CASTCLASS_UMulticastDelegateProperty)
+	DECLARE_CASTED_CLASS_INTRINSIC(UMulticastDelegateProperty, UProperty, CLASS_Abstract, TEXT("/Script/CoreUObject"), CASTCLASS_UMulticastDelegateProperty)
 
 	/** Points to the source delegate function (the function declared with the delegate keyword) used in the declaration of this delegate property. */
 	UFunction* SignatureFunction;
-public:
-
-	typedef UMulticastDelegateProperty_Super::TTypeFundamentals TTypeFundamentals;
-	typedef TTypeFundamentals::TCppType TCppType;
 
-	UMulticastDelegateProperty(ECppProperty, int32 InOffset, EPropertyFlags InFlags, UFunction* InSignatureFunction = NULL)
-		: TProperty(FObjectInitializer::Get(), EC_CppProperty, InOffset, InFlags)
-		, SignatureFunction(InSignatureFunction)
-	{
-	}
+public:
 
-	UMulticastDelegateProperty( const FObjectInitializer& ObjectInitializer, ECppProperty, int32 InOffset, EPropertyFlags InFlags, UFunction* InSignatureFunction = NULL )
-		: TProperty( ObjectInitializer, EC_CppProperty, InOffset, InFlags )
-		, SignatureFunction(InSignatureFunction)
+	UMulticastDelegateProperty(const FObjectInitializer& ObjectInitializer, ECppProperty, int32 InOffset, EPropertyFlags InFlags)
+		: UProperty(ObjectInitializer, EC_CppProperty, InOffset, InFlags)
+		, SignatureFunction(nullptr)
 	{
 	}
 
@@ -4223,22 +4213,161 @@ public:
 	virtual FString GetCPPType( FString* ExtendedTypeText, uint32 CPPExportFlags ) const override;
 	virtual FString GetCPPTypeForwardDeclaration() const override;
 	virtual bool Identical( const void* A, const void* B, uint32 PortFlags ) const override;
-	virtual void SerializeItem(FStructuredArchive::FSlot Slot, void* Value, void const* Defaults) const override;
 	virtual bool NetSerializeItem( FArchive& Ar, UPackageMap* Map, void* Data, TArray<uint8> * MetaData = NULL ) const override;
 	virtual void ExportTextItem( FString& ValueStr, const void* PropertyValue, const void* DefaultValue, UObject* Parent, int32 PortFlags, UObject* ExportRootScope ) const override;
-	virtual const TCHAR* ImportText_Internal( const TCHAR* Buffer, void* Data, int32 PortFlags, UObject* OwnerObject, FOutputDevice* ErrorText ) const override;
 	virtual bool ContainsWeakObjectReference() const override;
 	virtual void InstanceSubobjects( void* Data, void const* DefaultData, UObject* Owner, struct FObjectInstancingGraph* InstanceGraph ) override;
 	virtual bool SameType(const UProperty* Other) const override;
+	virtual EConvertFromTypeResult ConvertFromType(const FPropertyTag& Tag, FStructuredArchive::FSlot Slot, uint8* Data, UStruct* DefaultsStruct) override;
 	// End of UProperty interface
 
+	virtual const FMulticastScriptDelegate* GetMulticastDelegate(const void* PropertyValue) const PURE_VIRTUAL(UMulticastDelegateProperty::GetMulticastDelegate, return nullptr;);
+	virtual void SetMulticastDelegate(void* PropertyValue, FMulticastScriptDelegate ScriptDelegate) const PURE_VIRTUAL(UMulticastDelegateProperty::SetMulticastDelegate, );
+
+	virtual void AddDelegate(FScriptDelegate ScriptDelegate, UObject* Parent) const PURE_VIRTUAL(UMulticastDelegateProperty::AddDelegate, );
+	virtual void AddDelegate(FScriptDelegate ScriptDelegate, UObject* Parent, void* PropertyValue) const PURE_VIRTUAL(UMulticastDelegateProperty::AddDelegate, );
+	virtual void RemoveDelegate(const FScriptDelegate& ScriptDelegate, UObject* Parent) const PURE_VIRTUAL(UMulticastDelegateProperty::RemoveDelegate, );
+	virtual void RemoveDelegate(const FScriptDelegate& ScriptDelegate, UObject* Parent, void* PropertyValue) const PURE_VIRTUAL(UMulticastDelegateProperty::RemoveDelegate, );
+	virtual void ClearDelegate(UObject* Parent) const PURE_VIRTUAL(UMulticastDelegateProperty::ClearDelegate, );
+	virtual void ClearDelegate(UObject* Parent, void* PropertyValue)  const PURE_VIRTUAL(UMulticastDelegateProperty::ClearDelegate, );
+
 protected:
 	friend class UProperty;
 
+	static FMulticastScriptDelegate::FInvocationList EmptyList;
+	virtual FMulticastScriptDelegate::FInvocationList& GetInvocationList(const void* PropertyValue) const PURE_VIRTUAL(UMulticastDelegateProperty::GetInvocationList, return EmptyList;);
+
+
 	const TCHAR* ImportText_Add( const TCHAR* Buffer, void* PropertyValue, int32 PortFlags, UObject* Parent, FOutputDevice* ErrorText ) const;
 	const TCHAR* ImportText_Remove( const TCHAR* Buffer, void* PropertyValue, int32 PortFlags, UObject* Parent, FOutputDevice* ErrorText ) const;
+
+	const TCHAR* ImportDelegateFromText(FMulticastScriptDelegate& MulticastDelegate, const TCHAR* Buffer, void* Data, int32 PortFlags, UObject* OwnerObject, FOutputDevice* ErrorText) const;
+};
+
+template<class InTCppType>
+class COREUOBJECT_API TProperty_MulticastDelegate : public TProperty<InTCppType, UMulticastDelegateProperty>
+{
+public:
+	typedef TProperty<InTCppType, UMulticastDelegateProperty> Super;
+	typedef InTCppType TCppType;
+	typedef typename Super::TTypeFundamentals TTypeFundamentals;
+	TProperty_MulticastDelegate(const FObjectInitializer& ObjectInitializer = FObjectInitializer::Get(), UFunction* InSignatureFunction = nullptr)
+		: Super(ObjectInitializer)
+	{
+		this->SignatureFunction = InSignatureFunction;
+	}
+
+	TProperty_MulticastDelegate(ECppProperty, int32 InOffset, EPropertyFlags InFlags, UFunction* InSignatureFunction = nullptr)
+		: Super(FObjectInitializer::Get(), EC_CppProperty, InOffset, InFlags)
+	{
+		this->SignatureFunction = InSignatureFunction;
+	}
+
+	TProperty_MulticastDelegate(const FObjectInitializer& ObjectInitializer, ECppProperty, int32 InOffset, EPropertyFlags InFlags, UFunction* InSignatureFunction = nullptr)
+		: Super(ObjectInitializer, EC_CppProperty, InOffset, InFlags)
+	{
+		this->SignatureFunction = InSignatureFunction;
+	}
+
+	/** DO NOT USE. This constructor is for internal usage only for hot-reload purposes. */
+	TProperty_MulticastDelegate(FVTableHelper& Helper) : Super(Helper) {};
+};
+
+class COREUOBJECT_API UMulticastInlineDelegateProperty : public TProperty_MulticastDelegate<FMulticastScriptDelegate>
+{
+	DECLARE_CASTED_CLASS_INTRINSIC(UMulticastInlineDelegateProperty, TProperty_MulticastDelegate<FMulticastScriptDelegate>, 0, TEXT("/Script/CoreUObject"), CASTCLASS_UMulticastInlineDelegateProperty)
+
+public:
+
+	UMulticastInlineDelegateProperty(ECppProperty, int32 InOffset, EPropertyFlags InFlags, UFunction* InSignatureFunction = nullptr)
+		: TProperty_MulticastDelegate(FObjectInitializer::Get(), EC_CppProperty, InOffset, InFlags, InSignatureFunction)
+	{
+	}
+
+	UMulticastInlineDelegateProperty(const FObjectInitializer& ObjectInitializer, ECppProperty, int32 InOffset, EPropertyFlags InFlags, UFunction* InSignatureFunction = nullptr)
+		: TProperty_MulticastDelegate(ObjectInitializer, EC_CppProperty, InOffset, InFlags, InSignatureFunction)
+	{
+	}
+
+	// UProperty interface
+	virtual void SerializeItem(FStructuredArchive::FSlot Slot, void* Value, void const* Defaults) const override;
+	virtual const TCHAR* ImportText_Internal(const TCHAR* Buffer, void* Data, int32 PortFlags, UObject* OwnerObject, FOutputDevice* ErrorText) const override;
+	// End of UProperty interface
+
+	// UMulticastDelegateProperty interface
+	virtual const FMulticastScriptDelegate* GetMulticastDelegate(const void* PropertyValue) const override;
+	virtual void SetMulticastDelegate(void* PropertyValue, FMulticastScriptDelegate ScriptDelegate) const override;
+
+	virtual void AddDelegate(FScriptDelegate ScriptDelegate, UObject* Parent) const override
+	{
+		AddDelegate(MoveTemp(ScriptDelegate), Parent, GetPropertyValuePtr_InContainer(Parent));
+	}
+
+	virtual void RemoveDelegate(const FScriptDelegate& ScriptDelegate, UObject* Parent) const override
+	{
+		RemoveDelegate(ScriptDelegate, Parent, GetPropertyValuePtr_InContainer(Parent));
+	}
+
+	virtual void ClearDelegate(UObject* Parent) const override
+	{
+		ClearDelegate(Parent, GetPropertyValuePtr_InContainer(Parent));
+	}
+
+	virtual void AddDelegate(FScriptDelegate ScriptDelegate, UObject* Parent, void* PropertyValue) const override final;
+	virtual void RemoveDelegate(const FScriptDelegate& ScriptDelegate, UObject* Parent, void* PropertyValue) const override final;
+	virtual void ClearDelegate(UObject* Parent, void* PropertyValue) const override final;
+protected:
+	virtual FMulticastScriptDelegate::FInvocationList& GetInvocationList(const void* PropertyValue) const;
+	// End of UMulticastDelegateProperty interface
 };
 
+class COREUOBJECT_API UMulticastSparseDelegateProperty : public TProperty_MulticastDelegate<FSparseDelegate> 
+{
+	DECLARE_CASTED_CLASS_INTRINSIC(UMulticastSparseDelegateProperty, TProperty_MulticastDelegate<FSparseDelegate>, 0, TEXT("/Script/CoreUObject"), CASTCLASS_UMulticastSparseDelegateProperty)
+
+public:
+
+	UMulticastSparseDelegateProperty(ECppProperty, int32 InOffset, EPropertyFlags InFlags, UFunction* InSignatureFunction = nullptr)
+		: TProperty_MulticastDelegate(FObjectInitializer::Get(), EC_CppProperty, InOffset, InFlags, InSignatureFunction)
+	{
+	}
+
+	UMulticastSparseDelegateProperty(const FObjectInitializer& ObjectInitializer, ECppProperty, int32 InOffset, EPropertyFlags InFlags, UFunction* InSignatureFunction = nullptr)
+		: TProperty_MulticastDelegate(ObjectInitializer, EC_CppProperty, InOffset, InFlags, InSignatureFunction)
+	{
+	}
+
+	// UProperty interface
+	virtual void SerializeItem(FStructuredArchive::FSlot Slot, void* Value, void const* Defaults) const override;
+	virtual const TCHAR* ImportText_Internal(const TCHAR* Buffer, void* Data, int32 PortFlags, UObject* OwnerObject, FOutputDevice* ErrorText) const override;
+	// End of UProperty interface
+
+	// UMulticastDelegateProperty interface
+	virtual const FMulticastScriptDelegate* GetMulticastDelegate(const void* PropertyValue) const override;
+	virtual void SetMulticastDelegate(void* PropertyValue, FMulticastScriptDelegate ScriptDelegate) const override;
+
+	virtual void AddDelegate(FScriptDelegate ScriptDelegate, UObject* Parent) const override
+	{
+		AddDelegate(MoveTemp(ScriptDelegate), Parent, GetPropertyValuePtr_InContainer(Parent));
+	}
+
+	virtual void RemoveDelegate(const FScriptDelegate& ScriptDelegate, UObject* Parent) const override
+	{
+		RemoveDelegate(ScriptDelegate, Parent, GetPropertyValuePtr_InContainer(Parent));
+	}
+
+	virtual void ClearDelegate(UObject* Parent) const override
+	{
+		ClearDelegate(Parent, GetPropertyValuePtr_InContainer(Parent));
+	}
+
+	virtual void AddDelegate(FScriptDelegate ScriptDelegate, UObject* Parent, void* PropertyValue) const override final;
+	virtual void RemoveDelegate(const FScriptDelegate& ScriptDelegate, UObject* Parent, void* PropertyValue) const override final;
+	virtual void ClearDelegate(UObject* Parent, void* PropertyValue) const override final;
+protected:
+	virtual FMulticastScriptDelegate::FInvocationList& GetInvocationList(const void* PropertyValue) const;
+	// End of UMulticastDelegateProperty interface
+};
 
 /** Describes a single node in a custom property list. */
 struct COREUOBJECT_API FCustomPropertyListNode
diff --git a/Engine/Source/Runtime/Engine/Classes/Engine/ComponentDelegateBinding.h b/Engine/Source/Runtime/Engine/Classes/Engine/ComponentDelegateBinding.h
--- a/Engine/Source/Runtime/Engine/Classes/Engine/ComponentDelegateBinding.h
+++ b/Engine/Source/Runtime/Engine/Classes/Engine/ComponentDelegateBinding.h
@@ -48,8 +48,4 @@ class ENGINE_API UComponentDelegateBinding
 	virtual void UnbindDynamicDelegates(UObject* InInstance) const override;
 	virtual void UnbindDynamicDelegatesForProperty(UObject* InInstance, const UObjectProperty* InObjectProperty) const override;
 	//~ End DynamicBlueprintBinding Interface
-
-private:
-	// Utility method used to find the target delegate given an instance and a binding descriptor
-	static FMulticastScriptDelegate* FindComponentTargetDelegate(const UObject* InInstance, const FBlueprintComponentDelegateBinding& InBinding, const UObjectProperty* InObjectProperty = nullptr);
 };
diff --git a/Engine/Source/Runtime/Engine/Private/ComponentDelegateBinding.cpp b/Engine/Source/Runtime/Engine/Private/ComponentDelegateBinding.cpp
--- a/Engine/Source/Runtime/Engine/Private/ComponentDelegateBinding.cpp
+++ b/Engine/Source/Runtime/Engine/Private/ComponentDelegateBinding.cpp
@@ -9,45 +9,31 @@ UComponentDelegateBinding::UComponentDelegateBinding(const FObjectInitializer& O
 {
 }
 
-FMulticastScriptDelegate* UComponentDelegateBinding::FindComponentTargetDelegate(const UObject* InInstance, const FBlueprintComponentDelegateBinding& InBinding, const UObjectProperty* InObjectProperty)
-{
-	// Get the property that points to the component
-	const UObjectProperty* ObjProp = (InObjectProperty != nullptr && InObjectProperty->GetFName() == InBinding.ComponentPropertyName) ? InObjectProperty : FindField<UObjectProperty>(InInstance->GetClass(), InBinding.ComponentPropertyName);
-	if(ObjProp != nullptr)
-	{
-		// ..see if there is actually a component assigned
-		UObject* Component = ObjProp->GetObjectPropertyValue_InContainer(InInstance);
-		if(Component != nullptr)
-		{
-			// If there is, find and return the delegate property on it
-			UMulticastDelegateProperty* DelegateProp = FindField<UMulticastDelegateProperty>(Component->GetClass(), InBinding.DelegatePropertyName);
-			if(DelegateProp != nullptr)
-			{
-				return DelegateProp->GetPropertyValuePtr_InContainer(Component);
-			}
-		}
-	}
-
-	return nullptr;
-}
-
 void UComponentDelegateBinding::BindDynamicDelegates(UObject* InInstance) const
 {
 	for(int32 BindIdx=0; BindIdx<ComponentDelegateBindings.Num(); BindIdx++)
 	{
 		const FBlueprintComponentDelegateBinding& Binding = ComponentDelegateBindings[BindIdx];
 
-		// Get the delegate property on the component we want to bind to
-		FMulticastScriptDelegate* TargetDelegate = FindComponentTargetDelegate(InInstance, Binding);
-		// Get the function we want to bind
-		UFunction* FunctionToBind = FindField<UFunction>(InInstance->GetClass(), Binding.FunctionNameToBind);
-		// If we have both of those..
-		if(TargetDelegate != nullptr && FunctionToBind != nullptr)
+		// Get the property that points to the component
+		if (const UObjectProperty* ObjProp = FindField<UObjectProperty>(InInstance->GetClass(), Binding.ComponentPropertyName))
 		{
-			// Bind function on the instance to this delegate
-			FScriptDelegate Delegate;
-			Delegate.BindUFunction(InInstance, Binding.FunctionNameToBind);
-			TargetDelegate->AddUnique(Delegate);
+			// ..see if there is actually a component assigned
+			if (UObject* Component = ObjProp->GetObjectPropertyValue_InContainer(InInstance))
+			{
+				// If there is, find and return the delegate property on it
+				if (UMulticastDelegateProperty* MulticastDelegateProp = FindField<UMulticastDelegateProperty>(Component->GetClass(), Binding.DelegatePropertyName))
+				{
+					// Get the function we want to bind
+					if (UFunction* FunctionToBind = InInstance->GetClass()->FindFunctionByName(Binding.FunctionNameToBind))
+					{
+						// Bind function on the instance to this delegate
+						FScriptDelegate Delegate;
+						Delegate.BindUFunction(InInstance, Binding.FunctionNameToBind);
+						MulticastDelegateProp->AddDelegate(MoveTemp(Delegate), Component);
+					}
+				}
+			}
 		}
 	}
 }
@@ -58,12 +44,21 @@ void UComponentDelegateBinding::UnbindDynamicDelegates(UObject* InInstance) cons
 	{
 		const FBlueprintComponentDelegateBinding& Binding = ComponentDelegateBindings[BindIdx];
 
-		// Get the delegate property on the component we want to unbind from
-		FMulticastScriptDelegate* TargetDelegate = FindComponentTargetDelegate(InInstance, Binding);
-		if (TargetDelegate != nullptr)
+		// Get the property that points to the component
+		if (const UObjectProperty* ObjProp = FindField<UObjectProperty>(InInstance->GetClass(), Binding.ComponentPropertyName))
 		{
-			// Unbind function on the instance from this delegate
-			TargetDelegate->Remove(InInstance, Binding.FunctionNameToBind);
+			// ..see if there is actually a component assigned
+			if (UObject* Component = ObjProp->GetObjectPropertyValue_InContainer(InInstance))
+			{
+				// If there is, find and return the delegate property on it
+				if (UMulticastDelegateProperty* MulticastDelegateProp = FindField<UMulticastDelegateProperty>(Component->GetClass(), Binding.DelegatePropertyName))
+				{
+					// Unbind function on the instance to this delegate
+					FScriptDelegate Delegate;
+					Delegate.BindUFunction(InInstance, Binding.FunctionNameToBind);
+					MulticastDelegateProp->RemoveDelegate(Delegate, Component);
+				}
+			}
 		}
 	}
 }
@@ -73,18 +68,25 @@ void UComponentDelegateBinding::UnbindDynamicDelegatesForProperty(UObject* InIns
 	for(int32 BindIdx=0; BindIdx<ComponentDelegateBindings.Num(); BindIdx++)
 	{
 		const FBlueprintComponentDelegateBinding& Binding = ComponentDelegateBindings[BindIdx];
-		const UObjectProperty* ObjProp = FindField<UObjectProperty>(InInstance->GetClass(), Binding.ComponentPropertyName);
-		if(ObjProp == InObjectProperty)
+		if (InObjectProperty->GetFName() == Binding.ComponentPropertyName)
 		{
-			// Get the delegate property on the component we want to unbind from
-			FMulticastScriptDelegate* TargetDelegate = FindComponentTargetDelegate(InInstance, Binding, ObjProp);
-			if(TargetDelegate != nullptr)
+			const UObjectProperty* ObjProp = FindField<UObjectProperty>(InInstance->GetClass(), Binding.ComponentPropertyName);
+			if (ObjProp == InObjectProperty)
 			{
-				// Unbind function on the instance from this delegate
-				TargetDelegate->Remove(InInstance, Binding.FunctionNameToBind);
+				// ..see if there is actually a component assigned
+				if (UObject* Component = ObjProp->GetObjectPropertyValue_InContainer(InInstance))
+				{
+					// If there is, find and return the delegate property on it
+					if (UMulticastDelegateProperty* MulticastDelegateProp = FindField<UMulticastDelegateProperty>(Component->GetClass(), Binding.DelegatePropertyName))
+					{
+						// Unbind function on the instance from this delegate
+						FScriptDelegate Delegate;
+						Delegate.BindUFunction(InInstance, Binding.FunctionNameToBind);
+						MulticastDelegateProp->RemoveDelegate(Delegate, Component);
+					}
+				}
+				break;
 			}
-
-			break;
 		}
 	}
 }
