Provide a comment at the check that will occur if a sparse delegate is used on the stack so user has better understanding of why it is checking.
#rnx

diff --git a/Engine/Source/Runtime/CoreUObject/Public/UObject/SparseDelegate.h b/Engine/Source/Runtime/CoreUObject/Public/UObject/SparseDelegate.h
--- a/Engine/Source/Runtime/CoreUObject/Public/UObject/SparseDelegate.h
+++ b/Engine/Source/Runtime/CoreUObject/Public/UObject/SparseDelegate.h
@@ -172,7 +172,7 @@ private:
 		const size_t OffsetToOwner = DelegateOffsetClass::template GetDelegateOffset<OwningClass>();
 		check(OffsetToOwner);
 		UObject* DelegateOwner = reinterpret_cast<UObject*>((uint8*)this - OffsetToOwner);
-		check(DelegateOwner->IsValidLowLevelFast(false));
+		check(DelegateOwner->IsValidLowLevelFast(false)); // Most likely the delegate is trying to be used on the stack, in an object it wasn't defined for, or for a class member with a different name than it was defined for. It is only valid for a sparse delegate to be used for the exact class/property name it is defined with.
 		return DelegateOwner;
 	}
 
