Override GetCPPType in the multicast delegate TProperty specialization as the new hierarchy was using the TProperty version instead of the UMulticastDelegateProperty version.
This fixes nativization crash generating code for a multicast delegate.
#rnx

diff --git a/Engine/Source/Runtime/CoreUObject/Public/UObject/UnrealType.h b/Engine/Source/Runtime/CoreUObject/Public/UObject/UnrealType.h
--- a/Engine/Source/Runtime/CoreUObject/Public/UObject/UnrealType.h
+++ b/Engine/Source/Runtime/CoreUObject/Public/UObject/UnrealType.h
@@ -4271,6 +4271,13 @@ public:
 
 	/** DO NOT USE. This constructor is for internal usage only for hot-reload purposes. */
 	TProperty_MulticastDelegate(FVTableHelper& Helper) : Super(Helper) {};
+
+	// UProperty interface.
+	virtual FString GetCPPType(FString* ExtendedTypeText, uint32 CPPExportFlags) const override
+	{
+		return UMulticastDelegateProperty::GetCPPType(ExtendedTypeText, CPPExportFlags);
+	}
+	// End of UProperty interface
 };
 
 class COREUOBJECT_API UMulticastInlineDelegateProperty : public TProperty_MulticastDelegate<FMulticastScriptDelegate>
