Attempt to fix linux CIS
#rnx

diff --git a/Engine/Source/Runtime/CoreUObject/Public/UObject/SparseDelegate.h b/Engine/Source/Runtime/CoreUObject/Public/UObject/SparseDelegate.h
--- a/Engine/Source/Runtime/CoreUObject/Public/UObject/SparseDelegate.h
+++ b/Engine/Source/Runtime/CoreUObject/Public/UObject/SparseDelegate.h
@@ -167,9 +167,10 @@ protected:
 	}
 
 private:
+	template<typename OC=OwningClass>
 	UObject* GetDelegateOwner() const
 	{
-		const size_t OffsetToOwner = DelegateOffsetClass::template GetDelegateOffset<OwningClass>();
+		const size_t OffsetToOwner = DelegateOffsetClass::template GetDelegateOffset<OC>();
 		check(OffsetToOwner);
 		UObject* DelegateOwner = reinterpret_cast<UObject*>((uint8*)this - OffsetToOwner);
 		check(DelegateOwner->IsValidLowLevelFast(false)); // Most likely the delegate is trying to be used on the stack, in an object it wasn't defined for, or for a class member with a different name than it was defined for. It is only valid for a sparse delegate to be used for the exact class/property name it is defined with.
