A second attempt at soothing Linux build errors
#rnx

diff --git a/Engine/Source/Runtime/CoreUObject/Private/UObject/SparseDelegate.cpp b/Engine/Source/Runtime/CoreUObject/Private/UObject/SparseDelegate.cpp
--- a/Engine/Source/Runtime/CoreUObject/Private/UObject/SparseDelegate.cpp
+++ b/Engine/Source/Runtime/CoreUObject/Private/UObject/SparseDelegate.cpp
@@ -1,7 +1,6 @@
 // Copyright 1998-2019 Epic Games, Inc. All Rights Reserved.
 
 #include "UObject/SparseDelegate.h"
-#include "UObject/Object.h"
 #include "UObject/Class.h"
 #include "UObject/Package.h"
 #include "HAL/IConsoleManager.h"
diff --git a/Engine/Source/Runtime/CoreUObject/Public/UObject/SparseDelegate.h b/Engine/Source/Runtime/CoreUObject/Public/UObject/SparseDelegate.h
--- a/Engine/Source/Runtime/CoreUObject/Public/UObject/SparseDelegate.h
+++ b/Engine/Source/Runtime/CoreUObject/Public/UObject/SparseDelegate.h
@@ -3,6 +3,7 @@
 #pragma once
 
 #include "CoreMinimal.h"
+#include "UObject/Object.h"
 #include "UObject/UObjectAnnotation.h"
 #include "UObject/WeakObjectPtr.h"
 #include "Delegates/Delegate.h"
@@ -167,10 +168,9 @@ protected:
 	}
 
 private:
-	template<typename OC=OwningClass>
 	UObject* GetDelegateOwner() const
 	{
-		const size_t OffsetToOwner = DelegateOffsetClass::template GetDelegateOffset<OC>();
+		const size_t OffsetToOwner = DelegateOffsetClass::template GetDelegateOffset<OwningClass>();
 		check(OffsetToOwner);
 		UObject* DelegateOwner = reinterpret_cast<UObject*>((uint8*)this - OffsetToOwner);
 		check(DelegateOwner->IsValidLowLevelFast(false)); // Most likely the delegate is trying to be used on the stack, in an object it wasn't defined for, or for a class member with a different name than it was defined for. It is only valid for a sparse delegate to be used for the exact class/property name it is defined with.
