Rework sparse delegate storage to not use uobject annotation and manage its own map and lock state. This eliminates unnnecessary annotation cache overhead and ensures entire operation on map is completed while the lock is in place rather than copying per-object map out of the annotation and then moving modified inner per-object map back in to the annotation as two separately locked actions.
#rnx

diff --git a/Engine/Source/Runtime/CoreUObject/Private/UObject/PropertyMulticastDelegate.cpp b/Engine/Source/Runtime/CoreUObject/Private/UObject/PropertyMulticastDelegate.cpp
--- a/Engine/Source/Runtime/CoreUObject/Private/UObject/PropertyMulticastDelegate.cpp
+++ b/Engine/Source/Runtime/CoreUObject/Private/UObject/PropertyMulticastDelegate.cpp
@@ -194,7 +194,7 @@ void UMulticastDelegateProperty::ExportTextItem( FString& ValueStr, const void*
 	ValueStr += TEXT( ")" );
 }
 
-const TCHAR* UMulticastDelegateProperty::ImportDelegateFromText( FMulticastScriptDelegate& MulticastDelegate, const TCHAR* Buffer, void* PropertyValue, int32 PortFlags, UObject* Parent, FOutputDevice* ErrorText ) const
+const TCHAR* UMulticastDelegateProperty::ImportDelegateFromText( FMulticastScriptDelegate& MulticastDelegate, const TCHAR* Buffer, UObject* Parent, FOutputDevice* ErrorText ) const
 {
 	// Multi-cast delegates always expect an opening parenthesis when using assignment syntax, so that
 	// users don't accidentally blow away already-bound delegates in DefaultProperties.  This also helps
@@ -371,7 +371,7 @@ void UMulticastInlineDelegateProperty::SerializeItem(FStructuredArchive::FSlot S
 const TCHAR* UMulticastInlineDelegateProperty::ImportText_Internal(const TCHAR* Buffer, void* PropertyValue, int32 PortFlags, UObject* Parent, FOutputDevice* ErrorText) const
 {
 	FMulticastScriptDelegate& MulticastDelegate = (*(FMulticastScriptDelegate*)PropertyValue);
-	return ImportDelegateFromText(MulticastDelegate, Buffer, PropertyValue, PortFlags, Parent, ErrorText);
+	return ImportDelegateFromText(MulticastDelegate, Buffer, Parent, ErrorText);
 }
 
 void ResolveDelegateReference(const UMulticastInlineDelegateProperty* InlineProperty, UObject*& Parent, void*& PropertyValue)
@@ -436,14 +436,16 @@ void UMulticastSparseDelegateProperty::SetMulticastDelegate(void* PropertyValue,
 
 	USparseDelegateFunction* SparseDelegateFunc = CastChecked<USparseDelegateFunction>(SignatureFunction);
 	UObject* OwningObject = FSparseDelegateStorage::ResolveSparseOwner(SparseDelegate, SparseDelegateFunc->OwningClassName, SparseDelegateFunc->DelegateName);
-	FMulticastScriptDelegate& Delegate = FSparseDelegateStorage::GetOrCreateMulticastDelegate(OwningObject, SparseDelegateFunc->DelegateName);
 
-	Delegate = MoveTemp(ScriptDelegate);
-
-	SparseDelegate.bIsBound = Delegate.IsBound();
-	if (!SparseDelegate.bIsBound)
+	if (ScriptDelegate.IsBound())
+	{
+		FSparseDelegateStorage::SetMulticastDelegate(OwningObject, SparseDelegateFunc->DelegateName, MoveTemp(ScriptDelegate));
+		SparseDelegate.bIsBound = true;
+	}
+	else if (SparseDelegate.bIsBound)
 	{
 		FSparseDelegateStorage::Clear(OwningObject, SparseDelegateFunc->DelegateName);
+		SparseDelegate.bIsBound = false;
 	}
 }
 
@@ -480,8 +482,7 @@ void UMulticastSparseDelegateProperty::SerializeItem(FStructuredArchive::FSlot S
 		{
 			USparseDelegateFunction* SparseDelegateFunc = CastChecked<USparseDelegateFunction>(SignatureFunction);
 			UObject* OwningObject = FSparseDelegateStorage::ResolveSparseOwner(SparseDelegate, SparseDelegateFunc->OwningClassName, SparseDelegateFunc->DelegateName);
-			FMulticastScriptDelegate& BoundDelegate = FSparseDelegateStorage::GetOrCreateMulticastDelegate(OwningObject, SparseDelegateFunc->DelegateName);
-			BoundDelegate = MoveTemp(Delegate);
+			FSparseDelegateStorage::SetMulticastDelegate(OwningObject, SparseDelegateFunc->DelegateName, MoveTemp(Delegate));
 			SparseDelegate.bIsBound = true;
 		}
 		else if (SparseDelegate.bIsBound)
@@ -516,19 +517,23 @@ void UMulticastSparseDelegateProperty::SerializeItem(FStructuredArchive::FSlot S
 
 const TCHAR* UMulticastSparseDelegateProperty::ImportText_Internal(const TCHAR* Buffer, void* PropertyValue, int32 PortFlags, UObject* Parent, FOutputDevice* ErrorText) const
 {
-	FSparseDelegate& SparseDelegate = *(FSparseDelegate*)PropertyValue;
-
-	USparseDelegateFunction* SparseDelegateFunc = CastChecked<USparseDelegateFunction>(SignatureFunction);
-	FMulticastScriptDelegate& Delegate = FSparseDelegateStorage::GetOrCreateMulticastDelegate(Parent, SparseDelegateFunc->DelegateName);
-
-	const TCHAR* Result = ImportDelegateFromText(Delegate, Buffer, PropertyValue, PortFlags, Parent, ErrorText);
+	FMulticastScriptDelegate Delegate;
+	const TCHAR* Result = ImportDelegateFromText(Delegate, Buffer, Parent, ErrorText);
 
 	if (Result)
 	{
-		SparseDelegate.bIsBound = Delegate.IsBound();
-		if (!SparseDelegate.bIsBound)
+		FSparseDelegate& SparseDelegate = *(FSparseDelegate*)PropertyValue;
+		USparseDelegateFunction* SparseDelegateFunc = CastChecked<USparseDelegateFunction>(SignatureFunction);
+
+		if (Delegate.IsBound())
+		{
+			FSparseDelegateStorage::SetMulticastDelegate(Parent, SparseDelegateFunc->DelegateName, MoveTemp(Delegate));
+			SparseDelegate.bIsBound = true;
+		}
+		else
 		{
 			FSparseDelegateStorage::Clear(Parent, SparseDelegateFunc->DelegateName);
+			SparseDelegate.bIsBound = false;
 		}
 	}
 
diff --git a/Engine/Source/Runtime/CoreUObject/Private/UObject/SparseDelegate.cpp b/Engine/Source/Runtime/CoreUObject/Private/UObject/SparseDelegate.cpp
--- a/Engine/Source/Runtime/CoreUObject/Private/UObject/SparseDelegate.cpp
+++ b/Engine/Source/Runtime/CoreUObject/Private/UObject/SparseDelegate.cpp
@@ -5,9 +5,36 @@
 #include "UObject/Package.h"
 #include "HAL/IConsoleManager.h"
 
-FUObjectAnnotationSparse<FSparseDelegateStorage::FSparseDelegateMap, true> FSparseDelegateStorage::SparseDelegates;
+FSparseDelegateStorage::FObjectListener FSparseDelegateStorage::SparseDelegateObjectListener;
+FCriticalSection FSparseDelegateStorage::SparseDelegateMapCritical;
+TMap<const UObjectBase*, FSparseDelegateStorage::FSparseDelegateMap> FSparseDelegateStorage::SparseDelegates;
 TMap<TPair<FName,FName>, size_t> FSparseDelegateStorage::SparseDelegateObjectOffsets;
 
+FSparseDelegateStorage::FObjectListener::~FObjectListener()
+{
+	DisableListener();
+}
+
+void FSparseDelegateStorage::FObjectListener::NotifyUObjectDeleted(const UObjectBase* Object, int32 Index)
+{
+	FScopeLock SparseDelegateMapLock(&FSparseDelegateStorage::SparseDelegateMapCritical);
+	FSparseDelegateStorage::SparseDelegates.Remove(Object);
+	if (FSparseDelegateStorage::SparseDelegates.Num() == 0)
+	{
+		DisableListener();
+	}
+}
+
+void FSparseDelegateStorage::FObjectListener::EnableListener()
+{
+	GUObjectArray.AddUObjectDeleteListener(this);
+}
+
+void FSparseDelegateStorage::FObjectListener::DisableListener()
+{
+	GUObjectArray.RemoveUObjectDeleteListener(this);
+}
+
 void FSparseDelegateStorage::RegisterDelegateOffset(const UObject* OwningObject, const FName DelegateName, const size_t DelegateOffsetToOwner)
 {
 	check(OwningObject);
@@ -45,27 +72,36 @@ UObject* FSparseDelegateStorage::ResolveSparseOwner(const FSparseDelegate& Spars
 
 FMulticastScriptDelegate* FSparseDelegateStorage::GetMulticastDelegate(const UObject* DelegateOwner, const FName DelegateName)
 {
-	FSparseDelegateMap DelegateMap = SparseDelegates.GetAnnotation(DelegateOwner);
-	if (TSharedPtr<FMulticastScriptDelegate>* MulticastDelegatePtr = DelegateMap.Find(DelegateName))
+	FScopeLock SparseDelegateMapLock(&SparseDelegateMapCritical);
+
+	if (FSparseDelegateMap* DelegateMap = SparseDelegates.Find(DelegateOwner))
 	{
-		return MulticastDelegatePtr->Get();
+		if (TSharedPtr<FMulticastScriptDelegate>* MulticastDelegatePtr = DelegateMap->Find(DelegateName))
+		{
+			return MulticastDelegatePtr->Get();
+		}
 	}
-
 	return nullptr;
 }
 
-FMulticastScriptDelegate& FSparseDelegateStorage::GetOrCreateMulticastDelegate(const UObject* DelegateOwner, const FName DelegateName)
+void FSparseDelegateStorage::SetMulticastDelegate(const UObject* DelegateOwner, const FName DelegateName, FMulticastScriptDelegate Delegate)
 {
-	FSparseDelegateMap DelegateMap = SparseDelegates.GetAnnotation(DelegateOwner);
+	FScopeLock SparseDelegateMapLock(&SparseDelegateMapCritical);
+
+	if (SparseDelegates.Num() == 0)
+	{
+		SparseDelegateObjectListener.EnableListener();
+	}
+
+	FSparseDelegateMap& DelegateMap = SparseDelegates.FindOrAdd(DelegateOwner);
 	TSharedPtr<FMulticastScriptDelegate>& MulticastDelegate = DelegateMap.FindOrAdd(DelegateName);
 
 	if (!MulticastDelegate.IsValid())
 	{
 		MulticastDelegate = MakeShared<FMulticastScriptDelegate>();
-		SparseDelegates.AddAnnotation(DelegateOwner, MoveTemp(DelegateMap));
 	}
 
-	return *MulticastDelegate;
+	*MulticastDelegate = MoveTemp(Delegate);
 }
 
 bool FSparseDelegateStorage::Add(const UObject* DelegateOwner, const FName DelegateName, FScriptDelegate Delegate)
@@ -73,7 +109,14 @@ bool FSparseDelegateStorage::Add(const UObject* DelegateOwner, const FName Deleg
 	bool bDelegateWasBound = false;
 	if (Delegate.IsBound())
 	{
-		FSparseDelegateMap DelegateMap = SparseDelegates.GetAnnotation(DelegateOwner);
+		FScopeLock SparseDelegateMapLock(&SparseDelegateMapCritical);
+
+		if (SparseDelegates.Num() == 0)
+		{
+			SparseDelegateObjectListener.EnableListener();
+		}
+
+		FSparseDelegateMap& DelegateMap = SparseDelegates.FindOrAdd(DelegateOwner);
 		TSharedPtr<FMulticastScriptDelegate>& MulticastDelegate = DelegateMap.FindOrAdd(DelegateName);
 
 		if (!MulticastDelegate.IsValid())
@@ -82,7 +125,6 @@ bool FSparseDelegateStorage::Add(const UObject* DelegateOwner, const FName Deleg
 		}
 
 		MulticastDelegate->Add(MoveTemp(Delegate));
-		SparseDelegates.AddAnnotation(DelegateOwner, MoveTemp(DelegateMap));
 		bDelegateWasBound = true;
 	}
 	return bDelegateWasBound;
@@ -93,7 +135,14 @@ bool FSparseDelegateStorage::AddUnique(const UObject* DelegateOwner, const FName
 	bool bDelegateWasBound = false;
 	if (Delegate.IsBound())
 	{
-		FSparseDelegateMap DelegateMap = SparseDelegates.GetAnnotation(DelegateOwner);
+		FScopeLock SparseDelegateMapLock(&SparseDelegateMapCritical);
+
+		if (SparseDelegates.Num() == 0)
+		{
+			SparseDelegateObjectListener.EnableListener();
+		}
+
+		FSparseDelegateMap& DelegateMap = SparseDelegates.FindOrAdd(DelegateOwner);
 		TSharedPtr<FMulticastScriptDelegate>& MulticastDelegate = DelegateMap.FindOrAdd(DelegateName);
 
 		if (!MulticastDelegate.IsValid())
@@ -102,7 +151,6 @@ bool FSparseDelegateStorage::AddUnique(const UObject* DelegateOwner, const FName
 		}
 
 		MulticastDelegate->AddUnique(MoveTemp(Delegate));
-		SparseDelegates.AddAnnotation(DelegateOwner, MoveTemp(DelegateMap));
 		bDelegateWasBound = true;
 	}
 	return bDelegateWasBound;
@@ -111,12 +159,16 @@ bool FSparseDelegateStorage::AddUnique(const UObject* DelegateOwner, const FName
 bool FSparseDelegateStorage::Contains(const UObject* DelegateOwner, const FName DelegateName, const FScriptDelegate& Delegate)
 {
 	bool bContainsDelegate = false;
-	FSparseDelegateMap DelegateMap = SparseDelegates.GetAnnotation(DelegateOwner);
-	if (TSharedPtr<FMulticastScriptDelegate>* MulticastDelegatePtr = DelegateMap.Find(DelegateName))
+	FScopeLock SparseDelegateMapLock(&SparseDelegateMapCritical);
+
+	if (FSparseDelegateMap* DelegateMap = SparseDelegates.Find(DelegateOwner))
 	{
-		if (FMulticastScriptDelegate* MulticastDelegate = MulticastDelegatePtr->Get())
+		if (TSharedPtr<FMulticastScriptDelegate>* MulticastDelegatePtr = DelegateMap->Find(DelegateName))
 		{
-			bContainsDelegate = MulticastDelegate->Contains(Delegate);
+			if (FMulticastScriptDelegate* MulticastDelegate = MulticastDelegatePtr->Get())
+			{
+				bContainsDelegate = MulticastDelegate->Contains(Delegate);
+			}
 		}
 	}
 
@@ -126,12 +178,16 @@ bool FSparseDelegateStorage::Contains(const UObject* DelegateOwner, const FName
 bool FSparseDelegateStorage::Contains(const UObject* DelegateOwner, const FName DelegateName, const UObject* InObject, FName InFunctionName)
 {
 	bool bContainsDelegate = false;
-	FSparseDelegateMap DelegateMap = SparseDelegates.GetAnnotation(DelegateOwner);
-	if (TSharedPtr<FMulticastScriptDelegate>* MulticastDelegatePtr = DelegateMap.Find(DelegateName))
+	FScopeLock SparseDelegateMapLock(&SparseDelegateMapCritical);
+
+	if (FSparseDelegateMap* DelegateMap = SparseDelegates.Find(DelegateOwner))
 	{
-		if (FMulticastScriptDelegate* MulticastDelegate = MulticastDelegatePtr->Get())
+		if (TSharedPtr<FMulticastScriptDelegate>* MulticastDelegatePtr = DelegateMap->Find(DelegateName))
 		{
-			bContainsDelegate = MulticastDelegate->Contains(InObject, InFunctionName);
+			if (FMulticastScriptDelegate* MulticastDelegate = MulticastDelegatePtr->Get())
+			{
+				bContainsDelegate = MulticastDelegate->Contains(InObject, InFunctionName);
+			}
 		}
 	}
 	return bContainsDelegate;
@@ -140,31 +196,35 @@ bool FSparseDelegateStorage::Contains(const UObject* DelegateOwner, const FName
 bool FSparseDelegateStorage::Remove(const UObject* DelegateOwner, const FName DelegateName, const FScriptDelegate& Delegate)
 {
 	bool bSparseDelegateBound = false;
-	FSparseDelegateMap DelegateMap = SparseDelegates.GetAnnotation(DelegateOwner);
-	if (TSharedPtr<FMulticastScriptDelegate>* MulticastDelegatePtr = DelegateMap.Find(DelegateName))
+	FScopeLock SparseDelegateMapLock(&SparseDelegateMapCritical);
+
+	if (FSparseDelegateMap* DelegateMap = SparseDelegates.Find(DelegateOwner))
 	{
-		if (FMulticastScriptDelegate* MulticastDelegate = MulticastDelegatePtr->Get())
+		if (TSharedPtr<FMulticastScriptDelegate>* MulticastDelegatePtr = DelegateMap->Find(DelegateName))
 		{
-			MulticastDelegate->Remove(Delegate);
-			bSparseDelegateBound = MulticastDelegate->IsBound();
-			if (!bSparseDelegateBound)
+			if (FMulticastScriptDelegate* MulticastDelegate = MulticastDelegatePtr->Get())
 			{
-				DelegateMap.Remove(DelegateName);
+				MulticastDelegate->Remove(Delegate);
+				bSparseDelegateBound = MulticastDelegate->IsBound();
+				if (!bSparseDelegateBound)
+				{
+					DelegateMap->Remove(DelegateName);
+				}
+			}
+			else
+			{
+				DelegateMap->Remove(DelegateName);
 			}
 		}
-		else
+		if (DelegateMap->Num() == 0)
 		{
-			DelegateMap.Remove(DelegateName);
+			SparseDelegates.Remove(DelegateOwner);
+			if (SparseDelegates.Num() == 0)
+			{
+				SparseDelegateObjectListener.DisableListener();
+			}
 		}
 	}
-	if (DelegateMap.IsDefault())
-	{
-		SparseDelegates.RemoveAnnotation(DelegateOwner);
-	}
-	else
-	{
-		SparseDelegates.AddAnnotation(DelegateOwner, MoveTemp(DelegateMap));
-	}
 
 	return bSparseDelegateBound;
 }
@@ -172,31 +232,35 @@ bool FSparseDelegateStorage::Remove(const UObject* DelegateOwner, const FName De
 bool FSparseDelegateStorage::Remove(const UObject* DelegateOwner, const FName DelegateName, const UObject* InObject, FName InFunctionName)
 {
 	bool bSparseDelegateBound = false;
-	FSparseDelegateMap DelegateMap = SparseDelegates.GetAnnotation(DelegateOwner);
-	if (TSharedPtr<FMulticastScriptDelegate>* MulticastDelegatePtr = DelegateMap.Find(DelegateName))
+	FScopeLock SparseDelegateMapLock(&SparseDelegateMapCritical);
+
+	if (FSparseDelegateMap* DelegateMap = SparseDelegates.Find(DelegateOwner))
 	{
-		if (FMulticastScriptDelegate* MulticastDelegate = MulticastDelegatePtr->Get())
+		if (TSharedPtr<FMulticastScriptDelegate>* MulticastDelegatePtr = DelegateMap->Find(DelegateName))
 		{
-			MulticastDelegate->Remove(InObject, InFunctionName);
-			bSparseDelegateBound = MulticastDelegate->IsBound();
-			if (!bSparseDelegateBound)
+			if (FMulticastScriptDelegate* MulticastDelegate = MulticastDelegatePtr->Get())
 			{
-				DelegateMap.Remove(DelegateName);
+				MulticastDelegate->Remove(InObject, InFunctionName);
+				bSparseDelegateBound = MulticastDelegate->IsBound();
+				if (!bSparseDelegateBound)
+				{
+					DelegateMap->Remove(DelegateName);
+				}
+			}
+			else
+			{
+				DelegateMap->Remove(DelegateName);
 			}
 		}
-		else
+		if (DelegateMap->Num() == 0)
 		{
-			DelegateMap.Remove(DelegateName);
+			SparseDelegates.Remove(DelegateOwner);
+			if (SparseDelegates.Num() == 0)
+			{
+				SparseDelegateObjectListener.DisableListener();
+			}
 		}
 	}
-	if (DelegateMap.IsDefault())
-	{
-		SparseDelegates.RemoveAnnotation(DelegateOwner);
-	}
-	else
-	{
-		SparseDelegates.AddAnnotation(DelegateOwner, MoveTemp(DelegateMap));
-	}
 
 	return bSparseDelegateBound;
 }
@@ -204,52 +268,59 @@ bool FSparseDelegateStorage::Remove(const UObject* DelegateOwner, const FName De
 bool FSparseDelegateStorage::RemoveAll(const UObject* DelegateOwner, const FName DelegateName, const UObject* UserObject)
 {
 	bool bSparseDelegateBound = false;
-	FSparseDelegateMap DelegateMap = SparseDelegates.GetAnnotation(DelegateOwner);
-	if (TSharedPtr<FMulticastScriptDelegate>* MulticastDelegatePtr = DelegateMap.Find(DelegateName))
+	FScopeLock SparseDelegateMapLock(&SparseDelegateMapCritical);
+
+	if (FSparseDelegateMap* DelegateMap = SparseDelegates.Find(DelegateOwner))
 	{
-		if (FMulticastScriptDelegate* MulticastDelegate = MulticastDelegatePtr->Get())
+		if (TSharedPtr<FMulticastScriptDelegate>* MulticastDelegatePtr = DelegateMap->Find(DelegateName))
 		{
-			MulticastDelegate->RemoveAll(UserObject);
-			bSparseDelegateBound = MulticastDelegate->IsBound();
-			if (!bSparseDelegateBound)
+			if (FMulticastScriptDelegate* MulticastDelegate = MulticastDelegatePtr->Get())
 			{
-				DelegateMap.Remove(DelegateName);
+				MulticastDelegate->RemoveAll(UserObject);
+				bSparseDelegateBound = MulticastDelegate->IsBound();
+				if (!bSparseDelegateBound)
+				{
+					DelegateMap->Remove(DelegateName);
+				}
+			}
+			else
+			{
+				DelegateMap->Remove(DelegateName);
 			}
 		}
-		else
+		if (DelegateMap->Num() == 0)
 		{
-			DelegateMap.Remove(DelegateName);
+			SparseDelegates.Remove(DelegateOwner);
+			if (SparseDelegates.Num() == 0)
+			{
+				SparseDelegateObjectListener.DisableListener();
+			}
 		}
 	}
-	if (DelegateMap.IsDefault())
-	{
-		SparseDelegates.RemoveAnnotation(DelegateOwner);
-	}
-	else
-	{
-		SparseDelegates.AddAnnotation(DelegateOwner, MoveTemp(DelegateMap));
-	}
 	return bSparseDelegateBound;
 }
 
 void FSparseDelegateStorage::Clear(const UObject* DelegateOwner, const FName DelegateName)
 {
-	FSparseDelegateMap DelegateMap = SparseDelegates.GetAnnotation(DelegateOwner);
-	if (TSharedPtr<FMulticastScriptDelegate>* MulticastDelegatePtr = DelegateMap.Find(DelegateName))
+	FScopeLock SparseDelegateMapLock(&SparseDelegateMapCritical);
+	if (FSparseDelegateMap* DelegateMap = SparseDelegates.Find(DelegateOwner))
 	{
-		if (FMulticastScriptDelegate* MulticastDelegate = MulticastDelegatePtr->Get())
+		if (TSharedPtr<FMulticastScriptDelegate>* MulticastDelegatePtr = DelegateMap->Find(DelegateName))
 		{
-			MulticastDelegate->Clear();
+			if (FMulticastScriptDelegate* MulticastDelegate = MulticastDelegatePtr->Get())
+			{
+				MulticastDelegate->Clear();
+			}
+			DelegateMap->Remove(DelegateName);
+		}
+		if (DelegateMap->Num() == 0)
+		{
+			SparseDelegates.Remove(DelegateOwner);
+			if (SparseDelegates.Num() == 0)
+			{
+				SparseDelegateObjectListener.DisableListener();
+			}
 		}
-		DelegateMap.Remove(DelegateName);
-	}
-	if (DelegateMap.IsDefault())
-	{
-		SparseDelegates.RemoveAnnotation(DelegateOwner);
-	}
-	else
-	{
-		SparseDelegates.AddAnnotation(DelegateOwner, MoveTemp(DelegateMap));
 	}
 }
 
@@ -316,32 +387,36 @@ void FSparseDelegateStorage::SparseDelegateReport(const TArray<FString>& Args, U
 	TArray<FString> Details;
 	uint32 BoundObjects = 0;
 	uint32 BoundDelegates = 0;
-	for (const TPair<const UObjectBase*, FSparseDelegateMap>& ObjectAnnotation : SparseDelegates.GetAnnotationMap())
+
 	{
-		const UObject* Object = static_cast<const UObject*>(ObjectAnnotation.Key);
-		if (ObjectName.IsNone() || Object->GetFName() == ObjectName)
+		FScopeLock SparseDelegateMapLock(&SparseDelegateMapCritical);
+		for (const TPair<const UObjectBase*, FSparseDelegateMap>& ObjectAnnotation : SparseDelegates)
 		{
-			if (ObjectType == nullptr || Object->IsA(ObjectType))
+			const UObject* Object = static_cast<const UObject*>(ObjectAnnotation.Key);
+			if (ObjectName.IsNone() || Object->GetFName() == ObjectName)
 			{
-				if (!bSummary)
+				if (ObjectType == nullptr || Object->IsA(ObjectType))
 				{
-					Details.Add(FString::Printf(TEXT("%s"), *Object->GetPathName()));
-				}
-				++BoundObjects;
-				if (DelegateName.IsNone())
-				{
-					BoundDelegates += ObjectAnnotation.Value.Num();
 					if (!bSummary)
 					{
-						for (const TPair<FName, TSharedPtr<FMulticastScriptDelegate>>& BoundDelegate : ObjectAnnotation.Value)
+						Details.Add(FString::Printf(TEXT("%s"), *Object->GetPathName()));
+					}
+					++BoundObjects;
+					if (DelegateName.IsNone())
+					{
+						BoundDelegates += ObjectAnnotation.Value.Num();
+						if (!bSummary)
 						{
-							Details.Add(FString::Printf(TEXT("   %s"), *BoundDelegate.Key.ToString()));
+							for (const TPair<FName, TSharedPtr<FMulticastScriptDelegate>>& BoundDelegate : ObjectAnnotation.Value)
+							{
+								Details.Add(FString::Printf(TEXT("   %s"), *BoundDelegate.Key.ToString()));
+							}
 						}
 					}
-				}
-				else if (ObjectAnnotation.Value.Contains(DelegateName))
-				{
-					++BoundDelegates;
+					else if (ObjectAnnotation.Value.Contains(DelegateName))
+					{
+						++BoundDelegates;
+					}
 				}
 			}
 		}
diff --git a/Engine/Source/Runtime/CoreUObject/Public/UObject/SparseDelegate.h b/Engine/Source/Runtime/CoreUObject/Public/UObject/SparseDelegate.h
--- a/Engine/Source/Runtime/CoreUObject/Public/UObject/SparseDelegate.h
+++ b/Engine/Source/Runtime/CoreUObject/Public/UObject/SparseDelegate.h
@@ -4,9 +4,9 @@
 
 #include "CoreMinimal.h"
 #include "UObject/Object.h"
-#include "UObject/UObjectAnnotation.h"
 #include "UObject/WeakObjectPtr.h"
 #include "Delegates/Delegate.h"
+#include "Misc/ScopeLock.h"
 
 struct FSparseDelegate;
 
@@ -52,8 +52,8 @@ public:
 	/** Acquires the actual Multicast Delegate from the annotation if any delegates are bound to it. Will be null if no entry exists in the annotation for this object/delegatename. */
 	static FMulticastScriptDelegate* GetMulticastDelegate(const UObject* DelegateOwner, const FName DelegateName);
 
-	/** Acquires the actual Multicast Delegate from the annotation if any delegates are bound to it. Will be created if no entry exists in the annotation for this object/delegatename. */
-	static FMulticastScriptDelegate& GetOrCreateMulticastDelegate(const UObject* DelegateOwner, const FName DelegateName);
+	/** Directly sets the Multicast Delegate for this object/delegatename pair. If the delegate is unbound it will be assigned/inserted anyways. */
+	static void SetMulticastDelegate(const UObject* DelegateOwner, const FName DelegateName, FMulticastScriptDelegate Delegate);
 
 	/** Using the registry of sparse delegates recover the FSparseDelegate address from the UObject and name. */
 	static FSparseDelegate* ResolveSparseDelegate(const UObject* OwningObject, FName DelegateName);
@@ -66,13 +66,29 @@ public:
 
 private:
 
-	struct FSparseDelegateMap : public TMap<FName, TSharedPtr<FMulticastScriptDelegate>>
+	struct FObjectListener : public FUObjectArray::FUObjectDeleteListener
 	{
-		bool IsDefault() const { return Num() == 0; }
+		virtual ~FObjectListener();
+		virtual void NotifyUObjectDeleted(const UObjectBase* Object, int32 Index) override;
+		void EnableListener();
+		void DisableListener();
 	};
 
-	static FUObjectAnnotationSparse<FSparseDelegateMap, true> SparseDelegates;
+	/** Allow the object listener to use the critical section and remove objects from the map */
+	friend struct FObjectListener;
 
+	/** A listener to get notified when objects have been deleted and remove them from the map */
+	static FObjectListener SparseDelegateObjectListener;
+
+	/** Critical Section for locking access to the sparse delegate map */
+	static FCriticalSection SparseDelegateMapCritical;
+
+	/** Delegate map is a map of Delegate names to a shared pointer of the multicast script delegate */
+	typedef TMap<FName, TSharedPtr<FMulticastScriptDelegate>> FSparseDelegateMap;
+
+	/** Map of objects to the map of delegates that are bound to that object */
+	static TMap<const UObjectBase*, FSparseDelegateMap> SparseDelegates;
+	
 	/** Sparse delegate offsets are indexed by ActorClass/DelegateName pair */
 	static TMap<TPair<FName, FName>, size_t> SparseDelegateObjectOffsets;
 };
diff --git a/Engine/Source/Runtime/CoreUObject/Public/UObject/UnrealType.h b/Engine/Source/Runtime/CoreUObject/Public/UObject/UnrealType.h
--- a/Engine/Source/Runtime/CoreUObject/Public/UObject/UnrealType.h
+++ b/Engine/Source/Runtime/CoreUObject/Public/UObject/UnrealType.h
@@ -4238,7 +4238,7 @@ protected:
 	const TCHAR* ImportText_Add( const TCHAR* Buffer, void* PropertyValue, int32 PortFlags, UObject* Parent, FOutputDevice* ErrorText ) const;
 	const TCHAR* ImportText_Remove( const TCHAR* Buffer, void* PropertyValue, int32 PortFlags, UObject* Parent, FOutputDevice* ErrorText ) const;
 
-	const TCHAR* ImportDelegateFromText(FMulticastScriptDelegate& MulticastDelegate, const TCHAR* Buffer, void* Data, int32 PortFlags, UObject* OwnerObject, FOutputDevice* ErrorText) const;
+	const TCHAR* ImportDelegateFromText(FMulticastScriptDelegate& MulticastDelegate, const TCHAR* Buffer, UObject* OwnerObject, FOutputDevice* ErrorText) const;
 };
 
 template<class InTCppType>
