Fixing crash when the engine tries to open content saved with a newer engine version

diff --git a/Engine/Source/Runtime/CoreUObject/Private/UObject/Linker.cpp b/Engine/Source/Runtime/CoreUObject/Private/UObject/Linker.cpp
--- a/Engine/Source/Runtime/CoreUObject/Private/UObject/Linker.cpp
+++ b/Engine/Source/Runtime/CoreUObject/Private/UObject/Linker.cpp
@@ -715,9 +715,6 @@ FLinkerLoad* GetPackageLinker
 	// Create new linker.
 	if( !Result )
 	{
-		// @todo: RM so never create the context below?
-		//check(InExistingContext && InExistingContext->HasStartedLoading());
-
 		// we will already have found the filename above
 		check(NewFilename.Len() > 0);
 		TRefCountPtr<FUObjectSerializeContext> LoadContext(FUObjectThreadContext::Get().GetSerializeContext());
