Avoid resetting bExportLoadFailed when an object is destroyed, allows exports to stay invalidated after a GC call. Also we now invalidate exports for nodes replaced via ReplaceOldNodeWithNew

diff --git a/Engine/Source/Editor/BlueprintGraph/Private/EdGraphSchema_K2.cpp b/Engine/Source/Editor/BlueprintGraph/Private/EdGraphSchema_K2.cpp
--- a/Engine/Source/Editor/BlueprintGraph/Private/EdGraphSchema_K2.cpp
+++ b/Engine/Source/Editor/BlueprintGraph/Private/EdGraphSchema_K2.cpp
@@ -6112,7 +6112,8 @@ bool UEdGraphSchema_K2::ReplaceOldNodeWithNew(UK2Node* OldNode, UK2Node* NewNode
 		NewNode->NodeComment = OldNode->NodeComment;
 		NewNode->bCommentBubblePinned = OldNode->bCommentBubblePinned;
 		NewNode->bCommentBubbleVisible = OldNode->bCommentBubbleVisible;
-
+		
+		FLinkerLoad::InvalidateExport(OldNode);
 		OldNode->DestroyNode();
 	}
 	return !bFailedToFindPin;
diff --git a/Engine/Source/Runtime/CoreUObject/Private/UObject/Obj.cpp b/Engine/Source/Runtime/CoreUObject/Private/UObject/Obj.cpp
--- a/Engine/Source/Runtime/CoreUObject/Private/UObject/Obj.cpp
+++ b/Engine/Source/Runtime/CoreUObject/Private/UObject/Obj.cpp
@@ -753,9 +753,30 @@ void UObject::BeginDestroy()
 	}
 
 	LowLevelRename(NAME_None);
+	
+#if WITH_EDITORONLY_DATA
+	// Make sure the linker entry stays as 'bExportLoadFailed' if the entry was marked as such, 
+	// doing this prevents the object from being reloaded by subsequent load calls:
+	FLinkerLoad* Linker = GetLinker();
+	const int32 CachedLinkerIndex = GetLinkerIndex();
+	bool bLinkerEntryWasInvalid = false;
+	if(Linker && Linker->ExportMap.IsValidIndex(CachedLinkerIndex))
+	{
+		FObjectExport& ObjExport = Linker->ExportMap[CachedLinkerIndex];
+		bLinkerEntryWasInvalid = ObjExport.bExportLoadFailed;
+	}
+#endif // WITH_EDITORONLY_DATA
 
 	// Remove from linker's export table.
 	SetLinker( NULL, INDEX_NONE );
+	
+#if WITH_EDITORONLY_DATA
+	if(bLinkerEntryWasInvalid)
+	{
+		FObjectExport& ObjExport = Linker->ExportMap[CachedLinkerIndex];
+		ObjExport.bExportLoadFailed = true;
+	}
+#endif // WITH_EDITORONLY_DATA
 
 	// ensure BeginDestroy has been routed back to UObject::BeginDestroy.
 #if !(UE_BUILD_SHIPPING || UE_BUILD_TEST)
