Fix crash on shutdown depending on order of globals being destroyed
#rnx

diff --git a/Engine/Source/Runtime/CoreUObject/Private/UObject/SparseDelegate.cpp b/Engine/Source/Runtime/CoreUObject/Private/UObject/SparseDelegate.cpp
--- a/Engine/Source/Runtime/CoreUObject/Private/UObject/SparseDelegate.cpp
+++ b/Engine/Source/Runtime/CoreUObject/Private/UObject/SparseDelegate.cpp
@@ -12,7 +12,11 @@ TMap<TPair<FName,FName>, size_t> FSparseDelegateStorage::SparseDelegateObjectOff
 
 FSparseDelegateStorage::FObjectListener::~FObjectListener()
 {
-	DisableListener();
+	// Destroy order might result in GUObjectArray or its critical section being invalid so don't disable since we're shutting down anyways
+	if (!GIsRequestingExit)
+	{
+		DisableListener();
+	}
 }
 
 void FSparseDelegateStorage::FObjectListener::NotifyUObjectDeleted(const UObjectBase* Object, int32 Index)
