#UE4 Fix a bug that caused logs about missing assets to not show their referencer

diff --git a/Engine/Source/Runtime/CoreUObject/Private/UObject/LinkerLoad.cpp b/Engine/Source/Runtime/CoreUObject/Private/UObject/LinkerLoad.cpp
--- a/Engine/Source/Runtime/CoreUObject/Private/UObject/LinkerLoad.cpp
+++ b/Engine/Source/Runtime/CoreUObject/Private/UObject/LinkerLoad.cpp
@@ -2717,7 +2717,8 @@ bool FLinkerLoad::VerifyImportInner(const int32 ImportIndex, FString& WarningSuf
 		// Get the linker if the package hasn't been fully loaded already.
 		if (!bWasFullyLoaded)
 		{
-			Import.SourceLinker = GetPackageLinker( TmpPkg, NULL, InternalLoadFlags, NULL, NULL );
+			FUObjectSerializeContext* SerializeContext = GetSerializeContext();
+			Import.SourceLinker = GetPackageLinker( TmpPkg, nullptr, InternalLoadFlags, nullptr, nullptr, nullptr, &SerializeContext);
 #if WITH_EDITORONLY_DATA
 			if (Import.SourceLinker && !TmpPkg->HasAnyFlags(RF_WasLoaded))
 			{
