Fixing crashes in pkg info commandlet

diff --git a/Engine/Source/Editor/UnrealEd/Private/Commandlets/PackageUtilities.cpp b/Engine/Source/Editor/UnrealEd/Private/Commandlets/PackageUtilities.cpp
--- a/Engine/Source/Editor/UnrealEd/Private/Commandlets/PackageUtilities.cpp
+++ b/Engine/Source/Editor/UnrealEd/Private/Commandlets/PackageUtilities.cpp
@@ -1556,7 +1556,7 @@ int32 UPkgInfoCommandlet::Main( const FString& Params )
 			TRefCountPtr<FUObjectSerializeContext> LoadContext(FUObjectThreadContext::Get().GetSerializeContext());
 			BeginLoad(LoadContext);
 			Linker = CreateLinkerForFilename(LoadContext, Filename);
-			EndLoad(Linker->GetSerializeContext());
+			EndLoad(Linker ? Linker->GetSerializeContext() : LoadContext.GetReference());
 		}
 		else
 		{
diff --git a/Engine/Source/Runtime/CoreUObject/Private/Serialization/ArchiveStackTrace.cpp b/Engine/Source/Runtime/CoreUObject/Private/Serialization/ArchiveStackTrace.cpp
--- a/Engine/Source/Runtime/CoreUObject/Private/Serialization/ArchiveStackTrace.cpp
+++ b/Engine/Source/Runtime/CoreUObject/Private/Serialization/ArchiveStackTrace.cpp
@@ -1356,8 +1356,7 @@ FArchiveStackTraceReader::FArchiveStackTraceReader(const TCHAR* InFilename, cons
 void FArchiveStackTraceReader::Serialize(void* OutData, int64 Num)
 {
 	bool bAddData = true;
-	check(GetSerializeContext());
-	FSerializeData NewData(Tell(), Num, GetSerializeContext()->SerializedObject, GetSerializedProperty());
+	FSerializeData NewData(Tell(), Num, ThreadContext.GetSerializeContext()->SerializedObject, GetSerializedProperty());
 	if (SerializeTrace.Num())
 	{
 		FSerializeData& Last = SerializeTrace.Last();
