#UE4 Improve logging when assets fail to load. Failing to get the package linker will now produce fewer redundant log messages and all failure messages will include the asset that failed to load as well as the asset that was requesting the linker where possible.

diff --git a/Engine/Source/Runtime/CoreUObject/Private/UObject/Linker.cpp b/Engine/Source/Runtime/CoreUObject/Private/UObject/Linker.cpp
--- a/Engine/Source/Runtime/CoreUObject/Private/UObject/Linker.cpp
+++ b/Engine/Source/Runtime/CoreUObject/Private/UObject/Linker.cpp
@@ -384,7 +384,7 @@ void DeleteLoader(FLinkerLoad* Loader)
 	FLinkerManager::Get().RemoveLinker(Loader);
 }
 
-static void LogGetPackageLinkerError(FArchive* LinkerArchive, FUObjectSerializeContext* LoadContext, const TCHAR* InFilename, const FText& InFullErrorMessage, const FText& InSummaryErrorMessage, UObject* InOuter, uint32 LoadFlags)
+static void LogGetPackageLinkerError(FArchive* LinkerArchive, FUObjectSerializeContext* LoadContext, const TCHAR* InFilename, const FText& InErrorMessage, UObject* InOuter, uint32 LoadFlags)
 {
 	static FName NAME_LoadErrors("LoadErrors");
 	struct Local
@@ -417,24 +417,48 @@ static void LogGetPackageLinkerError(FArchive* LinkerArchive, FUObjectSerializeC
 		}
 	};
 
-	FMessageLog LoadErrors(NAME_LoadErrors);
+	FLinkerLoad* SerializedPackageLinker = LoadContext ? LoadContext->SerializedPackageLinker : nullptr;
+	UObject* SerializedObject = LoadContext ? LoadContext->SerializedObject : nullptr;
+	FString LoadingFile = InFilename ? InFilename : InOuter ? *InOuter->GetName() : TEXT("NULL");
+	
+	FFormatNamedArguments Arguments;
+	Arguments.Add(TEXT("LoadingFile"), FText::FromString(LoadingFile));
+	Arguments.Add(TEXT("ErrorMessage"), InErrorMessage);
 
-	// Display log error regardless LoadFlag settings
-	SET_WARN_COLOR(COLOR_RED);
-	if (LoadFlags & (LOAD_NoWarn | LOAD_Quiet))
+	FText FullErrorMessage;
+	if (SerializedPackageLinker || SerializedObject)
 	{
-		UE_LOG(LogLinker, Log, TEXT("%s"), *InFullErrorMessage.ToString());
+		FLinkerLoad* LinkerToUse = SerializedPackageLinker;
+		if (!LinkerToUse)
+		{
+			LinkerToUse = SerializedObject->GetLinker();
+		}
+		FString LoadedByFile = LinkerToUse ? *LinkerToUse->Filename : SerializedObject->GetOutermost()->GetName();
+		Arguments.Add(TEXT("LoadedByFile"), FText::FromString(LoadedByFile));
+		FullErrorMessage = FText::Format(LOCTEXT("FailedLoadWithLoadedBy", "Failed to load '{LoadingFile}' while loading '{LoadedByFile}': {ErrorMessage}"), Arguments);
 	}
 	else
 	{
-		UE_LOG(LogLinker, Warning, TEXT("%s"), *InFullErrorMessage.ToString());
+		FullErrorMessage = FText::Format(LOCTEXT("FailedLoad", "Failed to load '{LoadingFile}': {ErrorMessage}"), Arguments);
 	}
-	CLEAR_WARN_COLOR();
+
+	FMessageLog LoadErrors(NAME_LoadErrors);
+
 	if( GIsEditor && !IsRunningCommandlet() )
 	{
 		// if we don't want to be warned, skip the load warning
-		if (!(LoadFlags & (LOAD_NoWarn | LOAD_Quiet)))
+		// Display log error regardless LoadFlag settings
+		if (LoadFlags & (LOAD_NoWarn | LOAD_Quiet))
 		{
+			SET_WARN_COLOR(COLOR_RED);
+			UE_LOG(LogLinker, Log, TEXT("%s"), *FullErrorMessage.ToString());
+			CLEAR_WARN_COLOR();
+		}
+		else
+		{
+			SET_WARN_COLOR(COLOR_RED);
+			UE_LOG(LogLinker, Warning, TEXT("%s"), *FullErrorMessage.ToString());
+			CLEAR_WARN_COLOR();
 			// we only want to output errors that content creators will be able to make sense of,
 			// so any errors we cant get links out of we will just let be output to the output log (above)
 			// rather than clog up the message log
@@ -455,7 +479,7 @@ static void LogGetPackageLinkerError(FArchive* LinkerArchive, FUObjectSerializeC
 				TSharedRef<FTokenizedMessage> Message = LoadErrors.Error();
 				Message->AddToken(FAssetNameToken::Create(PackageName));
 				Message->AddToken(FTextToken::Create(FText::FromString(TEXT(":"))));
-				Message->AddToken(FTextToken::Create(InSummaryErrorMessage));
+				Message->AddToken(FTextToken::Create(FullErrorMessage));
 				Message->AddToken(FAssetNameToken::Create(OuterPackageName));
 			}
 
@@ -464,25 +488,34 @@ static void LogGetPackageLinkerError(FArchive* LinkerArchive, FUObjectSerializeC
 	}
 	else
 	{
-		if (!(LoadFlags & (LOAD_NoWarn | LOAD_Quiet)))
-		{
-			Local::OutputErrorDetail(LinkerArchive, LoadContext, NAME_LoadErrors);
-		}
-
-		FFormatNamedArguments Arguments;
-		Arguments.Add(TEXT("FileName"), FText::FromString(InFilename ? InFilename : InOuter ? *InOuter->GetName() : TEXT("NULL")));
-		Arguments.Add(TEXT("ErrorMessage"), InFullErrorMessage);
-		const FText Error = FText::Format(LOCTEXT("FailedLoad", "Failed to load '{FileName}': {ErrorMessage}"), Arguments);
-
+		bool bLogMessageEmitted = false;
 		// @see ResavePackagesCommandlet
 		if( FParse::Param(FCommandLine::Get(),TEXT("SavePackagesThatHaveFailedLoads")) == true )
 		{
-			LoadErrors.Warning(Error);
+			LoadErrors.Warning(FullErrorMessage);
 		}
 		else
 		{
 			// Gracefully handle missing packages
-			SafeLoadError(InOuter, LoadFlags, *Error.ToString());
+			bLogMessageEmitted = SafeLoadError(InOuter, LoadFlags, *FullErrorMessage.ToString());
+		}
+
+		// Only print out the message if it was not already handled by SafeLoadError
+		if (!bLogMessageEmitted)
+		{
+			if (LoadFlags & (LOAD_NoWarn | LOAD_Quiet))
+			{
+				SET_WARN_COLOR(COLOR_RED);
+				UE_LOG(LogLinker, Log, TEXT("%s"), *FullErrorMessage.ToString());
+				CLEAR_WARN_COLOR();
+			}
+			else
+			{
+				SET_WARN_COLOR(COLOR_RED);
+				UE_LOG(LogLinker, Warning, TEXT("%s"), *FullErrorMessage.ToString());
+				CLEAR_WARN_COLOR();
+				Local::OutputErrorDetail(LinkerArchive, LoadContext, NAME_LoadErrors);
+			}
 		}
 	}
 }
@@ -572,8 +605,7 @@ FLinkerLoad* GetPackageLinker
 		if( !InOuter )
 		{
 			// try to recover from this instead of throwing, it seems recoverable just by doing this
-			FText ErrorText(LOCTEXT("PackageResolveFailed", "Can't resolve asset name"));
-			LogGetPackageLinkerError(Result, InExistingContext, InLongPackageName, ErrorText, ErrorText, InOuter, LoadFlags);
+			LogGetPackageLinkerError(Result, InExistingContext, InLongPackageName, LOCTEXT("PackageResolveFailed", "Can't resolve asset name"), InOuter, LoadFlags);
 			return nullptr;
 		}
 	
@@ -601,15 +633,7 @@ FLinkerLoad* GetPackageLinker
 			// In memory-only packages have no linker and this is ok.
 			if (!(LoadFlags & LOAD_AllowDll) && !InOuter->HasAnyPackageFlags(PKG_InMemoryOnly) && !FLinkerLoad::IsKnownMissingPackage(InOuter->GetFName()))
 			{
-				FFormatNamedArguments Arguments;
-				FLinkerLoad* SerializedPackageLinker = InExistingContext ? InExistingContext->SerializedPackageLinker : nullptr;
-				Arguments.Add(TEXT("AssetName"), FText::FromString(PackageNameToLoad));
-				Arguments.Add(TEXT("PackageName"), FText::FromString(SerializedPackageLinker ? *SerializedPackageLinker->Filename : TEXT("NULL")));
-				LogGetPackageLinkerError(Result, InExistingContext, SerializedPackageLinker ? *SerializedPackageLinker->Filename : nullptr,
-											FText::Format(LOCTEXT("PackageNotFound", "Can't find file for asset '{AssetName}' while loading {PackageName}."), Arguments),
-											LOCTEXT("PackageNotFoundShort", "Can't find file for asset."),
-											InOuter,
-											LoadFlags);
+				LogGetPackageLinkerError(Result, InExistingContext, InLongPackageName, LOCTEXT("PackageNotFoundShort", "Can't find file."), InOuter, LoadFlags);
 			}
 
 			return nullptr;
@@ -621,8 +645,7 @@ FLinkerLoad* GetPackageLinker
 		if (!FPackageName::TryConvertFilenameToLongPackageName(InLongPackageName, PackageNameToCreate))
 		{
 			// try to recover from this instead of throwing, it seems recoverable just by doing this
-			FText ErrorText(LOCTEXT("PackageResolveFailed", "Can't resolve asset name"));
-			LogGetPackageLinkerError(Result, InExistingContext, InLongPackageName, ErrorText, ErrorText, InOuter, LoadFlags);
+			LogGetPackageLinkerError(Result, InExistingContext, InLongPackageName, LOCTEXT("PackageResolveFailed", "Can't resolve asset name"), InOuter, LoadFlags);
 			return nullptr;
 		}
 
@@ -657,11 +680,8 @@ FLinkerLoad* GetPackageLinker
 		{
 			if (!FLinkerLoad::IsKnownMissingPackage(InLongPackageName))
 			{
-				FFormatNamedArguments Arguments;
-				Arguments.Add(TEXT("Filename"), FText::FromString(InLongPackageName));
-
 				// try to recover from this instead of throwing, it seems recoverable just by doing this
-				LogGetPackageLinkerError(Result, InExistingContext, InLongPackageName, FText::Format(LOCTEXT("FileNotFound", "Can't find file '{Filename}'"), Arguments), LOCTEXT("FileNotFoundShort", "Can't find file"), InOuter, LoadFlags);
+				LogGetPackageLinkerError(Result, InExistingContext, InLongPackageName, LOCTEXT("FileNotFoundShort", "Can't find file."), InOuter, LoadFlags);
 			}
 			return nullptr;
 		}
@@ -684,9 +704,7 @@ FLinkerLoad* GetPackageLinker
 		{
 			if( !FilenamePkg )
 			{
-				FFormatNamedArguments Arguments;
-				Arguments.Add(TEXT("Filename"), FText::FromString(InLongPackageName));
-				LogGetPackageLinkerError(Result, InExistingContext, InLongPackageName, FText::Format(LOCTEXT("FilenameToPackage", "Can't convert filename '{Filename}' to asset name"), Arguments), LOCTEXT("FilenameToPackageShort", "Can't convert filename to asset name"), InOuter, LoadFlags);
+				LogGetPackageLinkerError(Result, InExistingContext, InLongPackageName, LOCTEXT("FilenameToPackageShort", "Can't convert filename to asset name"), InOuter, LoadFlags);
 				return nullptr;
 			}
 			InOuter = FilenamePkg;
@@ -704,10 +722,7 @@ FLinkerLoad* GetPackageLinker
 	// Make sure the package is accessible in the sandbox.
 	if( Sandbox && !Sandbox->SupportsPackage(InOuter) )
 	{
-		FFormatNamedArguments Arguments;
-		Arguments.Add(TEXT("AssetName"), FText::FromString(InOuter->GetName()));
-
-		LogGetPackageLinkerError(Result, InExistingContext, InLongPackageName, FText::Format(LOCTEXT("Sandbox", "Asset '{AssetName}' is not accessible in this sandbox"), Arguments), LOCTEXT("SandboxShort", "Asset is not accessible in this sandbox"), InOuter, LoadFlags);
+		LogGetPackageLinkerError(Result, InExistingContext, InLongPackageName, LOCTEXT("SandboxShort", "Asset is not accessible in this sandbox"), InOuter, LoadFlags);
 		return nullptr;
 	}
 #endif
@@ -755,11 +770,8 @@ FLinkerLoad* GetPackageLinker
 	// Verify compatibility.
 	if (Result && CompatibleGuid && Result->Summary.Guid != *CompatibleGuid)
 	{
-		FFormatNamedArguments Arguments;
-		Arguments.Add(TEXT("AssetName"), FText::FromString(InOuter->GetName()));
-
 		// This should never fire, because FindPackageFile should never return an incompatible file
-		LogGetPackageLinkerError(Result, InExistingContext, InLongPackageName, FText::Format(LOCTEXT("PackageVersion", "Asset '{AssetName}' version mismatch"), Arguments), LOCTEXT("PackageVersionShort", "Asset version mismatch"), InOuter, LoadFlags);
+		LogGetPackageLinkerError(Result, InExistingContext, InLongPackageName, LOCTEXT("PackageVersionShort", "Asset version mismatch"), InOuter, LoadFlags);
 		return nullptr;
 	}
 
diff --git a/Engine/Source/Runtime/CoreUObject/Private/UObject/UObjectGlobals.cpp b/Engine/Source/Runtime/CoreUObject/Private/UObject/UObjectGlobals.cpp
--- a/Engine/Source/Runtime/CoreUObject/Private/UObject/UObjectGlobals.cpp
+++ b/Engine/Source/Runtime/CoreUObject/Private/UObject/UObjectGlobals.cpp
@@ -513,13 +513,15 @@ void StaticTick( float DeltaTime, bool bUseFullTimeLimit, float AsyncLoadingTime
 -----------------------------------------------------------------------------*/
 
 //
-// Safe load error-handling.
+// Safe load error-handling. Returns true if a message was emitted.
 //
-void SafeLoadError( UObject* Outer, uint32 LoadFlags, const TCHAR* ErrorMessage)
+bool SafeLoadError( UObject* Outer, uint32 LoadFlags, const TCHAR* ErrorMessage)
 {
+	bool bRetVal = false;
 	if( FParse::Param( FCommandLine::Get(), TEXT("TREATLOADWARNINGSASERRORS") ) == true )
 	{
 		UE_LOG(LogUObjectGlobals, Error, TEXT("%s"), ErrorMessage);
+		bRetVal = true;
 	}
 	else
 	{
@@ -527,8 +529,11 @@ void SafeLoadError( UObject* Outer, uint32 LoadFlags, const TCHAR* ErrorMessage)
 		if( (LoadFlags & LOAD_Quiet) == 0 && (LoadFlags & LOAD_NoWarn) == 0)
 		{ 
 			UE_LOG(LogUObjectGlobals, Warning, TEXT("%s"), ErrorMessage);
+			bRetVal = true;
 		}
 	}
+
+	return bRetVal;
 }
 
 UPackage* FindPackage( UObject* InOuter, const TCHAR* PackageName )
diff --git a/Engine/Source/Runtime/CoreUObject/Public/UObject/UObjectGlobals.h b/Engine/Source/Runtime/CoreUObject/Public/UObject/UObjectGlobals.h
--- a/Engine/Source/Runtime/CoreUObject/Public/UObject/UObjectGlobals.h
+++ b/Engine/Source/Runtime/CoreUObject/Public/UObject/UObjectGlobals.h
@@ -182,8 +182,8 @@ COREUOBJECT_API FString ResolveIniObjectsReference(const FString& ObjectReferenc
  */
 COREUOBJECT_API bool ResolveName(UObject*& Outer, FString& ObjectsReferenceString, bool Create, bool Throw, uint32 LoadFlags = LOAD_None, FUObjectSerializeContext* InLoadContext = nullptr);
 
-/** Internal function used to possibly output an error message, taking into account the outer and LoadFlags */
-COREUOBJECT_API void SafeLoadError( UObject* Outer, uint32 LoadFlags, const TCHAR* ErrorMessage);
+/** Internal function used to possibly output an error message, taking into account the outer and LoadFlags. Returns true if a log message was emitted. */
+COREUOBJECT_API bool SafeLoadError( UObject* Outer, uint32 LoadFlags, const TCHAR* ErrorMessage);
 
 /** Internal function used to update the suffix to be given to the next newly-created unnamed object. */
 COREUOBJECT_API int32 UpdateSuffixForNextNewObject(UObject* Parent, UClass* Class, TFunctionRef<void(int32&)> IndexMutator);
