Edigration of an FLinkerLoad optimization from Dev-Core CL 5709116
Second optimization to a heavily hit function that uses an FArchiveFromStructuredArchive adapter. That class needs optimization, so temporarily use a straight FArchive path when text archive support is compiled out

fyi ben.woodhouse, gil.gribb, per.larsson

diff --git a/Engine/Source/Runtime/CoreUObject/Private/UObject/Class.cpp b/Engine/Source/Runtime/CoreUObject/Private/UObject/Class.cpp
--- a/Engine/Source/Runtime/CoreUObject/Private/UObject/Class.cpp
+++ b/Engine/Source/Runtime/CoreUObject/Private/UObject/Class.cpp
@@ -2169,6 +2169,7 @@ void UScriptStruct::SerializeItem(FStructuredArchive::FSlot Slot, void* Value, v
 		}
 		else
 		{
+#if WITH_TEXT_ARCHIVE_SUPPORT
 			FArchiveUObjectFromStructuredArchive Ar(Slot);
 			bItemSerialized = TheCppStructOps->Serialize(Ar, Value);
 			if (bItemSerialized && !Slot.IsFilled())
@@ -2176,6 +2177,9 @@ void UScriptStruct::SerializeItem(FStructuredArchive::FSlot Slot, void* Value, v
 				// The struct said that serialization succeeded but it didn't actually write anything.
 				Slot.EnterRecord();
 			}
+#else
+			bItemSerialized = TheCppStructOps->Serialize(Slot.GetUnderlyingArchive(), Value);
+#endif
 		}		
 	}
 
diff --git a/Engine/Source/Runtime/CoreUObject/Private/UObject/LinkerLoad.cpp b/Engine/Source/Runtime/CoreUObject/Private/UObject/LinkerLoad.cpp
--- a/Engine/Source/Runtime/CoreUObject/Private/UObject/LinkerLoad.cpp
+++ b/Engine/Source/Runtime/CoreUObject/Private/UObject/LinkerLoad.cpp
@@ -1420,17 +1420,28 @@ FLinkerLoad::ELinkerStatus FLinkerLoad::SerializeNameMap()
 	}
 
 	FStructuredArchive::FStream NameStream = StructuredArchiveRootRecord->EnterStream(FIELD_NAME_TEXT("Names"));
+	const bool bIsTextFormat = StructuredArchiveRootRecord->GetUnderlyingArchive().IsTextFormat();
 
-	while( bFinishedPrecaching && NameMapIndex < Summary.NameCount && !IsTimeLimitExceeded(TEXT("serializing name map"),100))
+	NameMap.Reserve(Summary.NameCount);
+
+	while (bFinishedPrecaching && NameMapIndex < Summary.NameCount && !IsTimeLimitExceeded(TEXT("serializing name map"), 100))
 	{
 		SCOPED_LOADTIMER(LinkerLoad_SerializeNameMap_ProcessingEntries);
 
-		// Read the name entry from the file.
 		FNameEntrySerialized NameEntry(ENAME_LinkerConstructor);
-		NameStream.EnterElement() << NameEntry;
+
+		if (bIsTextFormat)
+		{
+			NameStream.EnterElement() << NameEntry;
+		}
+		else
+		{
+			// Read the name from the underlying Archive when the format is binary to avoid the overhead from ArchiveProxy
+			NameStream.GetUnderlyingArchive() << NameEntry;
+		}
 
 		// Add it to the name table with no splitting and no hash calculations
-		NameMap.Add(FName(NameEntry));
+		NameMap.Emplace(NameEntry);
 
 		NameMapIndex++;
 	}
