#UE4 Changed the order of the asset load error message to show the referencer first

diff --git a/Engine/Source/Runtime/CoreUObject/Private/UObject/Linker.cpp b/Engine/Source/Runtime/CoreUObject/Private/UObject/Linker.cpp
--- a/Engine/Source/Runtime/CoreUObject/Private/UObject/Linker.cpp
+++ b/Engine/Source/Runtime/CoreUObject/Private/UObject/Linker.cpp
@@ -435,7 +435,7 @@ static void LogGetPackageLinkerError(FArchive* LinkerArchive, FUObjectSerializeC
 		}
 		FString LoadedByFile = LinkerToUse ? *LinkerToUse->Filename : SerializedObject->GetOutermost()->GetName();
 		Arguments.Add(TEXT("LoadedByFile"), FText::FromString(LoadedByFile));
-		FullErrorMessage = FText::Format(LOCTEXT("FailedLoadWithLoadedBy", "Failed to load '{LoadingFile}' while loading '{LoadedByFile}': {ErrorMessage}"), Arguments);
+		FullErrorMessage = FText::Format(LOCTEXT("FailedLoadWithLoadedBy", "While loading '{LoadedByFile}' failed to load '{LoadingFile}': {ErrorMessage}"), Arguments);
 	}
 	else
 	{
