Optimize UPackage::IsFullyLoaded() for cooker.
The editor specific feature of querying the file system to identify newly created packages (not having a corresponding file on disk) is really slow.
Early out from this expensive check when the package already has a known file size.

VTune user-mode cpu sampling time results when cooking a big data set (~200 000 packages):
UCookOnTheFlyServer::MaybeMarkPackageAsAlreadyLoaded() down from 10 sec to 0.1 sec, a 100X speedup.
FLinkerLoad::LoadAllObjects down from 400 sec to 280 sec, a 1.4X speedup, i.e. ~4% (120 sec) of the total load time of 2750 sec.

diff --git a/Engine/Source/Runtime/CoreUObject/Private/UObject/Package.cpp b/Engine/Source/Runtime/CoreUObject/Private/UObject/Package.cpp
--- a/Engine/Source/Runtime/CoreUObject/Private/UObject/Package.cpp
+++ b/Engine/Source/Runtime/CoreUObject/Private/UObject/Package.cpp
@@ -200,7 +200,7 @@ bool UPackage::IsFullyLoaded() const
 	// Newly created packages aren't loaded and therefore haven't been marked as being fully loaded. They are treated as fully
 	// loaded packages though in this case, which is why we are looking to see whether the package exists on disk and assume it
 	// has been fully loaded if it doesn't.
-	if( !bHasBeenFullyLoaded && !HasAnyInternalFlags(EInternalObjectFlags::AsyncLoading) )
+	if (!bHasBeenFullyLoaded && !HasAnyInternalFlags(EInternalObjectFlags::AsyncLoading) && FileSize == 0)
 	{
 		FString DummyFilename;
 		FString SourcePackageName = FileName != NAME_None ? FileName.ToString() : GetName();
