Adding const variants of some of the more common AddReferencedObject(s) functions to allow for FGCObjects to store const pointers to UObjects or collections and report those references easily, otherwise people make the mistake of creating local ptrs that can't be nulled out correctly.

diff --git a/Engine/Source/Runtime/CoreUObject/Public/UObject/UObjectGlobals.h b/Engine/Source/Runtime/CoreUObject/Public/UObject/UObjectGlobals.h
--- a/Engine/Source/Runtime/CoreUObject/Public/UObject/UObjectGlobals.h
+++ b/Engine/Source/Runtime/CoreUObject/Public/UObject/UObjectGlobals.h
@@ -1723,6 +1723,22 @@ public:
 		HandleObjectReference(*(UObject**)&Object, ReferencingObject, ReferencingProperty);
 	}
 
+	/**
+	 * Adds const object reference, this reference can still be nulled out if forcefully collected.
+	 *
+	 * @param Object Referenced object.
+	 * @param ReferencingObject Referencing object (if available).
+	 * @param ReferencingProperty Referencing property (if available).
+	 */
+	template<class UObjectType>
+	void AddReferencedObject(const UObjectType*& Object, const UObject* ReferencingObject = nullptr, const UProperty* ReferencingProperty = nullptr)
+	{
+		// @todo: should be uncommented when proper usage is fixed everywhere
+		// static_assert(sizeof(UObjectType) > 0, "AddReferencedObject: Element must be a pointer to a fully-defined type");
+		// static_assert(TPointerIsConvertibleFromTo<UObjectType, const UObjectBase>::Value, "AddReferencedObject: Element must be a pointer to a type derived from UObject");
+		HandleObjectReference(*(UObject**)const_cast<UObjectType**>(&Object), ReferencingObject, ReferencingProperty);
+	}
+
 	/**
 	* Adds references to an array of objects.
 	*
@@ -1738,6 +1754,21 @@ public:
 		HandleObjectReferences(reinterpret_cast<UObject**>(ObjectArray.GetData()), ObjectArray.Num(), ReferencingObject, ReferencingProperty);
 	}
 
+	/**
+	* Adds references to an array of const objects, these objects can still be nulled out if forcefully collected.
+	*
+	* @param ObjectArray Referenced objects array.
+	* @param ReferencingObject Referencing object (if available).
+	* @param ReferencingProperty Referencing property (if available).
+	*/
+	template<class UObjectType>
+	void AddReferencedObjects(TArray<const UObjectType*>& ObjectArray, const UObject* ReferencingObject = nullptr, const UProperty* ReferencingProperty = nullptr)
+	{
+		static_assert(sizeof(UObjectType) > 0, "AddReferencedObjects: Elements must be pointers to a fully-defined type");
+		static_assert(TPointerIsConvertibleFromTo<UObjectType, const UObjectBase>::Value, "AddReferencedObjects: Elements must be pointers to a type derived from UObject");
+		HandleObjectReferences(reinterpret_cast<UObject**>(const_cast<UObjectType**>(ObjectArray.GetData())), ObjectArray.Num(), ReferencingObject, ReferencingProperty);
+	}
+
 	/**
 	* Adds references to a set of objects.
 	*
