Fix race condition in cook -diffonly

The same ubulk file, typically for a texture, was written by two different threads. This caused a fatal SavePackage assert in FAsyncWriteWorker::DoWork(). There was no need to write out these files to disk at all.

Also cleaned up surrounding code a little bit.

diff --git a/Engine/Source/Runtime/CoreUObject/Private/UObject/SavePackage.cpp b/Engine/Source/Runtime/CoreUObject/Private/UObject/SavePackage.cpp
--- a/Engine/Source/Runtime/CoreUObject/Private/UObject/SavePackage.cpp
+++ b/Engine/Source/Runtime/CoreUObject/Private/UObject/SavePackage.cpp
@@ -3518,6 +3518,14 @@ void VerifyEDLCookInfo()
 	FEDLCookChecker::Verify();
 }
 
+static void AddArchiveToHash(FBufferArchive& Ar, FMD5& OutHash)
+{
+	if (int64 Size = Ar.TotalSize())
+	{
+		OutHash.Update(Ar.GetData(), Size);
+	}
+}
+
 void AddFileToHash(FString const &Filename, FMD5 &Hash)
 {
 	TArray<uint8> LocalScratch;
@@ -3555,7 +3563,12 @@ FSavePackageResultStruct UPackage::Save(UPackage* InOuter, UObject* Base, EObjec
 	{
 		TRefCountPtr<FUObjectSerializeContext> SaveContext(FUObjectThreadContext::Get().GetSerializeContext());
 
-#if WITH_EDITOR
+		const bool bComputeHash = (SaveFlags & SAVE_ComputeHash) != 0;
+#if !WITH_EDITOR
+		const bool bDiffing = false;
+#else
+		const bool bDiffing = (SaveFlags & (SAVE_DiffCallstack | SAVE_DiffOnly)) != 0;
+
 		struct FDiffSettings
 		{
 			int32 MaxDiffsToLog;
@@ -3585,7 +3598,7 @@ FSavePackageResultStruct UPackage::Save(UPackage* InOuter, UObject* Base, EObjec
 					bSaveForDiff = FParse::Param(FCommandLine::Get(), TEXT("SaveForDiff"));
 				}
 			}
-		} DiffSettings((SaveFlags & (SAVE_DiffCallstack | SAVE_DiffOnly)) != 0);
+		} DiffSettings(bDiffing);
 #endif
 
 		if (GIsSavingPackage && !bSavingConcurrent)
@@ -5596,11 +5609,13 @@ FSavePackageResultStruct UPackage::Save(UPackage* InOuter, UObject* Base, EObjec
 
 					const FString MappedBulkFilename = FPaths::ChangeExtension(Filename, MappedBulkFileExtension);
 
+					bool bBufferedBulkArchives = false;
 					if (bShouldUseSeparateBulkFile)
 					{
 						ExtraBulkDataFlags = BULKDATA_PayloadInSeperateFile;
-						if ( bSaveAsync )
+						if (bSaveAsync || bDiffing)
 						{
+							bBufferedBulkArchives = true;
 							BulkArchive = new FBufferArchive(true);
 							OptionalBulkArchive = new FBufferArchive(true);
 							MappedBulkArchive = new FBufferArchive(true);
@@ -5713,23 +5728,19 @@ FSavePackageResultStruct UPackage::Save(UPackage* InOuter, UObject* Base, EObjec
 					if (BulkArchive)
 					{
 						check( OptionalBulkArchive);
-						auto FinalizeBulkDataFile = [&bSaveAsync,&TotalPackageSizeUncompressed](FArchive* Archive, const FString& ArchiveFilename)
+
+						const bool bWriteBulkToDisk = bSaveAsync && !bDiffing; 
+						auto FinalizeBulkDataFile = [bWriteBulkToDisk](FArchive* Archive, const FString& ArchiveFilename)
 							{
-								TotalPackageSizeUncompressed += Archive->TotalSize();
 								Archive->Close();
-								if ( bSaveAsync )
+								if (bWriteBulkToDisk)
 								{
-									FBufferArchive* Buffer = (FBufferArchive*)(Archive);
-
-									if ( Buffer->TotalSize() > 0 )
+									if (int64 DataSize = Archive->TotalSize())
 									{
-										int64 DataSize = Buffer->TotalSize();
-
 										// TODO - update the BulkBuffer code to write into a FLargeMemoryWriter so we can 
 										// take ownership of the data
 										uint8* Data = new uint8[DataSize];
-										FMemory::Memcpy(Data, Buffer->GetData(), DataSize);
-
+										FMemory::Memcpy(Data, ((FBufferArchive*)(Archive))->GetData(), DataSize);
 										FLargeMemoryPtr DataPtr = FLargeMemoryPtr(Data);
 		
 										AsyncWriteFile(MoveTemp(DataPtr), DataSize, *ArchiveFilename, FDateTime::MinValue(), false);
@@ -5738,38 +5749,29 @@ FSavePackageResultStruct UPackage::Save(UPackage* InOuter, UObject* Base, EObjec
 								delete Archive;
 							};
 
-						if (SaveFlags & SAVE_ComputeHash)
+						if (bComputeHash && bBufferedBulkArchives)
 						{
-							if (bSaveAsync)
-							{
-								if (BulkArchive->TotalSize())
-								{
-									CookedPackageHash.Update(((FBufferArchive*)BulkArchive)->GetData(), BulkArchive->TotalSize());
-								}
-								if (OptionalBulkArchive->TotalSize())
-								{
-									CookedPackageHash.Update(((FBufferArchive*)OptionalBulkArchive)->GetData(), OptionalBulkArchive->TotalSize());
-								}
-								if (MappedBulkArchive->TotalSize())
-								{
-									CookedPackageHash.Update(((FBufferArchive*)MappedBulkArchive)->GetData(), MappedBulkArchive->TotalSize());
-								}
-							}
+							AddArchiveToHash(*(FBufferArchive*)BulkArchive, CookedPackageHash);
+							AddArchiveToHash(*(FBufferArchive*)OptionalBulkArchive, CookedPackageHash);
+							AddArchiveToHash(*(FBufferArchive*)MappedBulkArchive, CookedPackageHash);
 						}
+						
+						TotalPackageSizeUncompressed += BulkArchive->TotalSize();
+						TotalPackageSizeUncompressed += OptionalBulkArchive->TotalSize();		
+						TotalPackageSizeUncompressed += MappedBulkArchive->TotalSize();
+
 						FinalizeBulkDataFile(BulkArchive, BulkFilename);
 						FinalizeBulkDataFile(OptionalBulkArchive, OptionalBulkFilename);
 						FinalizeBulkDataFile(MappedBulkArchive, MappedBulkFilename);
 
-						if (SaveFlags & SAVE_ComputeHash)
+						if (bComputeHash && !bBufferedBulkArchives)
 						{
-							if (!bSaveAsync)
-							{
-								AddFileToHash(BulkFilename, CookedPackageHash);
-								AddFileToHash(OptionalBulkFilename, CookedPackageHash);
-								AddFileToHash(MappedBulkFilename, CookedPackageHash);
-							}
+							check(!bWriteBulkToDisk);
+
+							AddFileToHash(BulkFilename, CookedPackageHash);
+							AddFileToHash(OptionalBulkFilename, CookedPackageHash);
+							AddFileToHash(MappedBulkFilename, CookedPackageHash);
 						}
-						
 					}
 				}
 
@@ -5958,7 +5960,7 @@ FSavePackageResultStruct UPackage::Save(UPackage* InOuter, UObject* Base, EObjec
 					delete TextFormatArchive;
 				}
 				UNCLOCK_CYCLES(Time);
-				UE_CLOG(!(SaveFlags & (SAVE_DiffCallstack | SAVE_DiffOnly)), LogSavePackage, Verbose,  TEXT("Save=%.2fms"), FPlatformTime::ToMilliseconds(Time) );
+				UE_CLOG(!bDiffing, LogSavePackage, Verbose,  TEXT("Save=%.2fms"), FPlatformTime::ToMilliseconds(Time) );
 		
 				if ( EndSavingIfCancelled( Linker.Get(), TempFilename ) )
 				{ 
@@ -6024,7 +6026,7 @@ FSavePackageResultStruct UPackage::Save(UPackage* InOuter, UObject* Base, EObjec
 							COOK_STAT(FScopedDurationTimer SaveTimer(SavePackageStats::AsyncWriteTimeSec));
 							TotalPackageSizeUncompressed += DataSize;
 
-							if (SaveFlags & SAVE_ComputeHash)
+							if (bComputeHash)
 							{
 								CookedPackageHash.Update( Writer->GetData(), Writer->TotalSize() );
 							}
@@ -6180,7 +6182,7 @@ FSavePackageResultStruct UPackage::Save(UPackage* InOuter, UObject* Base, EObjec
 		// We're done!
 		SlowTask.EnterProgressFrame();
 
-		UE_CLOG(!(SaveFlags & (SAVE_DiffCallstack | SAVE_DiffOnly)), LogSavePackage, Verbose, TEXT("Finished SavePackage %s"), Filename);
+		UE_CLOG(!bDiffing, LogSavePackage, Verbose, TEXT("Finished SavePackage %s"), Filename);
 
 		if (Success)
 		{
