Fixing a crash when launching cooked projects with the Event Driven Loader disabled

diff --git a/Engine/Source/Runtime/CoreUObject/Private/UObject/UObjectGlobals.cpp b/Engine/Source/Runtime/CoreUObject/Private/UObject/UObjectGlobals.cpp
--- a/Engine/Source/Runtime/CoreUObject/Private/UObject/UObjectGlobals.cpp
+++ b/Engine/Source/Runtime/CoreUObject/Private/UObject/UObjectGlobals.cpp
@@ -1360,7 +1360,6 @@ UPackage* LoadPackageInternal(UPackage* InOuter, const TCHAR* InLongPackageNameO
 			{				
 				if (GGameThreadLoadCounter == 0)
 				{
-					check(!Linker->GetSerializeContext());
 					// Sanity check to make sure that Linker is the linker that loaded our Result package or the linker has already been detached
 					check(!Result || Result->LinkerLoad == Linker || Result->LinkerLoad == nullptr);
 					if (Result && Linker->HasLoader())
