- Added debug spikemark and spikemarkcheck to track objects that get destroyed after a mark, which can be used to track objects that go away right after being loaded; can help reduce memory spikes of objects being loaded and GCd right away

diff --git a/Engine/Source/Runtime/CoreUObject/Private/UObject/Obj.cpp b/Engine/Source/Runtime/CoreUObject/Private/UObject/Obj.cpp
--- a/Engine/Source/Runtime/CoreUObject/Private/UObject/Obj.cpp
+++ b/Engine/Source/Runtime/CoreUObject/Private/UObject/Obj.cpp
@@ -80,6 +80,9 @@ static UPackage*			GObjTransientPkg								= NULL;
 	/** Used for the "obj invmark" and "obj invmarkcheck" commands only			*/
 	static TArray<TWeakObjectPtr<UObject> >	DebugInvMarkWeakPtrs;
 	static TArray<FString>			DebugInvMarkNames;
+	/** Used for the "obj spikemark" and "obj spikemarkcheck" commands only			*/
+	static FUObjectAnnotationSparseBool DebugSpikeMarkAnnotation;
+	static TArray<FString>			DebugSpikeMarkNames;
 #endif
 
 UObject::UObject( EStaticConstructor, EObjectFlags InFlags )
@@ -926,6 +929,18 @@ static TMap<FName, FTimeCnt> MyProfile;
 
 bool UObject::ConditionalBeginDestroy()
 {
+#if !UE_BUILD_SHIPPING
+	// if this object wasn't marked (but some were) then that means it was created and destroyed since the SpikeMark command was given
+	// this object is contributing to the spike that is being investigated
+	if (DebugSpikeMarkAnnotation.Num() > 0)
+	{
+		if(!DebugSpikeMarkAnnotation.Get(this))
+		{
+			DebugSpikeMarkNames.Add(GetFullName());
+		}
+	}
+#endif
+	
 	check(IsValidLowLevel());
 	if( !HasAnyFlags(RF_BeginDestroyed) )
 	{
@@ -3825,6 +3840,28 @@ bool StaticExec( UWorld* InWorld, const TCHAR* Cmd, FOutputDevice& Ar )
 			}
 			return true;
 		}
+		else if( FParse::Command(&Str,TEXT("SPIKEMARK")) )
+		{
+			UE_LOG(LogObj, Log,  TEXT("Spikemarking objects") );
+			
+			FlushAsyncLoading();
+
+			DebugSpikeMarkAnnotation.ClearAll();
+			for( FObjectIterator It; It; ++It )
+			{
+				DebugSpikeMarkAnnotation.Set(*It);
+			}
+			return true;
+		}
+		else if( FParse::Command(&Str,TEXT("SPIKEMARKCHECK")) )
+		{
+			UE_LOG(LogObj, Log,  TEXT("Spikemarked (created and then destroyed) objects:") );
+			for( const FString& Name : DebugSpikeMarkNames )
+			{
+				UE_LOG(LogObj, Log,  TEXT("  %s"), *Name );
+			}
+			return true;
+		}
 		else if( FParse::Command(&Str,TEXT("REFS")) )
 		{
 			UObject* Object;
