Fix linker loop deadlock bug (causes loading deadlock in some cases)

diff --git a/Engine/Source/Runtime/CoreUObject/Private/UObject/LinkerLoad.cpp b/Engine/Source/Runtime/CoreUObject/Private/UObject/LinkerLoad.cpp
--- a/Engine/Source/Runtime/CoreUObject/Private/UObject/LinkerLoad.cpp
+++ b/Engine/Source/Runtime/CoreUObject/Private/UObject/LinkerLoad.cpp
@@ -5389,6 +5389,7 @@ bool FLinkerLoad::FinishExternalReadDependencies(double InTimeLimit)
 	double LocalStartTime = FPlatformTime::Seconds();
 	double RemainingTime = InTimeLimit;
 	const int32 Granularity = 5;
+	int32 Iteration = 0;
 	
 	while (ExternalReadDependencies.Num())
 	{
@@ -5404,7 +5405,7 @@ bool FLinkerLoad::FinishExternalReadDependencies(double InTimeLimit)
 		}
 
 		// Update remaining time
-		if (InTimeLimit > 0.0 && (ExternalReadDependencies.Num() % Granularity) == 0)
+		if (InTimeLimit > 0.0 && (++Iteration % Granularity) == 0)
 		{
 			RemainingTime = InTimeLimit - (FPlatformTime::Seconds() - LocalStartTime);
 			if (RemainingTime <= 0.0)
