Fixing a crash when reporting long reference chains with 'obj refs' command.

diff --git a/Engine/Source/Runtime/CoreUObject/Private/UObject/ReferenceChainSearch.cpp b/Engine/Source/Runtime/CoreUObject/Private/UObject/ReferenceChainSearch.cpp
--- a/Engine/Source/Runtime/CoreUObject/Private/UObject/ReferenceChainSearch.cpp
+++ b/Engine/Source/Runtime/CoreUObject/Private/UObject/ReferenceChainSearch.cpp
@@ -269,7 +269,7 @@ void FReferenceChainSearch::DumpChain(FReferenceChain* Chain)
 			const FNodeReferenceInfo& ReferenceInfo = Chain->GetReferenceInfo(NodeIndex);
 
 			UE_LOG(LogReferenceChain, Log, TEXT("%s%s%s%s"),
-				FCString::Spc(Chain->Num() - NodeIndex - 1),
+				FCString::Spc(FMath::Min<int32>(TCStringSpcHelper<TCHAR>::MAX_SPACES, Chain->Num() - NodeIndex - 1)),
 				*GetObjectFlags(Object),
 				*Object->GetFullName(),
 				*ReferenceInfo.ToString()
