- Now check if a "GettingStarted" path is a valid package name before trying to load it as an object to avoid some Verify's, etc
- Fix the call to IsValidLongPackageName() in FPackageName::DoesPackageExist() to pass the proper path in (it would pass in an empty string)

diff --git a/Engine/Source/Editor/IntroTutorials/Private/IntroTutorials.cpp b/Engine/Source/Editor/IntroTutorials/Private/IntroTutorials.cpp
--- a/Engine/Source/Editor/IntroTutorials/Private/IntroTutorials.cpp
+++ b/Engine/Source/Editor/IntroTutorials/Private/IntroTutorials.cpp
@@ -371,10 +371,13 @@ void FIntroTutorials::HandleCompilerNotFound()
 
 void FIntroTutorials::HandleSDKNotInstalled(const FString& PlatformName, const FString& InTutorialAsset)
 {
-	UBlueprint* Blueprint = LoadObject<UBlueprint>(nullptr, *InTutorialAsset);
-	if(Blueprint)
+	if (FPackageName::IsValidLongPackageName(InTutorialAsset, true))
 	{
-		LaunchTutorialByName( InTutorialAsset );
+		UBlueprint* Blueprint = LoadObject<UBlueprint>(nullptr, *InTutorialAsset);
+		if (Blueprint)
+		{
+			LaunchTutorialByName(InTutorialAsset);
+		}
 	}
 	else
 	{
diff --git a/Engine/Source/Runtime/CoreUObject/Private/Misc/PackageName.cpp b/Engine/Source/Runtime/CoreUObject/Private/Misc/PackageName.cpp
--- a/Engine/Source/Runtime/CoreUObject/Private/Misc/PackageName.cpp
+++ b/Engine/Source/Runtime/CoreUObject/Private/Misc/PackageName.cpp
@@ -844,7 +844,7 @@ bool FPackageName::DoesPackageExist(const FString& LongPackageName, const FGuid*
 
 	if (!FPackageName::TryConvertFilenameToLongPackageName(LongPackageName, PackageName))
 	{
-		verify(!FPackageName::IsValidLongPackageName(PackageName, true, &Reason));
+		verify(!FPackageName::IsValidLongPackageName(LongPackageName, true, &Reason));
 		UE_LOG(LogPackageName, Error, TEXT("Illegal call to DoesPackageExist: '%s' is not a standard unreal filename or a long path name. Reason: %s"), *LongPackageName, *Reason.ToString());
 		ensureMsgf(false, TEXT("Illegal call to DoesPackageExist: '%s' is not a standard unreal filename or a long path name. Reason: %s"), *LongPackageName, *Reason.ToString());
 		return false;
