Avoid callspace checking for non-net UFunctions

diff --git a/Engine/Source/Editor/KismetCompiler/Private/KismetCompilerVMBackend.cpp b/Engine/Source/Editor/KismetCompiler/Private/KismetCompilerVMBackend.cpp
--- a/Engine/Source/Editor/KismetCompiler/Private/KismetCompilerVMBackend.cpp
+++ b/Engine/Source/Editor/KismetCompiler/Private/KismetCompilerVMBackend.cpp
@@ -1189,8 +1189,8 @@ public:
 
 		const bool bFinalFunction = FunctionToCall->HasAnyFunctionFlags(FUNC_Final) || Statement.bIsParentContext;
 		const bool bMathCall = bFinalFunction
-			&& FunctionToCall->HasAllFunctionFlags(FUNC_Static | FUNC_Final | FUNC_Native)
-			&& !FunctionToCall->HasAnyFunctionFlags(FUNC_BlueprintAuthorityOnly | FUNC_BlueprintCosmetic | FUNC_NetFuncFlags)
+			&& FunctionToCall->HasAllFunctionFlags(FUNC_Static|FUNC_Final|FUNC_Native)
+			&& !FunctionToCall->HasAnyFunctionFlags(FUNC_NetFuncFlags|FUNC_BlueprintAuthorityOnly|FUNC_BlueprintCosmetic|FUNC_NetRequest|FUNC_NetResponse)
 			&& !FunctionToCall->GetOuterUClass()->IsChildOf(UInterface::StaticClass())
 			// If the function has wildcard parameters (e.g. UKismetArrayLibrary) we assume that it
 			// is doing things with reflection data at runtime and will therefore benefit from
diff --git a/Engine/Source/Runtime/CoreUObject/Private/UObject/ScriptCore.cpp b/Engine/Source/Runtime/CoreUObject/Private/UObject/ScriptCore.cpp
--- a/Engine/Source/Runtime/CoreUObject/Private/UObject/ScriptCore.cpp
+++ b/Engine/Source/Runtime/CoreUObject/Private/UObject/ScriptCore.cpp
@@ -824,6 +824,8 @@ DEFINE_FUNCTION(UObject::execCallMathFunction)
 	UFunction* Function = (UFunction*)Stack.ReadObject();
 	checkSlow(Function);
 	checkSlow(Function->FunctionFlags & FUNC_Native);
+	// ProcessContext is the arbiter of net callspace, so we can't call net functions using this instruction:
+	checkSlow(!Function->HasAnyFunctionFlags(FUNC_NetFuncFlags|FUNC_BlueprintAuthorityOnly|FUNC_BlueprintCosmetic|FUNC_NetRequest|FUNC_NetResponse));
 	UObject* NewContext = Function->GetOuterUClassUnchecked()->ClassDefaultObject;
 	checkSlow(NewContext);
 	{
@@ -857,7 +859,9 @@ void UObject::CallFunction( FFrame& Stack, RESULT_DECL, UFunction* Function )
 
 	if (Function->FunctionFlags & FUNC_Native)
 	{
-		int32 FunctionCallspace = GetFunctionCallspace( Function, &Stack );
+		const bool bNetFunction = Function->HasAnyFunctionFlags(FUNC_NetFuncFlags|FUNC_BlueprintAuthorityOnly|FUNC_BlueprintCosmetic|FUNC_NetRequest|FUNC_NetResponse);
+		const int32 FunctionCallspace = bNetFunction ? GetFunctionCallspace( Function, &Stack ) : FunctionCallspace::Local;
+
 		uint8* SavedCode = NULL;
 		if (FunctionCallspace & FunctionCallspace::Remote)
 		{
