Moved STAT_FAsyncPackage_CreateAsyncPackagesFromQueue to an inner scope to reduce the number of stat events it generates

diff --git a/Engine/Source/Runtime/CoreUObject/Private/Serialization/AsyncLoading.cpp b/Engine/Source/Runtime/CoreUObject/Private/Serialization/AsyncLoading.cpp
--- a/Engine/Source/Runtime/CoreUObject/Private/Serialization/AsyncLoading.cpp
+++ b/Engine/Source/Runtime/CoreUObject/Private/Serialization/AsyncLoading.cpp
@@ -594,7 +594,6 @@ void FAsyncLoadingThread::ProcessAsyncPackageRequest(FAsyncPackageDesc* InReques
 
 int32 FAsyncLoadingThread::CreateAsyncPackagesFromQueue(bool bUseTimeLimit, bool bUseFullTimeLimit, float TimeLimit, FFlushTree* FlushTree)
 {
-	SCOPE_CYCLE_COUNTER(STAT_FAsyncPackage_CreateAsyncPackagesFromQueue);
 	SCOPED_LOADTIMER(CreateAsyncPackagesFromQueueTime);
 
 	FAsyncLoadingTickScope InAsyncLoadingTick;
@@ -646,6 +645,7 @@ int32 FAsyncLoadingThread::CreateAsyncPackagesFromQueue(bool bUseTimeLimit, bool
 
 		if (QueueCopy.Num() > 0)
 		{
+			SCOPE_CYCLE_COUNTER(STAT_FAsyncPackage_CreateAsyncPackagesFromQueue);
 			double Timer = 0;
 			{
 				SCOPE_SECONDS_COUNTER(Timer);
@@ -4426,7 +4426,6 @@ struct FScopedRecursionNotAllowed
 
 EAsyncPackageState::Type FAsyncLoadingThread::ProcessAsyncLoading(int32& OutPackagesProcessed, bool bUseTimeLimit /*= false*/, bool bUseFullTimeLimit /*= false*/, float TimeLimit /*= 0.0f*/, FFlushTree* FlushTree)
 {
-	SCOPE_CYCLE_COUNTER(STAT_FAsyncLoadingThread_ProcessAsyncLoading);
 	SCOPED_LOADTIMER(AsyncLoadingTime);
 	check(!IsInGameThread() || !IsMultithreaded());
 
@@ -4439,8 +4438,9 @@ EAsyncPackageState::Type FAsyncLoadingThread::ProcessAsyncLoading(int32& OutPack
 
 	if (GEventDrivenLoaderEnabled)
 	{
+		SCOPE_CYCLE_COUNTER(STAT_FAsyncLoadingThread_ProcessAsyncLoading);
 #if !UE_BUILD_SHIPPING && !UE_BUILD_TEST
-	FScopedRecursionNotAllowed RecursionGuard;
+		FScopedRecursionNotAllowed RecursionGuard;
 #endif
 
 		FAsyncLoadingTickScope InAsyncLoadingTick;
@@ -4457,12 +4457,12 @@ EAsyncPackageState::Type FAsyncLoadingThread::ProcessAsyncLoading(int32& OutPack
 			bool bDidSomething = false;
 			{
 				bDidSomething = GPrecacheCallbackHandler.ProcessIncoming();
-			OutPackagesProcessed += (bDidSomething ? 1 : 0);
+				OutPackagesProcessed += (bDidSomething ? 1 : 0);
 
-			if (IsTimeLimitExceeded(TickStartTime, bUseTimeLimit, TimeLimit, TEXT("ProcessIncoming"), nullptr))
-			{
-				return EAsyncPackageState::TimeOut;
-			}
+				if (IsTimeLimitExceeded(TickStartTime, bUseTimeLimit, TimeLimit, TEXT("ProcessIncoming"), nullptr))
+				{
+					return EAsyncPackageState::TimeOut;
+				}
 			}
 
 			if (IsAsyncLoadingSuspended())
@@ -4547,7 +4547,7 @@ EAsyncPackageState::Type FAsyncLoadingThread::ProcessAsyncLoading(int32& OutPack
 					}
 
 					// We're done, at least on this thread, so we can remove the package now.
-				AddToLoadedPackages(Package);
+					AddToLoadedPackages(Package);
 				}
 				if (IsTimeLimitExceeded(TickStartTime, bUseTimeLimit, TimeLimit, TEXT("TickAsyncPackage")))
 				{
@@ -4604,8 +4604,10 @@ EAsyncPackageState::Type FAsyncLoadingThread::ProcessAsyncLoading(int32& OutPack
 			}
 		}
 	} // GEventDrivenLoaderEnabled
-	else
+	else if (AsyncPackages.Num())
 	{
+		SCOPE_CYCLE_COUNTER(STAT_FAsyncLoadingThread_ProcessAsyncLoading);
+
 		bool bDepthFirst = false;
 
 		// We need to loop as the function has to handle finish loading everything given no time limit
