Copying //UE4/Dev-Core to Dev-Main (//UE4/Dev-Main)

diff --git a/Engine/Source/Runtime/Core/Public/ProfilingDebugging/MiscTrace.h b/Engine/Source/Runtime/Core/Public/ProfilingDebugging/MiscTrace.h
new file mode 100644
--- /dev/null
+++ b/Engine/Source/Runtime/Core/Public/ProfilingDebugging/MiscTrace.h
@@ -0,0 +1,122 @@
+// Copyright 1998-2019 Epic Games, Inc. All Rights Reserved.
+
+#pragma once
+
+#include "CoreTypes.h"
+#include "Trace/Trace.h"
+
+#if UE_TRACE_ENABLED && !UE_BUILD_SHIPPING
+#define MISCTRACE_ENABLED 1
+#else
+#define MISCTRACE_ENABLED 0
+#endif
+
+enum ETraceFrameType
+{
+	TraceFrameType_Game,
+	TraceFrameType_Rendering,
+
+	TraceFrameType_Count
+};
+
+struct FTraceUtils
+{
+	static void Encode7bit(uint64 Value, uint8*& BufferPtr)
+	{
+		do
+		{
+			uint8 HasMoreBytes = (Value > uint64(0x7F)) << 7;
+			*(BufferPtr++) = (uint8)(Value & 0x7F) | HasMoreBytes;
+			Value >>= 7;
+		} while (Value > 0);
+	}
+
+	static void EncodeZigZag(int64 Value, uint8*& BufferPtr)
+	{
+		Encode7bit((Value << 1) ^ (Value >> 63), BufferPtr);
+	}
+};
+
+#if MISCTRACE_ENABLED
+
+#include "ProfilingDebugging/FormatArgsTrace.h"
+
+class FName;
+
+struct FMiscTrace
+{
+	CORE_API static void OutputRegisterGameThread(uint32 Id);
+	CORE_API static void OutputCreateThread(uint32 Id, const TCHAR* Name, uint32 Priority);
+	CORE_API static void OutputSetThreadGroup(uint32 Id, const ANSICHAR* GroupName);
+	CORE_API static void OutputBeginThreadGroupScope(const ANSICHAR* GroupName);
+	CORE_API static void OutputEndThreadGroupScope();
+	CORE_API static void OutputBookmarkSpec(const void* BookmarkPoint, const ANSICHAR* File, int32 Line, const TCHAR* Format);
+	template <typename... Types>
+	static void OutputBookmark(const void* BookmarkPoint, Types... FormatArgs)
+	{
+		uint8 FormatArgsBuffer[4096];
+		uint16 FormatArgsSize = FFormatArgsTrace::EncodeArguments(FormatArgsBuffer, FormatArgs...);
+		if (FormatArgsSize)
+		{
+			OutputBookmarkInternal(BookmarkPoint, FormatArgsSize, FormatArgsBuffer);
+		}
+	}
+
+	CORE_API static void OutputBeginFrame(ETraceFrameType FrameType);
+	CORE_API static void OutputEndFrame(ETraceFrameType FrameType);
+
+	struct FThreadGroupScope
+	{
+		FThreadGroupScope(const ANSICHAR* GroupName)
+		{
+			FMiscTrace::OutputBeginThreadGroupScope(GroupName);
+		}
+
+		~FThreadGroupScope()
+		{
+			FMiscTrace::OutputEndThreadGroupScope();
+		}
+	};
+
+private:
+	CORE_API static void OutputBookmarkInternal(const void* BookmarkPoint, uint16 EncodedFormatArgsSize, uint8* EncodedFormatArgs);
+};
+
+#define TRACE_REGISTER_GAME_THREAD(Id) \
+	FMiscTrace::OutputRegisterGameThread(Id);
+
+#define TRACE_CREATE_THREAD(Id, Name, Priority) \
+	FMiscTrace::OutputCreateThread(Id, Name, Priority);
+
+#define TRACE_SET_THREAD_GROUP(Id, Group) \
+	FMiscTrace::OutputSetThreadGroup(Id, Group);
+
+#define TRACE_THREAD_GROUP_SCOPE(Group) \
+	FMiscTrace::FThreadGroupScope ANONYMOUS_VARIABLE(ThreadGroupScope) (Group);
+
+#define TRACE_BOOKMARK(Format, ...) \
+	static bool PREPROCESSOR_JOIN(__BookmarkPoint, __LINE__); \
+	if (!PREPROCESSOR_JOIN(__BookmarkPoint, __LINE__)) \
+	{ \
+		FMiscTrace::OutputBookmarkSpec(&PREPROCESSOR_JOIN(__BookmarkPoint, __LINE__), __FILE__, __LINE__, Format); \
+		PREPROCESSOR_JOIN(__BookmarkPoint, __LINE__) = true; \
+	} \
+	FMiscTrace::OutputBookmark(&PREPROCESSOR_JOIN(__BookmarkPoint, __LINE__), ##__VA_ARGS__);
+
+#define TRACE_BEGIN_FRAME(FrameType) \
+	FMiscTrace::OutputBeginFrame(FrameType);
+
+#define TRACE_END_FRAME(FrameType) \
+	FMiscTrace::OutputEndFrame(FrameType);
+
+#else
+
+#define TRACE_REGISTER_GAME_THREAD(...)
+#define TRACE_CREATE_THREAD(...)
+#define TRACE_SET_THREAD_GROUP(...)
+#define TRACE_THREAD_GROUP_SCOPE(...)
+#define TRACE_BOOKMARK(...)
+#define TRACE_BEGIN_FRAME(...)
+#define TRACE_END_FRAME(...)
+
+#endif
\ No newline at end of file
diff --git a/Engine/Source/Runtime/Core/Public/ProfilingDebugging/PlatformFileTrace.h b/Engine/Source/Runtime/Core/Public/ProfilingDebugging/PlatformFileTrace.h
new file mode 100644
--- /dev/null
+++ b/Engine/Source/Runtime/Core/Public/ProfilingDebugging/PlatformFileTrace.h
@@ -0,0 +1,65 @@
+// Copyright 1998-2019 Epic Games, Inc. All Rights Reserved.
+
+#pragma once
+
+#include "CoreTypes.h"
+#include "Trace/Trace.h"
+
+#if UE_TRACE_ENABLED && !UE_BUILD_SHIPPING
+#define PLATFORMFILETRACE_ENABLED 0
+#else
+#define PLATFORMFILETRACE_ENABLED 0
+#endif
+
+#if PLATFORMFILETRACE_ENABLED
+
+struct FPlatformFileTrace
+{
+	static void BeginOpen(uint64 TempHandle, const TCHAR* Path);
+	static void EndOpen(uint64 TempHandle, uint64 FileHandle);
+	static void BeginClose(uint64 TempHandle, uint64 FileHandle);
+	static void EndClose(uint64 TempHandle);
+	static void BeginRead(uint64 ReadHandle, uint64 FileHandle, uint64 Offset, uint64 Size);
+	static void EndRead(uint64 ReadHandle, uint64 SizeRead);
+	static void BeginWrite(uint64 WriteHandle, uint64 FileHandle, uint64 Offset, uint64 Size);
+	static void EndWrite(uint64 WriteHandle, uint64 SizeWritten);
+};
+
+#define TRACE_PLATFORMFILE_BEGIN_OPEN(Path) \
+	uint64 __TempHandle; \
+	FPlatformFileTrace::BeginOpen(uint64(&__TempHandle), Path);
+
+#define TRACE_PLATFORMFILE_END_OPEN(FileHandle) \
+	FPlatformFileTrace::EndOpen(uint64(&__TempHandle), uint64(FileHandle));
+
+#define TRACE_PLATFORMFILE_BEGIN_CLOSE(FileHandle) \
+	uint64 __TempHandle; \
+	FPlatformFileTrace::BeginClose(uint64(&__TempHandle), uint64(FileHandle));
+
+#define TRACE_PLATFORMFILE_END_CLOSE() \
+	FPlatformFileTrace::EndClose(uint64(&__TempHandle));
+
+#define TRACE_PLATFORMFILE_BEGIN_READ(ReadHandle, FileHandle, Offset, Size) \
+	FPlatformFileTrace::BeginRead(uint64(ReadHandle), uint64(FileHandle), Offset, Size);
+
+#define TRACE_PLATFORMFILE_END_READ(ReadHandle, SizeRead) \
+	FPlatformFileTrace::EndRead(uint64(ReadHandle), SizeRead);
+
+#define TRACE_PLATFORMFILE_BEGIN_WRITE(WriteHandle, FileHandle, Offset, Size) \
+	FPlatformFileTrace::BeginWrite(uint64(WriteHandle), uint64(FileHandle), Offset, Size);
+
+#define TRACE_PLATFORMFILE_END_WRITE(WriteHandle, SizeWritten) \
+	FPlatformFileTrace::EndWrite(uint64(WriteHandle), SizeWritten);
+
+#else
+
+#define TRACE_PLATFORMFILE_BEGIN_OPEN(Path)
+#define TRACE_PLATFORMFILE_END_OPEN(FileHandle)
+#define TRACE_PLATFORMFILE_BEGIN_CLOSE(FileHandle)
+#define TRACE_PLATFORMFILE_END_CLOSE()
+#define TRACE_PLATFORMFILE_BEGIN_READ(ReadHandle, FileHandle, Offset, Size)
+#define TRACE_PLATFORMFILE_END_READ(ReadHandle, SizeRead)
+#define TRACE_PLATFORMFILE_BEGIN_WRITE(WriteHandle, FileHandle, Offset, Size)
+#define TRACE_PLATFORMFILE_END_WRITE(WriteHandle, SizeWritten)
+
+#endif
\ No newline at end of file
diff --git a/Engine/Source/Runtime/Core/Public/Stats/Stats2.h b/Engine/Source/Runtime/Core/Public/Stats/Stats2.h
--- a/Engine/Source/Runtime/Core/Public/Stats/Stats2.h
+++ b/Engine/Source/Runtime/Core/Public/Stats/Stats2.h
@@ -19,6 +19,9 @@
 #include "Math/Color.h"
 #include "StatsCommon.h"
 #include "Templates/UniquePtr.h"
+#include "ProfilingDebugging/CpuProfilerTrace.h"
+#include "ProfilingDebugging/MiscTrace.h"
+#include "StatsTrace.h"
 
 class FScopeCycleCounter;
 class FThreadStats;
@@ -120,6 +123,10 @@ struct TStatIdData
 
 	/** const WIDECHAR* pointer to a string describing the stat */
 	TUniquePtr<ANSICHAR[]> StatDescriptionAnsi;
+
+#if CPUPROFILERTRACE_ENABLED
+	uint32 TraceCpuProfilerSpecId = 0;
+#endif
 };
 
 struct TStatId
@@ -155,6 +162,13 @@ struct TStatId
 		return MinimalNameToName(StatIdPtr->Name);
 	}
 
+#if CPUPROFILERTRACE_ENABLED
+	FORCEINLINE uint16 GetTraceCpuProfilerSpecId() const
+	{
+		return StatIdPtr->TraceCpuProfilerSpecId;
+	}
+#endif
+
 	FORCEINLINE static const TStatIdData& GetStatNone()
 	{
 		return TStatId_NAME_None;
@@ -651,7 +665,7 @@ private:
 	/** For FName. */
 	CORE_API const FString GetName() const
 	{
-		return FName::SafeString( (int32)Cycles );
+		return FName::SafeString(FNameEntryId::FromUnstableInt(static_cast<uint32>(Cycles)));
 	}
 };
 
@@ -1545,6 +1559,9 @@ class FCycleCounter
 {
 	/** Name of the stat, usually a short name **/
 	FName StatId;
+#if CPUPROFILERTRACE_ENABLED
+	uint16 TraceCpuProfilerSpecId = 0;
+#endif
 
 public:
 
@@ -1559,6 +1576,13 @@ public:
 		{
 			return;
 		}
+#if CPUPROFILERTRACE_ENABLED
+		if (bAlways || GCycleStatsShouldEmitNamedEvents > 0)
+		{
+			TraceCpuProfilerSpecId = InStatId.GetTraceCpuProfilerSpecId();
+			FCpuProfilerTrace::OutputBeginEvent(TraceCpuProfilerSpecId);
+		}
+#endif
 
 		if( (bAlways && FThreadStats::WillEverCollectData()) || FThreadStats::IsCollectingData() )
 		{
@@ -1583,6 +1607,13 @@ public:
 	 */
 	FORCEINLINE_STATS void Stop()
 	{
+#if CPUPROFILERTRACE_ENABLED
+		if (TraceCpuProfilerSpecId)
+		{
+			FCpuProfilerTrace::OutputEndEvent();
+			TraceCpuProfilerSpecId = 0;
+		}
+#endif
 		if( !StatId.IsNone() )
 		{
 			FThreadStats::AddMessage(StatId, EStatOperation::CycleScopeEnd);
@@ -1946,74 +1977,126 @@ struct FStat_##StatName\
 #define INC_DWORD_STAT(Stat) \
 {\
 	if (FThreadStats::IsCollectingData() || !GET_STATISEVERYFRAME(Stat)) \
+	{ \
 		FThreadStats::AddMessage(GET_STATFNAME(Stat), EStatOperation::Add, int64(1));\
+		TRACE_STAT_INCREMENT(GET_STATFNAME(Stat)); \
+	} \
 }
 #define INC_FLOAT_STAT_BY(Stat, Amount) \
 {\
 	if (Amount != 0.0f) \
+	{ \
 		if (FThreadStats::IsCollectingData() || !GET_STATISEVERYFRAME(Stat)) \
-				FThreadStats::AddMessage(GET_STATFNAME(Stat), EStatOperation::Add, double(Amount));\
+		{ \
+			FThreadStats::AddMessage(GET_STATFNAME(Stat), EStatOperation::Add, double(Amount));\
+			TRACE_STAT_ADD(GET_STATFNAME(Stat), double(Amount)); \
+		} \
+	} \
 }
 #define INC_DWORD_STAT_BY(Stat, Amount) \
 {\
 	if (Amount != 0) \
+	{ \
 		if (FThreadStats::IsCollectingData() || !GET_STATISEVERYFRAME(Stat)) \
+		{ \
 			FThreadStats::AddMessage(GET_STATFNAME(Stat), EStatOperation::Add, int64(Amount));\
+			TRACE_STAT_ADD(GET_STATFNAME(Stat), int64(Amount)); \
+		} \
+	} \
 }
 #define INC_DWORD_STAT_FNAME_BY(StatFName, Amount) \
 {\
 	if (Amount != 0) \
-			FThreadStats::AddMessage(StatFName, EStatOperation::Add, int64(Amount));\
+	{ \
+		FThreadStats::AddMessage(StatFName, EStatOperation::Add, int64(Amount));\
+		TRACE_STAT_ADD(StatFName, int64(Amount)); \
+	} \
 }
 #define INC_MEMORY_STAT_BY(Stat, Amount) \
 {\
 	if (Amount != 0) \
+	{ \
 		if (FThreadStats::IsCollectingData() || !GET_STATISEVERYFRAME(Stat)) \
+		{ \
 			FThreadStats::AddMessage(GET_STATFNAME(Stat), EStatOperation::Add, int64(Amount));\
+			TRACE_STAT_ADD(GET_STATFNAME(Stat), int64(Amount)); \
+		} \
+	} \
 }
 #define DEC_DWORD_STAT(Stat) \
 {\
 	if (FThreadStats::IsCollectingData() || !GET_STATISEVERYFRAME(Stat)) \
+	{ \
 		FThreadStats::AddMessage(GET_STATFNAME(Stat), EStatOperation::Subtract, int64(1));\
+		TRACE_STAT_DECREMENT(GET_STATFNAME(Stat)); \
+	} \
 }
 #define DEC_FLOAT_STAT_BY(Stat,Amount) \
 {\
 	if (Amount != 0.0f) \
+	{ \
 		if (FThreadStats::IsCollectingData() || !GET_STATISEVERYFRAME(Stat)) \
+		{ \
 			FThreadStats::AddMessage(GET_STATFNAME(Stat), EStatOperation::Subtract, double(Amount));\
+			TRACE_STAT_ADD(GET_STATFNAME(Stat), -double(Amount)); \
+		} \
+	} \
 }
 #define DEC_DWORD_STAT_BY(Stat,Amount) \
 {\
 	if (Amount != 0) \
+	{ \
 		if (FThreadStats::IsCollectingData() || !GET_STATISEVERYFRAME(Stat)) \
+		{ \
 			FThreadStats::AddMessage(GET_STATFNAME(Stat), EStatOperation::Subtract, int64(Amount));\
+			TRACE_STAT_ADD(GET_STATFNAME(Stat), -int64(Amount)); \
+		} \
+	} \
 }
 #define DEC_DWORD_STAT_FNAME_BY(StatFName,Amount) \
 {\
 	if (Amount != 0) \
- 			FThreadStats::AddMessage(StatFName, EStatOperation::Subtract, int64(Amount));\
+	{ \
+		FThreadStats::AddMessage(StatFName, EStatOperation::Subtract, int64(Amount));\
+		TRACE_STAT_ADD(StatFName, -int64(Amount)); \
+	} \
 }
 #define DEC_MEMORY_STAT_BY(Stat,Amount) \
 {\
 	if (Amount != 0) \
+	{ \
 		if (FThreadStats::IsCollectingData() || !GET_STATISEVERYFRAME(Stat)) \
+		{ \
 			FThreadStats::AddMessage(GET_STATFNAME(Stat), EStatOperation::Subtract, int64(Amount));\
+			TRACE_STAT_ADD(GET_STATFNAME(Stat), -int64(Amount)); \
+		} \
+	} \
 }
 #define SET_MEMORY_STAT(Stat,Value) \
 {\
 	if (FThreadStats::IsCollectingData() || !GET_STATISEVERYFRAME(Stat)) \
+	{ \
 		FThreadStats::AddMessage(GET_STATFNAME(Stat), EStatOperation::Set, int64(Value));\
+		TRACE_STAT_SET(GET_STATFNAME(Stat), int64(Value)); \
+	} \
 }
 #define SET_DWORD_STAT(Stat,Value) \
 {\
 	if (FThreadStats::IsCollectingData() || !GET_STATISEVERYFRAME(Stat)) \
+	{ \
 		FThreadStats::AddMessage(GET_STATFNAME(Stat), EStatOperation::Set, int64(Value));\
+		TRACE_STAT_SET(GET_STATFNAME(Stat), int64(Value)); \
+	} \
 }
 #define SET_FLOAT_STAT(Stat,Value) \
 {\
 	if (FThreadStats::IsCollectingData() || !GET_STATISEVERYFRAME(Stat)) \
+	{ \
 		FThreadStats::AddMessage(GET_STATFNAME(Stat), EStatOperation::Set, double(Value));\
+		TRACE_STAT_SET(GET_STATFNAME(Stat), double(Value)); \
+	} \
 }
+
 #define STAT_ADD_CUSTOMMESSAGE_NAME(Stat,Value) \
 {\
 	FThreadStats::AddMessage(GET_STATFNAME(Stat), EStatOperation::SpecialMessageMarker, FName(Value));\
@@ -2031,52 +2114,75 @@ struct FStat_##StatName\
 #define INC_DWORD_STAT_FName(Stat) \
 {\
 	FThreadStats::AddMessage(Stat, EStatOperation::Add, int64(1));\
+	TRACE_STAT_INCREMENT(Stat); \
 }
 #define INC_FLOAT_STAT_BY_FName(Stat, Amount) \
 {\
 	if (Amount != 0.0f) \
+	{ \
 		FThreadStats::AddMessage(Stat, EStatOperation::Add, double(Amount));\
+		TRACE_STAT_ADD(Stat, double(Amount)); \
+	} \
 }
 #define INC_DWORD_STAT_BY_FName(Stat, Amount) \
 {\
 	if (Amount != 0) \
+	{ \
 		FThreadStats::AddMessage(Stat, EStatOperation::Add, int64(Amount));\
+		TRACE_STAT_ADD(Stat, int64(Amount)); \
+	} \
 }
 #define INC_MEMORY_STAT_BY_FName(Stat, Amount) \
 {\
 	if (Amount != 0) \
+	{ \
 		FThreadStats::AddMessage(Stat, EStatOperation::Add, int64(Amount));\
+		TRACE_STAT_ADD(Stat, int64(Amount)); \
+	} \
 }
 #define DEC_DWORD_STAT_FName(Stat) \
 {\
 	FThreadStats::AddMessage(Stat, EStatOperation::Subtract, int64(1));\
+	TRACE_STAT_DECREMENT(Stat); \
 }
 #define DEC_FLOAT_STAT_BY_FName(Stat,Amount) \
 {\
 	if (Amount != 0.0f) \
+	{ \
 		FThreadStats::AddMessage(Stat, EStatOperation::Subtract, double(Amount));\
+		TRACE_STAT_ADD(Stat, -double(Amount)); \
+	} \
 }
 #define DEC_DWORD_STAT_BY_FName(Stat,Amount) \
 {\
 	if (Amount != 0) \
+	{ \
 		FThreadStats::AddMessage(Stat, EStatOperation::Subtract, int64(Amount));\
+		TRACE_STAT_ADD(Stat, -int64(Amount)); \
+	} \
 }
 #define DEC_MEMORY_STAT_BY_FName(Stat,Amount) \
 {\
 	if (Amount != 0) \
+	{ \
 		FThreadStats::AddMessage(Stat, EStatOperation::Subtract, int64(Amount));\
+		TRACE_STAT_ADD(Stat, -int64(Amount)); \
+	} \
 }
 #define SET_MEMORY_STAT_FName(Stat,Value) \
 {\
 	FThreadStats::AddMessage(Stat, EStatOperation::Set, int64(Value));\
+	TRACE_STAT_SET(Stat, int64(Value)); \
 }
 #define SET_DWORD_STAT_FName(Stat,Value) \
 {\
 	FThreadStats::AddMessage(Stat, EStatOperation::Set, int64(Value));\
+	TRACE_STAT_SET(Stat, int64(Value)); \
 }
 #define SET_FLOAT_STAT_FName(Stat,Value) \
 {\
 	FThreadStats::AddMessage(Stat, EStatOperation::Set, double(Value));\
+	TRACE_STAT_SET(Stat, double(Value)); \
 }
 
 
diff --git a/Engine/Source/Runtime/CoreUObject/Private/Serialization/AsyncLoading.cpp b/Engine/Source/Runtime/CoreUObject/Private/Serialization/AsyncLoading.cpp
--- a/Engine/Source/Runtime/CoreUObject/Private/Serialization/AsyncLoading.cpp
+++ b/Engine/Source/Runtime/CoreUObject/Private/Serialization/AsyncLoading.cpp
@@ -41,6 +41,8 @@
 #include "HAL/LowLevelMemTracker.h"
 #include "ProfilingDebugging/CsvProfiler.h"
 #include "UObject/GCScopeLock.h"
+#include "ProfilingDebugging/MiscTrace.h"
+#include "Serialization/LoadTimeTracePrivate.h"
 
 #define FIND_MEMORY_STOMPS (1 && (PLATFORM_WINDOWS || PLATFORM_UNIX) && !WITH_EDITORONLY_DATA)
 
@@ -68,7 +70,7 @@ DECLARE_CYCLE_STAT(TEXT("CreateImports AsyncPackage"),STAT_FAsyncPackage_CreateI
 DECLARE_CYCLE_STAT(TEXT("CreateMetaData AsyncPackage"),STAT_FAsyncPackage_CreateMetaData,STATGROUP_AsyncLoad);
 DECLARE_CYCLE_STAT(TEXT("CreateExports AsyncPackage"),STAT_FAsyncPackage_CreateExports,STATGROUP_AsyncLoad);
 DECLARE_CYCLE_STAT(TEXT("FreeReferencedImports AsyncPackage"), STAT_FAsyncPackage_FreeReferencedImports, STATGROUP_AsyncLoad);
-DECLARE_CYCLE_STAT(TEXT("Precache ArchiveAsync"), STAT_FArchiveAsync_Precache, STATGROUP_AsyncLoad);
+DECLARE_CYCLE_STAT(TEXT("Precache AsyncArchive"), STAT_FAsyncArchive_Precache, STATGROUP_AsyncLoad);
 DECLARE_CYCLE_STAT(TEXT("PreLoadObjects AsyncPackage"),STAT_FAsyncPackage_PreLoadObjects,STATGROUP_AsyncLoad);
 DECLARE_CYCLE_STAT(TEXT("ExternalReadDependencies AsyncPackage"),STAT_FAsyncPackage_ExternalReadDependencies,STATGROUP_AsyncLoad);
 DECLARE_CYCLE_STAT(TEXT("PostLoadObjects AsyncPackage"),STAT_FAsyncPackage_PostLoadObjects,STATGROUP_AsyncLoad);
@@ -296,7 +298,7 @@ FORCEINLINE bool FAsyncPackage::IsTimeLimitExceeded()
 }
 
 DEFINE_LOG_CATEGORY_STATIC(LogAsyncArchive, Display, All);
-DECLARE_MEMORY_STAT(TEXT("FArchiveAsync2 Buffers"), STAT_FArchiveAsync2Mem, STATGROUP_Memory);
+DECLARE_MEMORY_STAT(TEXT("FAsyncArchive Buffers"), STAT_FAsyncArchiveMem, STATGROUP_Memory);
 
 //#define ASYNC_WATCH_FILE "SM_Boots_Wings.uasset"
 #define TRACK_SERIALIZE (0)
@@ -305,7 +307,7 @@ DECLARE_MEMORY_STAT(TEXT("FArchiveAsync2 Buffers"), STAT_FArchiveAsync2Mem, STAT
 
 
 
-FORCEINLINE void FArchiveAsync2::LogItem(const TCHAR* Item, int64 Offset, int64 Size, double StartTime)
+FORCEINLINE void FAsyncArchive::LogItem(const TCHAR* Item, int64 Offset, int64 Size, double StartTime)
 {
 	if (UE_LOG_ACTIVE(LogAsyncArchive, Verbose)
 #if defined(ASYNC_WATCH_FILE)
@@ -1788,7 +1790,11 @@ void FAsyncLoadingThread::QueueEvent_CreateLinker(FAsyncPackage* Package, int32
 	check(Package);
 	Package->AddNode(EEventLoadNode::Package_LoadSummary);
 	FWeakAsyncPackagePtr WeakPtr(Package);
-	EventQueue.AddAsyncEvent(Package->GetPriority(), GRandomizeLoadOrder ? GetRandomSerialNumber() : Package->SerialNumber, EventSystemPriority,
+
+	TAsyncLoadPriority UserPriority = Package->GetPriority();
+	int32 PackageSerialNumber = GRandomizeLoadOrder ? GetRandomSerialNumber() : Package->SerialNumber;
+	TRACE_LOADTIME_QUEUE_EVENT(Package, LoadTimeProfilerPackageEventType_CreateLinker, UserPriority, PackageSerialNumber, EventSystemPriority);
+	EventQueue.AddAsyncEvent(UserPriority, PackageSerialNumber, EventSystemPriority,
 		TFunction<void(FAsyncLoadEventArgs& Args)>(
 			[WeakPtr, this](FAsyncLoadEventArgs& Args)
 			{
@@ -1796,6 +1802,7 @@ void FAsyncLoadingThread::QueueEvent_CreateLinker(FAsyncPackage* Package, int32
 				check(Pkg);
 				if (Pkg)
 				{
+					TRACE_LOADTIME_ASYNC_PACKAGE_SCOPE(Pkg, LoadTimeProfilerPackageEventType_CreateLinker);
 					Pkg->SetTimeLimit(Args, TEXT("Create Linker"));
 					Pkg->Event_CreateLinker();
 					Args.OutLastObjectWorkWasPerformedOn = Pkg->GetLinkerRoot();
@@ -1850,7 +1857,10 @@ void FAsyncLoadingThread::QueueEvent_FinishLinker(FWeakAsyncPackagePtr WeakPtr,
 	FAsyncPackage* Pkg = GetPackage(WeakPtr);
 	if (Pkg)
 	{
-		EventQueue.AddAsyncEvent(Pkg->GetPriority(), GRandomizeLoadOrder ? GetRandomSerialNumber() : Pkg->SerialNumber, EventSystemPriority,
+		TAsyncLoadPriority UserPriority = Pkg->GetPriority();
+		int32 PackageSerialNumber = GRandomizeLoadOrder ? GetRandomSerialNumber() : Pkg->SerialNumber;
+		TRACE_LOADTIME_QUEUE_EVENT(Pkg, LoadTimeProfilerPackageEventType_FinishLinker, UserPriority, PackageSerialNumber, EventSystemPriority)
+		EventQueue.AddAsyncEvent(UserPriority, PackageSerialNumber, EventSystemPriority,
 			TFunction<void(FAsyncLoadEventArgs& Args)>(
 				[WeakPtr, this](FAsyncLoadEventArgs& Args)
 		{
@@ -1858,6 +1868,7 @@ void FAsyncLoadingThread::QueueEvent_FinishLinker(FWeakAsyncPackagePtr WeakPtr,
 			check(PkgInner);
 			if (PkgInner)
 			{
+				TRACE_LOADTIME_ASYNC_PACKAGE_SCOPE(PkgInner, LoadTimeProfilerPackageEventType_FinishLinker);
 				PkgInner->SetTimeLimit(Args, TEXT("Finish Linker"));
 				PkgInner->Event_FinishLinker();
 			}
@@ -1994,13 +2005,18 @@ void FAsyncLoadingThread::QueueEvent_StartImportPackages(FAsyncPackage* Package,
 {
 	check(Package);
 	FWeakAsyncPackagePtr WeakPtr(Package);
-	EventQueue.AddAsyncEvent(Package->GetPriority(), GRandomizeLoadOrder ? GetRandomSerialNumber() : Package->SerialNumber, EventSystemPriority,
+
+	TAsyncLoadPriority UserPriority = Package->GetPriority();
+	int32 PackageSerialNumber = GRandomizeLoadOrder ? GetRandomSerialNumber() : Package->SerialNumber;
+	TRACE_LOADTIME_QUEUE_EVENT(Package, LoadTimeProfilerPackageEventType_StartImportPackages, UserPriority, PackageSerialNumber, EventSystemPriority);
+	EventQueue.AddAsyncEvent(UserPriority, PackageSerialNumber, EventSystemPriority,
 		TFunction<void(FAsyncLoadEventArgs& Args)>(
 			[WeakPtr, this](FAsyncLoadEventArgs& Args)
 	{
 		FAsyncPackage* Pkg = GetPackage(WeakPtr);
 		if (Pkg)
 		{
+			TRACE_LOADTIME_ASYNC_PACKAGE_SCOPE(Pkg, LoadTimeProfilerPackageEventType_StartImportPackages);
 			Pkg->SetTimeLimit(Args, TEXT("Start Import Packages"));
 			Pkg->Event_StartImportPackages();
 		}
@@ -2161,6 +2177,7 @@ static const FName PrestreamPackageClassNameLoad = FName("PrestreamPackage");
 EAsyncPackageState::Type FAsyncPackage::LoadImports_Event()
 {
 	SCOPE_CYCLE_COUNTER(STAT_FAsyncPackage_LoadImports);
+	SCOPED_LOADTIMER(LoadImports_Event);
 	LastObjectWorkWasPerformedOn = LinkerRoot;
 	LastTypeOfWorkPerformed = TEXT("loading imports event");
 
@@ -2332,13 +2349,17 @@ void FAsyncLoadingThread::QueueEvent_SetupImports(FAsyncPackage* Package, int32
 	Package->AsyncPackageLoadingState = EAsyncPackageLoadingState::SetupImports;
 	check(Package);
 	FWeakAsyncPackagePtr WeakPtr(Package);
-	EventQueue.AddAsyncEvent(Package->GetPriority(), GRandomizeLoadOrder ? GetRandomSerialNumber() : Package->SerialNumber, EventSystemPriority,
+	TAsyncLoadPriority UserPriority = Package->GetPriority();
+	int32 PackageSerialNumber = GRandomizeLoadOrder ? GetRandomSerialNumber() : Package->SerialNumber;
+	TRACE_LOADTIME_QUEUE_EVENT(Package, LoadTimeProfilerPackageEventType_SetupImports, UserPriority, PackageSerialNumber, EventSystemPriority);
+	EventQueue.AddAsyncEvent(UserPriority, PackageSerialNumber, EventSystemPriority,
 		TFunction<void(FAsyncLoadEventArgs& Args)>(
 			[WeakPtr, this](FAsyncLoadEventArgs& Args)
 	{
 		FAsyncPackage* Pkg = GetPackage(WeakPtr);
 		if (Pkg)
 		{
+			TRACE_LOADTIME_ASYNC_PACKAGE_SCOPE(Pkg, LoadTimeProfilerPackageEventType_SetupImports);
 			Pkg->SetTimeLimit(Args, TEXT("Setup Imports"));
 			Pkg->Event_SetupImports();
 		}
@@ -2676,10 +2697,10 @@ EAsyncPackageState::Type FAsyncPackage::SetupImports_Event()
 	}
 
 #if 0
-	if (bAnyImportArcsAdded && Linker->GetFArchiveAsync2Loader())
+	if (bAnyImportArcsAdded && Linker->GetAsyncLoader())
 	{
 		// we are waiting for imports, so drop our precache requests
-		Linker->GetFArchiveAsync2Loader()->FlushCache();
+		Linker->GetAsyncLoader()->FlushCache();
 	}
 #endif
 	return ImportIndex == Linker->ImportMap.Num() ? EAsyncPackageState::Complete : EAsyncPackageState::TimeOut;
@@ -2691,13 +2712,17 @@ void FAsyncLoadingThread::QueueEvent_SetupExports(FAsyncPackage* Package, int32
 	check(Package->AsyncPackageLoadingState == EAsyncPackageLoadingState::SetupExports);
 	check(Package);
 	FWeakAsyncPackagePtr WeakPtr(Package);
-	EventQueue.AddAsyncEvent(Package->GetPriority(), GRandomizeLoadOrder ? GetRandomSerialNumber() : Package->SerialNumber, EventSystemPriority,
+	TAsyncLoadPriority UserPriority = Package->GetPriority();
+	int32 PackageSerialNumber = GRandomizeLoadOrder ? GetRandomSerialNumber() : Package->SerialNumber;
+	TRACE_LOADTIME_QUEUE_EVENT(Package, LoadTimeProfilerPackageEventType_SetupExports, UserPriority, PackageSerialNumber, EventSystemPriority);
+	EventQueue.AddAsyncEvent(UserPriority, PackageSerialNumber, EventSystemPriority,
 		TFunction<void(FAsyncLoadEventArgs& Args)>(
 			[WeakPtr, this](FAsyncLoadEventArgs& Args)
 	{
 		FAsyncPackage* Pkg = GetPackage(WeakPtr);
 		if (Pkg)
 		{
+			TRACE_LOADTIME_ASYNC_PACKAGE_SCOPE(Pkg, LoadTimeProfilerPackageEventType_SetupExports);
 			Pkg->SetTimeLimit(Args, TEXT("Setup Exports"));
 			Pkg->Event_SetupExports();
 		}
@@ -2725,6 +2750,9 @@ void FAsyncLoadingThread::QueueEvent_ProcessImportsAndExports(FAsyncPackage* Pac
 	check(Package->AsyncPackageLoadingState == EAsyncPackageLoadingState::ProcessNewImportsAndExports);
 	check(Package);
 	FWeakAsyncPackagePtr WeakPtr(Package);
+	TAsyncLoadPriority UserPriority = Package->GetPriority();
+	int32 PackageSerialNumber = GRandomizeLoadOrder ? GetRandomSerialNumber() : Package->SerialNumber;
+	TRACE_LOADTIME_QUEUE_EVENT(Package, LoadTimeProfilerPackageEventType_ProcessImportsAndExports, UserPriority, PackageSerialNumber, EventSystemPriority);
 	EventQueue.AddAsyncEvent(Package->GetPriority(), GRandomizeLoadOrder ? GetRandomSerialNumber() : Package->SerialNumber, EventSystemPriority,
 		TFunction<void(FAsyncLoadEventArgs& Args)>(
 			[WeakPtr, this](FAsyncLoadEventArgs& Args)
@@ -2732,6 +2760,7 @@ void FAsyncLoadingThread::QueueEvent_ProcessImportsAndExports(FAsyncPackage* Pac
 		FAsyncPackage* Pkg = GetPackage(WeakPtr);
 		if (Pkg)
 		{
+			TRACE_LOADTIME_ASYNC_PACKAGE_SCOPE(Pkg, LoadTimeProfilerPackageEventType_ProcessImportsAndExports);
 			Pkg->SetTimeLimit(Args, TEXT("ProcessImportsAndExports"));
 			Pkg->Event_ProcessImportsAndExports();
 		}
@@ -2744,13 +2773,17 @@ void FAsyncLoadingThread::QueueEvent_ProcessPostloadWait(FAsyncPackage* Package,
 	check(Package->AsyncPackageLoadingState == EAsyncPackageLoadingState::WaitingForPostLoad);
 	check(Package);
 	FWeakAsyncPackagePtr WeakPtr(Package);
-	EventQueue.AddAsyncEvent(Package->GetPriority(), GRandomizeLoadOrder ? GetRandomSerialNumber() : Package->SerialNumber, EventSystemPriority,
+	TAsyncLoadPriority UserPriority = Package->GetPriority();
+	int32 PackageSerialNumber = GRandomizeLoadOrder ? GetRandomSerialNumber() : Package->SerialNumber;
+	TRACE_LOADTIME_QUEUE_EVENT(Package, LoadTimeProfilerPackageEventType_PostLoadWait, UserPriority, PackageSerialNumber, EventSystemPriority);
+	EventQueue.AddAsyncEvent(UserPriority, PackageSerialNumber, EventSystemPriority,
 		TFunction<void(FAsyncLoadEventArgs& Args)>(
 			[WeakPtr, this](FAsyncLoadEventArgs& Args)
 	{
 		FAsyncPackage* Pkg = GetPackage(WeakPtr);
 		if (Pkg)
 		{
+			TRACE_LOADTIME_ASYNC_PACKAGE_SCOPE(Pkg, LoadTimeProfilerPackageEventType_PostLoadWait);
 			Pkg->SetTimeLimit(Args, TEXT("Process Process Postload Wait"));
 			Pkg->Event_ProcessPostloadWait();
 		}
@@ -2763,13 +2796,17 @@ void FAsyncLoadingThread::QueueEvent_ExportsDone(FAsyncPackage* Package, int32 E
 	check(Package->AsyncPackageLoadingState == EAsyncPackageLoadingState::ProcessNewImportsAndExports);
 	check(Package);
 	FWeakAsyncPackagePtr WeakPtr(Package);
-	EventQueue.AddAsyncEvent(Package->GetPriority(), GRandomizeLoadOrder ? GetRandomSerialNumber() : Package->SerialNumber, EventSystemPriority,
+	TAsyncLoadPriority UserPriority = Package->GetPriority();
+	int32 PackageSerialNumber = GRandomizeLoadOrder ? GetRandomSerialNumber() : Package->SerialNumber;
+	TRACE_LOADTIME_QUEUE_EVENT(Package, LoadTimeProfilerPackageEventType_ExportsDone, UserPriority, PackageSerialNumber, EventSystemPriority);
+	EventQueue.AddAsyncEvent(UserPriority, PackageSerialNumber, EventSystemPriority,
 		TFunction<void(FAsyncLoadEventArgs& Args)>(
 			[WeakPtr, this](FAsyncLoadEventArgs& Args)
 	{
 		FAsyncPackage* Pkg = GetPackage(WeakPtr);
 		if (Pkg)
 		{
+			TRACE_LOADTIME_ASYNC_PACKAGE_SCOPE(Pkg, LoadTimeProfilerPackageEventType_ExportsDone);
 			Pkg->SetTimeLimit(Args, TEXT("Exports Done"));
 			Pkg->Event_ExportsDone();
 		}
@@ -2782,6 +2819,9 @@ void FAsyncLoadingThread::QueueEvent_StartPostLoad(FAsyncPackage* Package, int32
 	check(Package->AsyncPackageLoadingState == EAsyncPackageLoadingState::ReadyForPostLoad);
 	check(Package);
 	FWeakAsyncPackagePtr WeakPtr(Package);
+	TAsyncLoadPriority UserPriority = Package->GetPriority();
+	int32 PackageSerialNumber = GRandomizeLoadOrder ? GetRandomSerialNumber() : Package->SerialNumber;
+	TRACE_LOADTIME_QUEUE_EVENT(Package, LoadTimeProfilerPackageEventType_StartPostLoad, UserPriority, PackageSerialNumber, EventSystemPriority);
 	EventQueue.AddAsyncEvent(Package->GetPriority(), GRandomizeLoadOrder ? GetRandomSerialNumber() : Package->SerialNumber, EventSystemPriority,
 		TFunction<void(FAsyncLoadEventArgs& Args)>(
 			[WeakPtr, this](FAsyncLoadEventArgs& Args)
@@ -2789,6 +2829,7 @@ void FAsyncLoadingThread::QueueEvent_StartPostLoad(FAsyncPackage* Package, int32
 		FAsyncPackage* Pkg = GetPackage(WeakPtr);
 		if (Pkg)
 		{
+			TRACE_LOADTIME_ASYNC_PACKAGE_SCOPE(Pkg, LoadTimeProfilerPackageEventType_StartPostLoad);
 			Pkg->SetTimeLimit(Args, TEXT("Start Post Load"));
 			Pkg->Event_StartPostload();
 		}
@@ -2847,7 +2888,7 @@ EAsyncPackageState::Type FAsyncPackage::SetupExports_Event()
 	FGCScopeGuard GCGuard;
 
 	FCheckedWeakAsyncPackagePtr WeakThis(this);
-	Linker->GetFArchiveAsync2Loader()->LogItem(TEXT("SetupExports_Event"));
+	Linker->GetAsyncLoader()->LogItem(TEXT("SetupExports_Event"));
 
 	LastTypeOfWorkPerformed = TEXT("SetupExports_Event");
 	LastObjectWorkWasPerformedOn = LinkerRoot;
@@ -2963,9 +3004,9 @@ void FAsyncPackage::LinkImport(int32 LocalImportIndex)
 	if (!Import.XObject && !Import.bImportFailed)
 	{
 		FScopedAddObjectreference OnExitAddReference(*this, Import.XObject);
-		if (Linker->GetFArchiveAsync2Loader())
+		if (Linker->GetAsyncLoader())
 		{
-			Linker->GetFArchiveAsync2Loader()->LogItem(TEXT("LinkImport"));
+			Linker->GetAsyncLoader()->LogItem(TEXT("LinkImport"));
 		}
 		if (Import.SourceLinker)
 		{
@@ -3207,6 +3248,8 @@ void FAsyncPackage::EventDrivenCreateExport(int32 LocalExportIndex)
 	SCOPED_LOADTIMER(Package_CreateExports);
 	FObjectExport& Export = Linker->ExportMap[LocalExportIndex];
 
+	TRACE_LOADTIME_CREATE_EXPORT_SCOPE(Linker, &Export.Object, Export.SerialOffset, Export.SerialSize, Export.bIsAsset);
+
 	LLM_SCOPE(ELLMTag::AsyncLoading);
 	LLM_SCOPED_TAG_WITH_OBJECT_IN_SET(GetLinkerRoot(), ELLMTagSet::Assets);
 	LLM_SCOPED_TAG_WITH_OBJECT_IN_SET((Export.DynamicType == FObjectExport::EDynamicType::DynamicType) ? UDynamicClass::StaticClass() : 
@@ -3224,9 +3267,9 @@ void FAsyncPackage::EventDrivenCreateExport(int32 LocalExportIndex)
 		{
 			SCOPED_ACCUM_LOADTIME(Construction, StaticGetNativeClassName(CastEventDrivenIndexToObject<UClass>(Export.ClassIndex, false)));
 
-			if (Linker->GetFArchiveAsync2Loader())
+			if (Linker->GetAsyncLoader())
 			{
-				Linker->GetFArchiveAsync2Loader()->LogItem(TEXT("EventDrivenCreateExport"), Export.SerialOffset, Export.SerialSize);
+				Linker->GetAsyncLoader()->LogItem(TEXT("EventDrivenCreateExport"), Export.SerialOffset, Export.SerialSize);
 			}
 			LastTypeOfWorkPerformed = TEXT("EventDrivenCreateExport");
 			LastObjectWorkWasPerformedOn = nullptr;
@@ -3547,7 +3590,7 @@ void FAsyncPackage::EventDrivenSerializeExport(int32 LocalExportIndex)
 	else if (Object && Object->HasAnyFlags(RF_NeedLoad))
 	{
 
-		Linker->GetFArchiveAsync2Loader()->LogItem(TEXT("EventDrivenSerializeExport"), Export.SerialOffset, Export.SerialSize);
+		Linker->GetAsyncLoader()->LogItem(TEXT("EventDrivenSerializeExport"), Export.SerialOffset, Export.SerialSize);
 
 		LastTypeOfWorkPerformed = TEXT("EventDrivenSerializeExport");
 		LastObjectWorkWasPerformedOn = Object;
@@ -3581,14 +3624,16 @@ void FAsyncPackage::EventDrivenSerializeExport(int32 LocalExportIndex)
 		}
 		check(Export.SerialOffset >= CurrentBlockOffset && Export.SerialOffset + Export.SerialSize <= CurrentBlockOffset + CurrentBlockBytes);
 
-		FArchiveAsync2* FAA2 = Linker->GetFArchiveAsync2Loader();
-		check(FAA2);
+		FAsyncArchive* AsyncLoader = Linker->GetAsyncLoader();
+		check(AsyncLoader);
 
-		const int64 SavedPos = FAA2->Tell();
-		FAA2->Seek(Export.SerialOffset);
+		const int64 SavedPos = AsyncLoader->Tell();
+		AsyncLoader->Seek(Export.SerialOffset);
 
 		Object->ClearFlags(RF_NeedLoad);
 
+		TRACE_LOADTIME_OBJECT_SCOPE(Object, LoadTimeProfilerObjectEventType_Serialize);
+
 		FUObjectSerializeContext* LoadContext = GetSerializeContext();
 		UObject* PrevSerializedObject = LoadContext->SerializedObject;
 		LoadContext->SerializedObject = Object;
@@ -3619,19 +3664,19 @@ void FAsyncPackage::EventDrivenSerializeExport(int32 LocalExportIndex)
 		LoadContext->SerializedObject = PrevSerializedObject;
 		Linker->bForceSimpleIndexToObject = false;
 
-		if (FAA2->Tell() - Export.SerialOffset != Export.SerialSize)
+		if (AsyncLoader->Tell() - Export.SerialOffset != Export.SerialSize)
 		{
 			if (Object->GetClass()->HasAnyClassFlags(CLASS_Deprecated))
 			{
-				UE_LOG(LogStreaming, Warning, TEXT("%s"), *FString::Printf(TEXT("%s: Serial size mismatch: Got %d, Expected %d"), *Object->GetFullName(), (int32)(FAA2->Tell() - Export.SerialOffset), Export.SerialSize));
+				UE_LOG(LogStreaming, Warning, TEXT("%s"), *FString::Printf(TEXT("%s: Serial size mismatch: Got %d, Expected %d"), *Object->GetFullName(), (int32)(AsyncLoader->Tell() - Export.SerialOffset), Export.SerialSize));
 			}
 			else
 			{
-				UE_LOG(LogStreaming, Fatal, TEXT("%s"), *FString::Printf(TEXT("%s: Serial size mismatch: Got %d, Expected %d"), *Object->GetFullName(), (int32)(FAA2->Tell() - Export.SerialOffset), Export.SerialSize));
+				UE_LOG(LogStreaming, Fatal, TEXT("%s"), *FString::Printf(TEXT("%s: Serial size mismatch: Got %d, Expected %d"), *Object->GetFullName(), (int32)(AsyncLoader->Tell() - Export.SerialOffset), Export.SerialSize));
 			}
 		}
 
-		FAA2->Seek(SavedPos);
+		AsyncLoader->Seek(SavedPos);
 #if DO_CHECK
 		if (Object->HasAnyFlags(RF_ClassDefaultObject) && Object->GetClass()->HasAnyClassFlags(CLASS_CompiledFromBlueprint))
 		{
@@ -3651,6 +3696,7 @@ void FAsyncPackage::EventDrivenSerializeExport(int32 LocalExportIndex)
 
 void FAsyncPackage::StartPrecacheRequest()
 {
+	SCOPED_LOADTIMER(StartPrecacheRequests);
 	if (Linker->bDynamicClassLinker)
 	{
 //native blueprint 
@@ -3764,10 +3810,10 @@ void FAsyncPackage::StartPrecacheRequest()
 		}
 	}
 	check(NewReq.ExportsToRead.Num());
-	FArchiveAsync2* FAA2 = Linker->GetFArchiveAsync2Loader();
-	check(FAA2);
+	FAsyncArchive* AsyncLoader = Linker->GetAsyncLoader();
+	check(AsyncLoader);
 
-	IAsyncReadRequest* Precache = FAA2->MakeEventDrivenPrecacheRequest(NewReq.Offset, NewReq.BytesToRead,GPrecacheCallbackHandler.GetCompletionCallback());
+	IAsyncReadRequest* Precache = AsyncLoader->MakeEventDrivenPrecacheRequest(NewReq.Offset, NewReq.BytesToRead,GPrecacheCallbackHandler.GetCompletionCallback());
 
 	NewReq.FirstExportCovered = LocalExportIndex;
 	NewReq.LastExportCovered = LastExportIndex;
@@ -3792,6 +3838,7 @@ int64 FAsyncPackage::PrecacheRequestReady(IAsyncReadRequest * Read)
 
 void FAsyncPackage::MakeNextPrecacheRequestCurrent()
 {
+	SCOPED_LOADTIMER(MakeNextPrecacheRequestCurrent);
 	LLM_SCOPE(ELLMTag::AsyncLoading);
 
 	check(ReadyPrecacheRequests.Num());
@@ -3803,11 +3850,11 @@ void FAsyncPackage::MakeNextPrecacheRequestCurrent()
 
 	GPrecacheCallbackHandler.FinishRequest(Req.BytesToRead);
 
-	FArchiveAsync2* FAA2 = Linker->GetFArchiveAsync2Loader();
-	check(FAA2);
+	FAsyncArchive* AsyncLoader = Linker->GetAsyncLoader();
+	check(AsyncLoader);
 	Read->WaitCompletion();
 
-	bool bReady = FAA2->PrecacheForEvent(CurrentBlockOffset, CurrentBlockBytes);
+	bool bReady = AsyncLoader->PrecacheForEvent(CurrentBlockOffset, CurrentBlockBytes);
 	UE_CLOG(!bReady, LogStreaming, Warning, TEXT("Precache request should have been hot %s."), *Linker->Filename);
 	for (int32 Index = Req.FirstExportCovered; Index <= Req.LastExportCovered; Index++)
 	{
@@ -3824,13 +3871,14 @@ void FAsyncPackage::MakeNextPrecacheRequestCurrent()
 
 void FAsyncPackage::FlushPrecacheBuffer()
 {
+	SCOPED_LOADTIMER(FlushPrecacheBuffer);
 	CurrentBlockOffset = -1;
 	CurrentBlockBytes = -1;
 	if (!Linker->bDynamicClassLinker)
 	{
-	FArchiveAsync2* FAA2 = Linker->GetFArchiveAsync2Loader();
-	check(FAA2);
-	FAA2->FlushPrecacheBlock();
+	FAsyncArchive* AsyncLoader = Linker->GetAsyncLoader();
+	check(AsyncLoader);
+	AsyncLoader->FlushPrecacheBlock();
 }
 }
 
@@ -3838,6 +3886,7 @@ int32 GCurrentExportIndex = -1;
 
 EAsyncPackageState::Type FAsyncPackage::ProcessImportsAndExports_Event()
 {
+	SCOPED_LOADTIMER(ProcessImportsAndExports_Event);
 	check(Linker);
 	bool bDidSomething = true;
 	int32 LoopIterations = 0;
@@ -3976,7 +4025,7 @@ EAsyncPackageState::Type FAsyncPackage::ProcessImportsAndExports_Event()
 
 void FAsyncPackage::Event_ExportsDone()
 {
-	Linker->GetFArchiveAsync2Loader()->LogItem(TEXT("Event_ExportsDone"));
+	Linker->GetAsyncLoader()->LogItem(TEXT("Event_ExportsDone"));
 	check(AsyncPackageLoadingState == EAsyncPackageLoadingState::ProcessNewImportsAndExports);
 	bAllExportsSerialized = true;
 	RemoveNode(EEventLoadNode::Package_ExportsSerialized);
@@ -4006,7 +4055,7 @@ void FAsyncPackage::Event_ExportsDone()
 
 void FAsyncPackage::Event_ProcessPostloadWait()
 {
-	Linker->GetFArchiveAsync2Loader()->LogItem(TEXT("Event_ProcessPostloadWait"));
+	Linker->GetAsyncLoader()->LogItem(TEXT("Event_ProcessPostloadWait"));
 	check(AsyncPackageLoadingState == EAsyncPackageLoadingState::WaitingForPostLoad);
 	check(bAllExportsSerialized && !OtherPackagesWaitingForMeBeforePostload.Num());
 	bProcessPostloadWaitInFlight = false;
@@ -4123,7 +4172,7 @@ void FAsyncPackage::Event_StartPostload()
 {
 	LLM_SCOPE(ELLMTag::AsyncLoading);
 
-	Linker->GetFArchiveAsync2Loader()->LogItem(TEXT("Event_StartPostload"));
+	Linker->GetAsyncLoader()->LogItem(TEXT("Event_StartPostload"));
 	check(AsyncPackageLoadingState == EAsyncPackageLoadingState::ReadyForPostLoad);
 	check(!PackagesIMayBeWaitingForBeforePostload.Num());
 	check(!PackagesIAmWaitingForBeforePostload.Num());
@@ -4739,8 +4788,6 @@ bool FAsyncPackage::AreAllDependenciesFullyLoaded(TSet<UPackage*>& VisitedPackag
 
 EAsyncPackageState::Type FAsyncLoadingThread::ProcessLoadedPackages(bool bUseTimeLimit, bool bUseFullTimeLimit, float TimeLimit, bool& bDidSomething, FFlushTree* FlushTree)
 {
-	SCOPED_LOADTIMER(TickAsyncLoading_ProcessLoadedPackages);
-
 	EAsyncPackageState::Type Result = EAsyncPackageState::Complete;
 
 	// This is for debugging purposes only. @todo remove
@@ -4753,10 +4800,16 @@ EAsyncPackageState::Type FAsyncLoadingThread::ProcessLoadedPackages(bool bUseTim
 		FScopeLock LoadedPackagesLock(&LoadedPackagesCritical);
 		FScopeLock LoadedPackagesToProcessLock(&LoadedPackagesToProcessCritical);
 #endif
-		LoadedPackagesToProcess.Append(LoadedPackages);
-		LoadedPackagesToProcessNameLookup.Append(LoadedPackagesNameLookup);
-		LoadedPackages.Reset();
-		LoadedPackagesNameLookup.Reset();
+		if (LoadedPackages.Num() != 0)
+		{
+			LoadedPackagesToProcess.Append(LoadedPackages);
+			LoadedPackages.Reset();
+		}
+		if (LoadedPackagesNameLookup.Num() != 0)
+		{
+			LoadedPackagesToProcessNameLookup.Append(LoadedPackagesNameLookup);
+			LoadedPackagesNameLookup.Reset();
+		}
 	}
 #if USE_EVENT_DRIVEN_ASYNC_LOAD_AT_BOOT_TIME
 	if (IsMultithreaded() && GEventDrivenLoaderEnabled &&
@@ -4776,11 +4829,12 @@ EAsyncPackageState::Type FAsyncLoadingThread::ProcessLoadedPackages(bool bUseTim
 	bDidSomething = LoadedPackagesToProcess.Num() > 0;
 	for (int32 PackageIndex = 0; PackageIndex < LoadedPackagesToProcess.Num() && !IsAsyncLoadingSuspended(); ++PackageIndex)
 	{
-		SCOPED_LOADTIMER(ProcessLoadedPackagesTime);
-
 		FAsyncPackage* Package = LoadedPackagesToProcess[PackageIndex];
 		if (Package->GetDependencyRefCount() == 0)
 		{
+			TRACE_LOADTIME_ASYNC_PACKAGE_SCOPE(Package, LoadTimeProfilerPackageEventType_DeferredPostLoad);
+			SCOPED_LOADTIMER(ProcessLoadedPackagesTime);
+
 			Result = Package->PostLoadDeferredObjects(TickStartTime, bUseTimeLimit, TimeLimit);
 			if (Result == EAsyncPackageState::Complete)
 			{
@@ -4827,7 +4881,6 @@ EAsyncPackageState::Type FAsyncLoadingThread::ProcessLoadedPackages(bool bUseTim
 				const bool bInternalCallbacks = false;
 				const EAsyncLoadingResult::Type LoadingResult = Package->HasLoadFailed() ? EAsyncLoadingResult::Failed : EAsyncLoadingResult::Succeeded;
 				Package->CallCompletionCallbacks(bInternalCallbacks, LoadingResult);
-
 #if WITH_EDITOR
 				// In the editor we need to find any assets and add them to list for later callback
 				Package->GetLoadedAssets(LoadedAssets);
@@ -5076,6 +5129,10 @@ FAsyncLoadingThread::FAsyncLoadingThread()
 			TEXT("Event driven async loader is NOT being used but it seems to be enabled in project settings."));
 	}
 
+#if LOADTIMEPROFILERTRACE_ENABLED
+	FLoadTimeProfilerTracePrivate::Init(FParse::Param(FCommandLine::Get(), TEXT("loadtimetrace")));
+#endif
+
 	QueuedRequestsEvent = FPlatformProcess::GetSynchEventFromPool();
 	CancelLoadingEvent = FPlatformProcess::GetSynchEventFromPool();
 	ThreadSuspendedEvent = FPlatformProcess::GetSynchEventFromPool();
@@ -5162,6 +5219,10 @@ void FAsyncLoadingThread::StartThread()
 		bThreadStarted = true;
 		FPlatformMisc::MemoryBarrier();
 		Thread = FRunnableThread::Create(this, TEXT("FAsyncLoadingThread"), 0, TPri_Normal);
+		if (Thread)
+		{
+			TRACE_SET_THREAD_GROUP(Thread->GetThreadID(), "AsyncLoading");
+		}
 	}
 }
 
@@ -5176,6 +5237,8 @@ uint32 FAsyncLoadingThread::Run()
 
 	AsyncLoadingThreadID = FPlatformTLS::GetCurrentThreadId();
 
+	TRACE_LOADTIME_START_ASYNC_LOADING();
+
 	if (!IsInGameThread())
 	{
 		FPlatformProcess::SetThreadAffinityMask(FPlatformAffinity::GetAsyncLoadingThreadMask());
@@ -5344,6 +5407,7 @@ void FAsyncLoadingThread::SuspendLoading()
 #endif
 	if (IsMultithreaded() && SuspendCount == 1)
 	{
+		TRACE_LOADTIME_SUSPEND_ASYNC_LOADING();
 		ThreadSuspendedEvent->Wait();
 	}
 }
@@ -5359,6 +5423,7 @@ void FAsyncLoadingThread::ResumeLoading()
 	if (IsMultithreaded() && SuspendCount == 0)
 	{
 		ThreadResumedEvent->Wait();
+		TRACE_LOADTIME_RESUME_ASYNC_LOADING();
 	}
 }
 
@@ -5498,11 +5563,13 @@ FAsyncPackage::FAsyncPackage(const FAsyncPackageDesc& InDesc)
 , FinishObjectsTime(0.0)
 #endif // PERF_TRACK_DETAILED_ASYNC_STATS
 {
+	TRACE_LOADTIME_NEW_ASYNC_PACKAGE(this, *InDesc.Name.ToString());
 	AddRequestID(InDesc.RequestID);
 }
 
 FAsyncPackage::~FAsyncPackage()
 {
+	TRACE_LOADTIME_DESTROY_ASYNC_PACKAGE(this);
 #if DO_CHECK
 	if (GEventDrivenLoaderEnabled)
 	{
@@ -5583,6 +5650,7 @@ void FAsyncPackage::AddRequestID(int32 Id)
 		}
 		RequestIDs.Add(Id);
 		AsyncLoadingThread.AddPendingRequest(Id);
+		TRACE_LOADTIME_ASYNC_PACKAGE_REQUEST_ASSOCIATION(this, Id);
 	}
 }
 
@@ -5733,6 +5801,7 @@ EAsyncPackageState::Type FAsyncPackage::TickAsyncPackage(bool InbUseTimeLimit, b
 
 	ReentryCount++;
 
+	TRACE_LOADTIME_ASYNC_PACKAGE_SCOPE(this, LoadTimeProfilerPackageEventType_Tick);
 	SCOPE_CYCLE_COUNTER(STAT_FAsyncPackage_Tick);
 	SCOPED_LOADTIMER(Package_Tick);
 
@@ -6083,7 +6152,7 @@ EAsyncPackageState::Type FAsyncPackage::CreateLinker()
 					if (Linker->bDynamicClassLinker)
 					{
 						//native blueprint 
-						check(!Linker->GetFArchiveAsync2Loader());
+						check(!Linker->GetAsyncLoader());
 						GPrecacheCallbackHandler.SummaryComplete(WeakPtr);
 					}
 				}
@@ -6098,6 +6167,7 @@ EAsyncPackageState::Type FAsyncPackage::CreateLinker()
 		check(Linker);
 		check(Linker->AsyncRoot == nullptr || Linker->AsyncRoot == this);
 		Linker->AsyncRoot = this;
+		TRACE_LOADTIME_ASYNC_PACKAGE_LINKER_ASSOCIATION(this, Linker);
 
 		UE_LOG(LogStreaming, Verbose, TEXT("FAsyncPackage::CreateLinker for %s finished."), *Desc.NameToLoad.ToString());
 	}
@@ -6460,10 +6530,10 @@ EAsyncPackageState::Type FAsyncPackage::CreateExports()
 		// Precache data and see whether it's already finished.
 		bool bReady;
 
-		FArchiveAsync2* FAA2 = Linker->GetFArchiveAsync2Loader();
-		if (FAA2)
+		FAsyncArchive* AsyncLoader = Linker->GetAsyncLoader();
+		if (AsyncLoader)
 		{
-			bReady = FAA2->Precache(Export.SerialOffset, Export.SerialSize, bUseTimeLimit, bUseFullTimeLimit, TickStartTime, TimeLimit);
+			bReady = AsyncLoader->PrecacheWithTimeLimit(Export.SerialOffset, Export.SerialSize, bUseTimeLimit, bUseFullTimeLimit, TickStartTime, TimeLimit);
 		}
 		else
 		{
@@ -6666,7 +6736,10 @@ EAsyncPackageState::Type FAsyncPackage::PostLoadObjects()
 				check(!GEventDrivenLoaderEnabled || !Object->HasAnyFlags(RF_NeedLoad));
 
 				ThreadContext.CurrentlyPostLoadedObjectByALT = Object;
-				Object->ConditionalPostLoad();
+				{
+					TRACE_LOADTIME_OBJECT_SCOPE(Object, LoadTimeProfilerObjectEventType_PostLoad);
+					Object->ConditionalPostLoad();
+				}
 				ThreadContext.CurrentlyPostLoadedObjectByALT = nullptr;
 
 				LastObjectWorkWasPerformedOn = Object;
@@ -6731,7 +6804,10 @@ EAsyncPackageState::Type FAsyncPackage::PostLoadDeferredObjects(double InTickSta
 		FScopeCycleCounterUObject ConstructorScope(Object, GET_STATID(STAT_FAsyncPackage_PostLoadObjectsGameThread));
 
 		PackageScope.ThreadContext.CurrentlyPostLoadedObjectByALT = Object;
-		Object->ConditionalPostLoad();
+		{
+			TRACE_LOADTIME_OBJECT_SCOPE(Object, LoadTimeProfilerObjectEventType_PostLoad);
+			Object->ConditionalPostLoad();
+		}
 		PackageScope.ThreadContext.CurrentlyPostLoadedObjectByALT = nullptr;
 
 		if (ObjLoadedInPostLoad.Num())
@@ -7167,6 +7243,7 @@ int32 LoadPackageAsync(const FString& InName, const FGuid* InGuid /*= nullptr*/,
 		// Generate new request ID and add it immediately to the global request list (it needs to be there before we exit
 		// this function, otherwise it would be added when the packages are being processed on the async thread).
 		RequestID = GPackageRequestID.Increment();
+		TRACE_LOADTIME_BEGIN_REQUEST(RequestID);
 		FAsyncLoadingThread::Get().AddPendingRequest(RequestID);
 
 		// Allocate delegate on Game Thread, it is not safe to copy delegates by value on other threads
@@ -7250,7 +7327,7 @@ void FlushAsyncLoading(int32 PackageID /* = INDEX_NONE */)
 
 	checkf(IsInGameThread(), TEXT("Unable to FlushAsyncLoading from any thread other than the game thread."));
 
-	if (IsAsyncLoading())
+ 	if (IsAsyncLoading())
 	{
 		FAsyncLoadingThread& AsyncThread = FAsyncLoadingThread::Get();
 		// Flushing async loading while loading is suspend will result in infinite stall
@@ -7263,6 +7340,8 @@ void FlushAsyncLoading(int32 PackageID /* = INDEX_NONE */)
 			return;
 		}
 
+		TRACE_LOADTIME_FLUSH_ASYNCLOADING_SCOPE(PackageID);
+
 		FCoreDelegates::OnAsyncLoadingFlush.Broadcast();
 
 #if NO_LOGGING == 0
@@ -7475,10 +7554,10 @@ void NotifyRegistrationComplete()
 	GetGEDLBootNotificationManager().NotifyRegistrationComplete();
 }
 
-#define USE_DETAILED_FARCHIVEASYNC2_MEMORY_TRACKING 0
+#define USE_DETAILED_FASYNCARCHIVE_MEMORY_TRACKING 0
 
-#if USE_DETAILED_FARCHIVEASYNC2_MEMORY_TRACKING
-class FArchiveAsync2MemTracker
+#if USE_DETAILED_FASYNCARCHIVE_MEMORY_TRACKING
+class FAsyncArchiveMemTracker
 {
 	TMap<FString, int64> AllocatedMem;
 	FCriticalSection AllocatedMemCritical;
@@ -7486,14 +7565,14 @@ class FArchiveAsync2MemTracker
 public:
 
 	void Allocate(const FString& Filename, int64 Mem)
-{
+	{
 		FScopeLock AllocatedMemLock(&AllocatedMemCritical);
 		int64& AllocatedMemAmount = AllocatedMem.FindOrAdd(Filename);
 		AllocatedMemAmount += Mem;
-}
+	}
 
 	void Deallocate(const FString& Filename, int64 Mem)
-{
+	{
 		FScopeLock AllocatedMemLock(&AllocatedMemCritical);
 		int64& AllocatedMemAmount = AllocatedMem.FindOrAdd(Filename);
 		AllocatedMemAmount -= Mem;
@@ -7501,36 +7580,36 @@ public:
 		if (AllocatedMemAmount == 0)
 		{
 			AllocatedMem.Remove(Filename);
+		}
 	}
-}
 
 	void Dump()
-{
+	{
 		FScopeLock AllocatedMemLock(&AllocatedMemCritical);
 
-		UE_LOG(LogStreaming, Display, TEXT("Dumping FArchiveAsync2 allocated memory (%d)"), AllocatedMem.Num());
+		UE_LOG(LogStreaming, Display, TEXT("Dumping FAsyncArchie allocated memory (%d)"), AllocatedMem.Num());
 		for (TPair<FString, int64>& ArchiveMem : AllocatedMem)
-{
+		{
 			UE_LOG(LogStreaming, Display, TEXT("  %s %lldb"), *ArchiveMem.Key, ArchiveMem.Value);
-}
-}
-} GArchiveAsync2MemTracker;
+		}
+	}
+} GAsyncArchiveMemTracker;
 
-void DumpArchiveAsync2Mem(const TArray<FString>& Args)
+void DumpAsyncArchiveMem(const TArray<FString>& Args)
 {
-	GArchiveAsync2MemTracker.Dump();
+	GAsyncArchiveMemTracker.Dump();
 }
 
 static FAutoConsoleCommand GDumpSerializeCmd(
-	TEXT("DumpFArchiveAsync2Mem"),
-	TEXT("Debug command to dump the memory allocated by existing FArhiveAsync2."),
-	FConsoleCommandWithArgsDelegate::CreateStatic(&DumpArchiveAsync2Mem)
+	TEXT("DumpAsyncArchiveMem"),
+	TEXT("Debug command to dump the memory allocated by existing FAsyncArchive."),
+	FConsoleCommandWithArgsDelegate::CreateStatic(&DumpAsyncArchiveMem)
 );
-#endif // USE_DETAILED_FARCHIVEASYNC2_MEMORY_TRACKING
+#endif // USE_DETAILED_FASYNCARCHIVE_MEMORY_TRACKING
 
 static FCriticalSection SummaryRacePreventer;
 
-FArchiveAsync2::FArchiveAsync2(const TCHAR* InFileName, TFunction<void()>&& InSummaryReadyCallback)
+FAsyncArchive::FAsyncArchive(const TCHAR* InFileName, TFunction<void()>&& InSummaryReadyCallback)
 	: Handle(nullptr)
 	, SizeRequestPtr(nullptr)
 	, EditorPrecacheRequestPtr(nullptr)
@@ -7577,7 +7656,7 @@ FArchiveAsync2::FArchiveAsync2(const TCHAR* InFileName, TFunction<void()>&& InSu
 
 }
 
-FArchiveAsync2::~FArchiveAsync2()
+FAsyncArchive::~FAsyncArchive()
 {
 	// Invalidate any precached data and free memory.
 	FlushCache();
@@ -7586,10 +7665,10 @@ FArchiveAsync2::~FArchiveAsync2()
 		delete Handle;
 		Handle = nullptr;
 	}
-	LogItem(TEXT("~FArchiveAsync2"), 0, 0);
+	LogItem(TEXT("~FAsyncArchive"), 0, 0);
 }
 
-void FArchiveAsync2::ReadCallback(bool bWasCancelled, IAsyncReadRequest* Request)
+void FAsyncArchive::ReadCallback(bool bWasCancelled, IAsyncReadRequest* Request)
 {
 	if (bWasCancelled || ArIsError)
 	{
@@ -7680,17 +7759,17 @@ void FArchiveAsync2::ReadCallback(bool bWasCancelled, IAsyncReadRequest* Request
 	}
 }
 
-void FArchiveAsync2::FlushPrecacheBlock()
+void FAsyncArchive::FlushPrecacheBlock()
 {
 #if DEVIRTUALIZE_FLinkerLoad_Serialize
 	DiscardInlineBufferAndUpdateCurrentPos();
 #endif
 	if (PrecacheBuffer)
 	{
-		DEC_MEMORY_STAT_BY(STAT_FArchiveAsync2Mem, PrecacheEndPos - PrecacheStartPos);
+		DEC_MEMORY_STAT_BY(STAT_FAsyncArchiveMem, PrecacheEndPos - PrecacheStartPos);
 		FMemory::Free(PrecacheBuffer);
-#if USE_DETAILED_FARCHIVEASYNC2_MEMORY_TRACKING
-		GArchiveAsync2MemTracker.Deallocate(FileName, PrecacheEndPos - PrecacheStartPos);
+#if USE_DETAILED_FASYNCARCHIVE_MEMORY_TRACKING
+		GAsyncArchieMemTracker.Deallocate(FileName, PrecacheEndPos - PrecacheStartPos);
 #endif
 		check(!GEventDrivenLoaderEnabled || LoadPhase > ELoadPhase::WaitingForHeader);
 	}
@@ -7699,7 +7778,7 @@ void FArchiveAsync2::FlushPrecacheBlock()
 	PrecacheEndPos = 0;
 }
 
-void FArchiveAsync2::FlushCache()
+void FAsyncArchive::FlushCache()
 {
 	bool nNonRedundantFlush = PrecacheEndPos || PrecacheBuffer || ReadRequestPtr;
 	LogItem(TEXT("Flush"));
@@ -7749,7 +7828,7 @@ void FArchiveAsync2::FlushCache()
 
 }
 
-bool FArchiveAsync2::Close()
+bool FAsyncArchive::Close()
 {
 	// Invalidate any precached data and free memory.
 	FlushCache();
@@ -7757,13 +7836,13 @@ bool FArchiveAsync2::Close()
 	return !ArIsError;
 }
 
-bool FArchiveAsync2::SetCompressionMap(TArray<FCompressedChunk>* InCompressedChunks, ECompressionFlags InCompressionFlags)
+bool FAsyncArchive::SetCompressionMap(TArray<FCompressedChunk>* InCompressedChunks, ECompressionFlags InCompressionFlags)
 {
 	check(0); // no support for compression
 	return false;
 }
 
-int64 FArchiveAsync2::TotalSize()
+int64 FAsyncArchive::TotalSize()
 {
 	if (SizeRequestPtr)
 	{
@@ -7780,7 +7859,7 @@ int64 FArchiveAsync2::TotalSize()
 }
 
 #if DEVIRTUALIZE_FLinkerLoad_Serialize
-FORCEINLINE void FArchiveAsync2::SetPosAndUpdatePrecacheBuffer(int64 Pos)
+FORCEINLINE void FAsyncArchive::SetPosAndUpdatePrecacheBuffer(int64 Pos)
 {
 	check(Pos >= 0 && Pos <= TotalSizeOrMaxInt64IfNotReady());
 	if (Pos < PrecacheStartPos || Pos >= PrecacheEndPos)
@@ -7800,7 +7879,7 @@ FORCEINLINE void FArchiveAsync2::SetPosAndUpdatePrecacheBuffer(int64 Pos)
 }
 #endif
 
-void FArchiveAsync2::Seek(int64 InPos)
+void FAsyncArchive::Seek(int64 InPos)
 {
 	if ((GEventDrivenLoaderEnabled || bCookedForEDLInEditor) && LoadPhase < ELoadPhase::ProcessingExports)
 	{
@@ -7810,7 +7889,7 @@ void FArchiveAsync2::Seek(int64 InPos)
 			FirstExportStarting();
 		}
 	}
-	checkf(InPos >= 0 && InPos <= TotalSizeOrMaxInt64IfNotReady(), TEXT("Bad position in FArchiveAsync2::Seek. Filename:%s InPos:%lu, Size:%lu"), *FileName, InPos, TotalSizeOrMaxInt64IfNotReady());
+	checkf(InPos >= 0 && InPos <= TotalSizeOrMaxInt64IfNotReady(), TEXT("Bad position in FAsyncArchive::Seek. Filename:%s InPos:%lu, Size:%lu"), *FileName, InPos, TotalSizeOrMaxInt64IfNotReady());
 #if DEVIRTUALIZE_FLinkerLoad_Serialize
 	SetPosAndUpdatePrecacheBuffer(InPos);
 #else
@@ -7818,7 +7897,7 @@ void FArchiveAsync2::Seek(int64 InPos)
 #endif
 }
 
-bool FArchiveAsync2::WaitRead(float TimeLimit)
+bool FAsyncArchive::WaitRead(float TimeLimit)
 {
 	if (ReadRequestPtr)
 	{
@@ -7838,7 +7917,7 @@ bool FArchiveAsync2::WaitRead(float TimeLimit)
 	return true;
 }
 
-void FArchiveAsync2::CompleteRead()
+void FAsyncArchive::CompleteRead()
 {
 	double StartTime = FPlatformTime::Seconds();
 	check(LoadPhase != ELoadPhase::WaitingForSize && LoadPhase != ELoadPhase::WaitingForSummary);
@@ -7860,10 +7939,10 @@ void FArchiveAsync2::CompleteRead()
 			PrecacheStartPos = ReadRequestOffset;
 			PrecacheEndPos = ReadRequestOffset + ReadRequestSize;
 			check(ReadRequestSize > 0 && PrecacheStartPos >= 0);
-			INC_MEMORY_STAT_BY(STAT_FArchiveAsync2Mem, PrecacheEndPos - PrecacheStartPos);
+			INC_MEMORY_STAT_BY(STAT_FAsyncArchiveMem, PrecacheEndPos - PrecacheStartPos);
 			DEC_MEMORY_STAT_BY(STAT_AsyncFileMemory, ReadRequestSize);
-#if USE_DETAILED_FARCHIVEASYNC2_MEMORY_TRACKING
-			GArchiveAsync2MemTracker.Allocate(FileName, PrecacheEndPos - PrecacheStartPos);
+#if USE_DETAILED_FASYNCARCHIVE_MEMORY_TRACKING
+			GAsyncArchiveMemTracker.Allocate(FileName, PrecacheEndPos - PrecacheStartPos);
 #endif
 			// keeps the last cache block of the header around until we process the first export
 			if (LoadPhase != ELoadPhase::ProcessingExports)
@@ -7881,7 +7960,7 @@ void FArchiveAsync2::CompleteRead()
 	ReadRequestSize = 0;
 }
 
-void FArchiveAsync2::CompleteCancel()
+void FAsyncArchive::CompleteCancel()
 {
 	if (CanceledReadRequestPtr)
 	{
@@ -7895,7 +7974,7 @@ void FArchiveAsync2::CompleteCancel()
 }
 
 
-void FArchiveAsync2::CancelRead()
+void FAsyncArchive::CancelRead()
 {
 	if (ReadRequestPtr)
 	{
@@ -7908,7 +7987,7 @@ void FArchiveAsync2::CancelRead()
 	ReadRequestSize = 0;
 }
 
-bool FArchiveAsync2::WaitForIntialPhases(float InTimeLimit)
+bool FAsyncArchive::WaitForIntialPhases(float InTimeLimit)
 {
 	if (SizeRequestPtr 
 		|| GEventDrivenLoaderEnabled || SummaryRequestPtr || SummaryPrecacheRequestPtr
@@ -7981,7 +8060,7 @@ bool FArchiveAsync2::WaitForIntialPhases(float InTimeLimit)
 	return true;
 }
 
-bool FArchiveAsync2::PrecacheInternal(int64 RequestOffset, int64 RequestSize, bool bApplyMinReadSize)
+bool FAsyncArchive::PrecacheInternal(int64 RequestOffset, int64 RequestSize, bool bApplyMinReadSize)
 {
 	// CAUTION! This is possibly called the first time from a random IO thread.
 
@@ -8008,7 +8087,7 @@ bool FArchiveAsync2::PrecacheInternal(int64 RequestOffset, int64 RequestSize, bo
 		if (ReadRequestPtr)
 		{
 			// this one does not have what we need
-			UE_LOG(LogStreaming, Warning, TEXT("FArchiveAsync2::PrecacheInternal Canceled read for %s  Offset = %lld   Size = %lld"), *FileName, RequestOffset, ReadRequestSize);
+			UE_LOG(LogStreaming, Warning, TEXT("FAsyncArchive::PrecacheInternal Canceled read for %s  Offset = %lld   Size = %lld"), *FileName, RequestOffset, ReadRequestSize);
 			CancelRead();
 		}
 	}
@@ -8068,7 +8147,7 @@ bool FArchiveAsync2::PrecacheInternal(int64 RequestOffset, int64 RequestSize, bo
 	return false;
 }
 
-void FArchiveAsync2::FirstExportStarting()
+void FAsyncArchive::FirstExportStarting()
 {
 	ExportReadTime = FPlatformTime::Seconds();
 	LogItem(TEXT("Exports"));
@@ -8098,7 +8177,7 @@ void FArchiveAsync2::FirstExportStarting()
 	}
 }
 
-IAsyncReadRequest* FArchiveAsync2::MakeEventDrivenPrecacheRequest(int64 Offset, int64 BytesToRead, FAsyncFileCallBack* CompleteCallback)
+IAsyncReadRequest* FAsyncArchive::MakeEventDrivenPrecacheRequest(int64 Offset, int64 BytesToRead, FAsyncFileCallBack* CompleteCallback)
 {
 	check(GEventDrivenLoaderEnabled);
 	if (LoadPhase == ELoadPhase::WaitingForFirstExport)
@@ -8143,7 +8222,7 @@ IAsyncReadRequest* FArchiveAsync2::MakeEventDrivenPrecacheRequest(int64 Offset,
 }
 
 
-bool FArchiveAsync2::Precache(int64 RequestOffset, int64 RequestSize, bool bUseTimeLimit, bool bUseFullTimeLimit, double TickStartTime, float TimeLimit)
+bool FAsyncArchive::PrecacheWithTimeLimit(int64 RequestOffset, int64 RequestSize, bool bUseTimeLimit, bool bUseFullTimeLimit, double TickStartTime, float TimeLimit)
 {
 #if defined(ASYNC_WATCH_FILE)
 	if (FileName.Contains(TEXT(ASYNC_WATCH_FILE)) && RequestOffset == 878129)
@@ -8176,7 +8255,7 @@ bool FArchiveAsync2::Precache(int64 RequestOffset, int64 RequestSize, bool bUseT
 	return bResult;
 }
 
-bool FArchiveAsync2::Precache(int64 RequestOffset, int64 RequestSize)
+bool FAsyncArchive::Precache(int64 RequestOffset, int64 RequestSize)
 {
 	if (LoadPhase == ELoadPhase::WaitingForSize || LoadPhase == ELoadPhase::WaitingForSummary)
 	{
@@ -8190,14 +8269,14 @@ bool FArchiveAsync2::Precache(int64 RequestOffset, int64 RequestSize)
 	return PrecacheInternal(RequestOffset, RequestSize);
 }
 
-bool FArchiveAsync2::PrecacheForEvent(int64 RequestOffset, int64 RequestSize)
+bool FAsyncArchive::PrecacheForEvent(int64 RequestOffset, int64 RequestSize)
 {
 	check(int32(LoadPhase) > int32(ELoadPhase::WaitingForHeader));
 	return PrecacheInternal(RequestOffset, RequestSize, false);
 }
 
 
-void FArchiveAsync2::StartReadingHeader()
+void FAsyncArchive::StartReadingHeader()
 {
 	//LogItem(TEXT("Start Header"));
 	WaitForIntialPhases();
@@ -8212,7 +8291,7 @@ void FArchiveAsync2::StartReadingHeader()
 	}
 }
 
-void FArchiveAsync2::EndReadingHeader()
+void FAsyncArchive::EndReadingHeader()
 {
 	LogItem(TEXT("End Header"));
 	check(LoadPhase == ELoadPhase::WaitingForHeader);
@@ -8220,7 +8299,7 @@ void FArchiveAsync2::EndReadingHeader()
 	FlushPrecacheBlock();
 }
 
-bool FArchiveAsync2::ReadyToStartReadingHeader(bool bUseTimeLimit, bool bUseFullTimeLimit, double TickStartTime, float TimeLimit)
+bool FAsyncArchive::ReadyToStartReadingHeader(bool bUseTimeLimit, bool bUseFullTimeLimit, double TickStartTime, float TimeLimit)
 {
 	if (SummaryReadTime == 0.0)
 	{
@@ -8254,7 +8333,7 @@ bool FArchiveAsync2::ReadyToStartReadingHeader(bool bUseTimeLimit, bool bUseFull
 void CallSerializeHook();
 #endif
 
-void FArchiveAsync2::Serialize(void* Data, int64 Count)
+void FAsyncArchive::Serialize(void* Data, int64 Count)
 {
 	if (!Count || ArIsError)
 	{
@@ -8342,7 +8421,7 @@ void FArchiveAsync2::Serialize(void* Data, int64 Count)
 	}
 	if (BeforeBlockSize)
 	{
-		UE_CLOG(GEventDrivenLoaderEnabled, LogAsyncArchive, Warning, TEXT("FArchiveAsync2::Serialize Backwards streaming in %s  CurrentPos = %lld   BeforeBlockOffset = %lld"), *FileName, CurrentPos, BeforeBlockOffset);
+		UE_CLOG(GEventDrivenLoaderEnabled, LogAsyncArchive, Warning, TEXT("FAsyncArchive::Serialize Backwards streaming in %s  CurrentPos = %lld   BeforeBlockOffset = %lld"), *FileName, CurrentPos, BeforeBlockOffset);
 		LogItem(TEXT("Sync Before Block"), BeforeBlockOffset, BeforeBlockSize);
 		if (!PrecacheInternal(BeforeBlockOffset, BeforeBlockSize))
 		{
@@ -8410,7 +8489,7 @@ void FArchiveAsync2::Serialize(void* Data, int64 Count)
 }
 
 #if DEVIRTUALIZE_FLinkerLoad_Serialize
-void FArchiveAsync2::DiscardInlineBufferAndUpdateCurrentPos()
+void FAsyncArchive::DiscardInlineBufferAndUpdateCurrentPos()
 {
 	CurrentPos += (ActiveFPLB->StartFastPathLoadBuffer - ActiveFPLB->OriginalFastPathLoadBuffer);
 	ActiveFPLB->Reset();
@@ -8423,17 +8502,17 @@ void FArchiveAsync2::DiscardInlineBufferAndUpdateCurrentPos()
 static TAutoConsoleVariable<int32> CVarLogAsyncArchiveSerializeChurn(
 	TEXT("LogAsyncArchiveSerializeChurn.Enable"),
 	0,
-	TEXT("If > 0, then sample game thread FArchiveAsync2::Serialize calls, periodically print a report of the worst offenders."));
+	TEXT("If > 0, then sample game thread FAsyncArchive::Serialize calls, periodically print a report of the worst offenders."));
 
 static TAutoConsoleVariable<int32> CVarLogAsyncArchiveSerializeChurn_Threshhold(
 	TEXT("LogAsyncArchiveSerializeChurn.Threshhold"),
 	1000,
-	TEXT("Minimum average number of FArchiveAsync2::Serialize calls to include in the report."));
+	TEXT("Minimum average number of FAsyncArchive::Serialize calls to include in the report."));
 
 static TAutoConsoleVariable<int32> CVarLogAsyncArchiveSerializeChurn_SampleFrequency(
 	TEXT("LogAsyncArchiveSerializeChurn.SampleFrequency"),
 	1000,
-	TEXT("Number of FArchiveAsync2::Serialize calls per sample. This is used to prevent sampling from slowing the game down too much."));
+	TEXT("Number of FAsyncArchive::Serialize calls per sample. This is used to prevent sampling from slowing the game down too much."));
 
 static TAutoConsoleVariable<int32> CVarLogAsyncArchiveSerializeChurn_StackIgnore(
 	TEXT("LogAsyncArchiveSerializeChurn.StackIgnore"),
diff --git a/Engine/Source/Runtime/CoreUObject/Private/Serialization/AsyncLoadingThread.h b/Engine/Source/Runtime/CoreUObject/Private/Serialization/AsyncLoadingThread.h
--- a/Engine/Source/Runtime/CoreUObject/Private/Serialization/AsyncLoadingThread.h
+++ b/Engine/Source/Runtime/CoreUObject/Private/Serialization/AsyncLoadingThread.h
@@ -9,6 +9,7 @@
 #include "HAL/ThreadSafeCounter.h"
 #include "UObject/ObjectMacros.h"
 #include "Serialization/AsyncLoading.h"
+#include "Serialization/LoadTimeTracePrivate.h"
 #include "Misc/ScopeLock.h"
 #include "Misc/CommandLine.h"
 #include "Misc/App.h"
@@ -517,6 +518,7 @@ public:
 		for (int32 ID : RequestIDs)
 		{
 			PendingRequests.Remove(ID);
+			TRACE_LOADTIME_END_REQUEST(ID);
 		}		
 	}
 
diff --git a/Engine/Source/Runtime/CoreUObject/Private/Serialization/LoadTimeTrace.cpp b/Engine/Source/Runtime/CoreUObject/Private/Serialization/LoadTimeTrace.cpp
new file mode 100644
--- /dev/null
+++ b/Engine/Source/Runtime/CoreUObject/Private/Serialization/LoadTimeTrace.cpp
@@ -0,0 +1,522 @@
+// Copyright 1998-2019 Epic Games, Inc. All Rights Reserved.
+
+#include "Serialization/LoadTimeTrace.h"
+
+#if LOADTIMEPROFILERTRACE_ENABLED
+
+#include "LoadTimeTracePrivate.h"
+#include "Trace/Trace.h"
+#include "Misc/CString.h"
+#include "HAL/PlatformTime.h"
+#include "HAL/PlatformTLS.h"
+#include "UObject/Object.h"
+#include "UObject/UObjectGlobals.h"
+
+UE_TRACE_EVENT_BEGIN(LoadTime, StartAsyncLoading, Always)
+	UE_TRACE_EVENT_FIELD(uint64, Cycle)
+	UE_TRACE_EVENT_FIELD(uint32, ThreadId)
+UE_TRACE_EVENT_END()
+
+UE_TRACE_EVENT_BEGIN(LoadTime, SuspendAsyncLoading)
+	UE_TRACE_EVENT_FIELD(uint64, Cycle)
+UE_TRACE_EVENT_END()
+
+UE_TRACE_EVENT_BEGIN(LoadTime, ResumeAsyncLoading)
+	UE_TRACE_EVENT_FIELD(uint64, Cycle)
+UE_TRACE_EVENT_END()
+
+UE_TRACE_EVENT_BEGIN(LoadTime, NewPackage)
+	UE_TRACE_EVENT_FIELD(const UPackage*, Package)
+UE_TRACE_EVENT_END()
+
+UE_TRACE_EVENT_BEGIN(LoadTime, NewLinker)
+	UE_TRACE_EVENT_FIELD(const FLinkerLoad*, Linker)
+UE_TRACE_EVENT_END()
+
+UE_TRACE_EVENT_BEGIN(LoadTime, DestroyLinker)
+	UE_TRACE_EVENT_FIELD(const FLinkerLoad*, Linker)
+UE_TRACE_EVENT_END()
+
+UE_TRACE_EVENT_BEGIN(LoadTime, PackageSummary)
+	UE_TRACE_EVENT_FIELD(const FLinkerLoad*, Linker)
+	UE_TRACE_EVENT_FIELD(uint32, TotalHeaderSize)
+	UE_TRACE_EVENT_FIELD(uint32, NameCount)
+	UE_TRACE_EVENT_FIELD(uint32, ImportCount)
+	UE_TRACE_EVENT_FIELD(uint32, ExportCount)
+UE_TRACE_EVENT_END()
+
+UE_TRACE_EVENT_BEGIN(LoadTime, BeginLinkerScope)
+	UE_TRACE_EVENT_FIELD(const FLinkerLoad*, Linker)
+	UE_TRACE_EVENT_FIELD(uint64, Cycle)
+	UE_TRACE_EVENT_FIELD(uint32, ThreadId)
+UE_TRACE_EVENT_END()
+
+UE_TRACE_EVENT_BEGIN(LoadTime, EndLinkerScope)
+	UE_TRACE_EVENT_FIELD(uint64, Cycle)
+	UE_TRACE_EVENT_FIELD(uint32, ThreadId)
+UE_TRACE_EVENT_END()
+
+UE_TRACE_EVENT_BEGIN(LoadTime, BeginCreateExport)
+	UE_TRACE_EVENT_FIELD(uint64, Cycle)
+	UE_TRACE_EVENT_FIELD(const FLinkerLoad*, Linker)
+	UE_TRACE_EVENT_FIELD(uint64, SerialOffset)
+	UE_TRACE_EVENT_FIELD(uint64, SerialSize)
+	UE_TRACE_EVENT_FIELD(uint32, ThreadId)
+	UE_TRACE_EVENT_FIELD(bool, IsAsset)
+UE_TRACE_EVENT_END()
+
+UE_TRACE_EVENT_BEGIN(LoadTime, EndCreateExport)
+	UE_TRACE_EVENT_FIELD(uint64, Cycle)
+	UE_TRACE_EVENT_FIELD(const UObject*, Object)
+	UE_TRACE_EVENT_FIELD(const UClass*, Class)
+	UE_TRACE_EVENT_FIELD(uint32, ThreadId)
+UE_TRACE_EVENT_END()
+
+UE_TRACE_EVENT_BEGIN(LoadTime, BeginObjectScope)
+	UE_TRACE_EVENT_FIELD(uint64, Cycle)
+	UE_TRACE_EVENT_FIELD(const UObject*, Object)
+	UE_TRACE_EVENT_FIELD(uint32, ThreadId)
+	UE_TRACE_EVENT_FIELD(uint8, EventType)
+UE_TRACE_EVENT_END()
+
+UE_TRACE_EVENT_BEGIN(LoadTime, EndObjectScope)
+	UE_TRACE_EVENT_FIELD(uint64, Cycle)
+	UE_TRACE_EVENT_FIELD(uint32, ThreadId)
+UE_TRACE_EVENT_END()
+
+UE_TRACE_EVENT_BEGIN(LoadTime, LinkerArchiveAssociation)
+	UE_TRACE_EVENT_FIELD(const FLinkerLoad*, Linker)
+	UE_TRACE_EVENT_FIELD(const FArchive*, Archive)
+UE_TRACE_EVENT_END()
+
+UE_TRACE_EVENT_BEGIN(LoadTime, BeginRequest)
+	UE_TRACE_EVENT_FIELD(uint64, Cycle)
+	UE_TRACE_EVENT_FIELD(uint64, RequestId)
+	UE_TRACE_EVENT_FIELD(uint32, ThreadId)
+UE_TRACE_EVENT_END()
+
+UE_TRACE_EVENT_BEGIN(LoadTime, EndRequest)
+	UE_TRACE_EVENT_FIELD(uint64, Cycle)
+	UE_TRACE_EVENT_FIELD(uint64, RequestId)
+UE_TRACE_EVENT_END()
+
+UE_TRACE_EVENT_BEGIN(LoadTime, BeginLoadMap)
+	UE_TRACE_EVENT_FIELD(uint64, Cycle)
+	UE_TRACE_EVENT_FIELD(uint32, ThreadId)
+UE_TRACE_EVENT_END()
+
+UE_TRACE_EVENT_BEGIN(LoadTime, EndLoadMap)
+	UE_TRACE_EVENT_FIELD(uint64, Cycle)
+	UE_TRACE_EVENT_FIELD(uint32, ThreadId)
+UE_TRACE_EVENT_END()
+
+UE_TRACE_EVENT_BEGIN(LoadTime, NewStreamableHandle)
+	UE_TRACE_EVENT_FIELD(const FStreamableHandle*, StreamableHandle)
+	UE_TRACE_EVENT_FIELD(bool, IsCombined)
+UE_TRACE_EVENT_END()
+
+UE_TRACE_EVENT_BEGIN(LoadTime, DestroyStreamableHandle)
+	UE_TRACE_EVENT_FIELD(const FStreamableHandle*, StreamableHandle)
+UE_TRACE_EVENT_END()
+
+UE_TRACE_EVENT_BEGIN(LoadTime, BeginLoadStreamableHandle)
+	UE_TRACE_EVENT_FIELD(uint64, Cycle)
+	UE_TRACE_EVENT_FIELD(const FStreamableHandle*, StreamableHandle)
+	UE_TRACE_EVENT_END()
+
+UE_TRACE_EVENT_BEGIN(LoadTime, EndLoadStreamableHandle)
+	UE_TRACE_EVENT_FIELD(uint64, Cycle)
+	UE_TRACE_EVENT_FIELD(const FStreamableHandle*, StreamableHandle)
+	UE_TRACE_EVENT_END()
+
+UE_TRACE_EVENT_BEGIN(LoadTime, BeginWaitForStreamableHandle)
+	UE_TRACE_EVENT_FIELD(uint64, Cycle)
+	UE_TRACE_EVENT_FIELD(const FStreamableHandle*, StreamableHandle)
+	UE_TRACE_EVENT_FIELD(uint32, ThreadId)
+UE_TRACE_EVENT_END()
+
+UE_TRACE_EVENT_BEGIN(LoadTime, EndWaitForStreamableHandle)
+	UE_TRACE_EVENT_FIELD(uint64, Cycle)
+	UE_TRACE_EVENT_FIELD(uint32, ThreadId)
+UE_TRACE_EVENT_END()
+
+UE_TRACE_EVENT_BEGIN(LoadTime, NewAsyncPackage)
+	UE_TRACE_EVENT_FIELD(const FAsyncPackage*, AsyncPackage)
+UE_TRACE_EVENT_END()
+
+UE_TRACE_EVENT_BEGIN(LoadTime, DestroyAsyncPackage)
+	UE_TRACE_EVENT_FIELD(const FAsyncPackage*, AsyncPackage)
+UE_TRACE_EVENT_END()
+
+UE_TRACE_EVENT_BEGIN(LoadTime, StreamableHandleRequestAssociation)
+	UE_TRACE_EVENT_FIELD(const FStreamableHandle*, StreamableHandle)
+	UE_TRACE_EVENT_FIELD(uint64, RequestId)
+UE_TRACE_EVENT_END()
+
+UE_TRACE_EVENT_BEGIN(LoadTime, AsyncPackageRequestAssociation)
+	UE_TRACE_EVENT_FIELD(const FAsyncPackage*, AsyncPackage)
+	UE_TRACE_EVENT_FIELD(uint64, RequestId)
+UE_TRACE_EVENT_END()
+
+UE_TRACE_EVENT_BEGIN(LoadTime, LoadImportDependency)
+	UE_TRACE_EVENT_FIELD(const FAsyncPackage*, Package)
+	UE_TRACE_EVENT_FIELD(const FAsyncPackage*, DependentPackage)
+UE_TRACE_EVENT_END()
+
+UE_TRACE_EVENT_BEGIN(LoadTime, AsyncPackageLinkerAssociation)
+	UE_TRACE_EVENT_FIELD(const FAsyncPackage*, AsyncPackage)
+	UE_TRACE_EVENT_FIELD(const FLinkerLoad*, Linker)
+UE_TRACE_EVENT_END()
+
+UE_TRACE_EVENT_BEGIN(LoadTime, BeginAsyncPackageScope)
+	UE_TRACE_EVENT_FIELD(const FAsyncPackage*, AsyncPackage)
+	UE_TRACE_EVENT_FIELD(uint64, Cycle)
+	UE_TRACE_EVENT_FIELD(uint32, ThreadId)
+	UE_TRACE_EVENT_FIELD(uint8, EventType)
+UE_TRACE_EVENT_END()
+
+UE_TRACE_EVENT_BEGIN(LoadTime, EndAsyncPackageScope)
+	UE_TRACE_EVENT_FIELD(uint64, Cycle)
+	UE_TRACE_EVENT_FIELD(uint32, ThreadId)
+UE_TRACE_EVENT_END()
+
+UE_TRACE_EVENT_BEGIN(LoadTime, BeginFlushAsyncLoading)
+	UE_TRACE_EVENT_FIELD(uint64, Cycle)
+	UE_TRACE_EVENT_FIELD(uint64, RequestId)
+	UE_TRACE_EVENT_FIELD(uint32, ThreadId)
+UE_TRACE_EVENT_END()
+
+UE_TRACE_EVENT_BEGIN(LoadTime, EndFlushAsyncLoading)
+	UE_TRACE_EVENT_FIELD(uint64, Cycle)
+	UE_TRACE_EVENT_FIELD(uint32, ThreadId)
+UE_TRACE_EVENT_END()
+
+UE_TRACE_EVENT_BEGIN(LoadTime, QueueEvent)
+	UE_TRACE_EVENT_FIELD(uint64, Cycle)
+	UE_TRACE_EVENT_FIELD(const FAsyncPackage*, AsyncPackage)
+	UE_TRACE_EVENT_FIELD(uint32, ThreadId)
+	UE_TRACE_EVENT_FIELD(uint32, UserPriority)
+	UE_TRACE_EVENT_FIELD(uint32, PackageSerialNumber)
+	UE_TRACE_EVENT_FIELD(uint32, EventSystemPriority)
+	UE_TRACE_EVENT_FIELD(uint8, EventType)
+UE_TRACE_EVENT_END();
+
+UE_TRACE_EVENT_BEGIN(LoadTime, ClassInfo, Always)
+	UE_TRACE_EVENT_FIELD(const UClass*, Class)
+UE_TRACE_EVENT_END()
+
+void FLoadTimeProfilerTracePrivate::Init(bool bStartEnabled)
+{
+	UE_TRACE_EVENT_IS_ENABLED(LoadTime, StartAsyncLoading);
+	UE_TRACE_EVENT_IS_ENABLED(LoadTime, SuspendAsyncLoading);
+	UE_TRACE_EVENT_IS_ENABLED(LoadTime, ResumeAsyncLoading);
+	UE_TRACE_EVENT_IS_ENABLED(LoadTime, NewPackage);
+	UE_TRACE_EVENT_IS_ENABLED(LoadTime, NewLinker);
+	UE_TRACE_EVENT_IS_ENABLED(LoadTime, DestroyLinker);
+	UE_TRACE_EVENT_IS_ENABLED(LoadTime, PackageSummary);
+	UE_TRACE_EVENT_IS_ENABLED(LoadTime, BeginLinkerScope);
+	UE_TRACE_EVENT_IS_ENABLED(LoadTime, EndLinkerScope);
+	UE_TRACE_EVENT_IS_ENABLED(LoadTime, BeginCreateExport);
+	UE_TRACE_EVENT_IS_ENABLED(LoadTime, EndCreateExport);
+	UE_TRACE_EVENT_IS_ENABLED(LoadTime, BeginObjectScope);
+	UE_TRACE_EVENT_IS_ENABLED(LoadTime, EndObjectScope);
+	UE_TRACE_EVENT_IS_ENABLED(LoadTime, LinkerArchiveAssociation);
+	UE_TRACE_EVENT_IS_ENABLED(LoadTime, BeginRequest);
+	UE_TRACE_EVENT_IS_ENABLED(LoadTime, EndRequest);
+	UE_TRACE_EVENT_IS_ENABLED(LoadTime, BeginLoadMap);
+	UE_TRACE_EVENT_IS_ENABLED(LoadTime, EndLoadMap);
+	UE_TRACE_EVENT_IS_ENABLED(LoadTime, NewStreamableHandle);
+	UE_TRACE_EVENT_IS_ENABLED(LoadTime, DestroyStreamableHandle);
+	UE_TRACE_EVENT_IS_ENABLED(LoadTime, BeginLoadStreamableHandle);
+	UE_TRACE_EVENT_IS_ENABLED(LoadTime, EndLoadStreamableHandle);
+	UE_TRACE_EVENT_IS_ENABLED(LoadTime, BeginWaitForStreamableHandle);
+	UE_TRACE_EVENT_IS_ENABLED(LoadTime, EndWaitForStreamableHandle);
+	UE_TRACE_EVENT_IS_ENABLED(LoadTime, NewAsyncPackage);
+	UE_TRACE_EVENT_IS_ENABLED(LoadTime, DestroyAsyncPackage);
+	UE_TRACE_EVENT_IS_ENABLED(LoadTime, StreamableHandleRequestAssociation);
+	UE_TRACE_EVENT_IS_ENABLED(LoadTime, AsyncPackageRequestAssociation);
+	UE_TRACE_EVENT_IS_ENABLED(LoadTime, LoadImportDependency);
+	UE_TRACE_EVENT_IS_ENABLED(LoadTime, AsyncPackageLinkerAssociation);
+	UE_TRACE_EVENT_IS_ENABLED(LoadTime, BeginAsyncPackageScope);
+	UE_TRACE_EVENT_IS_ENABLED(LoadTime, EndAsyncPackageScope);
+	UE_TRACE_EVENT_IS_ENABLED(LoadTime, BeginFlushAsyncLoading);
+	UE_TRACE_EVENT_IS_ENABLED(LoadTime, EndFlushAsyncLoading);
+	UE_TRACE_EVENT_IS_ENABLED(LoadTime, QueueEvent);
+	UE_TRACE_EVENT_IS_ENABLED(LoadTime, ClassInfo);
+	if (bStartEnabled)
+	{
+		Trace::ToggleEvent(TEXT("LoadTime"), true);
+	}
+}
+
+void FLoadTimeProfilerTracePrivate::OutputStartAsyncLoading()
+{
+	UE_TRACE_LOG(LoadTime, StartAsyncLoading)
+		<< StartAsyncLoading.Cycle(FPlatformTime::Cycles64())
+		<< StartAsyncLoading.ThreadId(FPlatformTLS::GetCurrentThreadId());
+}
+
+PRAGMA_DISABLE_SHADOW_VARIABLE_WARNINGS
+void FLoadTimeProfilerTracePrivate::OutputSuspendAsyncLoading()
+{
+	UE_TRACE_LOG(LoadTime, SuspendAsyncLoading)
+		<< SuspendAsyncLoading.Cycle(FPlatformTime::Cycles64());
+}
+
+void FLoadTimeProfilerTracePrivate::OutputResumeAsyncLoading()
+{
+	UE_TRACE_LOG(LoadTime, ResumeAsyncLoading)
+		<< ResumeAsyncLoading.Cycle(FPlatformTime::Cycles64());
+}
+PRAGMA_ENABLE_SHADOW_VARIABLE_WARNINGS
+
+void FLoadTimeProfilerTracePrivate::OutputNewLinker(const FLinkerLoad* Linker)
+{
+	UE_TRACE_LOG(LoadTime, NewLinker)
+		<< NewLinker.Linker(Linker);
+}
+
+void FLoadTimeProfilerTracePrivate::OutputDestroyLinker(const FLinkerLoad* Linker)
+{
+	UE_TRACE_LOG(LoadTime, DestroyLinker)
+		<< DestroyLinker.Linker(Linker);
+}
+
+void FLoadTimeProfilerTracePrivate::OutputBeginRequest(uint64 RequestId)
+{
+	UE_TRACE_LOG(LoadTime, BeginRequest)
+		<< BeginRequest.Cycle(FPlatformTime::Cycles64())
+		<< BeginRequest.RequestId(RequestId)
+		<< BeginRequest.ThreadId(FPlatformTLS::GetCurrentThreadId());
+}
+
+void FLoadTimeProfilerTracePrivate::OutputEndRequest(uint64 RequestId)
+{
+	UE_TRACE_LOG(LoadTime, EndRequest)
+		<< EndRequest.Cycle(FPlatformTime::Cycles64())
+		<< EndRequest.RequestId(RequestId);
+}
+
+void FLoadTimeProfilerTrace::OutputNewStreamableHandle(const FStreamableHandle* StreamableHandle, const TCHAR* DebugName, bool IsCombined)
+{
+	uint16 NameSize = (FCString::Strlen(DebugName) + 1) * sizeof(TCHAR);
+	UE_TRACE_LOG(LoadTime, NewStreamableHandle, NameSize)
+		<< NewStreamableHandle.StreamableHandle(StreamableHandle)
+		<< NewStreamableHandle.IsCombined(IsCombined)
+		<< NewStreamableHandle.Attachment(DebugName, NameSize);
+}
+
+void FLoadTimeProfilerTrace::OutputDestroyStreamableHandle(const FStreamableHandle* StreamableHandle)
+{
+	UE_TRACE_LOG(LoadTime, DestroyStreamableHandle)
+		<< DestroyStreamableHandle.StreamableHandle(StreamableHandle);
+}
+
+void FLoadTimeProfilerTrace::OutputBeginLoadStreamableHandle(const FStreamableHandle* StreamableHandle)
+{
+	UE_TRACE_LOG(LoadTime, BeginLoadStreamableHandle)
+		<< BeginLoadStreamableHandle.StreamableHandle(StreamableHandle)
+		<< BeginLoadStreamableHandle.Cycle(FPlatformTime::Cycles64());
+}
+
+void FLoadTimeProfilerTrace::OutputEndLoadStreamableHandle(const FStreamableHandle* StreamableHandle)
+{
+	UE_TRACE_LOG(LoadTime, EndLoadStreamableHandle)
+		<< EndLoadStreamableHandle.StreamableHandle(StreamableHandle)
+		<< EndLoadStreamableHandle.Cycle(FPlatformTime::Cycles64());
+}
+
+void FLoadTimeProfilerTracePrivate::OutputNewAsyncPackage(const FAsyncPackage* AsyncPackage, const TCHAR* PackageName)
+{
+	uint16 NameSize = (FCString::Strlen(PackageName) + 1) * sizeof(TCHAR);
+	UE_TRACE_LOG(LoadTime, NewAsyncPackage, NameSize)
+		<< NewAsyncPackage.AsyncPackage(AsyncPackage)
+		<< NewAsyncPackage.Attachment(PackageName, NameSize);
+}
+
+void FLoadTimeProfilerTracePrivate::OutputDestroyAsyncPackage(const FAsyncPackage* AsyncPackage)
+{
+	UE_TRACE_LOG(LoadTime, DestroyAsyncPackage)
+		<< DestroyAsyncPackage.AsyncPackage(AsyncPackage);
+}
+
+void FLoadTimeProfilerTracePrivate::OutputPackageSummary(const FLinkerLoad* Linker, uint32 TotalHeaderSize, uint32 NameCount, uint32 ImportCount, uint32 ExportCount)
+{
+	UE_TRACE_LOG(LoadTime, PackageSummary)
+		<< PackageSummary.Linker(Linker)
+		<< PackageSummary.TotalHeaderSize(TotalHeaderSize)
+		<< PackageSummary.NameCount(NameCount)
+		<< PackageSummary.ImportCount(ImportCount)
+		<< PackageSummary.ExportCount(ExportCount);
+}
+
+void FLoadTimeProfilerTrace::OutputStreamableHandleRequestAssociation(const FStreamableHandle* StreamableHandle, uint64 RequestId)
+{
+	UE_TRACE_LOG(LoadTime, StreamableHandleRequestAssociation)
+		<< StreamableHandleRequestAssociation.StreamableHandle(StreamableHandle)
+		<< StreamableHandleRequestAssociation.RequestId(RequestId);
+}
+
+void FLoadTimeProfilerTracePrivate::OutputAsyncPackageRequestAssociation(const FAsyncPackage* AsyncPackage, uint64 RequestId)
+{
+	UE_TRACE_LOG(LoadTime, AsyncPackageRequestAssociation)
+		<< AsyncPackageRequestAssociation.AsyncPackage(AsyncPackage)
+		<< AsyncPackageRequestAssociation.RequestId(RequestId);
+}
+
+void FLoadTimeProfilerTracePrivate::OutputLoadImportDependency(const FAsyncPackage* Package, const FAsyncPackage* DependentPackage)
+{
+	UE_TRACE_LOG(LoadTime, LoadImportDependency)
+		<< LoadImportDependency.Package(Package)
+		<< LoadImportDependency.DependentPackage(DependentPackage);
+}
+
+void FLoadTimeProfilerTracePrivate::OutputAsyncPackageLinkerAssociation(const FAsyncPackage* AsyncPackage, const FLinkerLoad* Linker)
+{
+	UE_TRACE_LOG(LoadTime, AsyncPackageLinkerAssociation)
+		<< AsyncPackageLinkerAssociation.AsyncPackage(AsyncPackage)
+		<< AsyncPackageLinkerAssociation.Linker(Linker);
+}
+
+void FLoadTimeProfilerTracePrivate::OutputLinkerArchiveAssociation(const FLinkerLoad* Linker, const FArchive* Archive)
+{
+	UE_TRACE_LOG(LoadTime, LinkerArchiveAssociation)
+		<< LinkerArchiveAssociation.Linker(Linker)
+		<< LinkerArchiveAssociation.Archive(Archive);
+}
+
+void FLoadTimeProfilerTracePrivate::OutputQueueEvent(const FAsyncPackage* AsyncPackage, ELoadTimeProfilerPackageEventType EventType, TAsyncLoadPriority UserPriority, int32 PackageSerialNumber, int32 EventSystemPriority)
+{
+	UE_TRACE_LOG(LoadTime, QueueEvent)
+		<< QueueEvent.Cycle(FPlatformTime::Cycles64())
+		<< QueueEvent.AsyncPackage(AsyncPackage)
+		<< QueueEvent.ThreadId(FPlatformTLS::GetCurrentThreadId())
+		<< QueueEvent.UserPriority(UserPriority)
+		<< QueueEvent.PackageSerialNumber(PackageSerialNumber)
+		<< QueueEvent.EventSystemPriority(EventSystemPriority)
+		<< QueueEvent.EventType(EventType);
+}
+
+void FLoadTimeProfilerTracePrivate::OutputClassInfo(const UClass* Class, const TCHAR* Name)
+{
+	uint16 NameSize = (FCString::Strlen(Name) + 1) * sizeof(TCHAR);
+	UE_TRACE_LOG(LoadTime, ClassInfo, NameSize)
+		<< ClassInfo.Class(Class)
+		<< ClassInfo.Attachment(Name, NameSize);
+}
+
+FLoadTimeProfilerTracePrivate::FAsyncPackageScope::FAsyncPackageScope(const FAsyncPackage* AsyncPackage, ELoadTimeProfilerPackageEventType EventType)
+{
+	UE_TRACE_LOG(LoadTime, BeginAsyncPackageScope)
+		<< BeginAsyncPackageScope.AsyncPackage(AsyncPackage)
+		<< BeginAsyncPackageScope.Cycle(FPlatformTime::Cycles64())
+		<< BeginAsyncPackageScope.ThreadId(FPlatformTLS::GetCurrentThreadId())
+		<< BeginAsyncPackageScope.EventType(EventType);
+}
+
+FLoadTimeProfilerTracePrivate::FAsyncPackageScope::~FAsyncPackageScope()
+{
+	UE_TRACE_LOG(LoadTime, EndAsyncPackageScope)
+		<< EndAsyncPackageScope.Cycle(FPlatformTime::Cycles64())
+		<< EndAsyncPackageScope.ThreadId(FPlatformTLS::GetCurrentThreadId());
+}
+
+FLoadTimeProfilerTracePrivate::FCreateExportScope::FCreateExportScope(const FLinkerLoad* Linker, uint64 SerialOffset, uint64 SerialSize, bool IsAsset, const UObject* const* InObject)
+	: Object(InObject)
+{
+	UE_TRACE_LOG(LoadTime, BeginCreateExport)
+		<< BeginCreateExport.Cycle(FPlatformTime::Cycles64())
+		<< BeginCreateExport.Linker(Linker)
+		<< BeginCreateExport.SerialOffset(SerialOffset)
+		<< BeginCreateExport.SerialSize(SerialSize)
+		<< BeginCreateExport.ThreadId(FPlatformTLS::GetCurrentThreadId())
+		<< BeginCreateExport.IsAsset(IsAsset);
+}
+
+FLoadTimeProfilerTracePrivate::FCreateExportScope::~FCreateExportScope()
+{
+	UE_TRACE_LOG(LoadTime, EndCreateExport)
+		<< EndCreateExport.Cycle(FPlatformTime::Cycles64())
+		<< EndCreateExport.Object(*Object)
+		<< EndCreateExport.Class(*Object ? (*Object)->GetClass() : nullptr)
+		<< EndCreateExport.ThreadId(FPlatformTLS::GetCurrentThreadId());
+}
+
+FLoadTimeProfilerTracePrivate::FObjectScope::FObjectScope(const UObject* Object, ELoadTimeProfilerObjectEventType EventType)
+{
+	UE_TRACE_LOG(LoadTime, BeginObjectScope)
+		<< BeginObjectScope.Cycle(FPlatformTime::Cycles64())
+		<< BeginObjectScope.Object(Object)
+		<< BeginObjectScope.ThreadId(FPlatformTLS::GetCurrentThreadId())
+		<< BeginObjectScope.EventType(EventType);
+}
+
+FLoadTimeProfilerTracePrivate::FObjectScope::~FObjectScope()
+{
+	UE_TRACE_LOG(LoadTime, EndObjectScope)
+		<< EndObjectScope.Cycle(FPlatformTime::Cycles64())
+		<< EndObjectScope.ThreadId(FPlatformTLS::GetCurrentThreadId());
+}
+
+FLoadTimeProfilerTracePrivate::FLinkerScope::FLinkerScope(const FLinkerLoad* Linker)
+{
+	UE_TRACE_LOG(LoadTime, BeginLinkerScope)
+		<< BeginLinkerScope.Linker(Linker)
+		<< BeginLinkerScope.Cycle(FPlatformTime::Cycles64())
+		<< BeginLinkerScope.ThreadId(FPlatformTLS::GetCurrentThreadId());
+}
+
+FLoadTimeProfilerTracePrivate::FLinkerScope::~FLinkerScope()
+{
+	UE_TRACE_LOG(LoadTime, EndLinkerScope)
+		<< EndLinkerScope.Cycle(FPlatformTime::Cycles64())
+		<< EndLinkerScope.ThreadId(FPlatformTLS::GetCurrentThreadId());
+}
+
+FLoadTimeProfilerTracePrivate::FFlushAsyncLoadingScope::FFlushAsyncLoadingScope(uint64 RequestId)
+{
+	UE_TRACE_LOG(LoadTime, BeginFlushAsyncLoading)
+		<< BeginFlushAsyncLoading.Cycle(FPlatformTime::Cycles64())
+		<< BeginFlushAsyncLoading.RequestId(RequestId)
+		<< BeginFlushAsyncLoading.ThreadId(FPlatformTLS::GetCurrentThreadId());
+}
+
+FLoadTimeProfilerTracePrivate::FFlushAsyncLoadingScope::~FFlushAsyncLoadingScope()
+{
+	UE_TRACE_LOG(LoadTime, EndFlushAsyncLoading)
+		<< EndFlushAsyncLoading.Cycle(FPlatformTime::Cycles64())
+		<< EndFlushAsyncLoading.ThreadId(FPlatformTLS::GetCurrentThreadId());
+}
+
+FLoadTimeProfilerTrace::FWaitForStreamableHandleScope::FWaitForStreamableHandleScope(const FStreamableHandle* StreamableHandle)
+{
+	UE_TRACE_LOG(LoadTime, BeginWaitForStreamableHandle)
+		<< BeginWaitForStreamableHandle.Cycle(FPlatformTime::Cycles64())
+		<< BeginWaitForStreamableHandle.StreamableHandle(StreamableHandle)
+		<< BeginWaitForStreamableHandle.ThreadId(FPlatformTLS::GetCurrentThreadId());
+}
+
+FLoadTimeProfilerTrace::FWaitForStreamableHandleScope::~FWaitForStreamableHandleScope()
+{
+	UE_TRACE_LOG(LoadTime, EndWaitForStreamableHandle)
+		<< EndWaitForStreamableHandle.Cycle(FPlatformTime::Cycles64())
+		<< EndWaitForStreamableHandle.ThreadId(FPlatformTLS::GetCurrentThreadId());
+}
+
+FLoadTimeProfilerTrace::FLoadMapScope::FLoadMapScope(const TCHAR* Name)
+{
+	uint16 NameSize = (FCString::Strlen(Name) + 1) * sizeof(TCHAR);
+	UE_TRACE_LOG(LoadTime, BeginLoadMap, NameSize)
+		<< BeginLoadMap.Cycle(FPlatformTime::Cycles64())
+		<< BeginLoadMap.ThreadId(FPlatformTLS::GetCurrentThreadId())
+		<< BeginLoadMap.Attachment(Name, NameSize);
+}
+
+FLoadTimeProfilerTrace::FLoadMapScope::~FLoadMapScope()
+{
+	UE_TRACE_LOG(LoadTime, EndLoadMap)
+		<< EndLoadMap.Cycle(FPlatformTime::Cycles64())
+		<< EndLoadMap.ThreadId(FPlatformTLS::GetCurrentThreadId());
+}
+
+#endif
\ No newline at end of file
diff --git a/Engine/Source/Runtime/CoreUObject/Private/Serialization/LoadTimeTracePrivate.h b/Engine/Source/Runtime/CoreUObject/Private/Serialization/LoadTimeTracePrivate.h
new file mode 100644
--- /dev/null
+++ b/Engine/Source/Runtime/CoreUObject/Private/Serialization/LoadTimeTracePrivate.h
@@ -0,0 +1,152 @@
+// Copyright 1998-2019 Epic Games, Inc. All Rights Reserved.
+
+#pragma once
+
+#include "Serialization/LoadTimeTrace.h"
+
+#if LOADTIMEPROFILERTRACE_ENABLED
+
+#include "UObject/UObjectGlobals.h"
+
+class FArchive;
+struct FAsyncPackage;
+class FLinkerLoad;
+class UObject;
+class UPackage;
+
+struct FLoadTimeProfilerTracePrivate
+{
+	static void Init(bool bStartEnabled);
+	static void OutputStartAsyncLoading();
+	static void OutputSuspendAsyncLoading();
+	static void OutputResumeAsyncLoading();
+	static void OutputNewLinker(const FLinkerLoad* Linker);
+	static void OutputDestroyLinker(const FLinkerLoad* Linker);
+	static void OutputNewAsyncPackage(const FAsyncPackage* AsyncPackage, const TCHAR* PackageName);
+	static void OutputDestroyAsyncPackage(const FAsyncPackage* AsyncPackage);
+	static void OutputBeginRequest(uint64 RequestId);
+	static void OutputEndRequest(uint64 RequestId);
+	static void OutputPackageSummary(const FLinkerLoad* Linker, uint32 TotalHeaderSize, uint32 NameCount, uint32 ImportCount, uint32 ExportCount);
+	static void OutputLoadImportDependency(const FAsyncPackage* Package, const FAsyncPackage* DependentPackage);
+	static void OutputAsyncPackageRequestAssociation(const FAsyncPackage* AsyncPackage, uint64 RequestId);
+	static void OutputAsyncPackageLinkerAssociation(const FAsyncPackage* AsyncPackage, const FLinkerLoad* Linker);
+	static void OutputLinkerArchiveAssociation(const FLinkerLoad* Linker, const FArchive* Archive);
+	static void OutputQueueEvent(const FAsyncPackage* AsyncPackage, ELoadTimeProfilerPackageEventType EventType, TAsyncLoadPriority UserPriority, int32 PackageSerialNumber, int32 EventSystemPriority);
+	static void OutputClassInfo(const UClass* Class, const TCHAR* Name);
+
+	struct FAsyncPackageScope
+	{
+		FAsyncPackageScope(const FAsyncPackage* AsyncPackage, ELoadTimeProfilerPackageEventType EventType);
+		~FAsyncPackageScope();
+	};
+
+	struct FCreateExportScope
+	{
+		FCreateExportScope(const FLinkerLoad* Linker, uint64 SerialOffset, uint64 SerialSize, bool IsAsset, const UObject* const* InObject);
+		~FCreateExportScope();
+
+	private:
+		const UObject* const* Object;
+	};
+
+	struct FObjectScope
+	{
+		FObjectScope(const UObject* Object, ELoadTimeProfilerObjectEventType EventType);
+		~FObjectScope();
+	};
+
+	struct FLinkerScope
+	{
+		FLinkerScope(const FLinkerLoad* Linker);
+		~FLinkerScope();
+	};
+
+	struct FFlushAsyncLoadingScope
+	{
+		FFlushAsyncLoadingScope(uint64 RequestId);
+		~FFlushAsyncLoadingScope();
+	};
+};
+
+#define TRACE_LOADTIME_START_ASYNC_LOADING() \
+	FLoadTimeProfilerTracePrivate::OutputStartAsyncLoading();
+
+#define TRACE_LOADTIME_SUSPEND_ASYNC_LOADING() \
+	FLoadTimeProfilerTracePrivate::OutputSuspendAsyncLoading();
+
+#define TRACE_LOADTIME_RESUME_ASYNC_LOADING() \
+	FLoadTimeProfilerTracePrivate::OutputResumeAsyncLoading();
+
+#define TRACE_LOADTIME_NEW_LINKER(Linker) \
+	FLoadTimeProfilerTracePrivate::OutputNewLinker(Linker);
+
+#define TRACE_LOADTIME_DESTROY_LINKER(Linker) \
+	FLoadTimeProfilerTracePrivate::OutputDestroyLinker(Linker);
+
+#define TRACE_LOADTIME_PACKAGE_SUMMARY(Linker, TotalHeaderSize, NameCount, ImportCount, ExportCount) \
+	FLoadTimeProfilerTracePrivate::OutputPackageSummary(Linker, TotalHeaderSize, NameCount, ImportCount, ExportCount);
+
+#define TRACE_LOADTIME_LINKER_ARCHIVE_ASSOCIATION(Linker, Archive) \
+	FLoadTimeProfilerTracePrivate::OutputLinkerArchiveAssociation(Linker, Archive);
+
+#define TRACE_LOADTIME_BEGIN_REQUEST(RequestId) \
+	FLoadTimeProfilerTracePrivate::OutputBeginRequest(RequestId);
+
+#define TRACE_LOADTIME_END_REQUEST(RequestId) \
+	FLoadTimeProfilerTracePrivate::OutputEndRequest(RequestId);
+
+#define TRACE_LOADTIME_NEW_ASYNC_PACKAGE(AsyncPackage, PackageName) \
+	FLoadTimeProfilerTracePrivate::OutputNewAsyncPackage(AsyncPackage, PackageName)
+
+#define TRACE_LOADTIME_DESTROY_ASYNC_PACKAGE(AsyncPackage) \
+	FLoadTimeProfilerTracePrivate::OutputDestroyAsyncPackage(AsyncPackage);
+
+#define TRACE_LOADTIME_ASYNC_PACKAGE_REQUEST_ASSOCIATION(AsyncPackage, RequestId) \
+	FLoadTimeProfilerTracePrivate::OutputAsyncPackageRequestAssociation(AsyncPackage, RequestId);
+
+#define TRACE_LOADTIME_ASYNC_PACKAGE_LINKER_ASSOCIATION(AsyncPackage, Linker) \
+	FLoadTimeProfilerTracePrivate::OutputAsyncPackageLinkerAssociation(AsyncPackage, Linker);
+
+#define TRACE_LOADTIME_QUEUE_EVENT(AsyncPackage, EventType, UserPriority, PackageSerialNumber, EventSystemPriority)
+//	FLoadTimeProfilerTracePrivate::OutputQueueEvent(AsyncPackage, EventType, UserPriority, PackageSerialNumber, EventSystemPriority);
+
+#define TRACE_LOADTIME_ASYNC_PACKAGE_SCOPE(AsyncPackage, EventType) \
+	FLoadTimeProfilerTracePrivate::FAsyncPackageScope __LoadTimeTraceAsyncPackageScope(AsyncPackage, EventType);
+
+#define TRACE_LOADTIME_CREATE_EXPORT_SCOPE(Linker, Object, SerialOffset, SerialSize, IsAsset) \
+	FLoadTimeProfilerTracePrivate::FCreateExportScope __LoadTimeTraceCreateExportScope(Linker, SerialOffset, SerialSize, IsAsset, Object);
+
+#define TRACE_LOADTIME_OBJECT_SCOPE(Object, EventType) \
+	FLoadTimeProfilerTracePrivate::FObjectScope __LoadTimeTraceObjectScope(Object, EventType);
+
+#define TRACE_LOADTIME_FLUSH_ASYNCLOADING_SCOPE(RequestId) \
+	FLoadTimeProfilerTracePrivate::FFlushAsyncLoadingScope __LoadTimeTraceFlushAsyncLoadingScope(RequestId);
+
+#define TRACE_LOADTIME_CLASS_INFO(Class, Name) \
+	FLoadTimeProfilerTracePrivate::OutputClassInfo(Class, Name);
+
+#else
+
+#define TRACE_LOADTIME_START_ASYNC_LOADING(...)
+#define TRACE_LOADTIME_SUSPEND_ASYNC_LOADING(...)
+#define TRACE_LOADTIME_RESUME_ASYNC_LOADING(...)
+#define TRACE_LOADTIME_NEW_PACKAGE(...)
+#define TRACE_LOADTIME_NEW_LINKER(...)
+#define TRACE_LOADTIME_DESTROY_LINKER(...)
+#define TRACE_LOADTIME_PACKAGE_SUMMARY(...)
+#define TRACE_LOADTIME_LINKER_ARCHIVE_ASSOCIATION(...)
+#define TRACE_LOADTIME_BEGIN_REQUEST(...)
+#define TRACE_LOADTIME_END_REQUEST(...)
+#define TRACE_LOADTIME_NEW_ASYNC_PACKAGE(...)
+#define TRACE_LOADTIME_DESTROY_ASYNC_PACKAGE(...)
+#define TRACE_LOADTIME_ASYNC_PACKAGE_REQUEST_ASSOCIATION(...)
+#define TRACE_LOADTIME_ASYNC_PACKAGE_LINKER_ASSOCIATION(...)
+#define TRACE_LOADTIME_FLUSH_ASYNCLOADING_SCOPE(...)
+#define TRACE_LOADTIME_QUEUE_EVENT(...)
+#define TRACE_LOADTIME_ASYNC_PACKAGE_SCOPE(...)
+#define TRACE_LOADTIME_TIMER_SCOPE(...)
+#define TRACE_LOADTIME_CREATE_EXPORT_SCOPE(...)
+#define TRACE_LOADTIME_OBJECT_SCOPE(...)
+#define TRACE_LOADTIME_CLASS_INFO(...)
+
+#endif
\ No newline at end of file
diff --git a/Engine/Source/Runtime/CoreUObject/Public/Serialization/LoadTimeTrace.h b/Engine/Source/Runtime/CoreUObject/Public/Serialization/LoadTimeTrace.h
new file mode 100644
--- /dev/null
+++ b/Engine/Source/Runtime/CoreUObject/Public/Serialization/LoadTimeTrace.h
@@ -0,0 +1,93 @@
+// Copyright 1998-2019 Epic Games, Inc. All Rights Reserved.
+
+#pragma once
+
+#include "CoreTypes.h"
+#include "Trace/Trace.h"
+
+#if UE_TRACE_ENABLED && !UE_BUILD_SHIPPING
+#define LOADTIMEPROFILERTRACE_ENABLED 1
+#else
+#define LOADTIMEPROFILERTRACE_ENABLED 0
+#endif
+
+enum ELoadTimeProfilerPackageEventType
+{
+	LoadTimeProfilerPackageEventType_CreateLinker,
+	LoadTimeProfilerPackageEventType_FinishLinker,
+	LoadTimeProfilerPackageEventType_StartImportPackages,
+	LoadTimeProfilerPackageEventType_SetupImports,
+	LoadTimeProfilerPackageEventType_SetupExports,
+	LoadTimeProfilerPackageEventType_ProcessImportsAndExports,
+	LoadTimeProfilerPackageEventType_ExportsDone,
+	LoadTimeProfilerPackageEventType_PostLoadWait,
+	LoadTimeProfilerPackageEventType_StartPostLoad,
+	LoadTimeProfilerPackageEventType_Tick,
+	LoadTimeProfilerPackageEventType_Finish,
+	LoadTimeProfilerPackageEventType_DeferredPostLoad,
+	LoadTimeProfilerPackageEventType_None,
+};
+
+enum ELoadTimeProfilerObjectEventType
+{
+	LoadTimeProfilerObjectEventType_Create,
+	LoadTimeProfilerObjectEventType_Serialize,
+	LoadTimeProfilerObjectEventType_PostLoad,
+	LoadTimeProfilerObjectEventType_None
+};
+
+#if LOADTIMEPROFILERTRACE_ENABLED
+
+struct FStreamableHandle;
+
+struct FLoadTimeProfilerTrace
+{
+	COREUOBJECT_API static void OutputNewStreamableHandle(const FStreamableHandle* StreamableHandle, const TCHAR* DebugName, bool IsCombined);
+	COREUOBJECT_API static void OutputDestroyStreamableHandle(const FStreamableHandle* StreamableHandle);
+	COREUOBJECT_API static void OutputBeginLoadStreamableHandle(const FStreamableHandle* StreamableHandle);
+	COREUOBJECT_API static void OutputEndLoadStreamableHandle(const FStreamableHandle* StreamableHandle);
+	COREUOBJECT_API static void OutputStreamableHandleRequestAssociation(const FStreamableHandle* StreamableHandle, uint64 RequestId);
+
+	struct FLoadMapScope
+	{
+		COREUOBJECT_API FLoadMapScope(const TCHAR* Name);
+		COREUOBJECT_API ~FLoadMapScope();
+	};
+
+	struct FWaitForStreamableHandleScope
+	{
+		COREUOBJECT_API FWaitForStreamableHandleScope(const FStreamableHandle* StreamableHandle);
+		COREUOBJECT_API ~FWaitForStreamableHandleScope();
+	};
+};
+
+#define TRACE_LOADTIME_LOAD_MAP_SCOPE(Name) \
+	FLoadTimeProfilerTrace::FLoadMapScope __LoadTimeTraceLoadMapScope(Name);
+
+#define TRACE_LOADTIME_NEW_STREAMABLE_HANDLE(StreamableHandle, DebugName, IsCombined) \
+	FLoadTimeProfilerTrace::OutputNewStreamableHandle(StreamableHandle, DebugName, IsCombined);
+
+#define TRACE_LOADTIME_DESTROY_STREAMABLE_HANDLE(StreamableHandle) \
+	FLoadTimeProfilerTrace::OutputDestroyStreamableHandle(StreamableHandle);
+
+#define TRACE_LOADTIME_BEGIN_LOAD_STREAMABLE_HANDLE(StreamableHandle) \
+	FLoadTimeProfilerTrace::OutputBeginLoadStreamableHandle(StreamableHandle);
+
+#define TRACE_LOADTIME_END_LOAD_STREAMABLE_HANDLE(StreamableHandle) \
+	FLoadTimeProfilerTrace::OutputEndLoadStreamableHandle(StreamableHandle);
+
+#define TRACE_LOADTIME_STREAMABLE_HANDLE_REQUEST_ASSOCIATION(StreamableHandle, RequestId) \
+	FLoadTimeProfilerTrace::OutputStreamableHandleRequestAssociation(StreamableHandle, RequestId);
+
+#define TRACE_LOADTIME_WAIT_FOR_STREAMABLE_HANDLE_SCOPE(StreamableHandle) \
+	FLoadTimeProfilerTrace::FWaitForStreamableHandleScope __LoadTimeTraceWaitForStreamableHandleScope(StreamableHandle);
+
+#else
+#define TRACE_LOADTIME_LOAD_MAP_SCOPE(...)
+#define TRACE_LOADTIME_NEW_STREAMABLE_HANDLE(...)
+#define TRACE_LOADTIME_DESTROY_STREAMABLE_HANDLE(...)
+#define TRACE_LOADTIME_BEGIN_LOAD_STREAMABLE_HANDLE(...)
+#define TRACE_LOADTIME_END_LOAD_STREAMABLE_HANDLE(...)
+#define TRACE_LOADTIME_STREAMABLE_HANDLE_REQUEST_ASSOCIATION(...)
+#define TRACE_LOADTIME_WAIT_FOR_STREAMABLE_HANDLE_SCOPE(...)
+#endif
diff --git a/Engine/Source/Runtime/Engine/Private/SceneUtils.cpp b/Engine/Source/Runtime/Engine/Private/SceneUtils.cpp
--- a/Engine/Source/Runtime/Engine/Private/SceneUtils.cpp
+++ b/Engine/Source/Runtime/Engine/Private/SceneUtils.cpp
@@ -184,7 +184,7 @@ public:
 	static const uint64 InvalidQueryResult = 0xFFFFFFFFFFFFFFFFull;
 
 public:
-	FRealtimeGPUProfilerEvent(const FName& InName, const FName& InStatName, FRenderQueryPool* RenderQueryPool)
+	FRealtimeGPUProfilerEvent(const FName& InName, const FName& InStatName, const FGpuProfilerTrace::FEventType* InTraceEventType, FRenderQueryPool* RenderQueryPool)
 		: StartResultMicroseconds(InvalidQueryResult)
 		, EndResultMicroseconds(InvalidQueryResult)
 		, FrameNumber(-1)
@@ -196,6 +196,7 @@ public:
 		StatName = InStatName;
 #endif
 		Name = InName;
+		TraceEventType = InTraceEventType;
 
 		const int MaxGPUQueries = CVarGPUStatsMaxQueriesPerFrame.GetValueOnRenderThread();
 		if ( MaxGPUQueries == -1 || RenderQueryPool->GetAllocatedQueryCount() < MaxGPUQueries )
@@ -321,6 +322,11 @@ public:
 		return Name;
 	}
 
+	const FGpuProfilerTrace::FEventType* GetTraceEventType() const
+	{
+		return TraceEventType;
+	}
+
 	uint64 GetStartResultMicroseconds(uint32 GPUIndex = 0) const
 	{
 		return StartResultMicroseconds;
@@ -343,6 +349,7 @@ private:
 	FName StatName;
 #endif
 	FName Name;
+	const FGpuProfilerTrace::FEventType* TraceEventType;
 	uint64 StartResultMicroseconds;
 	uint64 EndResultMicroseconds;
 	uint32 FrameNumber;
@@ -369,10 +376,10 @@ public:
 		Clear(nullptr);
 	}
 
-	void PushEvent(FRHICommandListImmediate& RHICmdList, const FName& Name, const FName& StatName)
+	void PushEvent(FRHICommandListImmediate& RHICmdList, const FName& Name, const FName& StatName, const FGpuProfilerTrace::FEventType* TraceEventType)
 	{
 		// TODO: this should really use a pool / free list
-		FRealtimeGPUProfilerEvent* Event = new FRealtimeGPUProfilerEvent(Name, StatName, RenderQueryPool);
+		FRealtimeGPUProfilerEvent* Event = new FRealtimeGPUProfilerEvent(Name, StatName, TraceEventType, RenderQueryPool);
 		const int32 EventIndex = GpuProfilerEvents.Num();
 
 		GpuProfilerEvents.Add(Event);
@@ -466,8 +473,11 @@ public:
 			EventAggregates.Push(Aggregate);
 		}
 
+		TRACE_GPUPROFILER_BEGIN_FRAME()
+
 		for (const FRealtimeGPUProfilerTimelineEvent& TimelineEvent : GpuProfilerTimelineEvents)
 		{
+			FRealtimeGPUProfilerEvent* Event = GpuProfilerEvents[TimelineEvent.EventIndex];
 			if (TimelineEvent.Type == FRealtimeGPUProfilerTimelineEvent::EType::PushEvent)
 			{
 				if (TimelineEventStack.Num() != 0)
@@ -475,13 +485,17 @@ public:
 					EventAggregates[TimelineEventStack.Last()].ExclusiveTime -= EventAggregates[TimelineEvent.EventIndex].InclusiveTime;
 				}
 				TimelineEventStack.Push(TimelineEvent.EventIndex);
+				TRACE_GPUPROFILER_BEGIN_EVENT(Event->GetTraceEventType(), Event->GetFrameNumber(), Event->GetStartResultMicroseconds());
 			}
 			else
 			{
 				TimelineEventStack.Pop();
+				TRACE_GPUPROFILER_END_EVENT(Event->GetEndResultMicroseconds());
 			}
 		}
 
+		TRACE_GPUPROFILER_END_FRAME()
+
 		// Update the stats
 
 		const bool GPUStatsChildTimesIncluded = !!CVarGPUStatsChildTimesIncluded.GetValueOnRenderThread();
@@ -679,7 +693,7 @@ void FRealtimeGPUProfiler::EndFrame(FRHICommandListImmediate& RHICmdList)
 	}
 }
 
-void FRealtimeGPUProfiler::PushEvent(FRHICommandListImmediate& RHICmdList, const FName& Name, const FName& StatName)
+void FRealtimeGPUProfiler::PushEvent(FRHICommandListImmediate& RHICmdList, const FName& Name, const FName& StatName, const FGpuProfilerTrace::FEventType* TraceEventType)
 {
 	check(IsInRenderingThread());
 	if (bStatGatheringPaused || !bInBeginEndBlock)
@@ -689,7 +703,7 @@ void FRealtimeGPUProfiler::PushEvent(FRHICommandListImmediate& RHICmdList, const
 	check(Frames.Num() > 0);
 	if (WriteBufferIndex >= 0)
 	{
-		Frames[WriteBufferIndex]->PushEvent(RHICmdList, Name, StatName);
+		Frames[WriteBufferIndex]->PushEvent(RHICmdList, Name, StatName, TraceEventType);
 	}
 }
 
@@ -710,7 +724,7 @@ void FRealtimeGPUProfiler::PopEvent(FRHICommandListImmediate& RHICmdList)
 /*-----------------------------------------------------------------------------
 FScopedGPUStatEvent
 -----------------------------------------------------------------------------*/
-void FScopedGPUStatEvent::Begin(FRHICommandList& InRHICmdList, const FName& Name, const FName& StatName)
+void FScopedGPUStatEvent::Begin(FRHICommandList& InRHICmdList, const FName& Name, const FName& StatName, const FGpuProfilerTrace::FEventType* TraceEventType)
 {
 	check(IsInRenderingThread());
 	if (!AreGPUStatsEnabled())
@@ -722,7 +736,7 @@ void FScopedGPUStatEvent::Begin(FRHICommandList& InRHICmdList, const FName& Name
 	if (InRHICmdList.IsImmediate())
 	{
 		RHICmdList = (FRHICommandListImmediate*)&InRHICmdList;
-		FRealtimeGPUProfiler::Get()->PushEvent(*RHICmdList, Name, StatName);
+		FRealtimeGPUProfiler::Get()->PushEvent(*RHICmdList, Name, StatName, TraceEventType);
 	}
 }
 
