UObject::PostLoadSubobjects does not need to instance templates in cooked builds as all instancing will be brought up to date during the cooking process

diff --git a/Engine/Source/Runtime/CoreUObject/Private/UObject/Obj.cpp b/Engine/Source/Runtime/CoreUObject/Private/UObject/Obj.cpp
--- a/Engine/Source/Runtime/CoreUObject/Private/UObject/Obj.cpp
+++ b/Engine/Source/Runtime/CoreUObject/Private/UObject/Obj.cpp
@@ -1112,32 +1112,37 @@ void UObject::PostLoadSubobjects( FObjectInstancingGraph* OuterInstanceGraph/*=N
 		// clear the flag so that we don't re-enter this method
 		ClearFlags(RF_NeedPostLoadSubobjects);
 
-		FObjectInstancingGraph CurrentInstanceGraph;
-
-		FObjectInstancingGraph* InstanceGraph = OuterInstanceGraph;
-		if ( InstanceGraph == NULL )
+		// Cooked data will already have its subobjects fully instanced as uninstanced subobjects are only due to newly introduced subobjects in
+		// an archetype that an instance of that object hasn't been saved with
+		if (!FPlatformProperties::RequiresCookedData())
 		{
-			CurrentInstanceGraph.SetDestinationRoot(this);
-			CurrentInstanceGraph.SetLoadingObject(true);
+			FObjectInstancingGraph CurrentInstanceGraph;
 
-			// if we weren't passed an instance graph to use, create a new one and use that
-			InstanceGraph = &CurrentInstanceGraph;
-		}
+			FObjectInstancingGraph* InstanceGraph = OuterInstanceGraph;
+			if (InstanceGraph == NULL)
+			{
+				CurrentInstanceGraph.SetDestinationRoot(this);
+				CurrentInstanceGraph.SetLoadingObject(true);
 
-		// this will be filled with the list of component instances which were serialized from disk
-		TArray<UObject*> SerializedComponents;
-		// fill the array with the component contained by this object that were actually serialized to disk through property references
-		CollectDefaultSubobjects(SerializedComponents, false);
+				// if we weren't passed an instance graph to use, create a new one and use that
+				InstanceGraph = &CurrentInstanceGraph;
+			}
 
-		// now, add all of the instanced components to the instance graph that will be used for instancing any components that have been added
-		// to this object's archetype since this object was last saved
-		for ( int32 ComponentIndex = 0; ComponentIndex < SerializedComponents.Num(); ComponentIndex++ )
-		{
-			UObject* PreviouslyInstancedComponent = SerializedComponents[ComponentIndex];
-			InstanceGraph->AddNewInstance(PreviouslyInstancedComponent);
-		}
+			// this will be filled with the list of component instances which were serialized from disk
+			TArray<UObject*> SerializedComponents;
+			// fill the array with the component contained by this object that were actually serialized to disk through property references
+			CollectDefaultSubobjects(SerializedComponents, false);
+
+			// now, add all of the instanced components to the instance graph that will be used for instancing any components that have been added
+			// to this object's archetype since this object was last saved
+			for (int32 ComponentIndex = 0; ComponentIndex < SerializedComponents.Num(); ComponentIndex++)
+			{
+				UObject* PreviouslyInstancedComponent = SerializedComponents[ComponentIndex];
+				InstanceGraph->AddNewInstance(PreviouslyInstancedComponent);
+			}
 
-		InstanceSubobjectTemplates(InstanceGraph);
+			InstanceSubobjectTemplates(InstanceGraph);
+		}
 	}
 	else
 	{
@@ -1463,19 +1468,19 @@ void UObject::BuildSubobjectMapping(UObject* OtherObject, TMap<UObject*, UObject
 
 void UObject::CollectDefaultSubobjects( TArray<UObject*>& OutSubobjectArray, bool bIncludeNestedSubobjects/*=false*/ ) const
 {
-	OutSubobjectArray.Empty();
-	GetObjectsWithOuter(this, OutSubobjectArray, bIncludeNestedSubobjects);
+		OutSubobjectArray.Empty();
+		GetObjectsWithOuter(this, OutSubobjectArray, bIncludeNestedSubobjects);
 
-	// Remove contained objects that are not subobjects.
-	for ( int32 ComponentIndex = 0; ComponentIndex < OutSubobjectArray.Num(); ComponentIndex++ )
-	{
-		UObject* PotentialComponent = OutSubobjectArray[ComponentIndex];
-		if (!PotentialComponent->IsDefaultSubobject())
+		// Remove contained objects that are not subobjects.
+		for (int32 ComponentIndex = 0; ComponentIndex < OutSubobjectArray.Num(); ComponentIndex++)
 		{
-			OutSubobjectArray.RemoveAtSwap(ComponentIndex--);
+			UObject* PotentialComponent = OutSubobjectArray[ComponentIndex];
+			if (!PotentialComponent->IsDefaultSubobject())
+			{
+				OutSubobjectArray.RemoveAtSwap(ComponentIndex--);
+			}
 		}
 	}
-}
 
 /**
  * FSubobjectReferenceFinder.
