Cache which fields from a struct/class types need gathering for localization

This removes repeated work iterating fields, and can save a considerable amount of time when processing large arrays of structs.

diff --git a/Engine/Source/Runtime/CoreUObject/Private/Serialization/PropertyLocalizationDataGathering.cpp b/Engine/Source/Runtime/CoreUObject/Private/Serialization/PropertyLocalizationDataGathering.cpp
--- a/Engine/Source/Runtime/CoreUObject/Private/Serialization/PropertyLocalizationDataGathering.cpp
+++ b/Engine/Source/Runtime/CoreUObject/Private/Serialization/PropertyLocalizationDataGathering.cpp
@@ -70,6 +70,84 @@ void FPropertyLocalizationDataGatherer::MarkObjectProcessed(const UObject* Objec
 	ProcessedObjects.Add(FObjectAndGatherFlags(Object, GatherTextFlags));
 }
 
+const FPropertyLocalizationDataGatherer::FGatherableFieldsForType& FPropertyLocalizationDataGatherer::GetGatherableFieldsForType(const UStruct* InType)
+{
+	if (const FGatherableFieldsForType* GatherableFieldsForTypePtr = GatherableFieldsForTypes.Find(InType))
+	{
+		// Already cached
+		return *GatherableFieldsForTypePtr;
+	}
+
+	// Not cached - work out the gatherable fields for this type and cache the result
+	// Note: This will also cache the result for the sub-structs within this type
+	return CacheGatherableFieldsForType(InType);
+}
+
+const FPropertyLocalizationDataGatherer::FGatherableFieldsForType& FPropertyLocalizationDataGatherer::CacheGatherableFieldsForType(const UStruct* InType)
+{
+	FGatherableFieldsForType GatherableFieldsForType;
+
+	for (TFieldIterator<const UField> FieldIt(InType, EFieldIteratorFlags::IncludeSuper, EFieldIteratorFlags::ExcludeDeprecated, EFieldIteratorFlags::IncludeInterfaces); FieldIt; ++FieldIt)
+	{
+		// Look for potential properties
+		{
+			auto ProcessInnerProperty = [this, &GatherableFieldsForType](const UProperty* InProp, const UProperty* InTypeProp)
+			{
+				if (CanGatherFromInnerProperty(InProp))
+				{
+					GatherableFieldsForType.Properties.AddUnique(InTypeProp);
+				}
+			};
+
+			const UProperty* PropertyField = Cast<UProperty>(*FieldIt);
+			if (PropertyField)
+			{
+				ProcessInnerProperty(PropertyField, PropertyField);
+				if (const UArrayProperty* ArrayProp = Cast<const UArrayProperty>(PropertyField))
+				{
+					ProcessInnerProperty(ArrayProp->Inner, PropertyField);
+				}
+				if (const UMapProperty* MapProp = Cast<const UMapProperty>(PropertyField))
+				{
+					ProcessInnerProperty(MapProp->KeyProp, PropertyField);
+					ProcessInnerProperty(MapProp->ValueProp, PropertyField);
+				}
+				if (const USetProperty* SetProp = Cast<const USetProperty>(PropertyField))
+				{
+					ProcessInnerProperty(SetProp->ElementProp, PropertyField);
+				}
+			}
+		}
+
+		// Look for potential functions
+		{
+			const UFunction* FunctionField = Cast<UFunction>(*FieldIt);
+			if (FunctionField && IsObjectValidForGather(FunctionField))
+			{
+				GatherableFieldsForType.Functions.Add(FunctionField);
+			}
+		}
+	}
+
+	return GatherableFieldsForTypes.Add(InType, MoveTemp(GatherableFieldsForType));
+}
+
+bool FPropertyLocalizationDataGatherer::CanGatherFromInnerProperty(const UProperty* InInnerProperty)
+{
+	if (InInnerProperty->IsA<UTextProperty>() || InInnerProperty->IsA<UObjectPropertyBase>())
+	{
+		return true;
+	}
+
+	if (const UStructProperty* StructInnerProp = Cast<const UStructProperty>(InInnerProperty))
+	{
+		// Call the "Get" version as we may have already cached a result for this type
+		return GetGatherableFieldsForType(StructInnerProp->Struct).HasFields();
+	}
+
+	return false;
+}
+
 void FPropertyLocalizationDataGatherer::GatherLocalizationDataFromObjectWithCallbacks(const UObject* Object, const EPropertyLocalizationGathererTextFlags GatherTextFlags)
 {
 	// See if we have a custom handler for this type
@@ -138,65 +216,56 @@ void FPropertyLocalizationDataGatherer::GatherLocalizationDataFromObject(const U
 void FPropertyLocalizationDataGatherer::GatherLocalizationDataFromObjectFields(const FString& PathToParent, const UObject* Object, const EPropertyLocalizationGathererTextFlags GatherTextFlags)
 {
 	const UObject* ArchetypeObject = Object->GetArchetype();
+	const FGatherableFieldsForType& GatherableFieldsForType = GetGatherableFieldsForType(Object->GetClass());
 
-	for (TFieldIterator<const UField> FieldIt(Object->GetClass(), EFieldIteratorFlags::IncludeSuper, EFieldIteratorFlags::ExcludeDeprecated, EFieldIteratorFlags::IncludeInterfaces); FieldIt; ++FieldIt)
+	// Gather text from the property data
+	for (const UProperty* PropertyField : GatherableFieldsForType.Properties)
 	{
-		// Gather text from the property data
+		const void* ValueAddress = PropertyField->ContainerPtrToValuePtr<void>(Object);
+		const void* DefaultValueAddress = nullptr;
+
+		if (ArchetypeObject)
 		{
-			const UProperty* PropertyField = Cast<UProperty>(*FieldIt);
-			if (PropertyField)
+			const UProperty* ArchetypePropertyField = FindField<UProperty>(ArchetypeObject->GetClass(), *PropertyField->GetName());
+			if (ArchetypePropertyField && ArchetypePropertyField->IsA(PropertyField->GetClass()))
 			{
-				const void* ValueAddress = PropertyField->ContainerPtrToValuePtr<void>(Object);
-				const void* DefaultValueAddress = nullptr;
-
-				if (ArchetypeObject)
-				{
-					const UProperty* ArchetypePropertyField = FindField<UProperty>(ArchetypeObject->GetClass(), *PropertyField->GetName());
-					if (ArchetypePropertyField && ArchetypePropertyField->IsA(PropertyField->GetClass()))
-					{
-						DefaultValueAddress = ArchetypePropertyField->ContainerPtrToValuePtr<void>(ArchetypeObject);
-					}
-				}
-
-				GatherLocalizationDataFromChildTextProperties(PathToParent, PropertyField, ValueAddress, DefaultValueAddress, GatherTextFlags | (PropertyField->HasAnyPropertyFlags(CPF_EditorOnly) ? EPropertyLocalizationGathererTextFlags::ForceEditorOnly : EPropertyLocalizationGathererTextFlags::None));
+				DefaultValueAddress = ArchetypePropertyField->ContainerPtrToValuePtr<void>(ArchetypeObject);
 			}
 		}
 
-		// Gather text from the script bytecode
+		GatherLocalizationDataFromChildTextProperties(PathToParent, PropertyField, ValueAddress, DefaultValueAddress, GatherTextFlags | (PropertyField->HasAnyPropertyFlags(CPF_EditorOnly) ? EPropertyLocalizationGathererTextFlags::ForceEditorOnly : EPropertyLocalizationGathererTextFlags::None));
+	}
+
+	// Gather text from the script bytecode
+	for (const UFunction* FunctionField : GatherableFieldsForType.Functions)
+	{
+		if (ShouldProcessObject(FunctionField, GatherTextFlags))
 		{
-			const UStruct* StructField = Cast<UStruct>(*FieldIt);
-			if (StructField && IsObjectValidForGather(StructField) && ShouldProcessObject(StructField, GatherTextFlags))
-			{
-				MarkObjectProcessed(StructField, GatherTextFlags);
-				GatherLocalizationDataFromObject(StructField, GatherTextFlags);
-			}
+			MarkObjectProcessed(FunctionField, GatherTextFlags);
+			GatherLocalizationDataFromObject(FunctionField, GatherTextFlags);
 		}
 	}
 }
 
 void FPropertyLocalizationDataGatherer::GatherLocalizationDataFromStructFields(const FString& PathToParent, const UStruct* Struct, const void* StructData, const void* DefaultStructData, const EPropertyLocalizationGathererTextFlags GatherTextFlags)
 {
-	for (TFieldIterator<const UField> FieldIt(Struct, EFieldIteratorFlags::IncludeSuper, EFieldIteratorFlags::ExcludeDeprecated, EFieldIteratorFlags::IncludeInterfaces); FieldIt; ++FieldIt)
+	const FGatherableFieldsForType& GatherableFieldsForType = GetGatherableFieldsForType(Struct);
+
+	// Gather text from the property data
+	for (const UProperty* PropertyField : GatherableFieldsForType.Properties)
 	{
-		// Gather text from the property data
-		{
-			const UProperty* PropertyField = Cast<UProperty>(*FieldIt);
-			if (PropertyField)
-			{
-				const void* ValueAddress = PropertyField->ContainerPtrToValuePtr<void>(StructData);
-				const void* DefaultValueAddress = DefaultStructData ? PropertyField->ContainerPtrToValuePtr<void>(DefaultStructData) : nullptr;
-				GatherLocalizationDataFromChildTextProperties(PathToParent, PropertyField, ValueAddress, DefaultValueAddress, GatherTextFlags | (PropertyField->HasAnyPropertyFlags(CPF_EditorOnly) ? EPropertyLocalizationGathererTextFlags::ForceEditorOnly : EPropertyLocalizationGathererTextFlags::None));
-			}
-		}
+		const void* ValueAddress = PropertyField->ContainerPtrToValuePtr<void>(StructData);
+		const void* DefaultValueAddress = DefaultStructData ? PropertyField->ContainerPtrToValuePtr<void>(DefaultStructData) : nullptr;
+		GatherLocalizationDataFromChildTextProperties(PathToParent, PropertyField, ValueAddress, DefaultValueAddress, GatherTextFlags | (PropertyField->HasAnyPropertyFlags(CPF_EditorOnly) ? EPropertyLocalizationGathererTextFlags::ForceEditorOnly : EPropertyLocalizationGathererTextFlags::None));
+	}
 
-		// Gather text from the script bytecode
+	// Gather text from the script bytecode
+	for (const UFunction* FunctionField : GatherableFieldsForType.Functions)
+	{
+		if (ShouldProcessObject(FunctionField, GatherTextFlags))
 		{
-			const UStruct* StructField = Cast<UStruct>(*FieldIt);
-			if (StructField && IsObjectValidForGather(StructField) && ShouldProcessObject(StructField, GatherTextFlags))
-			{
-				MarkObjectProcessed(StructField, GatherTextFlags);
-				GatherLocalizationDataFromObject(StructField, GatherTextFlags);
-			}
+			MarkObjectProcessed(FunctionField, GatherTextFlags);
+			GatherLocalizationDataFromObject(FunctionField, GatherTextFlags);
 		}
 	}
 }
@@ -209,6 +278,7 @@ void FPropertyLocalizationDataGatherer::GatherLocalizationDataFromChildTextPrope
 		return;
 	}
 
+	// If adding more type support here, also update CacheGatherableFieldsForType
 	const UTextProperty* const TextProperty = Cast<const UTextProperty>(Property);
 	const UArrayProperty* const ArrayProperty = Cast<const UArrayProperty>(Property);
 	const UMapProperty* const MapProperty = Cast<const UMapProperty>(Property);
@@ -218,18 +288,34 @@ void FPropertyLocalizationDataGatherer::GatherLocalizationDataFromChildTextPrope
 
 	const EPropertyLocalizationGathererTextFlags FixedChildPropertyGatherTextFlags = GatherTextFlags | (Property->HasAnyPropertyFlags(CPF_EditorOnly) ? EPropertyLocalizationGathererTextFlags::ForceEditorOnly : EPropertyLocalizationGathererTextFlags::None);
 
-	auto CanGatherFromInnerProperty = [](UProperty* InnerProperty)
+	FString PathToElementRoot;
 	{
-		return InnerProperty->IsA<UTextProperty>() ||
-			InnerProperty->IsA<UStructProperty>() ||
-			InnerProperty->IsA<UObjectPropertyBase>();
-	};
+		const FString PropertyNameStr = Property->GetName();
+
+		PathToElementRoot.Reserve(PathToParent.Len() + PropertyNameStr.Len() + 1);
+		PathToElementRoot += PathToParent;
+		if (!PathToParent.IsEmpty())
+		{
+			PathToElementRoot += TEXT('.');
+		}
+		PathToElementRoot += PropertyNameStr;
+	}
 
 	// Handles both native, fixed-size arrays and plain old non-array properties.
 	const bool IsFixedSizeArray = Property->ArrayDim > 1;
 	for(int32 i = 0; i < Property->ArrayDim; ++i)
 	{
-		const FString PathToElement = FString(PathToParent.IsEmpty() ? TEXT("") : PathToParent + TEXT(".")) + (IsFixedSizeArray ? Property->GetName() + FString::Printf(TEXT("[%d]"), i) : Property->GetName());
+		FString PathToFixedSizeArrayElement;
+		if (IsFixedSizeArray)
+		{
+			PathToFixedSizeArrayElement.Reserve(PathToElementRoot.Len() + 10); // +10 for some slack for the number and braces
+			PathToFixedSizeArrayElement += PathToElementRoot;
+			PathToFixedSizeArrayElement += TEXT('[');
+			PathToFixedSizeArrayElement.AppendInt(i);
+			PathToFixedSizeArrayElement += TEXT(']');
+		}
+
+		const FString& PathToElement = IsFixedSizeArray ? PathToFixedSizeArrayElement : PathToElementRoot;
 		const void* const ElementValueAddress = reinterpret_cast<const uint8*>(ValueAddress) + Property->ElementSize * i;
 		const void* const DefaultElementValueAddress = DefaultValueAddress ? (reinterpret_cast<const uint8*>(DefaultValueAddress) + Property->ElementSize * i) : nullptr;
 
@@ -267,31 +353,29 @@ void FPropertyLocalizationDataGatherer::GatherLocalizationDataFromChildTextPrope
 		// Property is a DYNAMIC array property.
 		else if (ArrayProperty)
 		{
-			// Skip arrays of types which can't contain any text.
-			if (CanGatherFromInnerProperty(ArrayProperty->Inner) == false)
-			{
-				continue;
-			}
-
 			// Iterate over all elements of the array.
 			FScriptArrayHelper ScriptArrayHelper(ArrayProperty, ElementValueAddress);
 			const int32 ElementCount = ScriptArrayHelper.Num();
 			for(int32 j = 0; j < ElementCount; ++j)
 			{
+				FString PathToInnerElement;
+				{
+					PathToInnerElement.Reserve(PathToElement.Len() + 10); // +10 for some slack for the number and braces
+					PathToInnerElement += PathToElement;
+					PathToInnerElement += TEXT('(');
+					PathToInnerElement.AppendInt(j);
+					PathToInnerElement += TEXT(')');
+				}
+
 				const uint8* ElementPtr = ScriptArrayHelper.GetRawPtr(j);
-				GatherLocalizationDataFromChildTextProperties(PathToElement + FString::Printf(TEXT("(%d)"), j), ArrayProperty->Inner, ElementPtr, nullptr, ElementChildPropertyGatherTextFlags);
+				GatherLocalizationDataFromChildTextProperties(PathToInnerElement, ArrayProperty->Inner, ElementPtr, nullptr, ElementChildPropertyGatherTextFlags);
 			}
 		}
 		// Property is a map property.
 		else if (MapProperty)
 		{
-			// Skip maps of types which can't contain any text.
 			const bool bGatherMapKey = CanGatherFromInnerProperty(MapProperty->KeyProp);
 			const bool bGatherMapValue = CanGatherFromInnerProperty(MapProperty->ValueProp);
-			if (!bGatherMapKey && !bGatherMapValue)
-			{
-				continue;
-			}
 
 			// Iterate over all elements of the map.
 			FScriptMapHelper ScriptMapHelper(MapProperty, ElementValueAddress);
@@ -307,14 +391,32 @@ void FPropertyLocalizationDataGatherer::GatherLocalizationDataFromChildTextPrope
 
 				if (bGatherMapKey)
 				{
+					FString PathToInnerElement;
+					{
+						PathToInnerElement.Reserve(PathToElement.Len() + 20); // +20 for some slack for the number, braces, and description
+						PathToInnerElement += PathToElement;
+						PathToInnerElement += TEXT('(');
+						PathToInnerElement.AppendInt(ElementIndex);
+						PathToInnerElement += TEXT(" - Key)");
+					}
+
 					const uint8* MapKeyPtr = MapPairPtr;
-					GatherLocalizationDataFromChildTextProperties(PathToElement + FString::Printf(TEXT("(%d - Key)"), ElementIndex), MapProperty->KeyProp, MapKeyPtr, nullptr, ElementChildPropertyGatherTextFlags);
+					GatherLocalizationDataFromChildTextProperties(PathToInnerElement, MapProperty->KeyProp, MapKeyPtr, nullptr, ElementChildPropertyGatherTextFlags);
 				}
 
 				if (bGatherMapValue)
 				{
+					FString PathToInnerElement;
+					{
+						PathToInnerElement.Reserve(PathToElement.Len() + 20); // +20 for some slack for the number, braces, and description
+						PathToInnerElement += PathToElement;
+						PathToInnerElement += TEXT('(');
+						PathToInnerElement.AppendInt(ElementIndex);
+						PathToInnerElement += TEXT(" - Value)");
+					}
+
 					const uint8* MapValuePtr = MapPairPtr + MapProperty->MapLayout.ValueOffset;
-					GatherLocalizationDataFromChildTextProperties(PathToElement + FString::Printf(TEXT("(%d - Value)"), ElementIndex), MapProperty->ValueProp, MapValuePtr, nullptr, ElementChildPropertyGatherTextFlags);
+					GatherLocalizationDataFromChildTextProperties(PathToInnerElement, MapProperty->ValueProp, MapValuePtr, nullptr, ElementChildPropertyGatherTextFlags);
 				}
 
 				++ElementIndex;
@@ -323,12 +425,6 @@ void FPropertyLocalizationDataGatherer::GatherLocalizationDataFromChildTextPrope
 		// Property is a set property.
 		else if (SetProperty)
 		{
-			// Skip sets of types which can't contain any text.
-			if (CanGatherFromInnerProperty(SetProperty->ElementProp) == false)
-			{
-				continue;
-			}
-
 			// Iterate over all elements of the Set.
 			FScriptSetHelper ScriptSetHelper(SetProperty, ElementValueAddress);
 			const int32 ElementCount = ScriptSetHelper.Num();
@@ -339,8 +435,17 @@ void FPropertyLocalizationDataGatherer::GatherLocalizationDataFromChildTextPrope
 					continue;
 				}
 
+				FString PathToInnerElement;
+				{
+					PathToInnerElement.Reserve(PathToElement.Len() + 10); // +10 for some slack for the number and braces
+					PathToInnerElement += PathToElement;
+					PathToInnerElement += TEXT('(');
+					PathToInnerElement.AppendInt(ElementIndex);
+					PathToInnerElement += TEXT(')');
+				}
+
 				const uint8* ElementPtr = ScriptSetHelper.GetElementPtr(j);
-				GatherLocalizationDataFromChildTextProperties(PathToElement + FString::Printf(TEXT("(%d)"), ElementIndex), SetProperty->ElementProp, ElementPtr, nullptr, ElementChildPropertyGatherTextFlags);
+				GatherLocalizationDataFromChildTextProperties(PathToInnerElement, SetProperty->ElementProp, ElementPtr, nullptr, ElementChildPropertyGatherTextFlags);
 				++ElementIndex;
 			}
 		}
diff --git a/Engine/Source/Runtime/CoreUObject/Public/Serialization/PropertyLocalizationDataGathering.h b/Engine/Source/Runtime/CoreUObject/Public/Serialization/PropertyLocalizationDataGathering.h
--- a/Engine/Source/Runtime/CoreUObject/Public/Serialization/PropertyLocalizationDataGathering.h
+++ b/Engine/Source/Runtime/CoreUObject/Public/Serialization/PropertyLocalizationDataGathering.h
@@ -74,11 +74,26 @@ ENUM_CLASS_FLAGS(EPropertyLocalizationGathererResultFlags);
 class COREUOBJECT_API FPropertyLocalizationDataGatherer
 {
 public:
+	struct FGatherableFieldsForType
+	{
+		TArray<const UProperty*> Properties;
+		TArray<const UFunction*> Functions;
+
+		bool HasFields() const
+		{
+			return Properties.Num() > 0 || Functions.Num() > 0;
+		}
+	};
+
 	typedef TFunction<void(const UObject* const, FPropertyLocalizationDataGatherer&, const EPropertyLocalizationGathererTextFlags)> FLocalizationDataGatheringCallback;
 	typedef TMap<const UClass*, FLocalizationDataGatheringCallback> FLocalizationDataGatheringCallbackMap;
 
 	FPropertyLocalizationDataGatherer(TArray<FGatherableTextData>& InOutGatherableTextDataArray, const UPackage* const InPackage, EPropertyLocalizationGathererResultFlags& OutResultFlags);
 
+	// Non-copyable
+	FPropertyLocalizationDataGatherer(const FPropertyLocalizationDataGatherer&) = delete;
+	FPropertyLocalizationDataGatherer& operator=(const FPropertyLocalizationDataGatherer&) = delete;
+
 	void GatherLocalizationDataFromObjectWithCallbacks(const UObject* Object, const EPropertyLocalizationGathererTextFlags GatherTextFlags);
 	void GatherLocalizationDataFromObject(const UObject* Object, const EPropertyLocalizationGathererTextFlags GatherTextFlags);
 	void GatherLocalizationDataFromObjectFields(const FString& PathToParent, const UObject* Object, const EPropertyLocalizationGathererTextFlags GatherTextFlags);
@@ -94,6 +109,8 @@ public:
 	bool ShouldProcessObject(const UObject* Object, const EPropertyLocalizationGathererTextFlags GatherTextFlags) const;
 	void MarkObjectProcessed(const UObject* Object, const EPropertyLocalizationGathererTextFlags GatherTextFlags);
 
+	const FGatherableFieldsForType& GetGatherableFieldsForType(const UStruct* InType);
+
 	static bool ExtractTextIdentity(const FText& Text, FString& OutNamespace, FString& OutKey, const bool bCleanNamespace);
 
 	static FLocalizationDataGatheringCallbackMap& GetTypeSpecificLocalizationDataGatheringCallbacks();
@@ -109,6 +126,9 @@ public:
 	}
 
 private:
+	const FGatherableFieldsForType& CacheGatherableFieldsForType(const UStruct* InType);
+	bool CanGatherFromInnerProperty(const UProperty* InInnerProperty);
+
 	struct FObjectAndGatherFlags
 	{
 		FObjectAndGatherFlags(const UObject* InObject, const EPropertyLocalizationGathererTextFlags InGatherTextFlags)
@@ -145,6 +165,7 @@ private:
 	const UPackage* Package;
 	FString PackageNamespace;
 	EPropertyLocalizationGathererResultFlags& ResultFlags;
+	TMap<const UStruct*, FGatherableFieldsForType> GatherableFieldsForTypes;
 	TSet<const UObject*> AllObjectsInPackage;
 	TSet<FObjectAndGatherFlags> ProcessedObjects;
 	TSet<FObjectAndGatherFlags> BytecodePendingGather;
