Fixed crash from CL# 7064412

These needed to be pointers as the map may have data added to it while we have a reference to an element within it.

#rnx

diff --git a/Engine/Source/Runtime/CoreUObject/Private/Serialization/PropertyLocalizationDataGathering.cpp b/Engine/Source/Runtime/CoreUObject/Private/Serialization/PropertyLocalizationDataGathering.cpp
--- a/Engine/Source/Runtime/CoreUObject/Private/Serialization/PropertyLocalizationDataGathering.cpp
+++ b/Engine/Source/Runtime/CoreUObject/Private/Serialization/PropertyLocalizationDataGathering.cpp
@@ -72,10 +72,11 @@ void FPropertyLocalizationDataGatherer::MarkObjectProcessed(const UObject* Objec
 
 const FPropertyLocalizationDataGatherer::FGatherableFieldsForType& FPropertyLocalizationDataGatherer::GetGatherableFieldsForType(const UStruct* InType)
 {
-	if (const FGatherableFieldsForType* GatherableFieldsForTypePtr = GatherableFieldsForTypes.Find(InType))
+	if (const TUniquePtr<FGatherableFieldsForType>* GatherableFieldsForTypePtr = GatherableFieldsForTypes.Find(InType))
 	{
 		// Already cached
-		return *GatherableFieldsForTypePtr;
+		check(GatherableFieldsForTypePtr->IsValid());
+		return *GatherableFieldsForTypePtr->Get();
 	}
 
 	// Not cached - work out the gatherable fields for this type and cache the result
@@ -85,7 +86,7 @@ const FPropertyLocalizationDataGatherer::FGatherableFieldsForType& FPropertyLoca
 
 const FPropertyLocalizationDataGatherer::FGatherableFieldsForType& FPropertyLocalizationDataGatherer::CacheGatherableFieldsForType(const UStruct* InType)
 {
-	FGatherableFieldsForType GatherableFieldsForType;
+	TUniquePtr<FGatherableFieldsForType> GatherableFieldsForType = MakeUnique<FGatherableFieldsForType>();
 
 	for (TFieldIterator<const UField> FieldIt(InType, EFieldIteratorFlags::IncludeSuper, EFieldIteratorFlags::ExcludeDeprecated, EFieldIteratorFlags::IncludeInterfaces); FieldIt; ++FieldIt)
 	{
@@ -95,7 +96,7 @@ const FPropertyLocalizationDataGatherer::FGatherableFieldsForType& FPropertyLoca
 			{
 				if (CanGatherFromInnerProperty(InProp))
 				{
-					GatherableFieldsForType.Properties.AddUnique(InTypeProp);
+					GatherableFieldsForType->Properties.AddUnique(InTypeProp);
 				}
 			};
 
@@ -124,12 +125,13 @@ const FPropertyLocalizationDataGatherer::FGatherableFieldsForType& FPropertyLoca
 			const UFunction* FunctionField = Cast<UFunction>(*FieldIt);
 			if (FunctionField && IsObjectValidForGather(FunctionField))
 			{
-				GatherableFieldsForType.Functions.Add(FunctionField);
+				GatherableFieldsForType->Functions.Add(FunctionField);
 			}
 		}
 	}
 
-	return GatherableFieldsForTypes.Add(InType, MoveTemp(GatherableFieldsForType));
+	check(GatherableFieldsForType.IsValid());
+	return *GatherableFieldsForTypes.Add(InType, MoveTemp(GatherableFieldsForType)).Get();
 }
 
 bool FPropertyLocalizationDataGatherer::CanGatherFromInnerProperty(const UProperty* InInnerProperty)
diff --git a/Engine/Source/Runtime/CoreUObject/Public/Serialization/PropertyLocalizationDataGathering.h b/Engine/Source/Runtime/CoreUObject/Public/Serialization/PropertyLocalizationDataGathering.h
--- a/Engine/Source/Runtime/CoreUObject/Public/Serialization/PropertyLocalizationDataGathering.h
+++ b/Engine/Source/Runtime/CoreUObject/Public/Serialization/PropertyLocalizationDataGathering.h
@@ -165,7 +165,7 @@ private:
 	const UPackage* Package;
 	FString PackageNamespace;
 	EPropertyLocalizationGathererResultFlags& ResultFlags;
-	TMap<const UStruct*, FGatherableFieldsForType> GatherableFieldsForTypes;
+	TMap<const UStruct*, TUniquePtr<FGatherableFieldsForType>> GatherableFieldsForTypes;
 	TSet<const UObject*> AllObjectsInPackage;
 	TSet<FObjectAndGatherFlags> ProcessedObjects;
 	TSet<FObjectAndGatherFlags> BytecodePendingGather;
