Adding a range check to the cluster array index operator that logs the bad index unlike the TArray range check

diff --git a/Engine/Source/Runtime/CoreUObject/Public/UObject/UObjectArray.h b/Engine/Source/Runtime/CoreUObject/Public/UObject/UObjectArray.h
--- a/Engine/Source/Runtime/CoreUObject/Public/UObject/UObjectArray.h
+++ b/Engine/Source/Runtime/CoreUObject/Public/UObject/UObjectArray.h
@@ -1073,6 +1073,7 @@ public:
 
 	FORCEINLINE FUObjectCluster& operator[](int32 Index)
 	{
+		checkf(Index >= 0 && Index < Clusters.Num(), TEXT("Cluster index %d out of range [0, %d]"), Index, Clusters.Num());
 		return Clusters[Index];
 	}
 
