FScriptMapHelper and FScriptSetHelper now have a GetInternalIndex that maps a logical index (ie. Nth element) to the internal index.

Source: 6932549 (Dev-Editor)

diff --git a/Engine/Source/Runtime/CoreUObject/Public/UObject/UnrealType.h b/Engine/Source/Runtime/CoreUObject/Public/UObject/UnrealType.h
--- a/Engine/Source/Runtime/CoreUObject/Public/UObject/UnrealType.h
+++ b/Engine/Source/Runtime/CoreUObject/Public/UObject/UnrealType.h
@@ -3264,6 +3264,30 @@ public:
 	 */
 	COREUOBJECT_API void Rehash();
 
+	/** 
+	 * Maps have gaps in their indices, so this function translates a logical index (ie. Nth element) 
+	 * to an internal index that can be used for the other functions in this class.
+	 * NOTE: This is slow, do not use this for iteration!
+	 */
+	int32 FindInternalIndex(int32 LogicalIdx) const
+	{
+		check(LogicalIdx >= 0 && LogicalIdx < Num());
+
+		int32 MaxIndex = GetMaxIndex();
+		for (int32 Actual = 0; Actual < MaxIndex; ++Actual)
+		{
+			if (IsValidIndex(Actual))
+			{
+				if (LogicalIdx == 0)
+				{
+					return Actual;
+				}
+				--LogicalIdx;
+			}
+		}
+		return INDEX_NONE;
+	}
+
 	/**
 	 * Finds the index of an element in a map which matches the key in another pair.
 	 *
@@ -3815,6 +3839,30 @@ public:
 	*/
 	COREUOBJECT_API void Rehash();
 
+	/**
+	 * Maps have gaps in their indices, so this function translates a logical index (ie. Nth element)
+	 * to an internal index that can be used for the other functions in this class.
+	 * NOTE: This is slow, do not use this for iteration!
+	 */
+	int32 FindInternalIndex(int32 LogicalIdx) const
+	{
+		check(LogicalIdx >= 0 && LogicalIdx < Num());
+
+		int32 MaxIndex = GetMaxIndex();
+		for (int32 Actual = 0; Actual < MaxIndex; ++Actual)
+		{
+			if (IsValidIndex(Actual))
+			{
+				if (LogicalIdx == 0)
+				{
+					return Actual;
+				}
+				--LogicalIdx;
+			}
+		}
+		return INDEX_NONE;
+	}
+	
 	/**
 	* Finds the index of an element in a set
 	*
