Make sure Blueprint Property Conditions are actually serialized out with UProperties, so they aren't lost in packaged games.
UNTESTED

diff --git a/Engine/Source/Runtime/Core/Public/UObject/ReleaseObjectVersion.h b/Engine/Source/Runtime/Core/Public/UObject/ReleaseObjectVersion.h
--- a/Engine/Source/Runtime/Core/Public/UObject/ReleaseObjectVersion.h
+++ b/Engine/Source/Runtime/Core/Public/UObject/ReleaseObjectVersion.h
@@ -74,6 +74,9 @@ struct CORE_API FReleaseObjectVersion
 		// Change trail anim node variable deprecation
 		TrailNodeBlendVariableNameChange,
 
+		// Make sure the Blueprint Replicated Property Conditions are actually serialized properly.
+		PropertiesSerializeRepCondition,
+
 		// -----<new versions can be added above this line>-------------------------------------------------
 		VersionPlusOne,
 		LatestVersion = VersionPlusOne - 1
diff --git a/Engine/Source/Runtime/CoreUObject/Private/UObject/Property.cpp b/Engine/Source/Runtime/CoreUObject/Private/UObject/Property.cpp
--- a/Engine/Source/Runtime/CoreUObject/Private/UObject/Property.cpp
+++ b/Engine/Source/Runtime/CoreUObject/Private/UObject/Property.cpp
@@ -17,6 +17,7 @@
 #include "UObject/CoreRedirects.h"
 #include "UObject/SoftObjectPath.h"
 #include "Math/Box2D.h"
+#include "UObject/ReleaseObjectVersion.h"
 
 DEFINE_LOG_CATEGORY(LogProperty);
 
@@ -478,6 +479,12 @@ void UProperty::Serialize( FArchive& Ar )
 		Offset_Internal = 0;
 		DestructorLinkNext = NULL;
 	}
+
+	Ar.UsingCustomVersion(FReleaseObjectVersion::GUID);
+	if (Ar.IsSaving() || Ar.CustomVer(FReleaseObjectVersion::GUID) >= FReleaseObjectVersion::PropertiesSerializeRepCondition)
+	{
+		Ar << BlueprintReplicationCondition;
+	}
 }
 
 void UProperty::CopySingleValueToScriptVM( void* Dest, void const* Src ) const
