Fixed a few issues with running the final Garbage Collector pass when exiting:

+ Fixed a crash when exiting early during init caused by disregard for GC set not being closed
+ Fixed a crash when exiting early during init caused by classes not being fully initialized
+ Added guards to prevent other threads from manipulating UObjects when the final GC runs
+ The final GC will not use multithreaded destruction as it will not help with GC speed when there's no time limit for GC

diff --git a/Engine/Source/Runtime/CoreUObject/Private/UObject/GarbageCollection.cpp b/Engine/Source/Runtime/CoreUObject/Private/UObject/GarbageCollection.cpp
--- a/Engine/Source/Runtime/CoreUObject/Private/UObject/GarbageCollection.cpp
+++ b/Engine/Source/Runtime/CoreUObject/Private/UObject/GarbageCollection.cpp
@@ -151,7 +151,7 @@ static FAutoConsoleVariableRef CIncrementalBeginDestroyEnabled(
 	ECVF_Default
 );
 
-static int32 GMultithreadedDestructionEnabled = 1;
+int32 GMultithreadedDestructionEnabled = 1;
 static FAutoConsoleVariableRef CMultithreadedDestructionEnabled(
 	TEXT("gc.MultithreadedDestructionEnabled"),
 	GMultithreadedDestructionEnabled,
@@ -1238,7 +1238,7 @@ static bool ShouldForceSingleThreadedGC()
 	return bForceSingleThreadedGC;
 }
 
-static void AcquireGCLock()
+void AcquireGCLock()
 {
 	const double StartTime = FPlatformTime::Seconds();
 	FGCCSyncObject::Get().GCLock();
@@ -1249,7 +1249,7 @@ static void AcquireGCLock()
 	}
 }
 
-static void ReleaseGCLock()
+void ReleaseGCLock()
 {
 	FGCCSyncObject::Get().GCUnlock();
 }
diff --git a/Engine/Source/Runtime/CoreUObject/Private/UObject/Obj.cpp b/Engine/Source/Runtime/CoreUObject/Private/UObject/Obj.cpp
--- a/Engine/Source/Runtime/CoreUObject/Private/UObject/Obj.cpp
+++ b/Engine/Source/Runtime/CoreUObject/Private/UObject/Obj.cpp
@@ -4341,6 +4341,9 @@ void StaticUObjectInit()
 void ShutdownGarbageCollection();
 void CleanupLinkerAnnotations();
 void CleanupCachedArchetypes();
+void AcquireGCLock();
+void ReleaseGCLock();
+extern int32 GMultithreadedDestructionEnabled;
 
 //
 // Shut down the object manager.
@@ -4355,12 +4358,26 @@ void StaticExit()
 	// Delete all linkers are pending destroy
 	DeleteLoaders();
 
+	// We'll be destroying objects without time limit during exit purge
+	// so doing it on a separate thread doesn't make anything faster,
+	// also the exit purge is not a standard GC pass so no need to overcompilcate things
+	GMultithreadedDestructionEnabled = false;
+
 	// Cleanup root.
 	if (GObjTransientPkg != NULL)
 	{
 		GObjTransientPkg->RemoveFromRoot();
 		GObjTransientPkg = NULL;
 	}
+	
+	// This can happen when we run into an error early in the init process
+	if (GUObjectArray.IsOpenForDisregardForGC())
+	{
+		GUObjectArray.CloseDisregardForGC();
+	}
+
+	// Make sure no other threads manipulate UObjects
+	AcquireGCLock();
 
 	GatherUnreachableObjects(false);
 	IncrementalPurgeGarbage(false);
@@ -4414,6 +4431,8 @@ void StaticExit()
 		IncrementalPurgeGarbage(false);
 	}
 
+	ReleaseGCLock();
+
 	ShutdownGarbageCollection();
 	UObjectBaseShutdown();
 
diff --git a/Engine/Source/Runtime/CoreUObject/Private/UObject/UObjectArray.cpp b/Engine/Source/Runtime/CoreUObject/Private/UObject/UObjectArray.cpp
--- a/Engine/Source/Runtime/CoreUObject/Private/UObject/UObjectArray.cpp
+++ b/Engine/Source/Runtime/CoreUObject/Private/UObject/UObjectArray.cpp
@@ -69,37 +69,37 @@ void FUObjectArray::CloseDisregardForGC()
 	if (!GIsRequestingExit)
 	{
 		ProcessNewlyLoadedUObjects();
-	}
 
-	UClass::AssembleReferenceTokenStreams();
+		UClass::AssembleReferenceTokenStreams();
 
-	if (GIsInitialLoad)
-	{
-		// Iterate over all objects and mark them to be part of root set.
-		int32 NumAlwaysLoadedObjects = 0;
-		int32 NumRootObjects = 0;
-		for (FObjectIterator It; It; ++It)
+		if (GIsInitialLoad)
 		{
-			UObject* Object = *It;
-			if (Object->IsSafeForRootSet())
+			// Iterate over all objects and mark them to be part of root set.
+			int32 NumAlwaysLoadedObjects = 0;
+			int32 NumRootObjects = 0;
+			for (FObjectIterator It; It; ++It)
 			{
-				NumRootObjects++;
-				Object->AddToRoot();
+				UObject* Object = *It;
+				if (Object->IsSafeForRootSet())
+				{
+					NumRootObjects++;
+					Object->AddToRoot();
+				}
+				else if (Object->IsRooted())
+				{
+					Object->RemoveFromRoot();
+				}
+				NumAlwaysLoadedObjects++;
 			}
-			else if (Object->IsRooted())
+
+			UE_LOG(LogUObjectArray, Log, TEXT("%i objects as part of root set at end of initial load."), NumAlwaysLoadedObjects);
+			if (GUObjectArray.DisregardForGCEnabled())
 			{
-				Object->RemoveFromRoot();
+				UE_LOG(LogUObjectArray, Log, TEXT("%i objects are not in the root set, but can never be destroyed because they are in the DisregardForGC set."), NumAlwaysLoadedObjects - NumRootObjects);
 			}
-			NumAlwaysLoadedObjects++;
-		}
 
-		UE_LOG(LogUObjectArray, Log, TEXT("%i objects as part of root set at end of initial load."), NumAlwaysLoadedObjects);
-		if (GUObjectArray.DisregardForGCEnabled())
-		{
-			UE_LOG(LogUObjectArray, Log, TEXT("%i objects are not in the root set, but can never be destroyed because they are in the DisregardForGC set."), NumAlwaysLoadedObjects - NumRootObjects);
+			GUObjectAllocator.BootMessage();
 		}
-
-		GUObjectAllocator.BootMessage();
 	}
 
 	// When disregard for GC pool is closed, make sure the first GC index is set after the last non-GC index.
