Missing include
#rnx

diff --git a/Engine/Source/Runtime/CoreUObject/Public/UObject/ConstructorHelpers.h b/Engine/Source/Runtime/CoreUObject/Public/UObject/ConstructorHelpers.h
--- a/Engine/Source/Runtime/CoreUObject/Public/UObject/ConstructorHelpers.h
+++ b/Engine/Source/Runtime/CoreUObject/Public/UObject/ConstructorHelpers.h
@@ -9,6 +9,7 @@
 #include "CoreMinimal.h"
 #include "Templates/SubclassOf.h"
 #include "UObject/Package.h"
+#include "UObject/GCObject.h"
 
 namespace ConstructorHelpersInternal
 {
