Perform a placeholder resolve pass immediately after the CDO has been serialized

diff --git a/Engine/Source/Runtime/CoreUObject/Private/Blueprint/BlueprintSupport.cpp b/Engine/Source/Runtime/CoreUObject/Private/Blueprint/BlueprintSupport.cpp
--- a/Engine/Source/Runtime/CoreUObject/Private/Blueprint/BlueprintSupport.cpp
+++ b/Engine/Source/Runtime/CoreUObject/Private/Blueprint/BlueprintSupport.cpp
@@ -542,6 +542,96 @@
 	}
 };
 
+/** 
+ * A helper utility for tracking exports whose classes we're currently running
+ * through ForceRegenerateClass(). This is primarily relied upon to help prevent
+ * infinite recursion since ForceRegenerateClass() doesn't do anything to 
+ * progress the state of the linker.
+ */
+struct FResolvingExportTracker : TThreadSingleton<FResolvingExportTracker>
+{
+public:
+	/**  */
+	void FlagLinkerExportAsResolving(FLinkerLoad* Linker, int32 ExportIndex)
+	{
+		ResolvingExports.FindOrAdd(Linker).Add(ExportIndex);
+	}
+
+	/**  */
+	bool IsLinkerExportBeingResolved(FLinkerLoad* Linker, int32 ExportIndex) const
+	{
+		if (const TSet<int32>* ExportIndices = ResolvingExports.Find(Linker))
+		{
+			return ExportIndices->Contains(ExportIndex);
+		}
+		return false;
+	}
+
+	/**  */
+	void FlagExportClassAsFullyResolved(FLinkerLoad* Linker, int32 ExportIndex)
+	{
+		if (TSet<int32>* ExportIndices = ResolvingExports.Find(Linker))
+		{
+			ExportIndices->Remove(ExportIndex);
+			if (ExportIndices->Num() == 0)
+			{
+				ResolvingExports.Remove(Linker);
+			}
+		}
+	}
+
+#if USE_DEFERRED_DEPENDENCY_CHECK_VERIFICATION_TESTS
+	void FlagFullExportResolvePassComplete(FLinkerLoad* Linker)
+	{
+		FullyResolvedLinkers.Add(Linker);
+	}
+
+	bool HasPerformedFullExportResolvePass(FLinkerLoad* Linker)
+	{
+		return FullyResolvedLinkers.Contains(Linker);
+	}
+#endif // USE_DEFERRED_DEPENDENCY_CHECK_VERIFICATION_TESTS
+
+	void Reset(FLinkerLoad* Linker)
+	{
+		ResolvingExports.Remove(Linker);
+#if USE_DEFERRED_DEPENDENCY_CHECK_VERIFICATION_TESTS
+		FullyResolvedLinkers.Remove(Linker);
+#endif // USE_DEFERRED_DEPENDENCY_CHECK_VERIFICATION_TESTS
+	}
+
+	void AddLinkerPlaceholderObject(UClass* ClassWaitingFor, ULinkerPlaceholderExportObject* Placeholder)
+	{
+		ClassToPlaceholderMap.FindOrAdd(ClassWaitingFor).Add(Placeholder);
+	}
+
+	void ResolvePlaceholders(UClass* ForClass)
+	{
+		TArray<ULinkerPlaceholderExportObject*>* Placeholders = ClassToPlaceholderMap.Find(ForClass);
+		if( Placeholders )
+		{
+			for(ULinkerPlaceholderExportObject* Placeholder : *Placeholders)
+			{
+				FLinkerLoad* Linker = Placeholder->GetLinker();
+				if(ensure(Linker))
+				{
+					Linker->ResolvePlaceholder( Placeholder );
+				}
+			}
+		}
+	}
+
+private:
+	/**  */
+	TMap< FLinkerLoad*, TSet<int32> > ResolvingExports;
+
+#if USE_DEFERRED_DEPENDENCY_CHECK_VERIFICATION_TESTS
+	TSet<FLinkerLoad*> FullyResolvedLinkers;
+#endif // USE_DEFERRED_DEPENDENCY_CHECK_VERIFICATION_TESTS
+
+	TMap< UClass*, TArray<ULinkerPlaceholderExportObject*> > ClassToPlaceholderMap;
+};
+
 /**
  * Regenerates/Refreshes a blueprint class
  *
@@ -589,6 +679,9 @@
 		LoadClass->StaticLink(true);
 		Preload(CurrentCDO);
 
+		// CDO preloaded - we can now resolve placeholders:
+		FResolvingExportTracker::Get().ResolvePlaceholders(LoadClass);
+
 		// Make sure that we regenerate any parent classes first before attempting to build a child
 		TArray<UClass*> ClassChainOrdered;
 		{
@@ -1020,72 +1113,5 @@
 #endif // WITH_EDITOR
 
 /** 
- * A helper utility for tracking exports whose classes we're currently running
- * through ForceRegenerateClass(). This is primarily relied upon to help prevent
- * infinite recursion since ForceRegenerateClass() doesn't do anything to 
- * progress the state of the linker.
- */
-struct FResolvingExportTracker : TThreadSingleton<FResolvingExportTracker>
-{
-public:
-	/**  */
-	void FlagLinkerExportAsResolving(FLinkerLoad* Linker, int32 ExportIndex)
-	{
-		ResolvingExports.FindOrAdd(Linker).Add(ExportIndex);
-	}
-
-	/**  */
-	bool IsLinkerExportBeingResolved(FLinkerLoad* Linker, int32 ExportIndex) const
-	{
-		if (const TSet<int32>* ExportIndices = ResolvingExports.Find(Linker))
-		{
-			return ExportIndices->Contains(ExportIndex);
-		}
-		return false;
-	}
-
-	/**  */
-	void FlagExportClassAsFullyResolved(FLinkerLoad* Linker, int32 ExportIndex)
-	{
-		if (TSet<int32>* ExportIndices = ResolvingExports.Find(Linker))
-		{
-			ExportIndices->Remove(ExportIndex);
-			if (ExportIndices->Num() == 0)
-			{
-				ResolvingExports.Remove(Linker);
-			}
-		}
-	}
-
-#if USE_DEFERRED_DEPENDENCY_CHECK_VERIFICATION_TESTS
-	void FlagFullExportResolvePassComplete(FLinkerLoad* Linker)
-	{
-		FullyResolvedLinkers.Add(Linker);
-	}
-
-	bool HasPerformedFullExportResolvePass(FLinkerLoad* Linker)
-	{
-		return FullyResolvedLinkers.Contains(Linker);
-	}
-#endif // USE_DEFERRED_DEPENDENCY_CHECK_VERIFICATION_TESTS
-
-	void Reset(FLinkerLoad* Linker)
-	{
-		ResolvingExports.Remove(Linker);
-#if USE_DEFERRED_DEPENDENCY_CHECK_VERIFICATION_TESTS
-		FullyResolvedLinkers.Remove(Linker);
-#endif // USE_DEFERRED_DEPENDENCY_CHECK_VERIFICATION_TESTS
-	}
-
-private:
-	/**  */
-	TMap< FLinkerLoad*, TSet<int32> > ResolvingExports;
-
-#if USE_DEFERRED_DEPENDENCY_CHECK_VERIFICATION_TESTS
-	TSet<FLinkerLoad*> FullyResolvedLinkers;
-#endif // USE_DEFERRED_DEPENDENCY_CHECK_VERIFICATION_TESTS
-};
-
-/** 
  * A helper struct that adds and removes its linker/export combo from the 
  * thread's FResolvingExportTracker (based off the scope it was declared within).
@@ -1141,6 +1167,7 @@
 		PlaceholderName = MakeUniqueObjectName(Outer, PlaceholderType, PlaceholderName);
 
 		ULinkerPlaceholderExportObject* Placeholder = NewObject<ULinkerPlaceholderExportObject>(Outer, PlaceholderType, PlaceholderName, RF_Public | RF_Transient);
+		Placeholder->SetLinker(this, Index, false);
 		Placeholder->PackageIndex = FPackageIndex::FromExport(Index);
 		
 		Export.Object = Placeholder;
@@ -1186,10 +1213,6 @@
 		ForceRegenerateClass(LoadClass);
 		return false;
 	}
-	// we haven't come across a scenario where this happens, but if this hits 
-	// then we're deferring exports that will NEVER be resolved (possibly 
-	// stemming from CDO serialization in FLinkerLoad::ResolveDeferredExports)
-	DEFERRED_DEPENDENCY_CHECK(!FResolvingExportTracker::Get().HasPerformedFullExportResolvePass(this));
 	
 	UPackage* PlaceholderOuter = LinkerRoot;
 	UClass*   PlaceholderType  = ULinkerPlaceholderExportObject::StaticClass();
@@ -1201,6 +1224,8 @@
 
 	ULinkerPlaceholderExportObject* Placeholder = NewObject<ULinkerPlaceholderExportObject>(PlaceholderOuter, PlaceholderType, PlaceholderName, RF_Public | RF_Transient);
 	Placeholder->PackageIndex = FPackageIndex::FromExport(Index);
+	Placeholder->SetLinker(this, Index, false);
+	FResolvingExportTracker::Get().AddLinkerPlaceholderObject(LoadClass, Placeholder);
 
 	Export.Object = Placeholder;
 #endif // USE_CIRCULAR_DEPENDENCY_LOAD_DEFERRING
@@ -1932,6 +1957,7 @@
 				}
 
 				// replace the placeholder with the proper object instance
+				PlaceholderExport->SetLinker(nullptr, INDEX_NONE);
 				Export.Object = nullptr;
 				UObject* ExportObj = CreateExport(ExportIndex);
 
@@ -2030,6 +2056,7 @@
 				if (ensure(PlaceholderExport))
 				{
 					// replace the placeholder with the proper object instance
+					PlaceholderExport->SetLinker(nullptr, INDEX_NONE);
 					Export.Object = nullptr;
 					UObject* ExportObj = CreateExport(ExportIndex);
 
@@ -2064,6 +2091,30 @@
 #endif // USE_CIRCULAR_DEPENDENCY_LOAD_DEFERRING
 }
 
+void FLinkerLoad::ResolvePlaceholder(ULinkerPlaceholderExportObject* Placeholder)
+{
+	int32 ExportIndex = Placeholder->PackageIndex.ToExport();
+
+	Placeholder->SetLinker(nullptr, INDEX_NONE);
+
+	FObjectExport& Export = ExportMap[ExportIndex];
+	Export.Object = nullptr;
+
+	UObject* ReplacementObject = CreateExport(ExportIndex);
+	Placeholder->ResolveAllPlaceholderReferences(ReplacementObject);
+	Placeholder->MarkPendingKill();
+	
+	// recurse:
+	ResolvedDeferredSubobjects(Placeholder);
+
+	// attempt to preload, we don't really care if this doesn't complete but we don't want to fail
+	// to serialize an object:
+	if (ReplacementObject != nullptr)
+	{
+		Preload(ReplacementObject);
+	}
+}
+
 void FLinkerLoad::ResolvedDeferredSubobjects(ULinkerPlaceholderExportObject* OwningPlaceholder)
 {
 #if USE_CIRCULAR_DEPENDENCY_LOAD_DEFERRING
@@ -2071,6 +2122,9 @@
 	for(ULinkerPlaceholderExportObject* PlaceholderSubobject : OwningPlaceholder->GetSubobjectPlaceholders() )
 	{
 		int32 ExportIndex = PlaceholderSubobject->PackageIndex.ToExport();
+	
+		PlaceholderSubobject->SetLinker(nullptr, INDEX_NONE);
+
 		FObjectExport& Export = ExportMap[ExportIndex];
 
 		Export.Object = nullptr;
diff --git a/Engine/Source/Runtime/CoreUObject/Public/UObject/LinkerLoad.h b/Engine/Source/Runtime/CoreUObject/Public/UObject/LinkerLoad.h
--- a/Engine/Source/Runtime/CoreUObject/Public/UObject/LinkerLoad.h
+++ b/Engine/Source/Runtime/CoreUObject/Public/UObject/LinkerLoad.h
@@ -80,6 +80,7 @@ class FLinkerLoad
 	friend class UObject;
 	friend class UPackageMap;
 	friend struct FAsyncPackage;
+	friend struct FResolvingExportTracker;
 protected:
 	/** Linker loading status. */
 	enum ELinkerStatus
@@ -1078,6 +1079,9 @@ private:
 	 */
 	void ResolveDeferredExports(UClass* LoadClass);
 
+	/** Called when a class is loaded and its now safe to load deferred instances */
+	void ResolvePlaceholder(ULinkerPlaceholderExportObject* Placeholder);
+
 	/** Helper function to recursively resolve placeholders that were waiting for their outer */
 	void ResolvedDeferredSubobjects(ULinkerPlaceholderExportObject* OwningPlaceholder);
 
