Interpreter Optimization 11 of N: Avoid accessing FDeferredCdoInitializationTracker if we're not potentially going to defer an object

diff --git a/Engine/Source/Runtime/CoreUObject/Private/Blueprint/BlueprintSupport.cpp b/Engine/Source/Runtime/CoreUObject/Private/Blueprint/BlueprintSupport.cpp
--- a/Engine/Source/Runtime/CoreUObject/Private/Blueprint/BlueprintSupport.cpp
+++ b/Engine/Source/Runtime/CoreUObject/Private/Blueprint/BlueprintSupport.cpp
@@ -2823,12 +2823,11 @@ FObjectInitializer* FDeferredObjInitializationHelper::DeferObjectInitializerIfNe
 	UObject* TargetObj = DeferringInitializer.GetObj();
 	if (TargetObj)
 	{
-		FDeferredCdoInitializationTracker& CdoInitDeferalSys = FDeferredCdoInitializationTracker::Get();
-		auto IsSuperCdoReadyToBeCopied = [&CdoInitDeferalSys](const UClass* LoadClass, const UObject* SuperCDO)->bool
+		auto IsSuperCdoReadyToBeCopied = [](FDeferredCdoInitializationTracker& InCdoInitDeferalSys, const UClass* LoadClass, const UObject* SuperCDO)->bool
 		{
 			// RF_WasLoaded indicates that this Super was loaded from disk (and hasn't been regenerated on load)
 			// regenerated CDOs will not have the RF_LoadCompleted
-			const bool bSuperCdoLoadPending = CdoInitDeferalSys.IsInitializationDeferred(SuperCDO) ||
+			const bool bSuperCdoLoadPending = InCdoInitDeferalSys.IsInitializationDeferred(SuperCDO) ||
 				SuperCDO->HasAnyFlags(RF_NeedLoad) || (SuperCDO->HasAnyFlags(RF_WasLoaded) && !SuperCDO->HasAnyFlags(RF_LoadCompleted));
 
 			if (bSuperCdoLoadPending)
@@ -2859,8 +2858,9 @@ FObjectInitializer* FDeferredObjInitializationHelper::DeferObjectInitializerIfNe
 				DEFERRED_DEPENDENCY_CHECK(SuperCDO && SuperCDO->HasAnyFlags(RF_ClassDefaultObject));
 				// use the ObjectArchetype for the super CDO because the SuperClass may have a REINST CDO cached currently
 				SuperClass = SuperCDO->GetClass();
-
-				if (!IsSuperCdoReadyToBeCopied(CdoClass, SuperCDO))
+				
+				FDeferredCdoInitializationTracker& CdoInitDeferalSys = FDeferredCdoInitializationTracker::Get();
+				if (!IsSuperCdoReadyToBeCopied(CdoInitDeferalSys, CdoClass, SuperCDO))
 				{
 					DeferredInitializerCopy = CdoInitDeferalSys.Add(SuperCDO, DeferringInitializer);
 				}
@@ -2885,7 +2885,8 @@ FObjectInitializer* FDeferredObjInitializationHelper::DeferObjectInitializerIfNe
 				// 
 				// So if the super CDO isn't ready, we need to defer this sub-object
 				const UObject* SuperCDO = SuperClass->ClassDefaultObject;
-				if (!IsSuperCdoReadyToBeCopied(OwnerClass, SuperCDO))
+				FDeferredCdoInitializationTracker& CdoInitDeferalSys = FDeferredCdoInitializationTracker::Get();
+				if (!IsSuperCdoReadyToBeCopied(CdoInitDeferalSys, OwnerClass, SuperCDO))
 				{
 					FDeferredSubObjInitializationTracker& SubObjInitDeferalSys = FDeferredSubObjInitializationTracker::Get();
 					DeferredInitializerCopy = SubObjInitDeferalSys.Add(SuperCDO, DeferringInitializer);
