When enabling cooked assets in the editor, the engine will now use FArchiveAsync2 in the editor to be able to open them.

(re-implemented CL #4400327 from Dev-Core)

diff --git a/Engine/Source/Runtime/CoreUObject/Private/UObject/LinkerLoad.cpp b/Engine/Source/Runtime/CoreUObject/Private/UObject/LinkerLoad.cpp
--- a/Engine/Source/Runtime/CoreUObject/Private/UObject/LinkerLoad.cpp
+++ b/Engine/Source/Runtime/CoreUObject/Private/UObject/LinkerLoad.cpp
@@ -75,6 +75,20 @@ FName FLinkerLoad::NAME_LoadErrors("LoadErrors");
 Helpers
 ----------------------------------------------------------------------------*/
 
+// Helper struct for getting the value of [Core.System] AllowCookedDataInEditorBuilds from ini
+struct FInitCookedDatataInEditorBuildsSupport
+{
+	bool bAllowCookedData;
+	FInitCookedDatataInEditorBuildsSupport()
+	{
+		if (!GConfig->GetBool(TEXT("Core.System"), TEXT("AllowCookedDataInEditorBuilds"), bAllowCookedData, GEngineIni))
+		{
+			bAllowCookedData = true;
+		}
+	}
+	FORCEINLINE operator bool() const { return bAllowCookedData; }
+};
+
 /**
 * Test whether the given package index is a valid import or export in this package
 */
@@ -984,7 +998,10 @@ FLinkerLoad::ELinkerStatus FLinkerLoad::CreateLoader(
 			else
 #endif
 			{
-				bool bCanUseFArchiveAsync2 = FPlatformProperties::RequiresCookedData();
+				// If want to be able to load cooked data in the editor we need to use FArchiveAsync2 which supports EDL cooked packages,
+				// otherwise the generic file reader is faster in the editor so use that
+				static FInitCookedDatataInEditorBuildsSupport AllowCookedDataInEditorBuilds;
+				bool bCanUseFArchiveAsync2 = FPlatformProperties::RequiresCookedData() || AllowCookedDataInEditorBuilds;
 				if (bCanUseFArchiveAsync2)
 				{
 					Loader = new FArchiveAsync2(*Filename
@@ -1178,18 +1195,7 @@ FLinkerLoad::ELinkerStatus FLinkerLoad::SerializePackageFileSummary()
 		if (FPlatformProperties::HasEditorOnlyData() && !!(Summary.PackageFlags & PKG_FilterEditorOnly))
 		{
 			// This warning can be disabled in ini with [Core.System] AllowCookedDataInEditorBuilds=False
-			static struct FInitCookedDatataInEditorBuildsSupport
-			{
-				bool bAllowCookedData;
-				FInitCookedDatataInEditorBuildsSupport()
-				{
-					if (!GConfig->GetBool(TEXT("Core.System"), TEXT("AllowCookedDataInEditorBuilds"), bAllowCookedData, GEngineIni))
-					{
-						bAllowCookedData = true;
-					}
-				}
-				FORCEINLINE operator bool() const { return bAllowCookedData; }
-			} AllowCookedDataInEditorBuilds;
+			static FInitCookedDatataInEditorBuildsSupport AllowCookedDataInEditorBuilds;
 			if (!AllowCookedDataInEditorBuilds)
 			{
 				UE_LOG(LogLinker, Warning, 
