Added logging to package reloading

diff --git a/Engine/Source/Runtime/CoreUObject/Private/UObject/PackageReload.cpp b/Engine/Source/Runtime/CoreUObject/Private/UObject/PackageReload.cpp
--- a/Engine/Source/Runtime/CoreUObject/Private/UObject/PackageReload.cpp
+++ b/Engine/Source/Runtime/CoreUObject/Private/UObject/PackageReload.cpp
@@ -420,6 +420,17 @@ UPackage* ReloadPackage(UPackage* InPackageToReload, const uint32 InLoadFlags)
 
 void ReloadPackages(const TArrayView<FReloadPackageData>& InPackagesToReload, TArray<UPackage*>& OutReloadedPackages, const int32 InNumPackagesPerBatch)
 {
+	FString Msg;
+	Msg.Append(FString::Printf(TEXT("Reloading %d Package(s):"), InPackagesToReload.Num()));
+	const int32 MAX_PACKAGES_TO_LOG = 10;
+	for(int32 I = 0; I < InPackagesToReload.Num() && I < MAX_PACKAGES_TO_LOG; ++I)
+	{
+		const FReloadPackageData& ReloadPackageData = InPackagesToReload[I];
+		Msg.Append(TEXT("\n"));
+		Msg.Append(FString::Printf(TEXT("\tAsset Name: %s"), *ReloadPackageData.PackageToReload->GetName()));
+	}
+	UE_LOG(LogUObjectGlobals, Log, TEXT("%s"), *Msg);
+
 	FScopedSlowTask ReloadingPackagesSlowTask(InPackagesToReload.Num(), NSLOCTEXT("CoreUObject", "ReloadingPackages", "Reloading Packages"));
 	ReloadingPackagesSlowTask.MakeDialog();
 
