Fixed a crash when setting the initial size of the disregard for Garbage Collection pool to a value greater than 65560.

diff --git a/Engine/Source/Runtime/CoreUObject/Public/UObject/UObjectArray.h b/Engine/Source/Runtime/CoreUObject/Public/UObject/UObjectArray.h
--- a/Engine/Source/Runtime/CoreUObject/Public/UObject/UObjectArray.h
+++ b/Engine/Source/Runtime/CoreUObject/Public/UObject/UObjectArray.h
@@ -322,7 +322,7 @@ class FChunkedFixedUObjectArray
 {
 	enum
 	{
-		NumElementsPerChunk = 65 * 1024,
+		NumElementsPerChunk = 64 * 1024,
 	};
 
 	/** Master table to chunks of pointers **/
@@ -349,7 +349,7 @@ class FChunkedFixedUObjectArray
 		while (ChunkIndex >= NumChunks)
 		{
 			// add a chunk, and make sure nobody else tries
-			FUObjectItem** Chunk = &Objects[ChunkIndex];
+			FUObjectItem** Chunk = &Objects[NumChunks];
 			FUObjectItem* NewChunk = new FUObjectItem[NumElementsPerChunk];
 			if (FPlatformAtomics::InterlockedCompareExchangePointer((void**)Chunk, NewChunk, nullptr))
 			{
@@ -359,6 +359,7 @@ class FChunkedFixedUObjectArray
 			else
 			{
 				NumChunks++;
+				check(NumChunks <= MaxChunks);
 			}
 		}
 		check(ChunkIndex < NumChunks && Objects[ChunkIndex]); // should have a valid pointer now
