If set the references will not be automatically cooked, reported to the asset registry, or cause warnings when deleting a referenced actor
They are meant to be used for cases where they are just used as a lazy cache, such as things converted from FLazyObjectPtr
The reference from Foliage to it's base actor is now an Untracked SoftObjectPtr instead of a LazyPtr

diff --git a/Engine/Source/Developer/AssetTools/Private/AssetRenameManager.cpp b/Engine/Source/Developer/AssetTools/Private/AssetRenameManager.cpp
--- a/Engine/Source/Developer/AssetTools/Private/AssetRenameManager.cpp
+++ b/Engine/Source/Developer/AssetTools/Private/AssetRenameManager.cpp
@@ -864,6 +864,22 @@ struct FSoftObjectPathRenameSerializer : public FArchiveUObject
 
 	FArchive& operator<<(FSoftObjectPath& Value)
 	{
+		// Ignore untracked references if just doing a search only. We still want to fix them up if they happen to be there
+		if (bSearchOnly)
+		{
+			FSoftObjectPathThreadContext& ThreadContext = FSoftObjectPathThreadContext::Get();
+			FName ReferencingPackageName, ReferencingPropertyName;
+			ESoftObjectPathCollectType CollectType = ESoftObjectPathCollectType::AlwaysCollect;
+			ESoftObjectPathSerializeType SerializeType = ESoftObjectPathSerializeType::AlwaysSerialize;
+
+			ThreadContext.GetSerializationOptions(ReferencingPackageName, ReferencingPropertyName, CollectType, SerializeType, this);
+
+			if (CollectType == ESoftObjectPathCollectType::NeverCollect)
+			{
+				return *this;
+			}
+		}
+
 		FString SubPath = Value.GetSubPathString();
 		for (const TPair<FSoftObjectPath, FSoftObjectPath>& Pair : RedirectorMap)
 		{
diff --git a/Engine/Source/Editor/UnrealEd/Private/EditorActor.cpp b/Engine/Source/Editor/UnrealEd/Private/EditorActor.cpp
--- a/Engine/Source/Editor/UnrealEd/Private/EditorActor.cpp
+++ b/Engine/Source/Editor/UnrealEd/Private/EditorActor.cpp
@@ -777,12 +777,7 @@ bool UUnrealEdEngine::edactDeleteSelected( UWorld* InWorld, bool bVerifyDeletion
 	{
 		FBlueprintEditorUtils::GetActorReferenceMap(InWorld, ClassTypesToIgnore, ReferencingActorsMap);
 
-		// For now, don't display the soft reference warning if we won't load packages to validate
-		// Re-evaluate this once we have control over rather less important soft references like foliage is tracked
-		const UBlueprintEditorProjectSettings* EditorProjectSettings = GetDefault<UBlueprintEditorProjectSettings>();
-		bool bLoadPackagesForSoftReferences = EditorProjectSettings->bValidateUnloadedSoftActorReferences;
-
-		if (bWarnAboutSoftReferences && bLoadPackagesForSoftReferences)
+		if (bWarnAboutSoftReferences)
 		{
 			FAssetToolsModule& AssetToolsModule = FModuleManager::GetModuleChecked<FAssetToolsModule>(TEXT("AssetTools"));
 
diff --git a/Engine/Source/Runtime/CoreUObject/Private/Serialization/ArchiveUObjectFromStructuredArchive.cpp b/Engine/Source/Runtime/CoreUObject/Private/Serialization/ArchiveUObjectFromStructuredArchive.cpp
--- a/Engine/Source/Runtime/CoreUObject/Private/Serialization/ArchiveUObjectFromStructuredArchive.cpp
+++ b/Engine/Source/Runtime/CoreUObject/Private/Serialization/ArchiveUObjectFromStructuredArchive.cpp
@@ -102,7 +102,7 @@ FArchive& FArchiveUObjectFromStructuredArchive::operator<<(FSoftObjectPath& Valu
 			ESoftObjectPathCollectType CollectType = ESoftObjectPathCollectType::AlwaysCollect;
 			ESoftObjectPathSerializeType SerializeType = ESoftObjectPathSerializeType::AlwaysSerialize;
 			FSoftObjectPathThreadContext& ThreadContext = FSoftObjectPathThreadContext::Get();
-			ThreadContext.GetSerializationOptions(PackageName, PropertyName, CollectType, SerializeType);
+			ThreadContext.GetSerializationOptions(PackageName, PropertyName, CollectType, SerializeType, this);
 
 			if (SerializeType == ESoftObjectPathSerializeType::AlwaysSerialize)
 			{
diff --git a/Engine/Source/Runtime/CoreUObject/Private/UObject/SavePackage.cpp b/Engine/Source/Runtime/CoreUObject/Private/UObject/SavePackage.cpp
--- a/Engine/Source/Runtime/CoreUObject/Private/UObject/SavePackage.cpp
+++ b/Engine/Source/Runtime/CoreUObject/Private/UObject/SavePackage.cpp
@@ -1376,11 +1376,21 @@ FArchive& FArchiveSaveTagImports::operator<<(FSoftObjectPath& Value)
 	{
 		Value.SerializePath(*this);
 
-		FString Path = Value.ToString();
-		FName PackageName = FName(*FPackageName::ObjectPathToPackageName(Path));
+		FSoftObjectPathThreadContext& ThreadContext = FSoftObjectPathThreadContext::Get();
+		FName ReferencingPackageName, ReferencingPropertyName;
+		ESoftObjectPathCollectType CollectType = ESoftObjectPathCollectType::AlwaysCollect;
+		ESoftObjectPathSerializeType SerializeType = ESoftObjectPathSerializeType::AlwaysSerialize;
 
-		SavePackageState.MarkNameAsReferenced(PackageName);
-		Linker->SoftPackageReferenceList.AddUnique(PackageName);	
+		ThreadContext.GetSerializationOptions(ReferencingPackageName, ReferencingPropertyName, CollectType, SerializeType, this);
+
+		if (CollectType != ESoftObjectPathCollectType::NeverCollect)
+		{
+			// Don't track if this is a never collect path
+			FString Path = Value.ToString();
+			FName PackageName = FName(*FPackageName::ObjectPathToPackageName(Path));
+			SavePackageState.MarkNameAsReferenced(PackageName);
+			Linker->SoftPackageReferenceList.AddUnique(PackageName);
+		}
 	}
 	return *this;
 }
diff --git a/Engine/Source/Runtime/CoreUObject/Private/UObject/SoftObjectPath.cpp b/Engine/Source/Runtime/CoreUObject/Private/UObject/SoftObjectPath.cpp
--- a/Engine/Source/Runtime/CoreUObject/Private/UObject/SoftObjectPath.cpp
+++ b/Engine/Source/Runtime/CoreUObject/Private/UObject/SoftObjectPath.cpp
@@ -139,7 +139,7 @@ void FSoftObjectPath::SerializePath(FArchive& Ar)
 	ESoftObjectPathSerializeType SerializeType = ESoftObjectPathSerializeType::AlwaysSerialize;
 
 	FSoftObjectPathThreadContext& ThreadContext = FSoftObjectPathThreadContext::Get();
-	ThreadContext.GetSerializationOptions(PackageName, PropertyName, CollectType, SerializeType);
+	ThreadContext.GetSerializationOptions(PackageName, PropertyName, CollectType, SerializeType, &Ar);
 
 	if (SerializeType == ESoftObjectPathSerializeType::NeverSerialize)
 	{
@@ -523,7 +523,7 @@ FSoftClassPath FSoftClassPath::GetOrCreateIDForClass(const UClass *InClass)
 	return FSoftClassPath(InClass);
 }
 
-bool FSoftObjectPathThreadContext::GetSerializationOptions(FName& OutPackageName, FName& OutPropertyName, ESoftObjectPathCollectType& OutCollectType, ESoftObjectPathSerializeType& OutSerializeType) const
+bool FSoftObjectPathThreadContext::GetSerializationOptions(FName& OutPackageName, FName& OutPropertyName, ESoftObjectPathCollectType& OutCollectType, ESoftObjectPathSerializeType& OutSerializeType, FArchive* Archive) const
 {
 	FName CurrentPackageName, CurrentPropertyName;
 	ESoftObjectPathCollectType CurrentCollectType = ESoftObjectPathCollectType::AlwaysCollect;
@@ -559,36 +559,57 @@ bool FSoftObjectPathThreadContext::GetSerializationOptions(FName& OutPackageName
 		bFoundAnything = true;
 	}
 	
-	// Check UObject thread context as a backup
+	
+	// Check UObject thread context as a backup, this only works for loads
 	FUObjectThreadContext& ThreadContext = FUObjectThreadContext::Get();
 	if (ThreadContext.SerializedObject)
 	{
 		FLinkerLoad* Linker = ThreadContext.SerializedObject->GetLinker();
+
 		if (Linker)
 		{
 			if (CurrentPackageName == NAME_None)
 			{
 				CurrentPackageName = FName(*FPackageName::FilenameToLongPackageName(Linker->Filename));
 			}
-			if (Linker->GetSerializedProperty() && CurrentPropertyName == NAME_None)
-			{
-				CurrentPropertyName = Linker->GetSerializedProperty()->GetFName();
-			}
-#if WITH_EDITORONLY_DATA
-			bool bEditorOnly = Linker->IsEditorOnlyPropertyOnTheStack();
-#else
-			bool bEditorOnly = false;
-#endif
-			// If we were always collect before and not overridden by stack options, set to editor only
-			if (bEditorOnly && CurrentCollectType == ESoftObjectPathCollectType::AlwaysCollect)
+			if (Archive == nullptr)
 			{
-				CurrentCollectType = ESoftObjectPathCollectType::EditorOnlyCollect;
+				// Use archive from linker if it wasn't passed in
+				Archive = Linker;
 			}
-
 			bFoundAnything = true;
 		}
 	}
 
+	// Check archive for property/editor only info, this works for any serialize if passed in
+	if (Archive)
+	{
+		UProperty* CurrentProperty = Archive->GetSerializedProperty();
+			
+		if (CurrentProperty && CurrentPropertyName == NAME_None)
+		{
+			CurrentPropertyName = CurrentProperty->GetFName();
+		}
+		bool bEditorOnly = false;
+#if WITH_EDITOR
+		bEditorOnly = Archive->IsEditorOnlyPropertyOnTheStack();
+
+		static FName UntrackedName = TEXT("Untracked");
+		if (CurrentProperty && CurrentProperty->HasMetaData(UntrackedName))
+		{
+			// Property has the Untracked metadata, so set to never collect references
+			CurrentCollectType = ESoftObjectPathCollectType::NeverCollect;
+		}
+#endif
+		// If we were always collect before and not overridden by stack options, set to editor only
+		if (bEditorOnly && CurrentCollectType == ESoftObjectPathCollectType::AlwaysCollect)
+		{
+			CurrentCollectType = ESoftObjectPathCollectType::EditorOnlyCollect;
+		}
+
+		bFoundAnything = true;
+	}
+
 	if (bFoundAnything)
 	{
 		OutPackageName = CurrentPackageName;
diff --git a/Engine/Source/Runtime/CoreUObject/Public/UObject/ObjectMacros.h b/Engine/Source/Runtime/CoreUObject/Public/UObject/ObjectMacros.h
--- a/Engine/Source/Runtime/CoreUObject/Public/UObject/ObjectMacros.h
+++ b/Engine/Source/Runtime/CoreUObject/Public/UObject/ObjectMacros.h
@@ -1221,6 +1221,9 @@ namespace UM
 
 		/// [PropertyMetadata] Used for float and integer properties.  Specifies the highest that the value slider should represent.
 		UIMax,
+
+		/// [PropertyMetadata] Used for SoftObjectPtr/SoftObjectPath properties to specify a reference should not be tracked. This reference will not be automatically cooked or saved into the asset registry for redirector/delete fixup.
+		Untracked,
 	};
 
 	// Metadata usable in UPROPERTY for customizing the behavior of Persona and UMG
diff --git a/Engine/Source/Runtime/CoreUObject/Public/UObject/SoftObjectPath.h b/Engine/Source/Runtime/CoreUObject/Public/UObject/SoftObjectPath.h
--- a/Engine/Source/Runtime/CoreUObject/Public/UObject/SoftObjectPath.h
+++ b/Engine/Source/Runtime/CoreUObject/Public/UObject/SoftObjectPath.h
@@ -322,8 +322,9 @@ public:
 	 * @param OutPackageName Package that this string asset belongs to
 	 * @param OutPropertyName Property that this string asset reference belongs to
 	 * @param OutCollectType Type of collecting that should be done
+	 * @param Archive The FArchive that is serializing this path if known. If null it will check FUObjectThreadContext
 	 */
-	bool GetSerializationOptions(FName& OutPackageName, FName& OutPropertyName, ESoftObjectPathCollectType& OutCollectType, ESoftObjectPathSerializeType& OutSerializeType) const;
+	bool GetSerializationOptions(FName& OutPackageName, FName& OutPropertyName, ESoftObjectPathCollectType& OutCollectType, ESoftObjectPathSerializeType& OutSerializeType, FArchive* Archive = nullptr) const;
 };
 
 /** Helper class to set and restore serialization options for string asset references */
diff --git a/Engine/Source/Runtime/Foliage/Private/FoliageInstanceBase.cpp b/Engine/Source/Runtime/Foliage/Private/FoliageInstanceBase.cpp
--- a/Engine/Source/Runtime/Foliage/Private/FoliageInstanceBase.cpp
+++ b/Engine/Source/Runtime/Foliage/Private/FoliageInstanceBase.cpp
@@ -52,6 +52,8 @@ FArchive& operator << (FArchive& Ar, FFoliageInstanceBaseInfo& BaseInfo)
 	}
 	else
 	{
+		// This is a cached reference, we don't want to warn about the target actor being deleted
+		FSoftObjectPathSerializationScope SerializationScope(NAME_None, NAME_None, ESoftObjectPathCollectType::NeverCollect, ESoftObjectPathSerializeType::AlwaysSerialize);
 		Ar << BaseInfo.BasePtr;
 	}
 
@@ -85,6 +87,9 @@ FArchive& operator << (FArchive& Ar, FFoliageInstanceBaseCache& InstanceBaseCach
 	}
 	else
 	{
+		// This is a cached reference, we don't want to warn about the target actor being deleted
+		FSoftObjectPathSerializationScope SerializationScope(NAME_None, NAME_None, ESoftObjectPathCollectType::NeverCollect, ESoftObjectPathSerializeType::AlwaysSerialize);
+
 		Ar << InstanceBaseCache.InstanceBaseLevelMap;
 
 		if (Ar.IsTransacting())
