[NO NON-PII COMMIT MESSAGE]

diff --git a/Engine/Source/Runtime/CoreUObject/Private/Blueprint/BlueprintSupport.cpp b/Engine/Source/Runtime/CoreUObject/Private/Blueprint/BlueprintSupport.cpp
--- a/Engine/Source/Runtime/CoreUObject/Private/Blueprint/BlueprintSupport.cpp
+++ b/Engine/Source/Runtime/CoreUObject/Private/Blueprint/BlueprintSupport.cpp
@@ -621,6 +621,8 @@ public:
 					}
 				}
 			}
+			// Remove from map as we could get GCd later
+			ClassToPlaceholderMap.Remove(ForClass);
 		}
 	}
 
