Fix leaking entries for PLACEHOLDER classes in the new ClassToPlaceholderMap
#rnx

diff --git a/Engine/Source/Runtime/CoreUObject/Private/Blueprint/BlueprintSupport.cpp b/Engine/Source/Runtime/CoreUObject/Private/Blueprint/BlueprintSupport.cpp
--- a/Engine/Source/Runtime/CoreUObject/Private/Blueprint/BlueprintSupport.cpp
+++ b/Engine/Source/Runtime/CoreUObject/Private/Blueprint/BlueprintSupport.cpp
@@ -598,6 +598,11 @@ public:
 #if USE_DEFERRED_DEPENDENCY_CHECK_VERIFICATION_TESTS
 		FullyResolvedLinkers.Remove(Linker);
 #endif // USE_DEFERRED_DEPENDENCY_CHECK_VERIFICATION_TESTS
+
+		// ClassToPlaceholderMap may have entries because instances of placeholder classes (which 
+		// will be resolved in ResolveDeferredExports()), will never have had ResolvePlaceholders
+		// for their class called. These entries are harmless and we can discard them here:
+		ClassToPlaceholderMap.Reset();
 	}
 
 	void AddLinkerPlaceholderObject(UClass* ClassWaitingFor, ULinkerPlaceholderExportObject* Placeholder)
