Add an overload of FSoftObjectPath::FixupForPIE that takes an explicit PIE instance ID for cases where one is available instead of using GPlayInEditorID.
Use the new overload in NetworkRemapPath_local, since the world context is available and it avoids edge cases where GPlayInEditorID may not be set correctly.

diff --git a/Engine/Source/Editor/UnrealEd/Private/EditorEngine.cpp b/Engine/Source/Editor/UnrealEd/Private/EditorEngine.cpp
--- a/Engine/Source/Editor/UnrealEd/Private/EditorEngine.cpp
+++ b/Engine/Source/Editor/UnrealEd/Private/EditorEngine.cpp
@@ -6869,7 +6869,7 @@ FORCEINLINE bool NetworkRemapPath_local(FWorldContext& Context, FString& Str, bo
 		// First strip any source prefix, then add the appropriate prefix for this context
 		FSoftObjectPath Path = UWorld::RemovePIEPrefix(Str);
 		
-		Path.FixupForPIE();
+		Path.FixupForPIE(Context.PIEInstance);
 		FString Remapped = Path.ToString();
 		if (!Remapped.Equals(Str, ESearchCase::CaseSensitive))
 		{
diff --git a/Engine/Source/Runtime/CoreUObject/Private/UObject/SoftObjectPath.cpp b/Engine/Source/Runtime/CoreUObject/Private/UObject/SoftObjectPath.cpp
--- a/Engine/Source/Runtime/CoreUObject/Private/UObject/SoftObjectPath.cpp
+++ b/Engine/Source/Runtime/CoreUObject/Private/UObject/SoftObjectPath.cpp
@@ -434,10 +434,10 @@ void FSoftObjectPath::ClearPIEPackageNames()
 	PIEPackageNames.Empty();
 }
 
-bool FSoftObjectPath::FixupForPIE()
+bool FSoftObjectPath::FixupForPIE(int32 PIEInstance)
 {
 #if WITH_EDITOR
-	if (GPlayInEditorID != INDEX_NONE && !IsNull())
+	if (PIEInstance != INDEX_NONE && !IsNull())
 	{
 		const FString Path = ToString();
 
@@ -448,7 +448,7 @@ bool FSoftObjectPath::FixupForPIE()
 			// Name of the ULevel subobject of UWorld, set in InitializeNewWorld
 			const bool bIsChildOfLevel = SubPathString.StartsWith(TEXT("PersistentLevel."));
 
-			FString PIEPath = FString::Printf(TEXT("%s/%s_%d_%s"), *FPackageName::GetLongPackagePath(Path), PLAYWORLD_PACKAGE_PREFIX, GPlayInEditorID, *ShortPackageOuterAndName);
+			FString PIEPath = FString::Printf(TEXT("%s/%s_%d_%s"), *FPackageName::GetLongPackagePath(Path), PLAYWORLD_PACKAGE_PREFIX, PIEInstance, *ShortPackageOuterAndName);
 			const FName PIEPackage = (!bIsChildOfLevel ? FName(*FPackageName::ObjectPathToPackageName(PIEPath)) : NAME_None);
 
 			// Duplicate if this an already registered PIE package or this looks like a level subobject reference
@@ -465,6 +465,11 @@ bool FSoftObjectPath::FixupForPIE()
 	return false;
 }
 
+bool FSoftObjectPath::FixupForPIE()
+{
+	return FixupForPIE(GPlayInEditorID);
+}
+
 bool FSoftObjectPath::FixupCoreRedirects()
 {
 	FString OldString = ToString();
diff --git a/Engine/Source/Runtime/CoreUObject/Public/UObject/SoftObjectPath.h b/Engine/Source/Runtime/CoreUObject/Public/UObject/SoftObjectPath.h
--- a/Engine/Source/Runtime/CoreUObject/Public/UObject/SoftObjectPath.h
+++ b/Engine/Source/Runtime/CoreUObject/Public/UObject/SoftObjectPath.h
@@ -160,9 +160,12 @@ struct COREUOBJECT_API FSoftObjectPath
 	/** Handles when a string asset reference has been loaded, call if loading with a method that skips SerializePath. This does not modify path but might call callbacks */
 	void PostLoadPath() const;
 
-	/** Fixes up this SoftObjectPath to add or remove the PIE prefix depending on what is currently active, returns true if it was modified */
+	/** Fixes up this SoftObjectPath to add the PIE prefix depending on what is currently active, returns true if it was modified. The overload that takes an explicit PIE instance is preferred, if it's available. */
 	bool FixupForPIE();
 
+	/** Fixes up this SoftObjectPath to add the PIE prefix for the given PIEInstance index, returns true if it was modified */
+	bool FixupForPIE(int32 PIEInstance);
+
 	/** Fixes soft object path for CoreRedirects to handle renamed native objects, returns true if it was modified */
 	bool FixupCoreRedirects();
 
