Fix regression in stepping behavior caused by 4235135

diff --git a/Engine/Source/Runtime/CoreUObject/Private/UObject/ScriptCore.cpp b/Engine/Source/Runtime/CoreUObject/Private/UObject/ScriptCore.cpp
--- a/Engine/Source/Runtime/CoreUObject/Private/UObject/ScriptCore.cpp
+++ b/Engine/Source/Runtime/CoreUObject/Private/UObject/ScriptCore.cpp
@@ -1322,7 +1322,7 @@ void UObject::ProcessEvent( UFunction* Function, void* Parms )
 
 #if DO_BLUEPRINT_GUARD
 	FBlueprintExceptionTracker& BlueprintExceptionTracker = FBlueprintExceptionTracker::Get();
-	BlueprintExceptionTracker.ScriptEntryTag++;
+	TGuardValue<int32> EntryCounter( BlueprintExceptionTracker.ScriptEntryTag, BlueprintExceptionTracker.ScriptEntryTag+1);
 
 	CONDITIONAL_SCOPE_CYCLE_COUNTER(STAT_BlueprintTime, IsInGameThread() && BlueprintExceptionTracker.ScriptEntryTag == 1);
 #endif
@@ -1496,10 +1496,6 @@ void UObject::ProcessEvent( UFunction* Function, void* Parms )
 	FBlueprintCoreDelegates::OnScriptExecutionEnd.Broadcast();
 #endif
 #endif
-
-#if DO_BLUEPRINT_GUARD
-	--BlueprintExceptionTracker.ScriptEntryTag;
-#endif
 }
 
 #ifdef _MSC_VER
