Made sure that GC's incremental ConditionBeginDestroy is not running when async loading in editor builds as this may result in the async loading code getting unreachable imports from linkers that haven't been detached yet.

diff --git a/Engine/Source/Runtime/CoreUObject/Private/Serialization/AsyncLoading.cpp b/Engine/Source/Runtime/CoreUObject/Private/Serialization/AsyncLoading.cpp
--- a/Engine/Source/Runtime/CoreUObject/Private/Serialization/AsyncLoading.cpp
+++ b/Engine/Source/Runtime/CoreUObject/Private/Serialization/AsyncLoading.cpp
@@ -4838,6 +4838,14 @@ EAsyncPackageState::Type FAsyncLoadingThread::TickAsyncLoading(bool bUseTimeLimi
 
 	if (!bLoadingSuspended)
 	{
+		// First make sure there's no objects pending to be unhashed. This is important in uncooked builds since we don't 
+		// detach linkers immediately there and we may end up in getting unreachable objects from Linkers in CreateImports
+		if (FPlatformProperties::RequiresCookedData() == false && IsIncrementalUnhashPending() && IsAsyncLoadingPackages())
+		{
+			// Call ConditionalBeginDestroy on all pending objects. CBD is where linkers get detached from objects.
+			UnhashUnreachableObjects(false);
+		}
+
 		const bool bIsMultithreaded = FAsyncLoadingThread::IsMultithreaded();
 		double TickStartTime = FPlatformTime::Seconds();
 		double TimeLimitUsedForProcessLoaded = 0;
@@ -5430,7 +5438,7 @@ void FAsyncPackage::AddObjectReference(UObject* InObject)
 				ReferencedObjects.Add(InObject);
 			}
 		}
-		UE_CLOG(InObject->HasAnyInternalFlags(EInternalObjectFlags::Unreachable), LogStreaming, Fatal, TEXT("Trying to add an object %s to FAsyncPackage %s referenced objects list that is unreachable."), *InObject->GetFullName(), *GetPackageName().ToString());
+		UE_CLOG(InObject->HasAnyInternalFlags(EInternalObjectFlags::Unreachable), LogStreaming, Fatal, TEXT("Trying to add an unreachable object %s to FAsyncPackage %s referenced objects list."), *InObject->GetFullName(), *GetPackageName().ToString());
 	}
 }
 
diff --git a/Engine/Source/Runtime/CoreUObject/Private/UObject/GarbageCollection.cpp b/Engine/Source/Runtime/CoreUObject/Private/UObject/GarbageCollection.cpp
--- a/Engine/Source/Runtime/CoreUObject/Private/UObject/GarbageCollection.cpp
+++ b/Engine/Source/Runtime/CoreUObject/Private/UObject/GarbageCollection.cpp
@@ -70,9 +70,6 @@ static int32 GUnrechableObjectIndex = 0;
 
 /** Helpful constant for determining how many token slots we need to store a pointer **/
 static const uint32 GNumTokensPerPointer = sizeof(void*) / sizeof(uint32); //-V514
-/** Calls ConditionalBeginDestroy on unreachable objects */
-static bool UnhashUnreachableObjects(bool bUseTimeLimit, float TimeLimit = 0.0f);
-
 
 FThreadSafeBool& FGCScopeLock::GetGarbageCollectingFlag()
 {
@@ -1615,6 +1612,11 @@ void CollectGarbageInternal(EObjectFlags KeepFlags, bool bPerformFullPurge)
 	STAT_ADD_CUSTOMMESSAGE_NAME( STAT_NamedMarker, TEXT( "GarbageCollection - End" ) );
 }
 
+bool IsIncrementalUnhashPending()
+{
+	return GUnrechableObjectIndex < GUnreachableObjects.Num();
+}
+
 bool UnhashUnreachableObjects(bool bUseTimeLimit, float TimeLimit)
 {
 	DECLARE_SCOPE_CYCLE_COUNTER(TEXT("UnhashUnreachableObjects"), STAT_UnhashUnreachableObjects, STATGROUP_GC);
diff --git a/Engine/Source/Runtime/CoreUObject/Private/UObject/LinkerLoad.cpp b/Engine/Source/Runtime/CoreUObject/Private/UObject/LinkerLoad.cpp
--- a/Engine/Source/Runtime/CoreUObject/Private/UObject/LinkerLoad.cpp
+++ b/Engine/Source/Runtime/CoreUObject/Private/UObject/LinkerLoad.cpp
@@ -4342,6 +4342,13 @@ UObject* FLinkerLoad::CreateImport( int32 Index )
 	DeferPotentialCircularImport(Index); 
 #endif // USE_CIRCULAR_DEPENDENCY_LOAD_DEFERRING
 
+	if (Import.XObject != nullptr && Import.XObject->HasAnyInternalFlags(EInternalObjectFlags::Unreachable))
+	{
+		// This is just a safeguard to catch potential bugs that should have been fixed by calling UnhashUnreachableObjects in Async Loading code
+		UE_LOG(LogLinker, Warning, TEXT("Unreachable object found when creating import %s from linker %s"), *Import.XObject->GetFullName(), *GetArchiveName());
+		Import.XObject = nullptr;
+	}
+
 	// Imports can have no name if they were filtered out due to package redirects, skip in that case
 	if (Import.XObject == nullptr && Import.ObjectName != NAME_None)
 	{
diff --git a/Engine/Source/Runtime/CoreUObject/Public/UObject/UObjectGlobals.h b/Engine/Source/Runtime/CoreUObject/Public/UObject/UObjectGlobals.h
--- a/Engine/Source/Runtime/CoreUObject/Public/UObject/UObjectGlobals.h
+++ b/Engine/Source/Runtime/CoreUObject/Public/UObject/UObjectGlobals.h
@@ -378,6 +378,20 @@ COREUOBJECT_API void CollectGarbage(EObjectFlags KeepFlags, bool bPerformFullPur
 */
 COREUOBJECT_API bool TryCollectGarbage(EObjectFlags KeepFlags, bool bPerformFullPurge = true);
 
+/**
+* Calls ConditionalBeginDestroy on unreachable objects
+*
+* @return true if the time limit passed and there's still objects pending to be unhashed
+*/
+COREUOBJECT_API bool UnhashUnreachableObjects(bool bUseTimeLimit, float TimeLimit = 0.0f);
+
+/**
+* Checks if there's objects pending to be unhashed when running incremental purge
+*
+* @return true if the time limit passed and there's still objects pending to be unhashed
+*/
+COREUOBJECT_API bool IsIncrementalUnhashPending();
+
 /**
  * Returns whether an incremental purge is still pending/ in progress.
  *
