Revert CppStructOps SerializeMismatchedTag() back to previous FArchive form, and add a new structured archive version that has to be explicitly opted into. Converted all existing engine code to use the new structured archive, but gives backwards compatibility for licensees with the old FArchive version

Also removed TStructOpsTypeTraitsBase which has been deprecated since 4.16

diff --git a/Engine/Plugins/Runtime/GameplayAbilities/Source/GameplayAbilities/Public/AttributeSet.h b/Engine/Plugins/Runtime/GameplayAbilities/Source/GameplayAbilities/Public/AttributeSet.h
--- a/Engine/Plugins/Runtime/GameplayAbilities/Source/GameplayAbilities/Public/AttributeSet.h
+++ b/Engine/Plugins/Runtime/GameplayAbilities/Source/GameplayAbilities/Public/AttributeSet.h
@@ -342,7 +342,7 @@ struct TStructOpsTypeTraits<FScalableFloat>
 {
 	enum
 	{
-		WithSerializeFromMismatchedTag = true,
+		WithStructuredSerializeFromMismatchedTag = true,
 	};
 };
 
diff --git a/Engine/Source/Editor/UMGEditor/Public/WidgetBlueprint.h b/Engine/Source/Editor/UMGEditor/Public/WidgetBlueprint.h
--- a/Engine/Source/Editor/UMGEditor/Public/WidgetBlueprint.h
+++ b/Engine/Source/Editor/UMGEditor/Public/WidgetBlueprint.h
@@ -176,7 +176,7 @@ struct TStructOpsTypeTraits<FWidgetAnimation_DEPRECATED> : public TStructOpsType
 {
 	enum
 	{
-		WithSerializeFromMismatchedTag = true,
+		WithStructuredSerializeFromMismatchedTag = true,
 	};
 };
 
diff --git a/Engine/Source/Runtime/AnimGraphRuntime/Public/AnimNodes/AnimNode_StateResult.h b/Engine/Source/Runtime/AnimGraphRuntime/Public/AnimNodes/AnimNode_StateResult.h
--- a/Engine/Source/Runtime/AnimGraphRuntime/Public/AnimNodes/AnimNode_StateResult.h
+++ b/Engine/Source/Runtime/AnimGraphRuntime/Public/AnimNodes/AnimNode_StateResult.h
@@ -24,6 +24,6 @@ struct TStructOpsTypeTraits<FAnimNode_StateResult>
 {
 	enum
 	{
-		WithSerializeFromMismatchedTag = true,
+		WithStructuredSerializeFromMismatchedTag = true,
 	};
 };
diff --git a/Engine/Source/Runtime/CoreUObject/Private/UObject/Class.cpp b/Engine/Source/Runtime/CoreUObject/Private/UObject/Class.cpp
--- a/Engine/Source/Runtime/CoreUObject/Private/UObject/Class.cpp
+++ b/Engine/Source/Runtime/CoreUObject/Private/UObject/Class.cpp
@@ -2048,7 +2048,7 @@ void UScriptStruct::PrepareCppStructOps()
 		UE_LOG(LogClass, Verbose, TEXT("Native struct %s has native ImportTextItem."),*GetName());
 		StructFlags = EStructFlags(StructFlags | STRUCT_ImportTextItemNative);
 	}
-	if (CppStructOps->HasSerializeFromMismatchedTag())
+	if (CppStructOps->HasSerializeFromMismatchedTag() || CppStructOps->HasStructuredSerializeFromMismatchedTag())
 	{
 		UE_LOG(LogClass, Verbose, TEXT("Native struct %s has native SerializeFromMismatchedTag."),*GetName());
 		StructFlags = EStructFlags(StructFlags | STRUCT_SerializeFromMismatchedTag);
diff --git a/Engine/Source/Runtime/CoreUObject/Private/UObject/Property.cpp b/Engine/Source/Runtime/CoreUObject/Private/UObject/Property.cpp
--- a/Engine/Source/Runtime/CoreUObject/Private/UObject/Property.cpp
+++ b/Engine/Source/Runtime/CoreUObject/Private/UObject/Property.cpp
@@ -302,7 +302,7 @@ struct TStructOpsTypeTraits<FSoftObjectPath> : public TStructOpsTypeTraitsBase2<
 		WithIdenticalViaEquality = true,
 		WithExportTextItem = true,
 		WithImportTextItem = true,
-		WithSerializeFromMismatchedTag = true,
+		WithStructuredSerializeFromMismatchedTag = true,
 	};
 };
 IMPLEMENT_STRUCT(SoftObjectPath);
@@ -318,7 +318,7 @@ struct TStructOpsTypeTraits<FSoftClassPath> : public TStructOpsTypeTraitsBase2<F
 		WithIdenticalViaEquality = true,
 		WithExportTextItem = true,
 		WithImportTextItem = true,
-		WithSerializeFromMismatchedTag = true,
+		WithStructuredSerializeFromMismatchedTag = true,
 	};
 };
 IMPLEMENT_STRUCT(SoftClassPath);
@@ -333,7 +333,7 @@ struct TStructOpsTypeTraits<FPrimaryAssetType> : public TStructOpsTypeTraitsBase
 		WithIdenticalViaEquality = true,
 		WithExportTextItem = true,
 		WithImportTextItem = true,
-		WithSerializeFromMismatchedTag = true,
+		WithStructuredSerializeFromMismatchedTag = true,
 	};
 };
 IMPLEMENT_STRUCT(PrimaryAssetType);
@@ -348,7 +348,7 @@ struct TStructOpsTypeTraits<FPrimaryAssetId> : public TStructOpsTypeTraitsBase2<
 		WithIdenticalViaEquality = true,
 		WithExportTextItem = true,
 		WithImportTextItem = true,
-		WithSerializeFromMismatchedTag = true,
+		WithStructuredSerializeFromMismatchedTag = true,
 	};
 };
 IMPLEMENT_STRUCT(PrimaryAssetId);
diff --git a/Engine/Source/Runtime/CoreUObject/Private/UObject/PropertyStruct.cpp b/Engine/Source/Runtime/CoreUObject/Private/UObject/PropertyStruct.cpp
--- a/Engine/Source/Runtime/CoreUObject/Private/UObject/PropertyStruct.cpp
+++ b/Engine/Source/Runtime/CoreUObject/Private/UObject/PropertyStruct.cpp
@@ -345,16 +345,25 @@ EConvertFromTypeResult UStructProperty::ConvertFromType(const FPropertyTag& Tag,
 		if ((Struct->StructFlags & STRUCT_SerializeFromMismatchedTag) && (Tag.Type != NAME_StructProperty || (Tag.StructName != Struct->GetFName())))
 		{
 			UScriptStruct::ICppStructOps* CppStructOps = Struct->GetCppStructOps();
-			check(CppStructOps && CppStructOps->HasSerializeFromMismatchedTag()); // else should not have STRUCT_SerializeFromMismatchedTag
+			check(CppStructOps && (CppStructOps->HasSerializeFromMismatchedTag() || CppStructOps->HasStructuredSerializeFromMismatchedTag())); // else should not have STRUCT_SerializeFromMismatchedTag
 			void* DestAddress = ContainerPtrToValuePtr<void>(Data, Tag.ArrayIndex);
-			if (CppStructOps->SerializeFromMismatchedTag(Tag, Slot, DestAddress))
+			if (CppStructOps->HasStructuredSerializeFromMismatchedTag() && CppStructOps->StructuredSerializeFromMismatchedTag(Tag, Slot, DestAddress))
 			{
 				return EConvertFromTypeResult::Converted;
 			}
-			else
+			else 
 			{
-				UE_LOG(LogClass, Warning, TEXT("SerializeFromMismatchedTag failed: Type mismatch in %s of %s - Previous (%s) Current(StructProperty) for package:  %s"), *Tag.Name.ToString(), *GetName(), *Tag.Type.ToString(), *UnderlyingArchive.GetArchiveName());
-				return EConvertFromTypeResult::CannotConvert;
+				FArchiveUObjectFromStructuredArchive Ar(Slot);
+				if (CppStructOps->HasSerializeFromMismatchedTag() && CppStructOps->SerializeFromMismatchedTag(Tag, Ar, DestAddress))
+				{
+
+					return EConvertFromTypeResult::Converted;
+				}
+				else
+				{
+					UE_LOG(LogClass, Warning, TEXT("SerializeFromMismatchedTag failed: Type mismatch in %s of %s - Previous (%s) Current(StructProperty) for package:  %s"), *Tag.Name.ToString(), *GetName(), *Tag.Type.ToString(), *UnderlyingArchive.GetArchiveName());
+					return EConvertFromTypeResult::CannotConvert;
+				}
 			}
 		}
 
diff --git a/Engine/Source/Runtime/CoreUObject/Public/UObject/Class.h b/Engine/Source/Runtime/CoreUObject/Public/UObject/Class.h
--- a/Engine/Source/Runtime/CoreUObject/Public/UObject/Class.h
+++ b/Engine/Source/Runtime/CoreUObject/Public/UObject/Class.h
@@ -538,33 +538,12 @@ struct TStructOpsTypeTraitsBase2
 		WithNetSerializer              = false,                         // struct has a NetSerialize function for serializing its state to an FArchive used for network replication.
 		WithNetDeltaSerializer         = false,                         // struct has a NetDeltaSerialize function for serializing differences in state from a previous NetSerialize operation.
 		WithSerializeFromMismatchedTag = false,                         // struct has a SerializeFromMismatchedTag function for converting from other property tags.
-		WithPostScriptConstruct        = false,							// struct has a PostScriptConstruct function which is called after it is constructed in blueprints
+		WithStructuredSerializeFromMismatchedTag = false,               // struct has an FStructuredArchive-based SerializeFromMismatchedTag function for converting from other property tags.
+		WithPostScriptConstruct        = false,				// struct has a PostScriptConstruct function which is called after it is constructed in blueprints
 		WithNetSharedSerialization     = false,                         // struct has a NetSerialize function that does not require the package map to serialize its state.
 	};
 };
 
-/** type traits to cover the custom aspects of a script struct **/
-struct DEPRECATED(4.16, "TStructOpsTypeTraitsBase has been deprecated, use TStructOpsTypeTraitsBase2<T> instead.") TStructOpsTypeTraitsBase
-{
-	enum
-	{
-		WithZeroConstructor            = false, // struct can be constructed as a valid object by filling its memory footprint with zeroes.
-		WithNoInitConstructor          = false, // struct has a constructor which takes an EForceInit parameter which will force the constructor to perform initialization, where the default constructor performs 'uninitialization'.
-		WithNoDestructor               = false, // struct will not have its destructor called when it is destroyed.
-		WithCopy                       = false, // struct can be copied via its copy assignment operator.
-		WithIdenticalViaEquality       = false, // struct can be compared via its operator==.  This should be mutually exclusive with WithIdentical.
-		WithIdentical                  = false, // struct can be compared via an Identical(const T* Other, uint32 PortFlags) function.  This should be mutually exclusive with WithIdenticalViaEquality.
-		WithExportTextItem             = false, // struct has an ExportTextItem function used to serialize its state into a string.
-		WithImportTextItem             = false, // struct has an ImportTextItem function used to deserialize a string into an object of that class.
-		WithAddStructReferencedObjects = false, // struct has an AddStructReferencedObjects function which allows it to add references to the garbage collector.
-		WithSerializer                 = false, // struct has a Serialize function for serializing its state to an FArchive.
-		WithPostSerialize              = false, // struct has a PostSerialize function which is called after it is serialized
-		WithNetSerializer              = false, // struct has a NetSerialize function for serializing its state to an FArchive used for network replication.
-		WithNetDeltaSerializer         = false, // struct has a NetDeltaSerialize function for serializing differences in state from a previous NetSerialize operation.
-		WithSerializeFromMismatchedTag = false, // struct has a SerializeFromMismatchedTag function for converting from other property tags.
-	};
-};
-
 template<class CPPSTRUCT>
 struct TStructOpsTypeTraits : public TStructOpsTypeTraitsBase2<CPPSTRUCT>
 {
@@ -783,13 +762,25 @@ FORCEINLINE typename TEnableIf<TStructOpsTypeTraits<CPPSTRUCT>::WithImportTextIt
  * Selection of SerializeFromMismatchedTag call.
  */
 template<class CPPSTRUCT>
-FORCEINLINE typename TEnableIf<!TStructOpsTypeTraits<CPPSTRUCT>::WithSerializeFromMismatchedTag, bool>::Type SerializeFromMismatchedTagOrNot(FPropertyTag const& Tag, FStructuredArchive::FSlot Slot, CPPSTRUCT *Data)
+FORCEINLINE typename TEnableIf<!TStructOpsTypeTraits<CPPSTRUCT>::WithSerializeFromMismatchedTag, bool>::Type SerializeFromMismatchedTagOrNot(FPropertyTag const& Tag, FArchive& Ar, CPPSTRUCT *Data)
 {
 	return false;
 }
 
 template<class CPPSTRUCT>
-FORCEINLINE typename TEnableIf<TStructOpsTypeTraits<CPPSTRUCT>::WithSerializeFromMismatchedTag, bool>::Type SerializeFromMismatchedTagOrNot(FPropertyTag const& Tag, FStructuredArchive::FSlot Slot, CPPSTRUCT *Data)
+FORCEINLINE typename TEnableIf<TStructOpsTypeTraits<CPPSTRUCT>::WithSerializeFromMismatchedTag, bool>::Type SerializeFromMismatchedTagOrNot(FPropertyTag const& Tag, FArchive& Ar, CPPSTRUCT *Data)
+{
+	return Data->SerializeFromMismatchedTag(Tag, Ar);
+}
+
+template<class CPPSTRUCT>
+FORCEINLINE typename TEnableIf<!TStructOpsTypeTraits<CPPSTRUCT>::WithStructuredSerializeFromMismatchedTag, bool>::Type StructuredSerializeFromMismatchedTagOrNot(FPropertyTag const& Tag, FStructuredArchive::FSlot Slot, CPPSTRUCT *Data)
+{
+	return false;
+}
+
+template<class CPPSTRUCT>
+FORCEINLINE typename TEnableIf<TStructOpsTypeTraits<CPPSTRUCT>::WithStructuredSerializeFromMismatchedTag, bool>::Type StructuredSerializeFromMismatchedTagOrNot(FPropertyTag const& Tag, FStructuredArchive::FSlot Slot, CPPSTRUCT *Data)
 {
 	return Data->SerializeFromMismatchedTag(Tag, Slot);
 }
@@ -937,11 +928,14 @@ public:
 
 		/** return true if this class wants to serialize from some other tag (usually for conversion purposes) **/
 		virtual bool HasSerializeFromMismatchedTag() = 0;
+		virtual bool HasStructuredSerializeFromMismatchedTag() = 0;
+
 		/** 
 		 * Serialize this structure, from some other tag
 		 * @return true if this succeeded, false will trigger a warning and not serialize at all
 		 */
-		virtual bool SerializeFromMismatchedTag(struct FPropertyTag const& Tag, FStructuredArchive::FSlot Slot, void *Data) = 0;
+		virtual bool SerializeFromMismatchedTag(struct FPropertyTag const& Tag, FArchive& Ar, void *Data) = 0;
+		virtual bool StructuredSerializeFromMismatchedTag(struct FPropertyTag const& Tag, FStructuredArchive::FSlot Slot, void *Data) = 0;
 
 		/** return true if this struct has a GetTypeHash */
 		virtual bool HasGetTypeHash() = 0;
@@ -1105,11 +1099,23 @@ public:
 		{
 			return TTraits::WithSerializeFromMismatchedTag;
 		}
-		virtual bool SerializeFromMismatchedTag(struct FPropertyTag const& Tag, FStructuredArchive::FSlot Slot, void *Data) override
+		virtual bool SerializeFromMismatchedTag(struct FPropertyTag const& Tag, FArchive& Ar, void *Data) override
 		{
 			check(TTraits::WithSerializeFromMismatchedTag); // don't call this if we have indicated it is not allowed
-			return SerializeFromMismatchedTagOrNot(Tag, Slot, (CPPSTRUCT*)Data);
+			return SerializeFromMismatchedTagOrNot(Tag, Ar, (CPPSTRUCT*)Data);
 		}
+		virtual bool HasStructuredSerializeFromMismatchedTag() override
+		{
+			return TTraits::WithStructuredSerializeFromMismatchedTag;
+		}
+		virtual bool StructuredSerializeFromMismatchedTag(struct FPropertyTag const& Tag, FStructuredArchive::FSlot Slot, void *Data) override
+		{
+			check(TTraits::WithStructuredSerializeFromMismatchedTag); // don't call this if we have indicated it is not allowed
+			return StructuredSerializeFromMismatchedTagOrNot(Tag, Slot, (CPPSTRUCT*)Data);
+		}
+
+		static_assert(!(TTraits::WithSerializeFromMismatchedTag && TTraits::WithStructuredSerializeFromMismatchedTag), "Structs cannot have both WithSerializeFromMismatchedTag and WithStructuredSerializeFromMismatchedTag set");
+
 		virtual bool HasGetTypeHash() override
 		{
 			return THasGetTypeHash<CPPSTRUCT>::Value;
diff --git a/Engine/Source/Runtime/Engine/Public/PerPlatformProperties.h b/Engine/Source/Runtime/Engine/Public/PerPlatformProperties.h
--- a/Engine/Source/Runtime/Engine/Public/PerPlatformProperties.h
+++ b/Engine/Source/Runtime/Engine/Public/PerPlatformProperties.h
@@ -75,20 +75,6 @@ struct ENGINE_API TPerPlatformProperty
 		return true;
 	}
 
-	/* Load old properties that have been converted to FPerPlatformX */
-	bool SerializeFromMismatchedTag(const FPropertyTag& Tag, FStructuredArchive::FSlot Slot)
-	{
-		if (Tag.Type == _BasePropertyName)
-		{
-			_StructType* This = StaticCast<_StructType*>(this);
-			_ValueType OldValue;
-			Slot << OldValue;
-			*This = _StructType(OldValue);
-			return true;
-		}
-		return false;
-	}
-
 	/* Serialization */
 	bool Serialize(FStructuredArchive::FSlot Slot)
 	{
diff --git a/Engine/Source/Runtime/GameplayTags/Classes/GameplayTagContainer.h b/Engine/Source/Runtime/GameplayTags/Classes/GameplayTagContainer.h
--- a/Engine/Source/Runtime/GameplayTags/Classes/GameplayTagContainer.h
+++ b/Engine/Source/Runtime/GameplayTags/Classes/GameplayTagContainer.h
@@ -256,7 +256,7 @@ struct TStructOpsTypeTraits< FGameplayTag > : public TStructOpsTypeTraitsBase2<
 		WithNetSerializer = true,
 		WithNetSharedSerialization = true,
 		WithPostSerialize = true,
-		WithSerializeFromMismatchedTag = true,
+		WithStructuredSerializeFromMismatchedTag = true,
 		WithImportTextItem = true,
 	};
 };
diff --git a/Engine/Source/Runtime/InputCore/Classes/InputCoreTypes.h b/Engine/Source/Runtime/InputCore/Classes/InputCoreTypes.h
--- a/Engine/Source/Runtime/InputCore/Classes/InputCoreTypes.h
+++ b/Engine/Source/Runtime/InputCore/Classes/InputCoreTypes.h
@@ -105,7 +105,7 @@ struct TStructOpsTypeTraits<FKey> : public TStructOpsTypeTraitsBase2<FKey>
 {
 	enum
 	{
-		WithSerializeFromMismatchedTag = true,
+		WithStructuredSerializeFromMismatchedTag = true,
 		WithExportTextItem = true,
 		WithImportTextItem = true,
 		WithPostSerialize = true,
diff --git a/Engine/Source/Runtime/MovieScene/Public/Channels/MovieSceneBoolChannel.h b/Engine/Source/Runtime/MovieScene/Public/Channels/MovieSceneBoolChannel.h
--- a/Engine/Source/Runtime/MovieScene/Public/Channels/MovieSceneBoolChannel.h
+++ b/Engine/Source/Runtime/MovieScene/Public/Channels/MovieSceneBoolChannel.h
@@ -136,7 +136,7 @@ private:
 template<>
 struct TStructOpsTypeTraits<FMovieSceneBoolChannel> : public TStructOpsTypeTraitsBase2<FMovieSceneBoolChannel>
 {
-	enum { WithSerializeFromMismatchedTag = true };
+	enum { WithStructuredSerializeFromMismatchedTag = true };
 };
 
 template<>
diff --git a/Engine/Source/Runtime/MovieScene/Public/Channels/MovieSceneByteChannel.h b/Engine/Source/Runtime/MovieScene/Public/Channels/MovieSceneByteChannel.h
--- a/Engine/Source/Runtime/MovieScene/Public/Channels/MovieSceneByteChannel.h
+++ b/Engine/Source/Runtime/MovieScene/Public/Channels/MovieSceneByteChannel.h
@@ -155,7 +155,7 @@ private:
 template<>
 struct TStructOpsTypeTraits<FMovieSceneByteChannel> : public TStructOpsTypeTraitsBase2<FMovieSceneByteChannel>
 {
-	enum { WithSerializeFromMismatchedTag = true };
+	enum { WithStructuredSerializeFromMismatchedTag = true };
 };
 
 template<>
diff --git a/Engine/Source/Runtime/MovieScene/Public/Channels/MovieSceneFloatChannel.h b/Engine/Source/Runtime/MovieScene/Public/Channels/MovieSceneFloatChannel.h
--- a/Engine/Source/Runtime/MovieScene/Public/Channels/MovieSceneFloatChannel.h
+++ b/Engine/Source/Runtime/MovieScene/Public/Channels/MovieSceneFloatChannel.h
@@ -287,7 +287,7 @@ public:
 template<>
 struct TStructOpsTypeTraits<FMovieSceneFloatChannel> : public TStructOpsTypeTraitsBase2<FMovieSceneFloatChannel>
 {
-	enum { WithSerializeFromMismatchedTag = true };
+	enum { WithStructuredSerializeFromMismatchedTag = true };
 };
 
 template<>
diff --git a/Engine/Source/Runtime/MovieScene/Public/Channels/MovieSceneIntegerChannel.h b/Engine/Source/Runtime/MovieScene/Public/Channels/MovieSceneIntegerChannel.h
--- a/Engine/Source/Runtime/MovieScene/Public/Channels/MovieSceneIntegerChannel.h
+++ b/Engine/Source/Runtime/MovieScene/Public/Channels/MovieSceneIntegerChannel.h
@@ -136,7 +136,7 @@ private:
 template<>
 struct TStructOpsTypeTraits<FMovieSceneIntegerChannel> : public TStructOpsTypeTraitsBase2<FMovieSceneIntegerChannel>
 {
-	enum { WithSerializeFromMismatchedTag = true };
+	enum { WithStructuredSerializeFromMismatchedTag = true };
 };
 
 template<>
diff --git a/Engine/Source/Runtime/MovieScene/Public/MovieSceneFrameMigration.h b/Engine/Source/Runtime/MovieScene/Public/MovieSceneFrameMigration.h
--- a/Engine/Source/Runtime/MovieScene/Public/MovieSceneFrameMigration.h
+++ b/Engine/Source/Runtime/MovieScene/Public/MovieSceneFrameMigration.h
@@ -96,7 +96,7 @@ private:
 template<>
 struct TStructOpsTypeTraits<FMovieSceneFrameRange> : public TStructOpsTypeTraitsBase2<FMovieSceneFrameRange>
 {
-	enum { WithSerializeFromMismatchedTag = true, WithSerializer = true, WithIdenticalViaEquality = true,
+	enum { WithStructuredSerializeFromMismatchedTag = true, WithSerializer = true, WithIdenticalViaEquality = true,
 		   WithExportTextItem = true, WithImportTextItem = true};
 };
 
diff --git a/Engine/Source/Runtime/MovieScene/Public/MovieSceneSequencePlayer.h b/Engine/Source/Runtime/MovieScene/Public/MovieSceneSequencePlayer.h
--- a/Engine/Source/Runtime/MovieScene/Public/MovieSceneSequencePlayer.h
+++ b/Engine/Source/Runtime/MovieScene/Public/MovieSceneSequencePlayer.h
@@ -111,7 +111,7 @@ struct FMovieSceneSequencePlaybackSettings
 
 template<> struct TStructOpsTypeTraits<FMovieSceneSequencePlaybackSettings> : public TStructOpsTypeTraitsBase2<FMovieSceneSequencePlaybackSettings>
 {
-	enum { WithCopy = true, WithSerializeFromMismatchedTag = true };
+	enum { WithCopy = true, WithStructuredSerializeFromMismatchedTag = true };
 };
 
 /**
diff --git a/Engine/Source/Runtime/MovieSceneTracks/Public/Sections/MovieSceneParticleSection.h b/Engine/Source/Runtime/MovieSceneTracks/Public/Sections/MovieSceneParticleSection.h
--- a/Engine/Source/Runtime/MovieSceneTracks/Public/Sections/MovieSceneParticleSection.h
+++ b/Engine/Source/Runtime/MovieSceneTracks/Public/Sections/MovieSceneParticleSection.h
@@ -31,7 +31,7 @@ struct FMovieSceneParticleChannel : public FMovieSceneByteChannel
 template<>
 struct TStructOpsTypeTraits<FMovieSceneParticleChannel> : public TStructOpsTypeTraitsBase2<FMovieSceneParticleChannel>
 {
-	enum { WithSerializeFromMismatchedTag = true };
+	enum { WithStructuredSerializeFromMismatchedTag = true };
 };
 
 
diff --git a/Engine/Source/Runtime/MovieSceneTracks/Public/Sections/MovieSceneStringChannel.h b/Engine/Source/Runtime/MovieSceneTracks/Public/Sections/MovieSceneStringChannel.h
--- a/Engine/Source/Runtime/MovieSceneTracks/Public/Sections/MovieSceneStringChannel.h
+++ b/Engine/Source/Runtime/MovieSceneTracks/Public/Sections/MovieSceneStringChannel.h
@@ -122,7 +122,7 @@ private:
 template<>
 struct TStructOpsTypeTraits<FMovieSceneStringChannel> : public TStructOpsTypeTraitsBase2<FMovieSceneStringChannel>
 {
-	enum { WithSerializeFromMismatchedTag = true };
+	enum { WithStructuredSerializeFromMismatchedTag = true };
 };
 
 
diff --git a/Engine/Source/Runtime/SlateCore/Public/Sound/SlateSound.h b/Engine/Source/Runtime/SlateCore/Public/Sound/SlateSound.h
--- a/Engine/Source/Runtime/SlateCore/Public/Sound/SlateSound.h
+++ b/Engine/Source/Runtime/SlateCore/Public/Sound/SlateSound.h
@@ -84,6 +84,6 @@ struct TStructOpsTypeTraits<FSlateSound>
 {
 	enum 
 	{
-		WithSerializeFromMismatchedTag = true,
+		WithStructuredSerializeFromMismatchedTag = true,
 	};
 };
diff --git a/Engine/Source/Runtime/SlateCore/Public/Styling/SlateColor.h b/Engine/Source/Runtime/SlateCore/Public/Styling/SlateColor.h
--- a/Engine/Source/Runtime/SlateCore/Public/Styling/SlateColor.h
+++ b/Engine/Source/Runtime/SlateCore/Public/Styling/SlateColor.h
@@ -210,6 +210,6 @@ struct TStructOpsTypeTraits<FSlateColor>
 {
 	enum 
 	{
-		WithSerializeFromMismatchedTag = true,
+		WithStructuredSerializeFromMismatchedTag = true,
 	};
 };
