Fix crash when reloading circularly dependent widget blueprints

diff --git a/Engine/Source/Editor/UnrealEd/Private/EditorEngine.cpp b/Engine/Source/Editor/UnrealEd/Private/EditorEngine.cpp
--- a/Engine/Source/Editor/UnrealEd/Private/EditorEngine.cpp
+++ b/Engine/Source/Editor/UnrealEd/Private/EditorEngine.cpp
@@ -727,7 +727,19 @@ void UEditorEngine::HandlePackageReloaded(const EPackageReloadPhase InPackageRel
 				{
 					if(NewObject && CastChecked<UBlueprint>(NewObject)->GeneratedClass)
 					{
-						FBlueprintCompileReinstancer::ReplaceInstancesOfClass(OldBlueprint->GeneratedClass, CastChecked<UBlueprint>(NewObject)->GeneratedClass);
+						// Don't change the class on instances that are being thrown away by the reload code. If we update
+						// the class and recompile the old class ::ReplaceInstancesOfClass will experience some crosstalk 
+						// with the compiler (both trying to create objects of the same class in the same location):
+						TArray<UObject*> OldInstances;
+						GetObjectsOfClass(OldBlueprint->GeneratedClass, OldInstances, false);
+						OldInstances.RemoveAllSwap(
+							[](UObject* Obj){ return !Obj->HasAnyFlags(RF_NewerVersionExists); }
+						);
+						
+						TSet<UObject*> InstancesToLeaveAlone(OldInstances);
+						FReplaceInstancesOfClassParameters ReplaceInstancesParameters(OldBlueprint->GeneratedClass, CastChecked<UBlueprint>(NewObject)->GeneratedClass);
+						ReplaceInstancesParameters.InstancesThatShouldUseOldClass = &InstancesToLeaveAlone;
+						FBlueprintCompileReinstancer::ReplaceInstancesOfClassEx(ReplaceInstancesParameters);
 					}
 					else
 					{
@@ -774,7 +786,7 @@ void UEditorEngine::HandlePackageReloaded(const EPackageReloadPhase InPackageRel
 				BlueprintToRecompile = ObjectReferencerPtr->GetTypedOuter<UBlueprint>();
 			}
 			
-			if (BlueprintToRecompile)
+			if (BlueprintToRecompile && !BlueprintToRecompile->HasAnyFlags(RF_NewerVersionExists))
 			{
 				BlueprintsToRecompileThisBatch.Add(BlueprintToRecompile);
 			}
diff --git a/Engine/Source/Editor/UnrealEd/Private/Kismet2/KismetReinstanceUtilities.cpp b/Engine/Source/Editor/UnrealEd/Private/Kismet2/KismetReinstanceUtilities.cpp
--- a/Engine/Source/Editor/UnrealEd/Private/Kismet2/KismetReinstanceUtilities.cpp
+++ b/Engine/Source/Editor/UnrealEd/Private/Kismet2/KismetReinstanceUtilities.cpp
@@ -205,6 +205,17 @@ struct FArchetypeReinstanceHelper
 	}
 };
 
+FReplaceInstancesOfClassParameters::FReplaceInstancesOfClassParameters(UClass* InOldClass, UClass* InNewClass)
+	: OldClass(InOldClass)
+	, NewClass(InNewClass)
+	, OriginalCDO(nullptr)
+	, ObjectsThatShouldUseOldStuff(nullptr)
+	, InstancesThatShouldUseOldClass(nullptr)
+	, bClassObjectReplaced(false)
+	, bPreserveRootComponent(true)
+{
+}
+
 /////////////////////////////////////////////////////////////////////////////////
 // FBlueprintCompileReinstancer
 
@@ -1461,6 +1472,13 @@ void FBlueprintCompileReinstancer::ReplaceInstancesOfClass(UClass* OldClass, UCl
 	ReplaceInstancesOfClass_Inner(OldToNewClassMap, OriginalCDO, ObjectsThatShouldUseOldStuff, bClassObjectReplaced, bPreserveRootComponent);
 }
 
+void FBlueprintCompileReinstancer::ReplaceInstancesOfClassEx(const FReplaceInstancesOfClassParameters& Parameters )
+{
+	TMap<UClass*, UClass*> OldToNewClassMap;
+	OldToNewClassMap.Add(Parameters.OldClass, Parameters.NewClass);
+	ReplaceInstancesOfClass_Inner(OldToNewClassMap, Parameters.OriginalCDO, Parameters.ObjectsThatShouldUseOldStuff, Parameters.bClassObjectReplaced, Parameters.bPreserveRootComponent, Parameters.InstancesThatShouldUseOldClass);
+}
+
 void FBlueprintCompileReinstancer::BatchReplaceInstancesOfClass(TMap<UClass*, UClass*>& InOldToNewClassMap, bool bArchetypesAreUpToDate)
 {
 	if (InOldToNewClassMap.Num() == 0)
@@ -1839,7 +1857,7 @@ static void ReplaceActorHelper(UObject* OldObject, UClass* OldClass, UObject*& N
 	OldToNewInstanceMap.Add(OldActor, NewActor);
 }
 
-void FBlueprintCompileReinstancer::ReplaceInstancesOfClass_Inner(TMap<UClass*, UClass*>& InOldToNewClassMap, UObject* InOriginalCDO, TSet<UObject*>* ObjectsThatShouldUseOldStuff, bool bClassObjectReplaced, bool bPreserveRootComponent, bool bArchetypesAreUpToDate)
+void FBlueprintCompileReinstancer::ReplaceInstancesOfClass_Inner(TMap<UClass*, UClass*>& InOldToNewClassMap, UObject* InOriginalCDO, TSet<UObject*>* ObjectsThatShouldUseOldStuff, bool bClassObjectReplaced, bool bPreserveRootComponent, bool bArchetypesAreUpToDate, const TSet<UObject*>* InstancesThatShouldUseOldClass )
 {
 	// If there is an original CDO, we are only reinstancing a single class
 	check((InOriginalCDO != nullptr && InOldToNewClassMap.Num() == 1) || InOriginalCDO == nullptr); // (InOldToNewClassMap.Num() > 1 && InOriginalCDO == nullptr) || (InOldToNewClassMap.Num() == 1 && InOriginalCDO != nullptr));
@@ -1940,7 +1958,8 @@ void FBlueprintCompileReinstancer::ReplaceInstancesOfClass_Inner(TMap<UClass*, U
 					const bool bIsComponent = NewClass->IsChildOf(UActorComponent::StaticClass());
 					const bool bIsChildActorTemplate = OldActor && OldActor->GetOuter()->IsA<UChildActorComponent>();
 					if (OldObject->IsPendingKill() || 
-						(!bIsComponent && !bIsChildActorTemplate && OldObject->IsTemplate()))
+						(!bIsComponent && !bIsChildActorTemplate && OldObject->IsTemplate()) ||
+						(InstancesThatShouldUseOldClass && InstancesThatShouldUseOldClass->Contains(OldObject)))
 					{
 						continue;
 					}
@@ -1981,12 +2000,15 @@ void FBlueprintCompileReinstancer::ReplaceInstancesOfClass_Inner(TMap<UClass*, U
 				for (int32 OldObjIndex = 0; OldObjIndex < ObjectsToReplace.Num(); ++OldObjIndex)
 				{
 					UObject* OldObject = ObjectsToReplace[OldObjIndex];
-					if (!OldObject->IsPendingKill())
+					if(OldObject->IsPendingKill() || 
+						(InstancesThatShouldUseOldClass && InstancesThatShouldUseOldClass->Contains(OldObject)))
 					{
-						if (AActor* OldActor = Cast<AActor>(OldObject))
-						{
-							ActorAttachmentData.Add(OldObject, FActorAttachmentData(OldActor));
-						}
+						continue;
+					}
+
+					if (AActor* OldActor = Cast<AActor>(OldObject))
+					{
+						ActorAttachmentData.Add(OldObject, FActorAttachmentData(OldActor));
 					}
 				}
 
@@ -1998,7 +2020,9 @@ void FBlueprintCompileReinstancer::ReplaceInstancesOfClass_Inner(TMap<UClass*, U
 
 					// Skip archetype instances, EXCEPT for child actor templates
 					const bool bIsChildActorTemplate = OldActor && OldActor->GetOuter()->IsA<UChildActorComponent>();
-					if (OldObject->IsPendingKill() || (!bIsChildActorTemplate && OldObject->IsTemplate()))
+					if (OldObject->IsPendingKill() || 
+						(!bIsChildActorTemplate && OldObject->IsTemplate()) ||
+						(InstancesThatShouldUseOldClass && InstancesThatShouldUseOldClass->Contains(OldObject)))
 					{
 						continue;
 					}
diff --git a/Engine/Source/Editor/UnrealEd/Public/Kismet2/KismetReinstanceUtilities.h b/Engine/Source/Editor/UnrealEd/Public/Kismet2/KismetReinstanceUtilities.h
--- a/Engine/Source/Editor/UnrealEd/Public/Kismet2/KismetReinstanceUtilities.h
+++ b/Engine/Source/Editor/UnrealEd/Public/Kismet2/KismetReinstanceUtilities.h
@@ -36,6 +36,32 @@ public:
 	~FRecreateUberGraphFrameScope();
 };
 
+struct UNREALED_API FReplaceInstancesOfClassParameters
+{
+	FReplaceInstancesOfClassParameters(UClass* InOldClass, UClass* InNewClass);
+
+	/** The old class, instances of which will be replaced */
+	UClass* OldClass;
+
+	/** The new class, used to create new instances to replace instances of the old class */
+	UClass* NewClass;
+
+	/** OriginalCDO, use if OldClass->ClassDefaultObject has been overwritten */
+	UObject* OriginalCDO; 
+
+	/** Set of objects that should not have their references updated if they refer to instances that are replaced */
+	TSet<UObject*>* ObjectsThatShouldUseOldStuff;
+
+	/** Set of objects for which new objects should not be created, useful if client has created new instances themself */
+	const TSet<UObject*>* InstancesThatShouldUseOldClass;
+
+	/** Set to true if class object has been replaced*/
+	bool bClassObjectReplaced; 
+
+	/** Defaults to true, indicates whether root component should be preserved */
+	bool bPreserveRootComponent; 
+};
+
 class UNREALED_API FBlueprintCompileReinstancer : public TSharedFromThis<FBlueprintCompileReinstancer>, public FGCObject
 {
 protected:
@@ -123,7 +149,8 @@ public:
 	void UpdateBytecodeReferences();
 
 	/** Worker function to replace all instances of OldClass with a new instance of NewClass */
-	static void ReplaceInstancesOfClass(UClass* OldClass, UClass* NewClass, UObject* OriginalCDO = NULL, TSet<UObject*>* ObjectsThatShouldUseOldStuff = NULL, bool bClassObjectReplaced = false, bool bPreserveRootComponent = true);
+	static void ReplaceInstancesOfClass(UClass* OldClass, UClass* NewClass, UObject* OriginalCDO = nullptr, TSet<UObject*>* ObjectsThatShouldUseOldStuff = nullptr, bool bClassObjectReplaced = false, bool bPreserveRootComponent = true);
+	static void ReplaceInstancesOfClassEx(const FReplaceInstancesOfClassParameters& Parameters );
 
 	/** Batch replaces a mapping of one or more classes to their new class by leveraging ReplaceInstancesOfClass */
 	static void BatchReplaceInstancesOfClass(TMap<UClass*, UClass*>& InOldToNewClassMap, bool bArchetypesAreUpToDate);
@@ -188,5 +215,5 @@ protected:
 
 private:
 	/** Handles the work of ReplaceInstancesOfClass, handling both normal replacement of instances and batch */
-	static void ReplaceInstancesOfClass_Inner(TMap<UClass*, UClass*>& InOldToNewClassMap, UObject* InOriginalCDO, TSet<UObject*>* ObjectsThatShouldUseOldStuff = NULL, bool bClassObjectReplaced = false, bool bPreserveRootComponent = true, bool bArchetypesAreUpToDate = false);
+	static void ReplaceInstancesOfClass_Inner(TMap<UClass*, UClass*>& InOldToNewClassMap, UObject* InOriginalCDO, TSet<UObject*>* ObjectsThatShouldUseOldStuff = NULL, bool bClassObjectReplaced = false, bool bPreserveRootComponent = true, bool bArchetypesAreUpToDate = false, const TSet<UObject*>* InstancesThatShouldUseOldClass = nullptr);
 };
diff --git a/Engine/Source/Runtime/CoreUObject/Private/UObject/PackageReload.cpp b/Engine/Source/Runtime/CoreUObject/Private/UObject/PackageReload.cpp
--- a/Engine/Source/Runtime/CoreUObject/Private/UObject/PackageReload.cpp
+++ b/Engine/Source/Runtime/CoreUObject/Private/UObject/PackageReload.cpp
@@ -1,9 +1,9 @@
 // Copyright 1998-2018 Epic Games, Inc. All Rights Reserved.
 
 #include "UObject/PackageReload.h"
-#include "UObject/StrongObjectPtr.h"
 #include "UObject/ReferenceChainSearch.h"
 #include "UObject/UObjectHash.h"
+#include "UObject/GCObject.h"
 #include "Misc/AssetRegistryInterface.h"
 #include "Misc/ScopedSlowTask.h"
 #include "UObject/Package.h"
@@ -15,25 +15,42 @@ namespace PackageReloadInternal
 {
 
 /**
- * Strong reference to an existing package that prevents it being GC'd while we're still using it.
+ * Reference to an existing package that prevents it being GC'd while we're still using it (via FExistingPackageReferences).
  * Once we're done with it, we clear out the strong reference and use the weak reference to verify that it was purged correctly via GC.
  */
 struct FExistingPackageReference
 {
 	FExistingPackageReference(UPackage* InPackage)
 		: RawRef(InPackage)
-		, WeakRef(InPackage)
 		, StrongRef(InPackage)
+		, WeakRef(InPackage)
 	{
 	}
 
 	UPackage* RawRef;
+	UPackage* StrongRef;
 	TWeakObjectPtr<UPackage> WeakRef;
-	TStrongObjectPtr<UPackage> StrongRef;
 };
 
 /**
- * Strong reference to an replacement package that prevents it being GC'd while we're still using it.
+ * Array wrapper that prevents the packages inside the FExistingPackageReference instances being GC'd while we're still using them.
+ */
+struct FExistingPackageReferences : public FGCObject
+{
+	virtual void AddReferencedObjects(FReferenceCollector& Collector) override
+	{
+		for (FExistingPackageReference& Ref : Refs)
+		{
+			// Note: We deliberate don't ARO RawRef here
+			Collector.AddReferencedObject(Ref.StrongRef);
+		}
+	}
+
+	TArray<FExistingPackageReference> Refs;
+};
+
+/**
+ * Reference to an replacement package that prevents it being GC'd while we're still using it (via FNewPackageReferences).
  * This also includes the event data used when broadcasting package reload events for this package.
  */
 struct FNewPackageReference
@@ -44,10 +61,30 @@ struct FNewPackageReference
 	{
 	}
 
-	TStrongObjectPtr<UPackage> Package;
+	UPackage* Package;
 	TSharedPtr<FPackageReloadedEvent> EventData;
 };
 
+/**
+ * Array wrapper that prevents the packages inside the FNewPackageReferences instances being GC'd while we're still using them.
+ */
+struct FNewPackageReferences : public FGCObject
+{
+	virtual void AddReferencedObjects(FReferenceCollector& Collector) override
+	{
+		for (FNewPackageReference& Ref : Refs)
+		{
+			Collector.AddReferencedObject(Ref.Package);
+			if(Ref.EventData)
+			{
+				Ref.EventData->AddReferencedObjects(Collector);
+			}
+		}
+	}
+
+	TArray<FNewPackageReference> Refs;
+};
+
 /**
  * Used to map objects from the old package, to objects in the new package, including the index of the package being reloaded.
  */
@@ -446,18 +483,18 @@ void ReloadPackages(const TArrayView<FReloadPackageData>& InPackagesToReload, TA
 	}, false);
 
 	// Gather up the list of all packages to reload (note: this array may include null packages!)
-	TArray<PackageReloadInternal::FExistingPackageReference> ExistingPackages;
-	ExistingPackages.Reserve(InPackagesToReload.Num());
+	PackageReloadInternal::FExistingPackageReferences ExistingPackages;
+	ExistingPackages.Refs.Reserve(InPackagesToReload.Num());
 	{
 		FScopedSlowTask PreparingPackagesForReloadSlowTask(InPackagesToReload.Num(), NSLOCTEXT("CoreUObject", "PreparingPackagesForReload", "Preparing Packages for Reload"));
 
 		for (const FReloadPackageData& PackageToReloadData : InPackagesToReload)
 		{
 			PreparingPackagesForReloadSlowTask.EnterProgressFrame(1.0f);
-			ExistingPackages.Emplace(PackageReloadInternal::ValidateAndPreparePackageForReload(PackageToReloadData.PackageToReload));
+			ExistingPackages.Refs.Emplace(PackageReloadInternal::ValidateAndPreparePackageForReload(PackageToReloadData.PackageToReload));
 		}
 
-		if (ExistingPackages.Num() > 0)
+		if (ExistingPackages.Refs.Num() > 0)
 		{
 			// Run a GC before we start to clean-up any lingering objects that may reference things we're about to reload
 			CollectGarbage(GARBAGE_COLLECTION_KEEPFLAGS);
@@ -465,30 +502,30 @@ void ReloadPackages(const TArrayView<FReloadPackageData>& InPackagesToReload, TA
 	}
 
 	// Rename the existing packages, load the new packages, then fix-up any references
-	TArray<PackageReloadInternal::FNewPackageReference> NewPackages;
-	NewPackages.Reserve(ExistingPackages.Num());
+	PackageReloadInternal::FNewPackageReferences NewPackages;
+	NewPackages.Refs.Reserve(ExistingPackages.Refs.Num());
 	{
 		// Process the packages in batches to avoid consuming too much memory due to a lack of GC
 		int32 PackageIndex = 0;
-		while (PackageIndex < ExistingPackages.Num())
+		while (PackageIndex < ExistingPackages.Refs.Num())
 		{
 			FCoreUObjectDelegates::OnPackageReloaded.Broadcast(EPackageReloadPhase::PreBatch, nullptr);
 
 			const int32 BatchStartIndex = PackageIndex;
-			for (; PackageIndex < ExistingPackages.Num(); ++PackageIndex)
+			for (; PackageIndex < ExistingPackages.Refs.Num(); ++PackageIndex)
 			{
-				UPackage* ExistingPackage = ExistingPackages[PackageIndex].RawRef;
+				UPackage* ExistingPackage = ExistingPackages.Refs[PackageIndex].RawRef;
 
 				const FText ProgressText = ExistingPackage 
 					? FText::Format(NSLOCTEXT("CoreUObject", "ReloadingPackagef", "Reloading {0}..."), FText::FromName(ExistingPackage->GetFName()))
 					: NSLOCTEXT("CoreUObject", "ReloadingPackages", "Reloading Packages");
 				ReloadingPackagesSlowTask.EnterProgressFrame(1, ProgressText);
 
-				check(NewPackages.Num() == PackageIndex);
-				NewPackages.Emplace(PackageReloadInternal::LoadReplacementPackage(ExistingPackage, InPackagesToReload[PackageIndex].LoadFlags));
+				check(NewPackages.Refs.Num() == PackageIndex);
+				NewPackages.Refs.Emplace(PackageReloadInternal::LoadReplacementPackage(ExistingPackage, InPackagesToReload[PackageIndex].LoadFlags));
 
-				UPackage* NewPackage = NewPackages[PackageIndex].Package.Get();
-				NewPackages[PackageIndex].EventData = PackageReloadInternal::GeneratePackageReloadEvent(ExistingPackage, NewPackage);
+				UPackage* NewPackage = NewPackages.Refs[PackageIndex].Package;
+				NewPackages.Refs[PackageIndex].EventData = PackageReloadInternal::GeneratePackageReloadEvent(ExistingPackage, NewPackage);
 
 				const bool bEndBatch = (PackageIndex == (BatchStartIndex + InNumPackagesPerBatch)) || (ExistingPackage && ExistingPackage->ContainsMap());
 				if (bEndBatch)
@@ -508,7 +545,7 @@ void ReloadPackages(const TArrayView<FReloadPackageData>& InPackagesToReload, TA
 			{
 				FixingUpReferencesSlowTask.EnterProgressFrame(1.0f);
 
-				PackageReloadInternal::FNewPackageReference& NewPackageData = NewPackages[BatchPackageIndex];
+				PackageReloadInternal::FNewPackageReference& NewPackageData = NewPackages.Refs[BatchPackageIndex];
 				if (NewPackageData.EventData.IsValid())
 				{
 					FCoreUObjectDelegates::OnPackageReloaded.Broadcast(EPackageReloadPhase::PrePackageFixup, NewPackageData.EventData.Get());
@@ -531,7 +568,7 @@ void ReloadPackages(const TArrayView<FReloadPackageData>& InPackagesToReload, TA
 
 				FixingUpReferencesSlowTask.EnterProgressFrame(1.0f);
 
-				PackageReloadInternal::FReplaceObjectReferencesArchive ReplaceRefsArchive(PotentialReferencer, OldObjectToNewData, ExistingPackages, NewPackages);
+				PackageReloadInternal::FReplaceObjectReferencesArchive ReplaceRefsArchive(PotentialReferencer, OldObjectToNewData, ExistingPackages.Refs, NewPackages.Refs);
 				PotentialReferencer->Serialize(ReplaceRefsArchive); // Deal with direct references during Serialization
 				PotentialReferencer->GetClass()->CallAddReferencedObjects(PotentialReferencer, ReplaceRefsArchive); // Deal with indirect references via AddReferencedObjects
 			}
@@ -541,7 +578,7 @@ void ReloadPackages(const TArrayView<FReloadPackageData>& InPackagesToReload, TA
 			{
 				FixingUpReferencesSlowTask.EnterProgressFrame(1.0f);
 
-				ExistingPackages[BatchPackageIndex].StrongRef.Reset(ExistingPackages[BatchPackageIndex].RawRef);
+				ExistingPackages.Refs[BatchPackageIndex].StrongRef = ExistingPackages.Refs[BatchPackageIndex].RawRef;
 			}
 
 			// Final pass to clean-up any remaining references prior to GC
@@ -550,7 +587,7 @@ void ReloadPackages(const TArrayView<FReloadPackageData>& InPackagesToReload, TA
 			{
 				FixingUpReferencesSlowTask.EnterProgressFrame(1.0f);
 
-				PackageReloadInternal::FNewPackageReference& NewPackageData = NewPackages[BatchPackageIndex];
+				PackageReloadInternal::FNewPackageReference& NewPackageData = NewPackages.Refs[BatchPackageIndex];
 				if (NewPackageData.EventData.IsValid())
 				{
 					FCoreUObjectDelegates::OnPackageReloaded.Broadcast(EPackageReloadPhase::PostPackageFixup, NewPackageData.EventData.Get());
@@ -564,14 +601,15 @@ void ReloadPackages(const TArrayView<FReloadPackageData>& InPackagesToReload, TA
 			{
 				FixingUpReferencesSlowTask.EnterProgressFrame(1.0f);
 
-				UPackage* ExistingPackage = ExistingPackages[BatchPackageIndex].RawRef;
-				UPackage* NewPackage = NewPackages[BatchPackageIndex].Package.Get();
+				UPackage* ExistingPackage = ExistingPackages.Refs[BatchPackageIndex].RawRef;
+				UPackage* NewPackage = NewPackages.Refs[BatchPackageIndex].Package;
 
 				if (ExistingPackage && NewPackage)
 				{
 					// Allow the old package to be GC'd
 					PackageReloadInternal::MakePackagePurgeable(ExistingPackage);
-					ExistingPackages[BatchPackageIndex].StrongRef.Reset();
+					ExistingPackages.Refs[BatchPackageIndex].StrongRef = nullptr;
+					NewPackages.Refs[BatchPackageIndex].EventData.Reset();
 				}
 			}
 			CollectGarbage(GARBAGE_COLLECTION_KEEPFLAGS);
@@ -591,11 +629,11 @@ void ReloadPackages(const TArrayView<FReloadPackageData>& InPackagesToReload, TA
 	}, false);
 
 	// Finalization and error reporting
-	OutReloadedPackages.Reserve(ExistingPackages.Num());
-	for (int32 PackageIndex = 0; PackageIndex < ExistingPackages.Num(); ++PackageIndex)
+	OutReloadedPackages.Reserve(ExistingPackages.Refs.Num());
+	for (int32 PackageIndex = 0; PackageIndex < ExistingPackages.Refs.Num(); ++PackageIndex)
 	{
-		UPackage* ExistingPackage = ExistingPackages[PackageIndex].WeakRef.Get();
-		UPackage* NewPackage = NewPackages[PackageIndex].Package.Get();
+		UPackage* ExistingPackage = ExistingPackages.Refs[PackageIndex].WeakRef.Get();
+		UPackage* NewPackage = NewPackages.Refs[PackageIndex].Package;
 
 		check(OutReloadedPackages.Num() == PackageIndex);
 		OutReloadedPackages.Emplace(NewPackage);
diff --git a/Engine/Source/Runtime/CoreUObject/Public/UObject/PackageReload.h b/Engine/Source/Runtime/CoreUObject/Public/UObject/PackageReload.h
--- a/Engine/Source/Runtime/CoreUObject/Public/UObject/PackageReload.h
+++ b/Engine/Source/Runtime/CoreUObject/Public/UObject/PackageReload.h
@@ -5,7 +5,9 @@
 #include "CoreMinimal.h"
 #include "Containers/ArrayView.h"
 #include "UObject/ObjectMacros.h"
+#include "UObject/UObjectGlobals.h"
 #include "UObject/WeakObjectPtr.h"
+#include "UObject/WeakObjectPtrTemplates.h"
 
 class IAssetRegistryInterface;
 
@@ -123,6 +125,11 @@ public:
 		ObjectReferencers.AddUnique(InObject);
 	}
 
+	void AddReferencedObjects(FReferenceCollector& Collector)
+	{
+		Collector.AddReferencedObjects(RepointedObjects);
+	}
+
 private:
 	bool GetRepointedObjectInternal(const UObject* InObject, const UObject*& OutRepointedObject) const
 	{
