Added code to prevent packages from being saved when the final serialization pass added new custom versions that were not detected when tagging exports which would result in saving a corrupt package and crashes when loading it.

Note that this is not a fix for UE-63571

diff --git a/Engine/Source/Runtime/CoreUObject/Private/UObject/SavePackage.cpp b/Engine/Source/Runtime/CoreUObject/Private/UObject/SavePackage.cpp
--- a/Engine/Source/Runtime/CoreUObject/Private/UObject/SavePackage.cpp
+++ b/Engine/Source/Runtime/CoreUObject/Private/UObject/SavePackage.cpp
@@ -4458,9 +4458,9 @@ FSavePackageResultStruct UPackage::Save(UPackage* InOuter, UObject* Base, EObjec
 					FArchiveStackTraceIgnoreScope IgnoreSummaryDiffsScope(DiffSettings.bIgnoreHeaderDiffs);
 #endif // WITH_EDITOR
 					if (!bTextFormat)
-				    {
+					{
 						StructuredArchiveRoot.GetUnderlyingArchive() << Linker->Summary;
-				    }
+					}
 				}
 				int32 OffsetAfterPackageFileSummary = Linker->Tell();
 		
@@ -5795,10 +5795,29 @@ FSavePackageResultStruct UPackage::Save(UPackage* InOuter, UObject* Base, EObjec
 				
 				// Update package flags from package, in case serialization has modified package flags.
 				Linker->Summary.PackageFlags = Linker->LinkerRoot->GetPackageFlags() & ~PKG_NewlyCreated;
+				
+				{
+					// Verify that the final serialization pass hasn't added any new custom versions. Otherwise this will result in crashes when loading the package.
+					bool bNewCustomVersionsUsed = false;
+					for (const FCustomVersion& LinkerCustomVer : Linker->GetCustomVersions().GetAllVersions())
+					{
+						if (Linker->Summary.GetCustomVersionContainer().GetVersion(LinkerCustomVer.Key) == nullptr)
+						{
+							UE_LOG(LogSavePackage, Error,
+								TEXT("Unexpected custom version \"%s\" found when saving %s. This usually happens when export tagging and final serialization paths differ. Package will not be saved."),
+								*LinkerCustomVer.GetFriendlyName().ToString(), *Linker->LinkerRoot->GetName());
+							bNewCustomVersionsUsed = true;
+						}
+					}
+					if (bNewCustomVersionsUsed)
+					{
+						return ESavePackageResult::Error;
+					}
+				}
 
 				if (!bTextFormat)
 				{
-				Linker->Seek(0);
+					Linker->Seek(0);
 				}
 				{
 #if WITH_EDITOR
@@ -5809,7 +5828,7 @@ FSavePackageResultStruct UPackage::Save(UPackage* InOuter, UObject* Base, EObjec
 
 				if (!bTextFormat)
 				{
-				check( Linker->Tell() == OffsetAfterPackageFileSummary );
+					check( Linker->Tell() == OffsetAfterPackageFileSummary );
 				}
 
 				if ( EndSavingIfCancelled( Linker.Get(), TempFilename ) )
