UObject - Update map property import from text to initialize all key/value pairs instead of just the first one.

diff --git a/Engine/Source/Runtime/CoreUObject/Private/UObject/PropertyMap.cpp b/Engine/Source/Runtime/CoreUObject/Private/UObject/PropertyMap.cpp
--- a/Engine/Source/Runtime/CoreUObject/Private/UObject/PropertyMap.cpp
+++ b/Engine/Source/Runtime/CoreUObject/Private/UObject/PropertyMap.cpp
@@ -656,14 +656,10 @@ const TCHAR* UMapProperty::ImportText_Internal(const TCHAR* Buffer, void* Data,
 	}
 
 	uint8* TempPairStorage   = (uint8*)FMemory::Malloc(MapLayout.ValueOffset + ValueProp->ElementSize);
-	KeyProp  ->InitializeValue(TempPairStorage);
-	ValueProp->InitializeValue(TempPairStorage + MapLayout.ValueOffset);
 
 	bool bSuccess = false;
 	ON_SCOPE_EXIT
 	{
-		ValueProp->DestroyValue(TempPairStorage + MapLayout.ValueOffset);
-		KeyProp  ->DestroyValue(TempPairStorage);
 		FMemory::Free(TempPairStorage);
 
 		// If we are returning because of an error, remove any already-added elements from the map before returning
@@ -676,6 +672,14 @@ const TCHAR* UMapProperty::ImportText_Internal(const TCHAR* Buffer, void* Data,
 
 	for (;;)
 	{
+		KeyProp->InitializeValue(TempPairStorage);
+		ValueProp->InitializeValue(TempPairStorage + MapLayout.ValueOffset);
+		ON_SCOPE_EXIT
+		{
+			ValueProp->DestroyValue(TempPairStorage + MapLayout.ValueOffset);
+			KeyProp->DestroyValue(TempPairStorage);
+		};
+
 		if (*Buffer++ != TCHAR('('))
 		{
 			return nullptr;
