[NO NON-PII COMMIT MESSAGE]

diff --git a/Engine/Source/Runtime/CoreUObject/Private/UObject/LinkerLoad.cpp b/Engine/Source/Runtime/CoreUObject/Private/UObject/LinkerLoad.cpp
--- a/Engine/Source/Runtime/CoreUObject/Private/UObject/LinkerLoad.cpp
+++ b/Engine/Source/Runtime/CoreUObject/Private/UObject/LinkerLoad.cpp
@@ -2669,6 +2669,14 @@ bool FLinkerLoad::VerifyImportInner(const int32 ImportIndex, FString& WarningSuf
 		if (!bWasFullyLoaded)
 		{
 			Import.SourceLinker = GetPackageLinker( TmpPkg, NULL, InternalLoadFlags, NULL, NULL );
+#if WITH_EDITORONLY_DATA
+			if (Import.SourceLinker && !TmpPkg->HasAnyFlags(RF_WasLoaded))
+			{
+				// If we didn't fully load, make sure our metadata is loaded before using this
+				// We need this case for user defined structs due to the LOAD_DeferDependencyLoads code above
+				Import.SourceLinker->LoadMetaDataFromExportMap(false);
+			}
+#endif
 		}
 	}
 	else
