This is required to avoid issues where UDataAsset instances were not correctly referencing their BP class, but we don't add native classes to avoid bloating the import table unnecessarily

diff --git a/Engine/Source/Runtime/CoreUObject/Private/UObject/SavePackage.cpp b/Engine/Source/Runtime/CoreUObject/Private/UObject/SavePackage.cpp
--- a/Engine/Source/Runtime/CoreUObject/Private/UObject/SavePackage.cpp
+++ b/Engine/Source/Runtime/CoreUObject/Private/UObject/SavePackage.cpp
@@ -50,7 +50,6 @@
 #include "ProfilingDebugging/CookStats.h"
 #include "UObject/DebugSerializationFlags.h"
 #include "UObject/EnumProperty.h"
-#include "Blueprint/BlueprintSupport.h"
 #include "HAL/IConsoleManager.h"
 #include "Serialization/ArchiveStackTrace.h"
 #include "UObject/CoreRedirects.h"
@@ -1357,6 +1356,15 @@ FArchive& FArchiveSaveTagImports::operator<<( UObject*& Obj )
 				{
 					*this << Parent;
 				}
+
+				// For things with a BP-created class we need to recurse into that class so the import ClassPackage will load properly
+				// We don't do this for native classes to avoid bloating the import table
+				UClass* ObjClass = Obj->GetClass();
+
+				if (!ObjClass->IsNative())
+				{
+					*this << ObjClass;
+				}
 			}
 		}
 	}
