Avoid crashing when attempting to load a subobject that has a missing outer

diff --git a/Engine/Source/Runtime/CoreUObject/Private/Blueprint/BlueprintSupport.cpp b/Engine/Source/Runtime/CoreUObject/Private/Blueprint/BlueprintSupport.cpp
--- a/Engine/Source/Runtime/CoreUObject/Private/Blueprint/BlueprintSupport.cpp
+++ b/Engine/Source/Runtime/CoreUObject/Private/Blueprint/BlueprintSupport.cpp
@@ -2040,7 +2040,7 @@ void FLinkerLoad::ResolveDeferredExports(UClass* LoadClass)
 		for (int32 ExportIndex = 0; ExportIndex < ExportMap.Num(); ++ExportIndex)
 		{
 			FObjectExport& Export = ExportMap[ExportIndex];
-			if((Export.ObjectFlags & RF_DefaultSubObject) != 0 && Export.OuterIndex.ToExport() == DeferredCDOIndex)
+			if((Export.ObjectFlags & RF_DefaultSubObject) != 0 && Export.OuterIndex.IsExport() && Export.OuterIndex.ToExport() == DeferredCDOIndex)
 			{
 				if (Export.Object == nullptr && Export.OuterIndex.IsExport())
 				{
