This fixes issues where soft object ptrs that weren't currently loaded would replicate as null

diff --git a/Engine/Source/Runtime/CoreUObject/Private/UObject/CoreNet.cpp b/Engine/Source/Runtime/CoreUObject/Private/UObject/CoreNet.cpp
--- a/Engine/Source/Runtime/CoreUObject/Private/UObject/CoreNet.cpp
+++ b/Engine/Source/Runtime/CoreUObject/Private/UObject/CoreNet.cpp
@@ -383,18 +383,18 @@ FArchive& FNetBitWriter::operator<<( UObject*& Object )
 
 FArchive& FNetBitWriter::operator<<(FSoftObjectPath& Value)
 {
+	// It's more efficient to serialize as a string then name+string
 	FString Path = Value.ToString();
-
 	*this << Path;
 
-	if (IsLoading())
-	{
-		Value.SetPath(MoveTemp(Path));
-	}
-
 	return *this;
 }
 
+FArchive& FNetBitWriter::operator<<(FSoftObjectPtr& Value)
+{
+	return FArchiveUObject::SerializeSoftObjectPtr(*this, Value);
+}
+
 FArchive& FNetBitWriter::operator<<(struct FWeakObjectPtr& WeakObjectPtr)
 {
 	return FArchiveUObject::SerializeWeakObjectPtr(*this, WeakObjectPtr);
@@ -438,18 +438,19 @@ FArchive& FNetBitReader::operator<<( class FName& N )
 
 FArchive& FNetBitReader::operator<<(FSoftObjectPath& Value)
 {
-	FString Path = Value.ToString();
-
+	FString Path;
 	*this << Path;
 
-	if (IsLoading())
-	{
-		Value.SetPath(MoveTemp(Path));
-	}
+	Value.SetPath(MoveTemp(Path));
 
 	return *this;
 }
 
+FArchive& FNetBitReader::operator<<(FSoftObjectPtr& Value)
+{
+	return FArchiveUObject::SerializeSoftObjectPtr(*this, Value);
+}
+
 FArchive& FNetBitReader::operator<<(struct FWeakObjectPtr& WeakObjectPtr)
 {
 	return FArchiveUObject::SerializeWeakObjectPtr(*this, WeakObjectPtr);
diff --git a/Engine/Source/Runtime/CoreUObject/Private/UObject/PropertySoftObjectPtr.cpp b/Engine/Source/Runtime/CoreUObject/Private/UObject/PropertySoftObjectPtr.cpp
--- a/Engine/Source/Runtime/CoreUObject/Private/UObject/PropertySoftObjectPtr.cpp
+++ b/Engine/Source/Runtime/CoreUObject/Private/UObject/PropertySoftObjectPtr.cpp
@@ -69,6 +69,14 @@ void USoftObjectProperty::SerializeItem( FStructuredArchive::FSlot Slot, void* V
 	}
 }
 
+bool USoftObjectProperty::NetSerializeItem(FArchive& Ar, UPackageMap* Map, void* Data, TArray<uint8>* MetaData) const
+{
+	// Serialize directly, will use FBitWriter/Reader
+	Ar << *(FSoftObjectPtr*)Data;
+
+	return true;
+}
+
 void USoftObjectProperty::ExportTextItem( FString& ValueStr, const void* PropertyValue, const void* DefaultValue, UObject* Parent, int32 PortFlags, UObject* ExportRootScope ) const
 {
 	FSoftObjectPtr& SoftObjectPtr = *(FSoftObjectPtr*)PropertyValue;
diff --git a/Engine/Source/Runtime/CoreUObject/Public/UObject/CoreNet.h b/Engine/Source/Runtime/CoreUObject/Public/UObject/CoreNet.h
--- a/Engine/Source/Runtime/CoreUObject/Public/UObject/CoreNet.h
+++ b/Engine/Source/Runtime/CoreUObject/Public/UObject/CoreNet.h
@@ -283,6 +283,7 @@ public:
 	virtual FArchive& operator<<(FName& Name) override;
 	virtual FArchive& operator<<(UObject*& Object) override;
 	virtual FArchive& operator<<(FSoftObjectPath& Value) override;
+	virtual FArchive& operator<<(FSoftObjectPtr& Value) override;
 	virtual FArchive& operator<<(struct FWeakObjectPtr& Value) override;
 
 	virtual void CountMemory(FArchive& Ar) const override;
@@ -304,6 +305,7 @@ public:
 	virtual FArchive& operator<<(FName& Name) override;
 	virtual FArchive& operator<<(UObject*& Object) override;
 	virtual FArchive& operator<<(FSoftObjectPath& Value) override;
+	virtual FArchive& operator<<(FSoftObjectPtr& Value) override;
 	virtual FArchive& operator<<(struct FWeakObjectPtr& Value) override;
 
 	virtual void CountMemory(FArchive& Ar) const override;
diff --git a/Engine/Source/Runtime/CoreUObject/Public/UObject/UnrealType.h b/Engine/Source/Runtime/CoreUObject/Public/UObject/UnrealType.h
--- a/Engine/Source/Runtime/CoreUObject/Public/UObject/UnrealType.h
+++ b/Engine/Source/Runtime/CoreUObject/Public/UObject/UnrealType.h
@@ -2170,6 +2170,7 @@ class COREUOBJECT_API USoftObjectProperty : public TUObjectPropertyBase<FSoftObj
 	virtual FName GetID() const override;
 	virtual bool Identical( const void* A, const void* B, uint32 PortFlags ) const override;
 	virtual void SerializeItem( FStructuredArchive::FSlot Slot, void* Value, void const* Defaults) const override;
+	virtual bool NetSerializeItem(FArchive& Ar, UPackageMap* Map, void* Data, TArray<uint8> * MetaData = NULL) const override;
 	virtual void ExportTextItem( FString& ValueStr, const void* PropertyValue, const void* DefaultValue, UObject* Parent, int32 PortFlags, UObject* ExportRootScope ) const override;
 	virtual const TCHAR* ImportText_Internal( const TCHAR* Buffer, void* Data, int32 PortFlags, UObject* OwnerObject, FOutputDevice* ErrorText ) const override;
 	virtual EConvertFromTypeResult ConvertFromType(const FPropertyTag& Tag, FStructuredArchive::FSlot Slot, uint8* Data, UStruct* DefaultsStruct) override;
