This code is from 2005 and predates one-asset-per-package or blueprints

diff --git a/Engine/Source/Runtime/CoreUObject/Private/UObject/LinkerLoad.cpp b/Engine/Source/Runtime/CoreUObject/Private/UObject/LinkerLoad.cpp
--- a/Engine/Source/Runtime/CoreUObject/Private/UObject/LinkerLoad.cpp
+++ b/Engine/Source/Runtime/CoreUObject/Private/UObject/LinkerLoad.cpp
@@ -3016,6 +3016,10 @@ bool FLinkerLoad::VerifyImportInner(const int32 ImportIndex, FString& WarningSuf
 				SafeReplace = true;
 			}
 		}
+		else
+		{
+			SafeReplace = true;
+		}
 
 		if (!Import.XObject && !SafeReplace)
 		{
