Fix several String Asset Reference -> Soft Object Path uses as the old name is deprecated

diff --git a/Engine/Plugins/FX/Niagara/Source/Niagara/Private/NiagaraModule.cpp b/Engine/Plugins/FX/Niagara/Source/Niagara/Private/NiagaraModule.cpp
--- a/Engine/Plugins/FX/Niagara/Source/Niagara/Private/NiagaraModule.cpp
+++ b/Engine/Plugins/FX/Niagara/Source/Niagara/Private/NiagaraModule.cpp
@@ -628,7 +628,7 @@ void FNiagaraTypeDefinition::RecreateUserDefinedTypeRegistry()
 
 		if (Obj != nullptr)
 		{
-			const FSoftObjectPath* ParamRefFound = Settings->AdditionalParameterEnums.FindByPredicate([&](const FStringAssetReference& Ref) { return Ref.ToString() == AssetRef.ToString(); });
+			const FSoftObjectPath* ParamRefFound = Settings->AdditionalParameterEnums.FindByPredicate([&](const FSoftObjectPath& Ref) { return Ref.ToString() == AssetRef.ToString(); });
 			const FSoftObjectPath* PayloadRefFound = nullptr;
 			UEnum* Enum = Cast<UEnum>(Obj);
 			if (Enum != nullptr)
diff --git a/Engine/Plugins/FX/Niagara/Source/NiagaraEditor/Private/NiagaraEmitterFactoryNew.cpp b/Engine/Plugins/FX/Niagara/Source/NiagaraEditor/Private/NiagaraEmitterFactoryNew.cpp
--- a/Engine/Plugins/FX/Niagara/Source/NiagaraEditor/Private/NiagaraEmitterFactoryNew.cpp
+++ b/Engine/Plugins/FX/Niagara/Source/NiagaraEditor/Private/NiagaraEmitterFactoryNew.cpp
@@ -82,7 +82,7 @@ bool UNiagaraEmitterFactoryNew::ConfigureProperties()
 
 UNiagaraNodeFunctionCall* AddModuleFromAssetPath(FString AssetPath, UNiagaraNodeOutput& TargetOutputNode)
 {
-	FStringAssetReference AssetRef(AssetPath);
+	FSoftObjectPath AssetRef(AssetPath);
 	UNiagaraScript* AssetScript = Cast<UNiagaraScript>(AssetRef.TryLoad());
 	FAssetData ScriptAssetData(AssetScript);
 	if (ScriptAssetData.IsValid())
diff --git a/Engine/Plugins/FX/Niagara/Source/NiagaraEditor/Private/NiagaraScriptFactory.cpp b/Engine/Plugins/FX/Niagara/Source/NiagaraEditor/Private/NiagaraScriptFactory.cpp
--- a/Engine/Plugins/FX/Niagara/Source/NiagaraEditor/Private/NiagaraScriptFactory.cpp
+++ b/Engine/Plugins/FX/Niagara/Source/NiagaraEditor/Private/NiagaraScriptFactory.cpp
@@ -41,7 +41,7 @@ UObject* UNiagaraScriptFactoryNew::FactoryCreateNew(UClass* Class, UObject* InPa
 	check(Settings);
 
 	UNiagaraScript* NewScript;
-	const FStringAssetReference& SettingDefaultScript = GetDefaultScriptFromSettings(Settings);
+	const FSoftObjectPath& SettingDefaultScript = GetDefaultScriptFromSettings(Settings);
 	if (UNiagaraScript* Default = Cast<UNiagaraScript>(SettingDefaultScript.TryLoad()))
 	{
 		NewScript = Cast<UNiagaraScript>(StaticDuplicateObject(Default, InParent, Name, Flags, Class));
@@ -58,7 +58,7 @@ UObject* UNiagaraScriptFactoryNew::FactoryCreateNew(UClass* Class, UObject* InPa
 	return NewScript;
 }
 
-const FStringAssetReference& UNiagaraScriptFactoryNew::GetDefaultScriptFromSettings(const UNiagaraEditorSettings* Settings)
+const FSoftObjectPath& UNiagaraScriptFactoryNew::GetDefaultScriptFromSettings(const UNiagaraEditorSettings* Settings)
 {
 	switch (GetScriptUsage())
 	{
diff --git a/Engine/Plugins/FX/Niagara/Source/NiagaraEditor/Private/NiagaraScriptSource.cpp b/Engine/Plugins/FX/Niagara/Source/NiagaraEditor/Private/NiagaraScriptSource.cpp
--- a/Engine/Plugins/FX/Niagara/Source/NiagaraEditor/Private/NiagaraScriptSource.cpp
+++ b/Engine/Plugins/FX/Niagara/Source/NiagaraEditor/Private/NiagaraScriptSource.cpp
@@ -153,7 +153,7 @@ void UNiagaraScriptSource::PostLoadFromEmitter(UNiagaraEmitter& OwningEmitter)
 
 bool UNiagaraScriptSource::AddModuleIfMissing(FString ModulePath, ENiagaraScriptUsage Usage, bool& bOutFoundModule)
 {
-	FStringAssetReference SystemUpdateScriptRef(ModulePath);
+	FSoftObjectPath SystemUpdateScriptRef(ModulePath);
 	FAssetData ModuleScriptAsset;
 	ModuleScriptAsset.ObjectPath = SystemUpdateScriptRef.GetAssetPathName();
 	bOutFoundModule = false;
diff --git a/Engine/Plugins/FX/Niagara/Source/NiagaraEditor/Private/NiagaraSystemFactoryNew.cpp b/Engine/Plugins/FX/Niagara/Source/NiagaraEditor/Private/NiagaraSystemFactoryNew.cpp
--- a/Engine/Plugins/FX/Niagara/Source/NiagaraEditor/Private/NiagaraSystemFactoryNew.cpp
+++ b/Engine/Plugins/FX/Niagara/Source/NiagaraEditor/Private/NiagaraSystemFactoryNew.cpp
@@ -162,7 +162,7 @@ void UNiagaraSystemFactoryNew::InitializeSystem(UNiagaraSystem* System, bool bCr
 
 	if (bCreateDefaultNodes)
 	{
-		FStringAssetReference SystemUpdateScriptRef(TEXT("/Niagara/Modules/System/SystemLifeCycle.SystemLifeCycle"));
+		FSoftObjectPath SystemUpdateScriptRef(TEXT("/Niagara/Modules/System/SystemLifeCycle.SystemLifeCycle"));
 		UNiagaraScript* Script = Cast<UNiagaraScript>(SystemUpdateScriptRef.TryLoad());
 
 		FAssetData ModuleScriptAsset(Script);
diff --git a/Engine/Plugins/FX/Niagara/Source/NiagaraEditor/Private/ViewModels/Stack/NiagaraStackScriptItemGroup.cpp b/Engine/Plugins/FX/Niagara/Source/NiagaraEditor/Private/ViewModels/Stack/NiagaraStackScriptItemGroup.cpp
--- a/Engine/Plugins/FX/Niagara/Source/NiagaraEditor/Private/ViewModels/Stack/NiagaraStackScriptItemGroup.cpp
+++ b/Engine/Plugins/FX/Niagara/Source/NiagaraEditor/Private/ViewModels/Stack/NiagaraStackScriptItemGroup.cpp
@@ -397,7 +397,7 @@ void UNiagaraStackScriptItemGroup::RefreshIssues(TArray<FStackIssue>& NewIssues)
 			// The factor ensures this, but older assets may not have it or it may have been removed accidentally.
 			// For now, treat this as an error and allow them to resolve.
 			FString ModulePath = TEXT("/Niagara/Modules/System/SystemLifeCycle.SystemLifeCycle");
-			FStringAssetReference SystemUpdateScriptRef(ModulePath);
+			FSoftObjectPath SystemUpdateScriptRef(ModulePath);
 			FAssetData ModuleScriptAsset;
 			ModuleScriptAsset.ObjectPath = SystemUpdateScriptRef.GetAssetPathName();
 
diff --git a/Engine/Plugins/FX/Niagara/Source/NiagaraEditor/Public/NiagaraEditorSettings.h b/Engine/Plugins/FX/Niagara/Source/NiagaraEditor/Public/NiagaraEditorSettings.h
--- a/Engine/Plugins/FX/Niagara/Source/NiagaraEditor/Public/NiagaraEditorSettings.h
+++ b/Engine/Plugins/FX/Niagara/Source/NiagaraEditor/Public/NiagaraEditorSettings.h
@@ -39,15 +39,15 @@ public:
 
 	/** Niagara script to duplicate as the base of all new script assets created. */
 	UPROPERTY(config, EditAnywhere, Category = Niagara)
-	FStringAssetReference DefaultDynamicInputScript;
+	FSoftObjectPath DefaultDynamicInputScript;
 
 	/** Niagara script to duplicate as the base of all new script assets created. */
 	UPROPERTY(config, EditAnywhere, Category = Niagara)
-	FStringAssetReference DefaultFunctionScript;
+	FSoftObjectPath DefaultFunctionScript;
 
 	/** Niagara script to duplicate as the base of all new script assets created. */
 	UPROPERTY(config, EditAnywhere, Category = Niagara)
-	FStringAssetReference DefaultModuleScript;
+	FSoftObjectPath DefaultModuleScript;
 
 	/** Shortcut key bindings that if held down while doing a mouse click, will spawn the specified type of Niagara node.*/
 	UPROPERTY(config, EditAnywhere, Category = Niagara)
diff --git a/Engine/Plugins/FX/Niagara/Source/NiagaraEditor/Public/NiagaraScriptFactoryNew.h b/Engine/Plugins/FX/Niagara/Source/NiagaraEditor/Public/NiagaraScriptFactoryNew.h
--- a/Engine/Plugins/FX/Niagara/Source/NiagaraEditor/Public/NiagaraScriptFactoryNew.h
+++ b/Engine/Plugins/FX/Niagara/Source/NiagaraEditor/Public/NiagaraScriptFactoryNew.h
@@ -25,7 +25,7 @@ public:
 protected:
 
 	/** Gets the script asset to create a new one from that based on the script usage type. */
-	const FStringAssetReference& GetDefaultScriptFromSettings(const class UNiagaraEditorSettings* Settings);
+	const FSoftObjectPath& GetDefaultScriptFromSettings(const class UNiagaraEditorSettings* Settings);
 
 	/** Give child factory classes the ability to define which asset action it represents. */
 	virtual FName GetAssetTypeActionName() const { return NAME_None; };
diff --git a/Engine/Source/Editor/DetailCustomizations/Private/SoftObjectPathCustomization.h b/Engine/Source/Editor/DetailCustomizations/Private/SoftObjectPathCustomization.h
--- a/Engine/Source/Editor/DetailCustomizations/Private/SoftObjectPathCustomization.h
+++ b/Engine/Source/Editor/DetailCustomizations/Private/SoftObjectPathCustomization.h
@@ -8,7 +8,7 @@
 class IPropertyHandle;
 
 /**
- * Customizes a string asset reference to look like a UObject property
+ * Customizes a soft object path to look like a UObject property
  */
 class FSoftObjectPathCustomization : public IPropertyTypeCustomization
 {
diff --git a/Engine/Source/Editor/PhysicsAssetEditor/Private/PhysicsAssetEditorSharedData.cpp b/Engine/Source/Editor/PhysicsAssetEditor/Private/PhysicsAssetEditorSharedData.cpp
--- a/Engine/Source/Editor/PhysicsAssetEditor/Private/PhysicsAssetEditorSharedData.cpp
+++ b/Engine/Source/Editor/PhysicsAssetEditor/Private/PhysicsAssetEditorSharedData.cpp
@@ -84,7 +84,7 @@ void FPhysicsAssetEditorSharedData::Initialize(const TSharedRef<IPersonaPreviewS
 
 	EditorSkelComp = nullptr;
 	PhysicalAnimationComponent = nullptr;
-	FStringAssetReference PreviewMeshStringRef = PhysicsAsset->PreviewSkeletalMesh.ToSoftObjectPath();
+	FSoftObjectPath PreviewMeshStringRef = PhysicsAsset->PreviewSkeletalMesh.ToSoftObjectPath();
 
 	// Look for body setups with no shapes (how does this happen?).
 	// If we find one- just bang on a default box.
diff --git a/Engine/Source/Editor/PropertyEditor/Private/PropertyHandleImpl.cpp b/Engine/Source/Editor/PropertyEditor/Private/PropertyHandleImpl.cpp
--- a/Engine/Source/Editor/PropertyEditor/Private/PropertyHandleImpl.cpp
+++ b/Engine/Source/Editor/PropertyEditor/Private/PropertyHandleImpl.cpp
@@ -3635,7 +3635,7 @@ FPropertyAccess::Result FPropertyHandleObject::SetValue(const FAssetData& NewVal
 	{
 		if (!PropertyNode->GetProperty()->IsA(USoftObjectProperty::StaticClass()))
 		{
-			// Make sure the asset is loaded if we are not a string asset reference.
+			// Make sure the asset is loaded if we are not a soft reference
 			NewValue.GetAsset();
 		}
 
diff --git a/Engine/Source/Editor/UnrealEd/Private/CookOnTheFlyServer.cpp b/Engine/Source/Editor/UnrealEd/Private/CookOnTheFlyServer.cpp
--- a/Engine/Source/Editor/UnrealEd/Private/CookOnTheFlyServer.cpp
+++ b/Engine/Source/Editor/UnrealEd/Private/CookOnTheFlyServer.cpp
@@ -6860,7 +6860,7 @@ void UCookOnTheFlyServer::StartCookByTheBook( const FCookByTheBookStartupOptions
 	TArray<FName> FilesInPath;
 	TSet<FName> StartupSoftObjectPackages;
 
-	// Get the list of string asset references, for both empty package and all startup packages
+	// Get the list of soft references, for both empty package and all startup packages
 	GRedirectCollector.ProcessSoftObjectPathPackageList(NAME_None, false, StartupSoftObjectPackages);
 
 	for (const FName& StartupPackage : CookByTheBookOptions->StartupPackages)
diff --git a/Engine/Source/Editor/UnrealEd/Private/PlayLevel.cpp b/Engine/Source/Editor/UnrealEd/Private/PlayLevel.cpp
--- a/Engine/Source/Editor/UnrealEd/Private/PlayLevel.cpp
+++ b/Engine/Source/Editor/UnrealEd/Private/PlayLevel.cpp
@@ -3642,7 +3642,7 @@ UWorld* UEditorEngine::CreatePIEWorldByDuplication(FWorldContext &WorldContext,
 		// Reset any GUID fixups with lazy pointers
 		FLazyObjectPtr::ResetPIEFixups();
 
-		// Prepare string asset references for fixup
+		// Prepare soft object paths for fixup
 		FSoftObjectPath::AddPIEPackageName(FName(*PlayWorldMapName));
 		for (ULevelStreaming* StreamingLevel : InWorld->GetStreamingLevels())
 		{
diff --git a/Engine/Source/Runtime/Core/Public/Serialization/Archive.h b/Engine/Source/Runtime/Core/Public/Serialization/Archive.h
--- a/Engine/Source/Runtime/Core/Public/Serialization/Archive.h
+++ b/Engine/Source/Runtime/Core/Public/Serialization/Archive.h
@@ -255,7 +255,7 @@ public:
 	/**
 	 * Serializes soft object paths from or into this archive.
 	 *
-	 * @param Value String asset reference to serialize.
+	 * @param Value Soft object path to serialize.
 	 * @return This instance.
 	 */
 	virtual FArchive& operator<<(struct FSoftObjectPath& Value);
diff --git a/Engine/Source/Runtime/CoreUObject/Private/Misc/PackageName.cpp b/Engine/Source/Runtime/CoreUObject/Private/Misc/PackageName.cpp
--- a/Engine/Source/Runtime/CoreUObject/Private/Misc/PackageName.cpp
+++ b/Engine/Source/Runtime/CoreUObject/Private/Misc/PackageName.cpp
@@ -1031,12 +1031,12 @@ FString FPackageName::GetNormalizedObjectPath(const FString& ObjectPath)
 	{
 		FString LongPath;
 
-		UE_LOG(LogPackageName, Warning, TEXT("String asset reference \"%s\" is in short form, which is unsupported and -- even if valid -- resolving it will be really slow."), *ObjectPath);
+		UE_LOG(LogPackageName, Warning, TEXT("Asset path \"%s\" is in short form, which is unsupported and -- even if valid -- resolving it will be really slow."), *ObjectPath);
 		UE_LOG(LogPackageName, Warning, TEXT("Please consider resaving package in order to speed-up loading."));
 		
 		if (!FPackageName::TryConvertShortPackagePathToLongInObjectPath(ObjectPath, LongPath))
 		{
-			UE_LOG(LogPackageName, Warning, TEXT("String asset reference \"%s\" could not be resolved."), *ObjectPath);
+			UE_LOG(LogPackageName, Warning, TEXT("Asset path \"%s\" could not be resolved."), *ObjectPath);
 		}
 
 		return LongPath;
diff --git a/Engine/Source/Runtime/CoreUObject/Private/Misc/RedirectCollector.cpp b/Engine/Source/Runtime/CoreUObject/Private/Misc/RedirectCollector.cpp
--- a/Engine/Source/Runtime/CoreUObject/Private/Misc/RedirectCollector.cpp
+++ b/Engine/Source/Runtime/CoreUObject/Private/Misc/RedirectCollector.cpp
@@ -105,7 +105,7 @@ void FRedirectCollector::ResolveAllSoftObjectPaths(FName FilterPackage)
 
 			if (ToLoad.Len() > 0 )
 			{
-				UE_LOG(LogRedirectors, Verbose, TEXT("String Asset Reference '%s'"), *ToLoad);
+				UE_LOG(LogRedirectors, Verbose, TEXT("Resolving Soft Object Path '%s'"), *ToLoad);
 				UE_CLOG(RefFilenameAndProperty.GetProperty().ToString().Len(), LogRedirectors, Verbose, TEXT("    Referenced by '%s'"), *RefFilenameAndProperty.GetProperty().ToString());
 
 				int32 DotIndex = ToLoad.Find(TEXT("."));
@@ -131,7 +131,7 @@ void FRedirectCollector::ResolveAllSoftObjectPaths(FName FilterPackage)
 				else
 				{
 					const FString Referencer = RefFilenameAndProperty.GetProperty().ToString().Len() ? RefFilenameAndProperty.GetProperty().ToString() : TEXT("Unknown");
-					UE_LOG(LogRedirectors, Warning, TEXT("String Asset Reference '%s' was not found! (Referencer '%s')"), *ToLoad, *Referencer);
+					UE_LOG(LogRedirectors, Warning, TEXT("Soft Object Path '%s' was not found when resolving paths! (Referencer '%s')"), *ToLoad, *Referencer);
 				}
 			}
 		}
diff --git a/Engine/Source/Runtime/CoreUObject/Public/Misc/AssetRegistryInterface.h b/Engine/Source/Runtime/CoreUObject/Public/Misc/AssetRegistryInterface.h
--- a/Engine/Source/Runtime/CoreUObject/Public/Misc/AssetRegistryInterface.h
+++ b/Engine/Source/Runtime/CoreUObject/Public/Misc/AssetRegistryInterface.h
@@ -10,7 +10,7 @@ namespace EAssetRegistryDependencyType
 {
 	enum Type
 	{
-		// Dependencies which don't need to be loaded for the object to be used (i.e. string asset references)
+		// Dependencies which don't need to be loaded for the object to be used (i.e. soft object paths)
 		Soft = 0x01,
 
 		// Dependencies which are required for correct usage of the source asset, and must be loaded at the same time
diff --git a/Engine/Source/Runtime/CoreUObject/Public/Misc/RedirectCollector.h b/Engine/Source/Runtime/CoreUObject/Public/Misc/RedirectCollector.h
--- a/Engine/Source/Runtime/CoreUObject/Public/Misc/RedirectCollector.h
+++ b/Engine/Source/Runtime/CoreUObject/Public/Misc/RedirectCollector.h
@@ -15,7 +15,7 @@ class COREUOBJECT_API FRedirectCollector
 {
 private:
 	
-	/** Helper struct for string asset reference tracking */
+	/** Helper struct for soft object path tracking */
 	struct FPackagePropertyPair
 	{
 		FPackagePropertyPair() : bReferencedByEditorOnlyProperty(false) {}
@@ -129,7 +129,7 @@ private:
 	/** A map of assets referenced by soft object paths, with the key being the asset being referenced and the value equal to the package with the reference */
 	TMultiMap<FName, FPackagePropertyPair> SoftObjectPathMap;
 
-	/** When saving, apply this remapping to all string asset references */
+	/** When saving, apply this remapping to all soft object paths */
 	TMap<FName, FName> AssetPathRedirectionMap;
 
 	/** For SoftObjectPackageMap map */
diff --git a/Engine/Source/Runtime/CoreUObject/Public/UObject/SoftObjectPath.h b/Engine/Source/Runtime/CoreUObject/Public/UObject/SoftObjectPath.h
--- a/Engine/Source/Runtime/CoreUObject/Public/UObject/SoftObjectPath.h
+++ b/Engine/Source/Runtime/CoreUObject/Public/UObject/SoftObjectPath.h
@@ -154,10 +154,10 @@ struct COREUOBJECT_API FSoftObjectPath
 	/** Serializes the internal path and also handles save/PIE fixups. Call this from the archiver overrides */
 	void SerializePath(FArchive& Ar);
 
-	/** Fixes up string asset reference for saving, call if saving with a method that skips SerializePath. This can modify the path, it will return true if it was modified */
+	/** Fixes up path for saving, call if saving with a method that skips SerializePath. This can modify the path, it will return true if it was modified */
 	bool PreSavePath(bool* bReportSoftObjectPathRedirects = nullptr);
 
-	/** Handles when a string asset reference has been loaded, call if loading with a method that skips SerializePath. This does not modify path but might call callbacks */
+	/** Handles when a path has been loaded, call if loading with a method that skips SerializePath. This does not modify path but might call callbacks */
 	void PostLoadPath() const;
 
 	/** Fixes up this SoftObjectPath to add the PIE prefix depending on what is currently active, returns true if it was modified. The overload that takes an explicit PIE instance is preferred, if it's available. */
@@ -323,21 +323,21 @@ public:
 	 * Returns the current serialization options that were added using SerializationScope or LinkerLoad
 	 *
 	 * @param OutPackageName Package that this string asset belongs to
-	 * @param OutPropertyName Property that this string asset reference belongs to
+	 * @param OutPropertyName Property that this path belongs to
 	 * @param OutCollectType Type of collecting that should be done
 	 * @param Archive The FArchive that is serializing this path if known. If null it will check FUObjectThreadContext
 	 */
 	bool GetSerializationOptions(FName& OutPackageName, FName& OutPropertyName, ESoftObjectPathCollectType& OutCollectType, ESoftObjectPathSerializeType& OutSerializeType, FArchive* Archive = nullptr) const;
 };
 
-/** Helper class to set and restore serialization options for string asset references */
+/** Helper class to set and restore serialization options for soft object paths */
 struct FSoftObjectPathSerializationScope
 {
 	/** 
-	 * Create a new serialization scope, which affects the way that string asset references are saved
+	 * Create a new serialization scope, which affects the way that soft object paths are saved
 	 *
 	 * @param SerializingPackageName Package that this string asset belongs to
-	 * @param SerializingPropertyName Property that this string asset reference belongs to
+	 * @param SerializingPropertyName Property that this path belongs to
 	 * @param CollectType Set type of collecting that should be done, can be used to disable tracking entirely
 	 */
 	FSoftObjectPathSerializationScope(FName SerializingPackageName, FName SerializingPropertyName, ESoftObjectPathCollectType CollectType, ESoftObjectPathSerializeType SerializeType)
diff --git a/Engine/Source/Runtime/Engine/Classes/Engine/StreamableManager.h b/Engine/Source/Runtime/Engine/Classes/Engine/StreamableManager.h
--- a/Engine/Source/Runtime/Engine/Classes/Engine/StreamableManager.h
+++ b/Engine/Source/Runtime/Engine/Classes/Engine/StreamableManager.h
@@ -307,7 +307,7 @@ struct ENGINE_API FStreamableManager : public FGCObject
 	/** Returns true if all pending async loads have finished for this target */
 	bool IsAsyncLoadComplete(const FSoftObjectPath& Target) const;
 
-	/** This will release any managed active handles pointing to the target string asset reference, even if they include other requested assets in the same load */
+	/** This will release any managed active handles pointing to the target soft object path, even if they include other requested assets in the same load */
 	void Unload(const FSoftObjectPath& Target);
 
 	/** Checks for any redirectors that were previously loaded, and returns the redirected target if found. This will not handle redirects that it doesn't yet know about */
diff --git a/Engine/Source/Runtime/Engine/Private/AssetManager.cpp b/Engine/Source/Runtime/Engine/Private/AssetManager.cpp
--- a/Engine/Source/Runtime/Engine/Private/AssetManager.cpp
+++ b/Engine/Source/Runtime/Engine/Private/AssetManager.cpp
@@ -2984,7 +2984,7 @@ const TMap<int32, FAssetManagerChunkInfo>& UAssetManager::GetChunkManagementMap(
 
 void UAssetManager::ApplyPrimaryAssetLabels()
 {
-	// Load all of them off disk. Turn off string asset reference tracking to avoid them getting cooked
+	// Load all of them off disk. Turn off soft object path tracking to avoid them getting cooked
 	FSoftObjectPathSerializationScope SerializationScope(NAME_None, NAME_None, ESoftObjectPathCollectType::NeverCollect, ESoftObjectPathSerializeType::AlwaysSerialize);
 
 	TSharedPtr<FStreamableHandle> Handle = LoadPrimaryAssetsWithType(PrimaryAssetLabelType);
diff --git a/Engine/Source/Runtime/PropertyPath/Public/PropertyTypeCompatibility.h b/Engine/Source/Runtime/PropertyPath/Public/PropertyTypeCompatibility.h
--- a/Engine/Source/Runtime/PropertyPath/Public/PropertyTypeCompatibility.h
+++ b/Engine/Source/Runtime/PropertyPath/Public/PropertyTypeCompatibility.h
@@ -260,15 +260,15 @@ inline bool IsConcreteTypeCompatibleWithReflectedType_Impl<FInt32Interval>(UProp
 }
 
 template<>
-inline bool IsConcreteTypeCompatibleWithReflectedType_Impl<FStringAssetReference>(UProperty* Property)
+inline bool IsConcreteTypeCompatibleWithReflectedType_Impl<FSoftObjectPath>(UProperty* Property)
 {
-	return IsConcreteTypeCompatibleWithReflectedType_BuiltInStruct<FStringAssetReference>(Property);
+	return IsConcreteTypeCompatibleWithReflectedType_BuiltInStruct<FSoftObjectPath>(Property);
 }
 
 template<>
-inline bool IsConcreteTypeCompatibleWithReflectedType_Impl<FStringClassReference>(UProperty* Property)
+inline bool IsConcreteTypeCompatibleWithReflectedType_Impl<FSoftClassPath>(UProperty* Property)
 {
-	return IsConcreteTypeCompatibleWithReflectedType_BuiltInStruct<FStringClassReference>(Property);
+	return IsConcreteTypeCompatibleWithReflectedType_BuiltInStruct<FSoftClassPath>(Property);
 }
 
 template<>
