The cooker will now attempt to collect garbage when the engine is getting close to running out of available UObject indices (reimplementing CL #4581474 from Dev-Core). Bumped max UObject count in the editor.

diff --git a/Engine/Config/BaseEngine.ini b/Engine/Config/BaseEngine.ini
--- a/Engine/Config/BaseEngine.ini
+++ b/Engine/Config/BaseEngine.ini
@@ -765,7 +765,7 @@ gc.NumRetriesBeforeForcingGC=10
 gc.AllowParallelGC=True
 ; pick a fractional number to keep phase shifting and avoid collisions
 gc.TimeBetweenPurgingPendingKillObjects=61.1
-gc.MaxObjectsInEditor=8388607
+gc.MaxObjectsInEditor=16777216
 gc.IncrementalBeginDestroyEnabled=True
 gc.CreateGCClusters=True
 gc.MergeGCClusters=False
diff --git a/Engine/Source/Editor/UnrealEd/Classes/CookOnTheSide/CookOnTheFlyServer.h b/Engine/Source/Editor/UnrealEd/Classes/CookOnTheSide/CookOnTheFlyServer.h
--- a/Engine/Source/Editor/UnrealEd/Classes/CookOnTheSide/CookOnTheFlyServer.h
+++ b/Engine/Source/Editor/UnrealEd/Classes/CookOnTheSide/CookOnTheFlyServer.h
@@ -135,6 +135,8 @@ private:
 	int32 MaxNumPackagesBeforePartialGC;
 	/** Max number of concurrent shader jobs reducing this too low will increase cook time */
 	int32 MaxConcurrentShaderJobs;
+	/** Min number of free UObject indices before the cooker should partial gc */
+	int32 MinFreeUObjectIndicesBeforeGC;
 
 	ECookInitializationFlags CookFlags = ECookInitializationFlags::None;
 	TUniquePtr<class FSandboxPlatformFile> SandboxFile;
diff --git a/Engine/Source/Editor/UnrealEd/Private/CookOnTheFlyServer.cpp b/Engine/Source/Editor/UnrealEd/Private/CookOnTheFlyServer.cpp
--- a/Engine/Source/Editor/UnrealEd/Private/CookOnTheFlyServer.cpp
+++ b/Engine/Source/Editor/UnrealEd/Private/CookOnTheFlyServer.cpp
@@ -3205,7 +3205,13 @@ bool UCookOnTheFlyServer::HasExceededMaxMemory() const
 		return true;
 	}
 
-	
+#if UE_GC_TRACK_OBJ_AVAILABLE
+	if (GUObjectArray.GetObjectArrayEstimatedAvailable() < MinFreeUObjectIndicesBeforeGC)
+	{
+		UE_LOG(LogCook, Display, TEXT("Running out of available UObject indices (%d remaining)"), GUObjectArray.GetObjectArrayEstimatedAvailable());
+		return true;
+	}
+#endif // UE_GC_TRACK_OBJ_AVAILABLE
 
 	return false;
 }
@@ -4084,6 +4090,10 @@ void UCookOnTheFlyServer::Initialize( ECookMode::Type DesiredCookMode, ECookInit
 	MinMemoryBeforeGC = MinMemoryBeforeGCInMB * 1024LL * 1024LL;
 	MinMemoryBeforeGC = FMath::Min(MaxMemoryAllowance, MinMemoryBeforeGC);
 
+	MinFreeUObjectIndicesBeforeGC = 5000;
+	GConfig->GetInt(TEXT("CookSettings"), TEXT("MinFreeUObjectIndicesBeforeGC"), MinFreeUObjectIndicesBeforeGC, GEditorIni);
+	MinFreeUObjectIndicesBeforeGC = FMath::Max(MinFreeUObjectIndicesBeforeGC, 0);
+
 	int32 MinFreeMemoryInMB = 0;
 	GConfig->GetInt(TEXT("CookSettings"), TEXT("MinFreeMemory"), MinFreeMemoryInMB, GEditorIni);
 	MinFreeMemoryInMB = FMath::Max(MinFreeMemoryInMB, 0);
diff --git a/Engine/Source/Runtime/CoreUObject/Public/UObject/UObjectArray.h b/Engine/Source/Runtime/CoreUObject/Public/UObject/UObjectArray.h
--- a/Engine/Source/Runtime/CoreUObject/Public/UObject/UObjectArray.h
+++ b/Engine/Source/Runtime/CoreUObject/Public/UObject/UObjectArray.h
@@ -272,6 +272,17 @@ public:
 	{
 		return NumElements;
 	}
+
+	/**
+	* Return the number max capacity of the array
+	* Thread safe, but you know, someone might have added more elements before this even returns
+	* @return	the maximum number of elements in the array
+	**/
+	FORCEINLINE int32 Capacity() const TSAN_SAFE
+	{
+		return MaxElements;
+	}
+
 	/**
 	* Return if this index is valid
 	* Thread safe, if it is valid now, it is valid forever. Other threads might be adding during this call.
@@ -426,6 +437,17 @@ public:
 	{
 		return NumElements;
 	}
+
+	/**
+	* Return the number max capacity of the array
+	* Thread safe, but you know, someone might have added more elements before this even returns
+	* @return	the maximum number of elements in the array
+	**/
+	FORCEINLINE int32 Capacity() const TSAN_SAFE
+	{
+		return MaxElements;
+	}
+
 	/**
 	* Return if this index is valid
 	* Thread safe, if it is valid now, it is valid forever. Other threads might be adding during this call.
@@ -802,10 +824,18 @@ public:
 	 *
 	 * @return	The number of objects claimed
 	 */
-	FORCEINLINE int32 GetObjectArrayNumMinusAvailable()
+	int32 GetObjectArrayNumMinusAvailable()
 	{
 		return ObjObjects.Num() - ObjAvailableCount.GetValue();
 	}
+
+	/**
+	* Returns the estimated number of object indices available for allocation
+	*/
+	int32 GetObjectArrayEstimatedAvailable()
+	{
+		return ObjObjects.Capacity() - GetObjectArrayNumMinusAvailable();
+	}
 #endif
 
 	/**
