Making sure the editor custom version is serialized with package file summary in cases where the package doesn't have UMetaData object. Otherwise the package may get corrupted because the localization gather code already forces the editor custom version on the linker after the file summary has been serialized. Fixes package redirects not being able to save.

diff --git a/Engine/Source/Runtime/CoreUObject/Private/UObject/SavePackage.cpp b/Engine/Source/Runtime/CoreUObject/Private/UObject/SavePackage.cpp
--- a/Engine/Source/Runtime/CoreUObject/Private/UObject/SavePackage.cpp
+++ b/Engine/Source/Runtime/CoreUObject/Private/UObject/SavePackage.cpp
@@ -3947,6 +3947,14 @@ FSavePackageResultStruct UPackage::Save(UPackage* InOuter, UObject* Base, EObjec
 				}
 #endif
 
+				if (!(Linker->Summary.PackageFlags & PKG_FilterEditorOnly))
+				{
+					// The Editor version is used as part of the check to see if a package is too old to use the gather cache, so we always have to add it if we have gathered loc for this asset
+					// We need to set the editor custom version before we copy the version container to the summary, otherwise we may end up with corrupt assets
+					// because we later do it on the Linker when actually gathering loc data
+					ExportTaggerArchive.UsingCustomVersion(FEditorObjectVersion::GUID);
+				}
+
 				// Use the custom versions we had previously gleaned from the export tag pass
 				Linker->Summary.SetCustomVersionContainer(ExportTaggerArchive.GetCustomVersions());
 
@@ -4537,6 +4545,7 @@ FSavePackageResultStruct UPackage::Save(UPackage* InOuter, UObject* Base, EObjec
 				if (!(Linker->Summary.PackageFlags & PKG_FilterEditorOnly))
 				{
 					// The Editor version is used as part of the check to see if a package is too old to use the gather cache, so we always have to add it if we have gathered loc for this asset
+					// Note that using custom version here only works because we already added it to the export tagger before the package summary was serialized
 					Linker->UsingCustomVersion(FEditorObjectVersion::GUID);
 
 					bool bCanCacheGatheredText = false;
