Stabilize the keys of FTexts imported from raw strings in a Data Table

diff --git a/Engine/Source/Runtime/Core/Private/Internationalization/Text.cpp b/Engine/Source/Runtime/Core/Private/Internationalization/Text.cpp
--- a/Engine/Source/Runtime/Core/Private/Internationalization/Text.cpp
+++ b/Engine/Source/Runtime/Core/Private/Internationalization/Text.cpp
@@ -1044,6 +1044,11 @@ bool FText::IsCultureInvariant() const
 	return (Flags & ETextFlag::CultureInvariant) != 0;
 }
 
+bool FText::IsInitializedFromString() const
+{
+	return (Flags & ETextFlag::InitializedFromString) != 0;
+}
+
 bool FText::IsFromStringTable() const
 {
 	return TextData->GetTextHistory().GetType() == ETextHistoryType::StringTableEntry;
diff --git a/Engine/Source/Runtime/Core/Public/Internationalization/Text.h b/Engine/Source/Runtime/Core/Public/Internationalization/Text.h
--- a/Engine/Source/Runtime/Core/Public/Internationalization/Text.h
+++ b/Engine/Source/Runtime/Core/Public/Internationalization/Text.h
@@ -523,6 +523,7 @@ public:
 
 	bool IsTransient() const;
 	bool IsCultureInvariant() const;
+	bool IsInitializedFromString() const;
 	bool IsFromStringTable() const;
 
 	bool ShouldGatherForLocalization() const;
diff --git a/Engine/Source/Runtime/CoreUObject/Private/Internationalization/StabilizeLocalizationKeys.cpp b/Engine/Source/Runtime/CoreUObject/Private/Internationalization/StabilizeLocalizationKeys.cpp
new file mode 100644
--- /dev/null
+++ b/Engine/Source/Runtime/CoreUObject/Private/Internationalization/StabilizeLocalizationKeys.cpp
@@ -0,0 +1,191 @@
+// Copyright 1998-2018 Epic Games, Inc. All Rights Reserved.
+
+#include "Internationalization/StabilizeLocalizationKeys.h"
+#include "UObject/TextProperty.h"
+
+#if WITH_EDITOR
+
+void StabilizeLocalizationKeys::StabilizeLocalizationKeysForProperty(UProperty* InProp, void* InPropData, const FString& InNamespace, const FString& InKeyRoot, const bool bAppendPropertyNameToKey)
+{
+	auto ShouldStabilizeInnerProperty = [](UProperty* InnerProperty) -> bool
+	{
+		return InnerProperty->IsA<UTextProperty>() || InnerProperty->IsA<UStructProperty>();
+	};
+
+	auto GetPropertyName = [InProp]() -> FString
+	{
+		// UserDefinedStructs should use the display name rather than the raw name, since the raw name doesn't match what the user generally sees or entered as the property name
+		for (UStruct* PropOwnerStruct = InProp->GetOwnerStruct(); PropOwnerStruct; PropOwnerStruct = PropOwnerStruct->GetSuperStruct())
+		{
+			static const FName UserDefinedStructName = TEXT("UserDefinedStruct");
+			if (PropOwnerStruct->GetClass()->GetFName() == UserDefinedStructName)
+			{
+				static const FName DisplayNameKey = TEXT("DisplayName");
+				return InProp->HasMetaData(DisplayNameKey) ? InProp->GetMetaData(DisplayNameKey) : InProp->GetName();
+			}
+		}
+		return InProp->GetName();
+	};
+
+	const FString PropKeyRoot = bAppendPropertyNameToKey ? FString::Printf(TEXT("%s_%s"), *InKeyRoot, *GetPropertyName()) : InKeyRoot;
+
+	if (UTextProperty* TextProp = Cast<UTextProperty>(InProp))
+	{
+		for (int32 ArrIndex = 0; ArrIndex < TextProp->ArrayDim; ++ArrIndex)
+		{
+			void* PropValueData = ((uint8*)InPropData) + (TextProp->ElementSize * ArrIndex);
+
+			FText* TextValuePtr = TextProp->GetPropertyValuePtr(PropValueData);
+			check(TextValuePtr);
+
+			if (TextValuePtr->IsInitializedFromString())
+			{
+				if (TextProp->ArrayDim > 1)
+				{
+					*TextValuePtr = FText::ChangeKey(InNamespace, FString::Printf(TEXT("%s_Index%d"), *PropKeyRoot, ArrIndex), *TextValuePtr);
+				}
+				else
+				{
+					*TextValuePtr = FText::ChangeKey(InNamespace, PropKeyRoot, *TextValuePtr);
+				}
+			}
+		}
+		return;
+	}
+
+	if (UStructProperty* StructProp = Cast<UStructProperty>(InProp))
+	{
+		for (int32 ArrIndex = 0; ArrIndex < StructProp->ArrayDim; ++ArrIndex)
+		{
+			void* PropValueData = ((uint8*)InPropData) + (StructProp->ElementSize * ArrIndex);
+
+			if (StructProp->ArrayDim > 1)
+			{
+				StabilizeLocalizationKeysForStruct(
+					StructProp->Struct,
+					PropValueData,
+					InNamespace,
+					FString::Printf(TEXT("%s_Index%d"), *PropKeyRoot, ArrIndex)
+					);
+			}
+			else
+			{
+				StabilizeLocalizationKeysForStruct(
+					StructProp->Struct,
+					PropValueData,
+					InNamespace,
+					PropKeyRoot
+					);
+			}
+		}
+		return;
+	}
+
+	if (UArrayProperty* ArrayProp = Cast<UArrayProperty>(InProp))
+	{
+		if (ShouldStabilizeInnerProperty(ArrayProp->Inner))
+		{
+			FScriptArrayHelper ScriptArrayHelper(ArrayProp, InPropData);
+
+			const int32 ElementCount = ScriptArrayHelper.Num();
+			for (int32 ArrIndex = 0; ArrIndex < ElementCount; ++ArrIndex)
+			{
+				StabilizeLocalizationKeysForProperty(
+					ArrayProp->Inner,
+					ScriptArrayHelper.GetRawPtr(ArrIndex),
+					InNamespace,
+					FString::Printf(TEXT("%s_Index%d"), *PropKeyRoot, ArrIndex),
+					/*bAppendPropertyNameToKey*/false
+					);
+			}
+		}
+		return;
+	}
+
+	if (USetProperty* SetProp = Cast<USetProperty>(InProp))
+	{
+		if (ShouldStabilizeInnerProperty(SetProp->ElementProp))
+		{
+			FScriptSetHelper ScriptSetHelper(SetProp, InPropData);
+
+			const int32 ElementCount = ScriptSetHelper.Num();
+			for (int32 RawIndex = 0, ElementIndex = 0; ElementIndex < ElementCount; ++RawIndex)
+			{
+				if (!ScriptSetHelper.IsValidIndex(RawIndex))
+				{
+					continue;
+				}
+
+				StabilizeLocalizationKeysForProperty(
+					SetProp->ElementProp,
+					ScriptSetHelper.GetElementPtr(RawIndex),
+					InNamespace,
+					FString::Printf(TEXT("%s_Index%d"), *PropKeyRoot, ElementIndex),
+					/*bAppendPropertyNameToKey*/false
+					);
+
+				++ElementIndex;
+			}
+
+			ScriptSetHelper.Rehash();
+		}
+		return;
+	}
+
+	if (UMapProperty* MapProp = Cast<UMapProperty>(InProp))
+	{
+		if (ShouldStabilizeInnerProperty(MapProp->KeyProp) || ShouldStabilizeInnerProperty(MapProp->ValueProp))
+		{
+			FScriptMapHelper ScriptMapHelper(MapProp, InPropData);
+
+			const int32 ElementCount = ScriptMapHelper.Num();
+			for (int32 RawIndex = 0, ElementIndex = 0; ElementIndex < ElementCount; ++RawIndex)
+			{
+				if (!ScriptMapHelper.IsValidIndex(RawIndex))
+				{
+					continue;
+				}
+
+				if (ShouldStabilizeInnerProperty(MapProp->KeyProp))
+				{
+					StabilizeLocalizationKeysForProperty(
+						MapProp->KeyProp,
+						ScriptMapHelper.GetKeyPtr(RawIndex),
+						InNamespace,
+						FString::Printf(TEXT("%s_KeyIndex%d"), *PropKeyRoot, ElementIndex),
+						/*bAppendPropertyNameToKey*/false
+						);
+				}
+
+				if (ShouldStabilizeInnerProperty(MapProp->ValueProp))
+				{
+					StabilizeLocalizationKeysForProperty(
+						MapProp->ValueProp,
+						ScriptMapHelper.GetValuePtr(RawIndex),
+						InNamespace,
+						FString::Printf(TEXT("%s_ValueIndex%d"), *PropKeyRoot, ElementIndex),
+						/*bAppendPropertyNameToKey*/false
+						);
+				}
+
+				++ElementIndex;
+			}
+
+			if (ShouldStabilizeInnerProperty(MapProp->KeyProp))
+			{
+				ScriptMapHelper.Rehash();
+			}
+		}
+		return;
+	}
+}
+
+void StabilizeLocalizationKeys::StabilizeLocalizationKeysForStruct(UStruct* InStruct, void* InStructData, const FString& InNamespace, const FString& InKeyRoot)
+{
+	for (TFieldIterator<UProperty> It(InStruct); It; ++It)
+	{
+		StabilizeLocalizationKeysForProperty(*It, It->ContainerPtrToValuePtr<void>(InStructData), InNamespace, InKeyRoot);
+	}
+}
+
+#endif	// WITH_EDITOR
diff --git a/Engine/Source/Runtime/CoreUObject/Public/Internationalization/StabilizeLocalizationKeys.h b/Engine/Source/Runtime/CoreUObject/Public/Internationalization/StabilizeLocalizationKeys.h
new file mode 100644
--- /dev/null
+++ b/Engine/Source/Runtime/CoreUObject/Public/Internationalization/StabilizeLocalizationKeys.h
@@ -0,0 +1,26 @@
+// Copyright 1998-2018 Epic Games, Inc. All Rights Reserved.
+
+#pragma once
+
+#include "CoreMinimal.h"
+
+namespace StabilizeLocalizationKeys
+{
+
+#if WITH_EDITOR
+
+/**
+ * Convert any text that has been initialized from a raw string to have a deterministic key based on the given key root and the property name.
+ * @note This function will recurse into arrays, sets, maps, and sub-structures.
+ */
+COREUOBJECT_API void StabilizeLocalizationKeysForProperty(UProperty* InProp, void* InPropData, const FString& InNamespace, const FString& InKeyRoot, const bool bAppendPropertyNameToKey = true);
+
+/**
+ * Walk through the struct and convert any text that has been initialized from a raw string to have a deterministic key based on the given key root and the property name.
+ * @note This function will recurse into arrays, sets, maps, and sub-structures.
+ */
+COREUOBJECT_API void StabilizeLocalizationKeysForStruct(UStruct* InStruct, void* InStructData, const FString& InNamespace, const FString& InKeyRoot);
+
+#endif	// WITH_EDITOR
+
+}
diff --git a/Engine/Source/Runtime/CoreUObject/Public/Internationalization/TextPackageNamespaceUtil.h b/Engine/Source/Runtime/CoreUObject/Public/Internationalization/TextPackageNamespaceUtil.h
--- a/Engine/Source/Runtime/CoreUObject/Public/Internationalization/TextPackageNamespaceUtil.h
+++ b/Engine/Source/Runtime/CoreUObject/Public/Internationalization/TextPackageNamespaceUtil.h
@@ -3,6 +3,7 @@
 #pragma once
 
 #include "CoreMinimal.h"
+#include "Internationalization/TextNamespaceUtil.h"
 
 namespace TextNamespaceUtil
 {
diff --git a/Engine/Source/Runtime/Engine/Private/DataTable.cpp b/Engine/Source/Runtime/Engine/Private/DataTable.cpp
--- a/Engine/Source/Runtime/Engine/Private/DataTable.cpp
+++ b/Engine/Source/Runtime/Engine/Private/DataTable.cpp
@@ -1,6 +1,8 @@
 // Copyright 1998-2018 Epic Games, Inc. All Rights Reserved.
 
 #include "Engine/DataTable.h"
+#include "Internationalization/TextPackageNamespaceUtil.h"
+#include "Internationalization/StabilizeLocalizationKeys.h"
 #include "Serialization/PropertyLocalizationDataGathering.h"
 #include "Serialization/ObjectWriter.h"
 #include "Serialization/ObjectReader.h"
@@ -155,12 +157,28 @@ void UDataTable::GetPreloadDependencies(TArray<UObject*>& OutDeps)
 
 void UDataTable::OnPostDataImported(TArray<FString>& OutCollectedImportProblems)
 {
-	if (RowStruct && RowStruct->IsChildOf(FTableRowBase::StaticStruct()))
+	if (RowStruct)
 	{
+		const bool bIsNativeRowStruct = RowStruct->IsChildOf(FTableRowBase::StaticStruct());
+
+		FString DataTableTextNamespace = GetName();
+#if USE_STABLE_LOCALIZATION_KEYS
+		DataTableTextNamespace = TextNamespaceUtil::BuildFullNamespace(DataTableTextNamespace, TextNamespaceUtil::EnsurePackageNamespace(this), /*bAlwaysApplyPackageNamespace*/true);
+#endif
+
 		for (const TPair<FName, uint8*>& TableRowPair : RowMap)
 		{
-			FTableRowBase* CurRow = reinterpret_cast<FTableRowBase*>(TableRowPair.Value);
-			CurRow->OnPostDataImport(this, TableRowPair.Key, OutCollectedImportProblems);
+			if (bIsNativeRowStruct)
+			{
+				FTableRowBase* CurRow = reinterpret_cast<FTableRowBase*>(TableRowPair.Value);
+				CurRow->OnPostDataImport(this, TableRowPair.Key, OutCollectedImportProblems);
+			}
+
+#if WITH_EDITOR
+			// Perform automatic fix-up on any text properties that have been imported from a raw string to assign them deterministic keys
+			// We do this after OnPostDataImport has been run on the row, as that function may perform custom fix-up logic that will fix the keys differently than the default logic
+			StabilizeLocalizationKeys::StabilizeLocalizationKeysForStruct(RowStruct, TableRowPair.Value, DataTableTextNamespace, TableRowPair.Key.ToString());
+#endif
 		}
 	}
 
