Temp log for debugging quoted string issue

diff --git a/Engine/Source/Runtime/CoreUObject/Private/UObject/Property.cpp b/Engine/Source/Runtime/CoreUObject/Private/UObject/Property.cpp
--- a/Engine/Source/Runtime/CoreUObject/Private/UObject/Property.cpp
+++ b/Engine/Source/Runtime/CoreUObject/Private/UObject/Property.cpp
@@ -366,7 +366,8 @@ const TCHAR* UPropertyHelpers::ReadToken( const TCHAR* Buffer, FString& String,
 		int32 NumCharsRead = 0;
 		if (!FParse::QuotedString(Buffer, String, &NumCharsRead))
 		{
-			UE_LOG(LogProperty, Warning, TEXT("ReadToken: Bad quoted string") );
+			UE_LOG(LogProperty, Warning, TEXT("ReadToken: Bad quoted string: |%s|"), *String );
+			UE_LOG(LogProperty, Warning, TEXT("  Buffer: |%s|"), Buffer );
 			return nullptr;
 		}
 		Buffer += NumCharsRead;
