Remove temp ensure related to Warning: ReadToken: Bad quoted string

diff --git a/Engine/Source/Runtime/CoreUObject/Private/UObject/Property.cpp b/Engine/Source/Runtime/CoreUObject/Private/UObject/Property.cpp
--- a/Engine/Source/Runtime/CoreUObject/Private/UObject/Property.cpp
+++ b/Engine/Source/Runtime/CoreUObject/Private/UObject/Property.cpp
@@ -367,7 +367,6 @@ const TCHAR* UPropertyHelpers::ReadToken( const TCHAR* Buffer, FString& String,
 		if (!FParse::QuotedString(Buffer, String, &NumCharsRead))
 		{
 			UE_LOG(LogProperty, Warning, TEXT("ReadToken: Bad quoted string: %s"), Buffer );
-			ensure(false);
 			return nullptr;
 		}
 		Buffer += NumCharsRead;
