* Adding CSV stats for overall and incremental garbage collection

diff --git a/Engine/Source/Runtime/CoreUObject/Private/UObject/GarbageCollection.cpp b/Engine/Source/Runtime/CoreUObject/Private/UObject/GarbageCollection.cpp
--- a/Engine/Source/Runtime/CoreUObject/Private/UObject/GarbageCollection.cpp
+++ b/Engine/Source/Runtime/CoreUObject/Private/UObject/GarbageCollection.cpp
@@ -1019,7 +1019,8 @@ struct FConditionalGCLock
 void IncrementalPurgeGarbage( bool bUseTimeLimit, float TimeLimit )
 {
 	SCOPED_NAMED_EVENT(IncrementalPurgeGarbage, FColor::Red);
-	DECLARE_SCOPE_CYCLE_COUNTER( TEXT( "IncrementalPurgeGarbage" ), STAT_IncrementalPurgeGarbage, STATGROUP_GC );
+	DECLARE_SCOPE_CYCLE_COUNTER(TEXT("IncrementalPurgeGarbage"), STAT_IncrementalPurgeGarbage, STATGROUP_GC);
+	CSV_SCOPED_TIMING_STAT_EXCLUSIVE(IncrementalGarbageCollection);
 
 	if (GExitPurge)
 	{
@@ -1488,6 +1489,7 @@ void CollectGarbageInternal(EObjectFlags KeepFlags, bool bPerformFullPurge)
 	SCOPE_TIME_GUARD(TEXT("Collect Garbage"));
 	SCOPED_NAMED_EVENT(CollectGarbageInternal, FColor::Red);
 	CSV_EVENT_GLOBAL(TEXT("GC"));
+	CSV_SCOPED_TIMING_STAT_EXCLUSIVE(GarbageCollection);
 
 	FGCCSyncObject::Get().ResetGCIsWaiting();
 
