* Reducing GC CSV stats down to one, incremental purge stat is redundant and subsumed by GC stat.

diff --git a/Engine/Source/Runtime/CoreUObject/Private/UObject/GarbageCollection.cpp b/Engine/Source/Runtime/CoreUObject/Private/UObject/GarbageCollection.cpp
--- a/Engine/Source/Runtime/CoreUObject/Private/UObject/GarbageCollection.cpp
+++ b/Engine/Source/Runtime/CoreUObject/Private/UObject/GarbageCollection.cpp
@@ -1020,7 +1020,6 @@ void IncrementalPurgeGarbage( bool bUseTimeLimit, float TimeLimit )
 {
 	SCOPED_NAMED_EVENT(IncrementalPurgeGarbage, FColor::Red);
 	DECLARE_SCOPE_CYCLE_COUNTER(TEXT("IncrementalPurgeGarbage"), STAT_IncrementalPurgeGarbage, STATGROUP_GC);
-	CSV_SCOPED_TIMING_STAT_EXCLUSIVE(IncrementalGarbageCollection);
 
 	if (GExitPurge)
 	{
