Fixed text properties within containers being gathered, even if they matched their default value

Also fixed some places incorrectly using Num() rather than GetMaxIndex(), as this could lead to map or set elements being missed.

diff --git a/Engine/Source/Editor/UnrealEd/Private/Commandlets/FixConflictingLocalizationKeys.cpp b/Engine/Source/Editor/UnrealEd/Private/Commandlets/FixConflictingLocalizationKeys.cpp
--- a/Engine/Source/Editor/UnrealEd/Private/Commandlets/FixConflictingLocalizationKeys.cpp
+++ b/Engine/Source/Editor/UnrealEd/Private/Commandlets/FixConflictingLocalizationKeys.cpp
@@ -178,7 +178,7 @@ bool ReKeyTextProperty(UStruct* InOuterType, void* InAddrToUpdate, const TArray<
 						// ContainerIndex is the element index, but we need the sparse index
 						int32 SparseIndex = 0;
 						{
-							const int32 ElementCount = ScriptMapHelper.Num();
+							const int32 ElementCount = ScriptMapHelper.GetMaxIndex();
 							for (int32 ElementIndex = 0; ElementIndex < ElementCount; ++SparseIndex)
 							{
 								if (ScriptMapHelper.IsValidIndex(SparseIndex))
@@ -215,7 +215,7 @@ bool ReKeyTextProperty(UStruct* InOuterType, void* InAddrToUpdate, const TArray<
 						// ContainerIndex is the element index, but we need the sparse index
 						int32 SparseIndex = 0;
 						{
-							const int32 ElementCount = ScriptSetHelper.Num();
+							const int32 ElementCount = ScriptSetHelper.GetMaxIndex();
 							for (int32 ElementIndex = 0; ElementIndex < ElementCount; ++SparseIndex)
 							{
 								if (ScriptSetHelper.IsValidIndex(SparseIndex))
diff --git a/Engine/Source/Runtime/CoreUObject/Private/Serialization/PropertyLocalizationDataGathering.cpp b/Engine/Source/Runtime/CoreUObject/Private/Serialization/PropertyLocalizationDataGathering.cpp
--- a/Engine/Source/Runtime/CoreUObject/Private/Serialization/PropertyLocalizationDataGathering.cpp
+++ b/Engine/Source/Runtime/CoreUObject/Private/Serialization/PropertyLocalizationDataGathering.cpp
@@ -268,21 +268,34 @@ void FPropertyLocalizationDataGatherer::GatherLocalizationDataFromChildTextPrope
 			const int32 ElementCount = ScriptArrayHelper.Num();
 			for(int32 j = 0; j < ElementCount; ++j)
 			{
-				GatherLocalizationDataFromChildTextProperties(PathToElement + FString::Printf(TEXT("(%d)"), j), ArrayProperty->Inner, ScriptArrayHelper.GetRawPtr(j), nullptr, ChildPropertyGatherTextFlags);
+				const uint8* ElementPtr = ScriptArrayHelper.GetRawPtr(j);
+				const uint8* DefaultElementPtr = nullptr;
+				if (DefaultElementValueAddress)
+				{
+					FScriptArrayHelper DefaultScriptArrayHelper(ArrayProperty, DefaultElementValueAddress);
+					if (DefaultScriptArrayHelper.IsValidIndex(j))
+					{
+						DefaultElementPtr = DefaultScriptArrayHelper.GetRawPtr(j);
+					}
+				}
+
+				GatherLocalizationDataFromChildTextProperties(PathToElement + FString::Printf(TEXT("(%d)"), j), ArrayProperty->Inner, ElementPtr, DefaultElementPtr, ChildPropertyGatherTextFlags);
 			}
 		}
 		// Property is a map property.
 		else if (MapProperty)
 		{
 			// Skip maps of types which can't contain any text.
-			if (CanGatherFromInnerProperty(MapProperty->KeyProp) == false && CanGatherFromInnerProperty(MapProperty->ValueProp) == false)
+			const bool bGatherMapKey = CanGatherFromInnerProperty(MapProperty->KeyProp);
+			const bool bGatherMapValue = CanGatherFromInnerProperty(MapProperty->ValueProp);
+			if (!bGatherMapKey && !bGatherMapValue)
 			{
 				continue;
 			}
 
 			// Iterate over all elements of the map.
 			FScriptMapHelper ScriptMapHelper(MapProperty, ElementValueAddress);
-			const int32 ElementCount = ScriptMapHelper.Num();
+			const int32 ElementCount = ScriptMapHelper.GetMaxIndex();
 			for(int32 j = 0, ElementIndex = 0; ElementIndex < ElementCount; ++j)
 			{
 				if (!ScriptMapHelper.IsValidIndex(j))
@@ -291,13 +304,25 @@ void FPropertyLocalizationDataGatherer::GatherLocalizationDataFromChildTextPrope
 				}
 
 				const uint8* MapPairPtr = ScriptMapHelper.GetPairPtr(j);
-				if (CanGatherFromInnerProperty(MapProperty->KeyProp))
+				const uint8* DefaultPairPtr = nullptr;
+				if (DefaultElementValueAddress)
 				{
-					GatherLocalizationDataFromChildTextProperties(PathToElement + FString::Printf(TEXT("(%d - Key)"), ElementIndex), MapProperty->KeyProp, MapPairPtr + MapProperty->MapLayout.KeyOffset, nullptr, ChildPropertyGatherTextFlags);
+					const uint8* MapKeyPtr = MapPairPtr + MapProperty->MapLayout.KeyOffset;
+					FScriptMapHelper DefaultScriptMapHelper(MapProperty, DefaultElementValueAddress);
+					DefaultPairPtr = DefaultScriptMapHelper.FindMapPairPtrFromHash(MapKeyPtr);
 				}
-				if (CanGatherFromInnerProperty(MapProperty->ValueProp))
+
+				if (bGatherMapKey)
 				{
-					GatherLocalizationDataFromChildTextProperties(PathToElement + FString::Printf(TEXT("(%d - Value)"), ElementIndex), MapProperty->ValueProp, MapPairPtr + MapProperty->MapLayout.ValueOffset, nullptr, ChildPropertyGatherTextFlags);
+					const uint8* MapKeyPtr = MapPairPtr + MapProperty->MapLayout.KeyOffset;
+					const uint8* DefaultMapKeyPtr = DefaultPairPtr ? DefaultPairPtr + MapProperty->MapLayout.KeyOffset : nullptr;
+					GatherLocalizationDataFromChildTextProperties(PathToElement + FString::Printf(TEXT("(%d - Key)"), ElementIndex), MapProperty->KeyProp, MapKeyPtr, DefaultMapKeyPtr, ChildPropertyGatherTextFlags);
+				}
+				if (bGatherMapValue)
+				{
+					const uint8* MapValuePtr = MapPairPtr + MapProperty->MapLayout.ValueOffset;
+					const uint8* DefaultMapValuePtr = DefaultPairPtr ? DefaultPairPtr + MapProperty->MapLayout.ValueOffset : nullptr;
+					GatherLocalizationDataFromChildTextProperties(PathToElement + FString::Printf(TEXT("(%d - Value)"), ElementIndex), MapProperty->ValueProp, MapValuePtr, DefaultMapValuePtr, ChildPropertyGatherTextFlags);
 				}
 				++ElementIndex;
 			}
@@ -313,7 +338,7 @@ void FPropertyLocalizationDataGatherer::GatherLocalizationDataFromChildTextPrope
 
 			// Iterate over all elements of the Set.
 			FScriptSetHelper ScriptSetHelper(SetProperty, ElementValueAddress);
-			const int32 ElementCount = ScriptSetHelper.Num();
+			const int32 ElementCount = ScriptSetHelper.GetMaxIndex();
 			for(int32 j = 0, ElementIndex = 0; ElementIndex < ElementCount; ++j)
 			{
 				if (!ScriptSetHelper.IsValidIndex(j))
@@ -322,7 +347,14 @@ void FPropertyLocalizationDataGatherer::GatherLocalizationDataFromChildTextPrope
 				}
 
 				const uint8* ElementPtr = ScriptSetHelper.GetElementPtr(j);
-				GatherLocalizationDataFromChildTextProperties(PathToElement + FString::Printf(TEXT("(%d)"), ElementIndex), SetProperty->ElementProp, ElementPtr + SetProperty->SetLayout.ElementOffset, nullptr, ChildPropertyGatherTextFlags);
+				const uint8* DefaultElementPtr = nullptr;
+				if (DefaultElementValueAddress)
+				{
+					FScriptSetHelper DefaultScriptSetHelper(SetProperty, DefaultElementValueAddress);
+					DefaultElementPtr = DefaultScriptSetHelper.FindElementPtrFromHash(ElementPtr);
+				}
+
+				GatherLocalizationDataFromChildTextProperties(PathToElement + FString::Printf(TEXT("(%d)"), ElementIndex), SetProperty->ElementProp, ElementPtr, DefaultElementPtr, ChildPropertyGatherTextFlags);
 				++ElementIndex;
 			}
 		}
diff --git a/Engine/Source/Runtime/CoreUObject/Public/UObject/UnrealType.h b/Engine/Source/Runtime/CoreUObject/Public/UObject/UnrealType.h
--- a/Engine/Source/Runtime/CoreUObject/Public/UObject/UnrealType.h
+++ b/Engine/Source/Runtime/CoreUObject/Public/UObject/UnrealType.h
@@ -3293,6 +3293,20 @@ public:
 		return Result;
 	}
 
+	/** Finds the associated pair from hash, rather than linearly searching */
+	uint8* FindMapPairPtrFromHash(const void* KeyPtr)
+	{
+		UProperty* LocalKeyPropForCapture = KeyProp;
+		int32 Index = Map->FindPairIndex(
+			KeyPtr,
+			MapLayout,
+			[LocalKeyPropForCapture](const void* ElementKey) { return LocalKeyPropForCapture->GetValueTypeHash(ElementKey); },
+			[LocalKeyPropForCapture](const void* A, const void* B) { return LocalKeyPropForCapture->Identical(A, B); }
+			);
+		uint8* Result = (Index >= 0) ? GetPairPtr(Index) : nullptr;
+		return Result;
+	}
+
 	/** Finds the associated value from hash, rather than linearly searching */
 	uint8* FindValueFromHash(const void* KeyPtr)
 	{
@@ -3829,6 +3843,14 @@ public:
 		);
 	}
 
+	/** Finds element pointer from hash, rather than linearly searching */
+	FORCEINLINE uint8* FindElementPtrFromHash(const void* ElementToFind)
+	{
+		const int32 Index = FindElementIndexFromHash(ElementToFind);
+		uint8* Result = (Index >= 0 ? GetElementPtr(Index) : nullptr);
+		return Result;
+	}
+
 	/** Adds the element to the set, returning true if the element was added, or false if the element was already present */
 	void AddElement(const void* ElementToAdd)
 	{
