Added UObject 'OverridePerObjectConfigSection' virtual, so that PerObjectConfig instances can move away from using the transient package, for achieving short section names in .ini files.
Using the transient package, would trigger name conflicts whenever two transient objects/PerObjectConfig-instances shared the same name, causing hard to predict crashes.

This also allows previously inaccessible ini sections, to be exposed to UObject reflection in the engine and editor - creating a path for moving GConfig sections to UObject/PerObjectConfig, in a backwards-compatible way.

diff --git a/Engine/Source/Runtime/CoreUObject/Private/UObject/Obj.cpp b/Engine/Source/Runtime/CoreUObject/Private/UObject/Obj.cpp
--- a/Engine/Source/Runtime/CoreUObject/Private/UObject/Obj.cpp
+++ b/Engine/Source/Runtime/CoreUObject/Private/UObject/Obj.cpp
@@ -2092,11 +2092,13 @@ void UObject::LoadConfig( UClass* ConfigClass/*=NULL*/, const TCHAR* InFilename/
 
 	FString ClassSection;
 	FName LongCommitName;
-	if ( bPerObject == true )
+
+	if (bPerObject)
 	{
 		FString PathNameString;
 		UObject* Outermost = GetOutermost();
-		if ( Outermost == GetTransientPackage() )
+
+		if (Outermost == GetTransientPackage())
 		{
 			PathNameString = GetName();
 		}
@@ -2105,7 +2107,10 @@ void UObject::LoadConfig( UClass* ConfigClass/*=NULL*/, const TCHAR* InFilename/
 			GetPathName(Outermost, PathNameString);
 			LongCommitName = Outermost->GetFName();
 		}
+
 		ClassSection = PathNameString + TEXT(" ") + GetClass()->GetName();
+
+		OverridePerObjectConfigSection(ClassSection);
 	}
 
 	// If any of my properties are class variables, then LoadConfig() would also be called for each one of those classes.
@@ -2310,11 +2315,13 @@ void UObject::SaveConfig( uint64 Flags, const TCHAR* InFilename, FConfigCacheIni
 
 	const bool bPerObject = UsesPerObjectConfig(this);
 	FString Section;
-	if ( bPerObject == true )
+
+	if (bPerObject == true)
 	{
 		FString PathNameString;
 		UObject* Outermost = GetOutermost();
-		if ( Outermost == GetTransientPackage() )
+
+		if (Outermost == GetTransientPackage())
 		{
 			PathNameString = GetName();
 		}
@@ -2322,7 +2329,10 @@ void UObject::SaveConfig( uint64 Flags, const TCHAR* InFilename, FConfigCacheIni
 		{
 			GetPathName(Outermost, PathNameString);
 		}
+
 		Section = PathNameString + TEXT(" ") + GetClass()->GetName();
+
+		OverridePerObjectConfigSection(Section);
 	}
 
 	UObject* CDO = GetClass()->GetDefaultObject();
diff --git a/Engine/Source/Runtime/CoreUObject/Private/UObject/TextProperty.cpp b/Engine/Source/Runtime/CoreUObject/Private/UObject/TextProperty.cpp
--- a/Engine/Source/Runtime/CoreUObject/Private/UObject/TextProperty.cpp
+++ b/Engine/Source/Runtime/CoreUObject/Private/UObject/TextProperty.cpp
@@ -152,6 +152,8 @@ const TCHAR* UTextProperty::ImportText_Internal( const TCHAR* Buffer, void* Data
 				PathNameString = Parent->GetPathName(ParentOutermost);
 			}
 			TextNamespace = PathNameString + TEXT(" ") + Parent->GetClass()->GetName();
+
+			Parent->OverridePerObjectConfigSection(TextNamespace);
 		}
 		else
 		{
diff --git a/Engine/Source/Runtime/CoreUObject/Private/UObject/UObjectGlobals.cpp b/Engine/Source/Runtime/CoreUObject/Private/UObject/UObjectGlobals.cpp
--- a/Engine/Source/Runtime/CoreUObject/Private/UObject/UObjectGlobals.cpp
+++ b/Engine/Source/Runtime/CoreUObject/Private/UObject/UObjectGlobals.cpp
@@ -2414,11 +2414,20 @@ UObject* StaticAllocateObject
 		// This breaks new UObject assumptions and these need to be fixed.
 		if (Obj && !Obj->GetClass()->IsChildOf(InClass))
 		{
+			const TCHAR* ErrorPrefix = TEXT("");
+
+			if (InClass->HasAnyClassFlags(CLASS_PerObjectConfig) && InOuter != nullptr && InOuter->GetOutermost() == GetTransientPackage())
+			{
+				ErrorPrefix = TEXT("PerObjectConfig object using the transient package, has triggered a name conflict and will now crash.\n")
+						TEXT("To avoid this, don't use the transient package for PerObjectConfig objects.\n")
+						TEXT("This has the side effect, of using the full path name for config ini sections. Use 'OverridePerObjectConfigSection' to keep the short name.\n\n");
+			}
+
 			UE_LOG(LogUObjectGlobals, Fatal,
-				TEXT("Objects have the same fully qualified name but different paths.\n")
+				TEXT("%sObjects have the same fully qualified name but different paths.\n")
 				TEXT("\tNew Object: %s %s.%s\n")
 				TEXT("\tExisting Object: %s"),
-				*InClass->GetName(), InOuter ? *InOuter->GetPathName() : TEXT(""), *InName.ToString(),
+				ErrorPrefix, *InClass->GetName(), InOuter ? *InOuter->GetPathName() : TEXT(""), *InName.ToString(),
 				*Obj->GetFullName());
 		}
 	}
diff --git a/Engine/Source/Runtime/CoreUObject/Public/UObject/Object.h b/Engine/Source/Runtime/CoreUObject/Public/UObject/Object.h
--- a/Engine/Source/Runtime/CoreUObject/Public/UObject/Object.h
+++ b/Engine/Source/Runtime/CoreUObject/Public/UObject/Object.h
@@ -964,6 +964,13 @@ public:
 	/** Returns the override config hierarchy platform (if NDAd platforms need defaults to not be in Base*.ini but still want editor to load them) */
 	virtual const TCHAR* GetConfigOverridePlatform() const { return nullptr; }
 
+	/**
+	 * Allows PerObjectConfig classes, to override the ini section name used for the PerObjectConfig object.
+	 *
+	 * @param SectionName	Reference to the unmodified config section name, that can be altered/modified
+	 */
+	virtual void OverridePerObjectConfigSection(FString& SectionName) {}
+
 	/**
 	 * Imports property values from an .ini file.
 	 *
