Fixing GC cluster verification code reporting false positives when a cluster is referencing another cluster through 'mutable' objects list.

diff --git a/Engine/Source/Runtime/CoreUObject/Private/UObject/GarbageCollectionVerification.cpp b/Engine/Source/Runtime/CoreUObject/Private/UObject/GarbageCollectionVerification.cpp
--- a/Engine/Source/Runtime/CoreUObject/Private/UObject/GarbageCollectionVerification.cpp
+++ b/Engine/Source/Runtime/CoreUObject/Private/UObject/GarbageCollectionVerification.cpp
@@ -232,22 +232,27 @@ public:
 				else if (ObjectItem->HasAnyFlags(EInternalObjectFlags::ClusterRoot))
 				{
 					// However, clusters need to be referenced by the current cluster otherwise they can also get GC'd too early.
+					const FUObjectItem* ClusterRootObjectItem = GUObjectArray.ObjectToObjectItem(ClusterRootObject);
 					const int32 OtherClusterRootIndex = GUObjectArray.ObjectToIndex(Object);
 					const FUObjectItem* OtherClusterRootItem = GUObjectArray.IndexToObjectUnsafeForGC(OtherClusterRootIndex);
 					check(OtherClusterRootItem && OtherClusterRootItem->Object);
 					UObject* OtherClusterRootObject = static_cast<UObject*>(OtherClusterRootItem->Object);
-					UE_CLOG(OtherClusterRootIndex != Cluster->RootIndex && !Cluster->ReferencedClusters.Contains(OtherClusterRootIndex), LogGarbage, Fatal,
-						TEXT("Object %s from source cluster %s is referencing object %s (0x%016llx) from cluster %s which is not referenced by the source cluster."),
-						*CurrentObject->GetFullName(),
+					UE_CLOG(OtherClusterRootIndex != Cluster->RootIndex &&
+						!Cluster->ReferencedClusters.Contains(OtherClusterRootIndex) &&
+						!Cluster->MutableObjects.Contains(OtherClusterRootIndex), LogGarbage, Fatal,
+						TEXT("Object %s from source cluster %s (%d) is referencing cluster root object %s (0x%016llx) (%d) which is not referenced by the source cluster."),
+						*GetFullNameSafe(ReferencingObject),
 						*ClusterRootObject->GetFullName(),
+						ClusterRootObjectItem->GetClusterIndex(),
 						*Object->GetFullName(),
 						(int64)(PTRINT)Object,
-						*OtherClusterRootObject->GetFullName());
+						OtherClusterRootItem->GetClusterIndex());
 				}
 			}
 			else if (ObjectItem->GetOwnerIndex() != Cluster->RootIndex)
 			{
 				// If we're referencing an object from another cluster, make sure the other cluster is actually referenced by this cluster
+				const FUObjectItem* ClusterRootObjectItem = GUObjectArray.ObjectToObjectItem(ClusterRootObject);
 				const int32 OtherClusterRootIndex = ObjectItem->GetOwnerIndex();
 				check(OtherClusterRootIndex > 0);
 				const FUObjectItem* OtherClusterRootItem = GUObjectArray.IndexToObjectUnsafeForGC(OtherClusterRootIndex);
@@ -256,12 +261,14 @@ public:
 				UE_CLOG(OtherClusterRootIndex != Cluster->RootIndex &&
 					!Cluster->ReferencedClusters.Contains(OtherClusterRootIndex) &&
 					!Cluster->MutableObjects.Contains(GUObjectArray.ObjectToIndex(Object)), LogGarbage, Fatal,
-					TEXT("Object %s from source cluster %s is referencing object %s (0x%016llx) from cluster %s which is not referenced by the source cluster."),
-					*CurrentObject->GetFullName(),
+					TEXT("Object %s from source cluster %s (%d) is referencing object %s (0x%016llx) from cluster %s (%d) which is not referenced by the source cluster."),
+					*GetFullNameSafe(ReferencingObject),
 					*ClusterRootObject->GetFullName(),
+					ClusterRootObjectItem->GetClusterIndex(),
 					*Object->GetFullName(),
 					(int64)(PTRINT)Object,
-					*OtherClusterRootObject->GetFullName());
+					*OtherClusterRootObject->GetFullName(),
+					OtherClusterRootItem->GetClusterIndex());
 			}
 		}
 	}
