Changed log verbosity when reporting individual cases of GC cluster assumption violations as they are followed by an asser anyway and this way we get the chance to see all issues before we assert at the end of these checks.

diff --git a/Engine/Source/Runtime/CoreUObject/Private/UObject/GarbageCollectionVerification.cpp b/Engine/Source/Runtime/CoreUObject/Private/UObject/GarbageCollectionVerification.cpp
--- a/Engine/Source/Runtime/CoreUObject/Private/UObject/GarbageCollectionVerification.cpp
+++ b/Engine/Source/Runtime/CoreUObject/Private/UObject/GarbageCollectionVerification.cpp
@@ -239,7 +239,7 @@ public:
 					UObject* OtherClusterRootObject = static_cast<UObject*>(OtherClusterRootItem->Object);
 					UE_CLOG(OtherClusterRootIndex != Cluster->RootIndex &&
 						!Cluster->ReferencedClusters.Contains(OtherClusterRootIndex) &&
-						!Cluster->MutableObjects.Contains(OtherClusterRootIndex), LogGarbage, Fatal,
+						!Cluster->MutableObjects.Contains(OtherClusterRootIndex), LogGarbage, Warning,
 						TEXT("Object %s from source cluster %s (%d) is referencing cluster root object %s (0x%016llx) (%d) which is not referenced by the source cluster."),
 						*GetFullNameSafe(ReferencingObject),
 						*ClusterRootObject->GetFullName(),
@@ -260,7 +260,7 @@ public:
 				UObject* OtherClusterRootObject = static_cast<UObject*>(OtherClusterRootItem->Object);
 				UE_CLOG(OtherClusterRootIndex != Cluster->RootIndex &&
 					!Cluster->ReferencedClusters.Contains(OtherClusterRootIndex) &&
-					!Cluster->MutableObjects.Contains(GUObjectArray.ObjectToIndex(Object)), LogGarbage, Fatal,
+					!Cluster->MutableObjects.Contains(GUObjectArray.ObjectToIndex(Object)), LogGarbage, Warning,
 					TEXT("Object %s from source cluster %s (%d) is referencing object %s (0x%016llx) from cluster %s (%d) which is not referenced by the source cluster."),
 					*GetFullNameSafe(ReferencingObject),
 					*ClusterRootObject->GetFullName(),
