TBA: Script struct serialization to structured archives

diff --git a/Engine/Source/Runtime/Core/Private/Misc/Guid.cpp b/Engine/Source/Runtime/Core/Private/Misc/Guid.cpp
--- a/Engine/Source/Runtime/Core/Private/Misc/Guid.cpp
+++ b/Engine/Source/Runtime/Core/Private/Misc/Guid.cpp
@@ -261,6 +261,11 @@ bool FGuid::ParseExact(const FString& GuidString, EGuidFormats Format, FGuid& Ou
 	return true;
 }
 
+FArchive& operator<<(FArchive& Ar, FGuid& G)
+{
+	return Ar << G.A << G.B << G.C << G.D;
+}
+
 void operator<<(FStructuredArchive::FSlot Slot, FGuid& G)
 {
 	if (Slot.GetUnderlyingArchive().IsTextFormat())
diff --git a/Engine/Source/Runtime/Core/Private/Serialization/ArchiveFromStructuredArchive.cpp b/Engine/Source/Runtime/Core/Private/Serialization/ArchiveFromStructuredArchive.cpp
--- a/Engine/Source/Runtime/Core/Private/Serialization/ArchiveFromStructuredArchive.cpp
+++ b/Engine/Source/Runtime/Core/Private/Serialization/ArchiveFromStructuredArchive.cpp
@@ -6,26 +6,14 @@
 FArchiveFromStructuredArchive::FArchiveFromStructuredArchive(FStructuredArchive::FSlot Slot)
 	: FArchiveProxy(Slot.GetUnderlyingArchive())
 	, bPendingSerialize(true)
+	, bWasOpened(false)
 	, Pos(0)
+	, RootSlot(Slot)
 {
 	// For some reason, the FArchive copy constructor will copy all the trivial members of the source archive, but then specifically set ArIsFilterEditorOnly to false, with a comment saying
 	// they don't know why it's doing this... make sure we inherit this flag here!
 	ArIsFilterEditorOnly = InnerArchive.ArIsFilterEditorOnly;
 	ArIsTextFormat = false;
-	
-	if (InnerArchive.IsTextFormat())
-	{
-		Root = Slot.EnterRecord();
-
-		if (IsLoading())
-		{
-			Commit();
-		}
-	}
-	else
-	{
-		Slot.EnterStream();
-	}
 }
 
 FArchiveFromStructuredArchive::~FArchiveFromStructuredArchive()
@@ -90,6 +78,8 @@ bool FArchiveFromStructuredArchive::AtEnd()
 
 FArchive& FArchiveFromStructuredArchive::operator<<(class FName& Value)
 {
+	OpenArchive();
+
 	if (InnerArchive.IsTextFormat())
 	{
 		if (IsLoading())
@@ -118,6 +108,8 @@ FArchive& FArchiveFromStructuredArchive::operator<<(class FName& Value)
 
 FArchive& FArchiveFromStructuredArchive::operator<<(class UObject*& Value)
 {
+	OpenArchive();
+
 	if(InnerArchive.IsTextFormat())
 	{
 		if (IsLoading())
@@ -187,6 +179,8 @@ FArchive& FArchiveFromStructuredArchive::operator<<(class UObject*& Value)
 
 FArchive& FArchiveFromStructuredArchive::operator<<(class FText& Value)
 {
+	OpenArchive();
+
 	if (InnerArchive.IsTextFormat())
 	{
 		FText::SerializeText(*this, Value);
@@ -200,6 +194,8 @@ FArchive& FArchiveFromStructuredArchive::operator<<(class FText& Value)
 
 void FArchiveFromStructuredArchive::Serialize(void* V, int64 Length)
 {
+	OpenArchive();
+
 	if (InnerArchive.IsTextFormat())
 	{
 		if (IsLoading())
@@ -229,7 +225,7 @@ void FArchiveFromStructuredArchive::Serialize(void* V, int64 Length)
 
 void FArchiveFromStructuredArchive::Commit()
 {
-	if (InnerArchive.IsTextFormat())
+	if (bWasOpened && InnerArchive.IsTextFormat())
 	{
 		SerializeInternal(Root.GetValue());
 	}
@@ -237,6 +233,8 @@ void FArchiveFromStructuredArchive::Commit()
 
 void FArchiveFromStructuredArchive::SerializeInternal(FStructuredArchive::FRecord Record)
 {
+	check(bWasOpened);
+
 	if (bPendingSerialize)
 	{
 		FStructuredArchive::FSlot DataSlot = Record.EnterField(FIELD_NAME_TEXT("Data"));
@@ -275,3 +273,25 @@ void FArchiveFromStructuredArchive::SerializeInternal(FStructuredArchive::FRecor
 		bPendingSerialize = false;
 	}
 }
+
+void FArchiveFromStructuredArchive::OpenArchive()
+{
+	if (!bWasOpened)
+	{
+		if (InnerArchive.IsTextFormat())
+		{
+			Root = RootSlot.EnterRecord();
+
+			if (IsLoading())
+			{
+				Commit();
+			}
+		}
+		else
+		{
+			RootSlot.EnterStream();
+		}
+
+		bWasOpened = true;
+	}
+}
\ No newline at end of file
diff --git a/Engine/Source/Runtime/Core/Public/Misc/Guid.h b/Engine/Source/Runtime/Core/Public/Misc/Guid.h
--- a/Engine/Source/Runtime/Core/Public/Misc/Guid.h
+++ b/Engine/Source/Runtime/Core/Public/Misc/Guid.h
@@ -180,16 +180,12 @@ public:
 	 * @param Ar The archive to serialize from or into.
 	 * @param G The GUID to serialize.
 	 */
-	friend FArchive& operator<<(FArchive& Ar, FGuid& G)
-	{
-		FStructuredArchiveFromArchive(Ar).GetSlot() << G;
-		return Ar;
-	}
+	CORE_API friend FArchive& operator<<(FArchive& Ar, FGuid& G);
 
 	/**
 	 * Serializes a GUID from or into a structured archive slot.
 	 *
-	 * @param Ar The structured archive slot to serialize from or into.
+	 * @param Slot The structured archive slot to serialize from or into
 	 * @param G The GUID to serialize.
 	 */
 	CORE_API friend void operator<<(FStructuredArchive::FSlot Slot, FGuid& G);
@@ -200,6 +196,12 @@ public:
 		return true;
 	}
 
+	bool Serialize(FStructuredArchive::FSlot Slot)
+	{
+		Slot << *this;
+		return true;
+	}
+
 public:
 
 	/**
diff --git a/Engine/Source/Runtime/Core/Public/Serialization/ArchiveFromStructuredArchive.h b/Engine/Source/Runtime/Core/Public/Serialization/ArchiveFromStructuredArchive.h
--- a/Engine/Source/Runtime/Core/Public/Serialization/ArchiveFromStructuredArchive.h
+++ b/Engine/Source/Runtime/Core/Public/Serialization/ArchiveFromStructuredArchive.h
@@ -40,11 +40,13 @@ protected:
 	TOptional<FStructuredArchive::FRecord> Root;
 	void Commit();
 	virtual void SerializeInternal(FStructuredArchive::FRecord Record);
+	void OpenArchive();
 
 private:
 	static const int32 MaxBufferSize = 128;
 
 	bool bPendingSerialize;
+	bool bWasOpened;
 
 	TArray<uint8> Buffer;
 	int32 Pos;
@@ -55,4 +57,6 @@ private:
 	TArray<UObject*> Objects;
 	TBitArray<> ObjectsValid;
 	TMap<UObject*, int32> ObjectToIndex;
+
+	FStructuredArchive::FSlot RootSlot;
 };
\ No newline at end of file
diff --git a/Engine/Source/Runtime/Core/Public/Serialization/StructuredArchive.h b/Engine/Source/Runtime/Core/Public/Serialization/StructuredArchive.h
--- a/Engine/Source/Runtime/Core/Public/Serialization/StructuredArchive.h
+++ b/Engine/Source/Runtime/Core/Public/Serialization/StructuredArchive.h
@@ -140,6 +140,15 @@ public:
 			return Ar.GetUnderlyingArchive();
 		}
 
+		FORCEINLINE bool IsFilled() const
+		{
+#if WITH_TEXT_ARCHIVE_SUPPORT
+			return Ar.CurrentSlotElementId != ElementId;
+#else
+			return true;
+#endif
+		}
+
 	private:
 		friend FStructuredArchive;
 		friend class FStructuredArchiveChildReader;
diff --git a/Engine/Source/Runtime/Core/Public/UObject/ScriptDelegates.h b/Engine/Source/Runtime/Core/Public/UObject/ScriptDelegates.h
--- a/Engine/Source/Runtime/Core/Public/UObject/ScriptDelegates.h
+++ b/Engine/Source/Runtime/Core/Public/UObject/ScriptDelegates.h
@@ -141,6 +141,13 @@ public:
 		return Ar;
 	}
 
+	/** Delegate serialization */
+	friend void operator<<(FStructuredArchive::FSlot Slot, TScriptDelegate& D)
+	{
+		FStructuredArchive::FRecord Record = Slot.EnterRecord();
+		Record << NAMED_ITEM("Object", D.Object) << NAMED_ITEM("FunctionName",D.FunctionName);
+	}
+
 	/** Comparison operators */
 	FORCEINLINE bool operator==( const TScriptDelegate& Other ) const
 	{
@@ -442,6 +449,25 @@ public:
 		return Ar;
 	}
 
+	friend void operator<<(FStructuredArchive::FSlot Slot, TMulticastScriptDelegate<TWeakPtr>& D)
+	{
+		FArchive& UnderlyingArchive = Slot.GetUnderlyingArchive();
+
+		if (UnderlyingArchive.IsSaving())
+		{
+			// When saving the delegate, clean up the list to make sure there are no bad object references
+			D.CompactInvocationList();
+		}
+
+		Slot << D.InvocationList;
+
+		if (UnderlyingArchive.IsLoading())
+		{
+			// After loading the delegate, clean up the list to make sure there are no bad object references
+			D.CompactInvocationList();
+		}
+	}
+
 	/**
 	 * Executes a multi-cast delegate by calling all functions on objects bound to the delegate.  Always
 	 * safe to call, even if when no objects are bound, or if objects have expired.  In general, you should
diff --git a/Engine/Source/Runtime/CoreUObject/Private/Serialization/ArchiveUObjectFromStructuredArchive.cpp b/Engine/Source/Runtime/CoreUObject/Private/Serialization/ArchiveUObjectFromStructuredArchive.cpp
--- a/Engine/Source/Runtime/CoreUObject/Private/Serialization/ArchiveUObjectFromStructuredArchive.cpp
+++ b/Engine/Source/Runtime/CoreUObject/Private/Serialization/ArchiveUObjectFromStructuredArchive.cpp
@@ -6,10 +6,6 @@ FArchiveUObjectFromStructuredArchive::FArchiveUObjectFromStructuredArchive(FStru
 	: FArchiveFromStructuredArchive(Slot)
 	, bPendingSerialize(true)
 {
-	if (InnerArchive.IsTextFormat() && IsLoading())
-	{
-		Commit();
-	}
 }
 
 FArchiveUObjectFromStructuredArchive::~FArchiveUObjectFromStructuredArchive()
@@ -19,6 +15,8 @@ FArchiveUObjectFromStructuredArchive::~FArchiveUObjectFromStructuredArchive()
 
 FArchive& FArchiveUObjectFromStructuredArchive::operator<<(FLazyObjectPtr& Value)
 {
+	OpenArchive();
+
 	if (InnerArchive.IsTextFormat())
 	{
 		if (IsLoading())
@@ -47,6 +45,8 @@ FArchive& FArchiveUObjectFromStructuredArchive::operator<<(FLazyObjectPtr& Value
 
 FArchive& FArchiveUObjectFromStructuredArchive::operator<<(FSoftObjectPtr& Value)
 {
+	OpenArchive();
+
 	if (InnerArchive.IsTextFormat())
 	{
 		if (IsLoading())
@@ -75,6 +75,8 @@ FArchive& FArchiveUObjectFromStructuredArchive::operator<<(FSoftObjectPtr& Value
 
 FArchive& FArchiveUObjectFromStructuredArchive::operator<<(FSoftObjectPath& Value)
 {
+	OpenArchive();
+
 	if (InnerArchive.IsTextFormat())
 	{
 		if (IsLoading())
@@ -103,6 +105,8 @@ FArchive& FArchiveUObjectFromStructuredArchive::operator<<(FSoftObjectPath& Valu
 
 FArchive& FArchiveUObjectFromStructuredArchive::operator<<(FWeakObjectPtr& Value)
 {
+	OpenArchive();
+
 	if (InnerArchive.IsTextFormat())
 	{
 		if (IsLoading())
diff --git a/Engine/Source/Runtime/CoreUObject/Private/UObject/Class.cpp b/Engine/Source/Runtime/CoreUObject/Private/UObject/Class.cpp
--- a/Engine/Source/Runtime/CoreUObject/Private/UObject/Class.cpp
+++ b/Engine/Source/Runtime/CoreUObject/Private/UObject/Class.cpp
@@ -836,34 +836,42 @@ void UStruct::DestroyStruct(void* Dest, int32 ArrayDim) const
 	}
 }
 
+void UStruct::SerializeBin(FArchive& Ar, void* Data) const
+{
+	SerializeBin(FStructuredArchiveFromArchive(Ar).GetSlot(), Data);
+}
 //
 // Serialize all of the class's data that belongs in a particular
 // bin and resides in Data.
 //
-void UStruct::SerializeBin( FArchive& Ar, void* Data ) const
+void UStruct::SerializeBin( FStructuredArchive::FSlot Slot, void* Data ) const
 {
-	if( Ar.IsObjectReferenceCollector() )
+	FArchive& UnderlyingArchive = Slot.GetUnderlyingArchive();
+
+	FStructuredArchive::FStream PropertyStream = Slot.EnterStream();
+
+	if( UnderlyingArchive.IsObjectReferenceCollector() )
 	{
 		for( UProperty* RefLinkProperty=RefLink; RefLinkProperty!=NULL; RefLinkProperty=RefLinkProperty->NextRef )
 		{
-			RefLinkProperty->SerializeBinProperty( Ar, Data );
+			RefLinkProperty->SerializeBinProperty(PropertyStream.EnterElement(), Data );
 		}
 	}
-	else if( Ar.ArUseCustomPropertyList )
+	else if( UnderlyingArchive.ArUseCustomPropertyList )
 	{
-		const FCustomPropertyListNode* CustomPropertyList = Ar.ArCustomPropertyList;
+		const FCustomPropertyListNode* CustomPropertyList = UnderlyingArchive.ArCustomPropertyList;
 		for (auto PropertyNode = CustomPropertyList; PropertyNode; PropertyNode = PropertyNode->PropertyListNext)
 		{
 			UProperty* Property = PropertyNode->Property;
 			if( Property )
 			{
 				// Temporarily set to the sub property list, in case we're serializing a UStruct property.
-				Ar.ArCustomPropertyList = PropertyNode->SubPropertyList;
+				UnderlyingArchive.ArCustomPropertyList = PropertyNode->SubPropertyList;
 
-				Property->SerializeBinProperty(Ar, Data, PropertyNode->ArrayIndex);
+				Property->SerializeBinProperty(PropertyStream.EnterElement(), Data, PropertyNode->ArrayIndex);
 
 				// Restore the original property list.
-				Ar.ArCustomPropertyList = CustomPropertyList;
+				UnderlyingArchive.ArCustomPropertyList = CustomPropertyList;
 			}
 		}
 	}
@@ -871,22 +879,22 @@ void UStruct::SerializeBin( FArchive& Ar, void* Data ) const
 	{
 		for (UProperty* Property = PropertyLink; Property != NULL; Property = Property->PropertyLinkNext)
 		{
-			Property->SerializeBinProperty(Ar, Data);
+			Property->SerializeBinProperty(PropertyStream.EnterElement(), Data);
 		}
 	}
 }
 
-void UStruct::SerializeBinEx( FArchive& Ar, void* Data, void const* DefaultData, UStruct* DefaultStruct ) const
+void UStruct::SerializeBinEx( FStructuredArchive::FSlot Slot, void* Data, void const* DefaultData, UStruct* DefaultStruct ) const
 {
 	if ( !DefaultData || !DefaultStruct )
 	{
-		SerializeBin(Ar, Data);
+		SerializeBin(Slot, Data);
 		return;
 	}
 
 	for( TFieldIterator<UProperty> It(this); It; ++It )
 	{
-		It->SerializeNonMatchingBinProperty(Ar, Data, DefaultData, DefaultStruct);
+		It->SerializeNonMatchingBinProperty(Slot, Data, DefaultData, DefaultStruct);
 	}
 }
 
@@ -1171,7 +1179,12 @@ void UStruct::SerializeTaggedProperties(FStructuredArchive::FSlot Slot, uint8* D
 							UnderlyingArchive.ArCustomPropertyList = CustomPropertyNode->SubPropertyList;
 						}
 
-						Tag.SerializeTaggedProperty(PropertyRecord.EnterField(FIELD_NAME_TEXT("Value")), Property, DataPtr, DefaultValue);
+						FStructuredArchive::FSlot PropertyField = PropertyRecord.EnterField(FIELD_NAME_TEXT("Value"));
+						Tag.SerializeTaggedProperty(PropertyField, Property, DataPtr, DefaultValue);
+						if (!PropertyField.IsFilled())
+						{
+							Tag.SerializeTaggedProperty(PropertyField, Property, DataPtr, DefaultValue);
+						}
 
 						// restore the original custom property list after serializing
 						if (SavedCustomPropertyList)
@@ -1216,6 +1229,9 @@ void UStruct::FinishDestroy()
 	Script.Empty();
 	Super::FinishDestroy();
 }
+
+IMPLEMENT_FSTRUCTUREDARCHIVE_SERIALIZER(UStruct);
+
 void UStruct::Serialize( FArchive& Ar )
 {
 	Super::Serialize( Ar );
@@ -2025,6 +2041,8 @@ void UScriptStruct::PrepareCppStructOps()
 	bPrepareCppStructOpsCompleted = true;
 }
 
+IMPLEMENT_FSTRUCTUREDARCHIVE_SERIALIZER(UScriptStruct);
+
 void UScriptStruct::Serialize( FArchive& Ar )
 {
 	Super::Serialize(Ar);
@@ -2048,13 +2066,20 @@ bool UScriptStruct::UseBinarySerialization(const FArchive& Ar) const
 
 void UScriptStruct::SerializeItem(FArchive& Ar, void* Value, void const* Defaults)
 {
-	const bool bUseBinarySerialization = UseBinarySerialization(Ar);
+	SerializeItem(FStructuredArchiveFromArchive(Ar).GetSlot(), Value, Defaults);
+}
+
+void UScriptStruct::SerializeItem(FStructuredArchive::FSlot Slot, void* Value, void const* Defaults)
+{
+	FArchive& UnderlyingArchive = Slot.GetUnderlyingArchive();
+
+	const bool bUseBinarySerialization = UseBinarySerialization(UnderlyingArchive);
 	const bool bUseNativeSerialization = UseNativeSerialization();
 
 	// Preload struct before serialization tracking to not double count time.
 	if (bUseBinarySerialization || bUseNativeSerialization)
 	{
-		Ar.Preload(this);
+		UnderlyingArchive.Preload(this);
 	}
 
 	bool bItemSerialized = false;
@@ -2062,7 +2087,20 @@ void UScriptStruct::SerializeItem(FArchive& Ar, void* Value, void const* Default
 	{
 		UScriptStruct::ICppStructOps* TheCppStructOps = GetCppStructOps();
 		check(TheCppStructOps); // else should not have STRUCT_SerializeNative
-		bItemSerialized = TheCppStructOps->Serialize(Ar, Value);
+
+		if (TheCppStructOps->HasStructuredSerializer())
+		{
+			bItemSerialized = TheCppStructOps->Serialize(Slot, Value);
+		}
+		else
+		{
+			FArchiveUObjectFromStructuredArchive Ar(Slot);
+			bItemSerialized = TheCppStructOps->Serialize(Ar, Value);
+			if (bItemSerialized && !Slot.IsFilled())
+			{
+				Slot.EnterStream();
+			}
+		}		
 	}
 
 	if (!bItemSerialized)
@@ -2070,18 +2108,18 @@ void UScriptStruct::SerializeItem(FArchive& Ar, void* Value, void const* Default
 		if (bUseBinarySerialization)
 		{
 			// Struct is already preloaded above.
-			if (!Ar.IsPersistent() && Ar.GetPortFlags() != 0 && !ShouldSerializeAtomically(Ar) && !Ar.ArUseCustomPropertyList)
+			if (!UnderlyingArchive.IsPersistent() && UnderlyingArchive.GetPortFlags() != 0 && !ShouldSerializeAtomically(UnderlyingArchive) && !UnderlyingArchive.ArUseCustomPropertyList)
 			{
-				SerializeBinEx(Ar, Value, Defaults, this);
+				SerializeBinEx(Slot, Value, Defaults, this);
 			}
 			else
 			{
-				SerializeBin(Ar, Value);
+				SerializeBin(Slot, Value);
 			}
 		}
 		else
 		{
-			SerializeTaggedProperties(Ar, (uint8*)Value, this, (uint8*)Defaults);
+			SerializeTaggedProperties(Slot, (uint8*)Value, this, (uint8*)Defaults);
 		}
 	}
 
@@ -2089,7 +2127,7 @@ void UScriptStruct::SerializeItem(FArchive& Ar, void* Value, void const* Default
 	{
 		UScriptStruct::ICppStructOps* TheCppStructOps = GetCppStructOps();
 		check(TheCppStructOps); // else should not have STRUCT_PostSerializeNative
-		TheCppStructOps->PostSerialize(Ar, Value);
+		TheCppStructOps->PostSerialize(UnderlyingArchive, Value);
 	}
 }
 
@@ -3776,31 +3814,38 @@ bool UClass::ImplementsInterface( const class UClass* SomeInterface ) const
 	return false;
 }
 
+void UClass::SerializeDefaultObject(UObject* Object, FArchive& Ar)
+{
+	SerializeDefaultObject(Object, FStructuredArchiveFromArchive(Ar).GetSlot());
+}
+
 /** serializes the passed in object as this class's default object using the given archive
  * @param Object the object to serialize as default
  * @param Ar the archive to serialize from
  */
-void UClass::SerializeDefaultObject(UObject* Object, FArchive& Ar)
+void UClass::SerializeDefaultObject(UObject* Object, FStructuredArchive::FSlot Slot)
 {
 	// tell the archive that it's allowed to load data for transient properties
-	Ar.StartSerializingDefaults();
+	FArchive& UnderlyingArchive = Slot.GetUnderlyingArchive();
 
-	if( ((Ar.IsLoading() || Ar.IsSaving()) && !Ar.WantBinaryPropertySerialization()) )
+	UnderlyingArchive.StartSerializingDefaults();
+
+	if( ((UnderlyingArchive.IsLoading() || UnderlyingArchive.IsSaving()) && !UnderlyingArchive.WantBinaryPropertySerialization()) )
 	{
 	    // class default objects do not always have a vtable when saved
 		// so use script serialization as opposed to native serialization to
 	    // guarantee that all property data is loaded into the correct location
-	    SerializeTaggedProperties(Ar, (uint8*)Object, GetSuperClass(), (uint8*)Object->GetArchetype());
+	    SerializeTaggedProperties(Slot, (uint8*)Object, GetSuperClass(), (uint8*)Object->GetArchetype());
 	}
-	else if ( Ar.GetPortFlags() != 0 )
+	else if (UnderlyingArchive.GetPortFlags() != 0 )
 	{
-		SerializeBinEx(Ar, Object, Object->GetArchetype(), GetSuperClass() );
+		SerializeBinEx(Slot, Object, Object->GetArchetype(), GetSuperClass() );
 	}
 	else
 	{
-		SerializeBin(Ar, Object);
+		SerializeBin(Slot, Object);
 	}
-	Ar.StopSerializingDefaults();
+	UnderlyingArchive.StopSerializingDefaults();
 }
 
 
diff --git a/Engine/Source/Runtime/CoreUObject/Private/UObject/LinkerLoad.cpp b/Engine/Source/Runtime/CoreUObject/Private/UObject/LinkerLoad.cpp
--- a/Engine/Source/Runtime/CoreUObject/Private/UObject/LinkerLoad.cpp
+++ b/Engine/Source/Runtime/CoreUObject/Private/UObject/LinkerLoad.cpp
@@ -3448,9 +3448,7 @@ void FLinkerLoad::Preload( UObject* Object )
 						{
 							FString ObjectName = Object->GetPathName(Object->GetOutermost());
 							FStructuredArchive::FSlot ExportSlot = StructuredArchiveRootRecord->EnterField(FIELD_NAME(*ObjectName));
-
-							FArchiveUObjectFromStructuredArchive Adapter(ExportSlot);
-							Object->GetClass()->SerializeDefaultObject(Object, Adapter);
+							Object->GetClass()->SerializeDefaultObject(Object, ExportSlot);
 						}
 						else
 #endif
diff --git a/Engine/Source/Runtime/CoreUObject/Private/UObject/Obj.cpp b/Engine/Source/Runtime/CoreUObject/Private/UObject/Obj.cpp
--- a/Engine/Source/Runtime/CoreUObject/Private/UObject/Obj.cpp
+++ b/Engine/Source/Runtime/CoreUObject/Private/UObject/Obj.cpp
@@ -1290,7 +1290,6 @@ void UObject::SerializeScriptProperties( FStructuredArchive::FSlot Slot ) const
 	}
 
 	UClass *ObjClass = GetClass();
-	bool bSlotEntered = false;
 
 	if(UnderlyingArchive.IsTextFormat() || ((UnderlyingArchive.IsLoading() || UnderlyingArchive.IsSaving()) && !UnderlyingArchive.WantBinaryPropertySerialization()))
 	{
@@ -1313,7 +1312,6 @@ void UObject::SerializeScriptProperties( FStructuredArchive::FSlot Slot ) const
 #endif
 
 		ObjClass->SerializeTaggedProperties(Slot, (uint8*)this, HasAnyFlags(RF_ClassDefaultObject) ? ObjClass->GetSuperClass() : ObjClass, (uint8*)DiffObject, bBreakSerializationRecursion ? this : NULL);
-		bSlotEntered = true;
 	}
 	else if (UnderlyingArchive.GetPortFlags() != 0 && !UnderlyingArchive.ArUseCustomPropertyList )
 	{
@@ -1323,11 +1321,11 @@ void UObject::SerializeScriptProperties( FStructuredArchive::FSlot Slot ) const
 		{
 			DiffObject = GetArchetype();
 		}
-		ObjClass->SerializeBinEx( UnderlyingArchive, const_cast<UObject *>(this), DiffObject, DiffObject ? DiffObject->GetClass() : NULL );
+		ObjClass->SerializeBinEx(Slot, const_cast<UObject *>(this), DiffObject, DiffObject ? DiffObject->GetClass() : NULL);
 	}
 	else
 	{
-		ObjClass->SerializeBin(UnderlyingArchive, const_cast<UObject *>(this) );
+		ObjClass->SerializeBin(Slot, const_cast<UObject *>(this));
 	}
 
 	if( HasAnyFlags(RF_ClassDefaultObject) )
@@ -1335,11 +1333,6 @@ void UObject::SerializeScriptProperties( FStructuredArchive::FSlot Slot ) const
 		UnderlyingArchive.StopSerializingDefaults();
 	}
 	UnderlyingArchive.MarkScriptSerializationEnd(this);
-
-	if (!bSlotEntered)
-	{
-		Slot.EnterStream();
-	}
 }
 
 
diff --git a/Engine/Source/Runtime/CoreUObject/Private/UObject/Property.cpp b/Engine/Source/Runtime/CoreUObject/Private/UObject/Property.cpp
--- a/Engine/Source/Runtime/CoreUObject/Private/UObject/Property.cpp
+++ b/Engine/Source/Runtime/CoreUObject/Private/UObject/Property.cpp
@@ -224,6 +224,7 @@ struct TStructOpsTypeTraits<FGuid> : public TStructOpsTypeTraitsBase2<FGuid>
 		WithImportTextItem = true,
 		WithZeroConstructor = true,
 		WithSerializer = true,
+		WithStructuredSerializer = true,
 	};
 };
 IMPLEMENT_STRUCT(Guid);
diff --git a/Engine/Source/Runtime/CoreUObject/Private/UObject/PropertyDelegate.cpp b/Engine/Source/Runtime/CoreUObject/Private/UObject/PropertyDelegate.cpp
--- a/Engine/Source/Runtime/CoreUObject/Private/UObject/PropertyDelegate.cpp
+++ b/Engine/Source/Runtime/CoreUObject/Private/UObject/PropertyDelegate.cpp
@@ -61,8 +61,7 @@ bool UDelegateProperty::Identical( const void* A, const void* B, uint32 PortFlag
 
 void UDelegateProperty::SerializeItem(FStructuredArchive::FSlot Slot, void* Value, void const* Defaults) const
 {
-	FArchiveUObjectFromStructuredArchive Ar(Slot);
-	Ar << *GetPropertyValuePtr(Value);
+	Slot << *GetPropertyValuePtr(Value);
 }
 
 
diff --git a/Engine/Source/Runtime/CoreUObject/Private/UObject/PropertyStruct.cpp b/Engine/Source/Runtime/CoreUObject/Private/UObject/PropertyStruct.cpp
--- a/Engine/Source/Runtime/CoreUObject/Private/UObject/PropertyStruct.cpp
+++ b/Engine/Source/Runtime/CoreUObject/Private/UObject/PropertyStruct.cpp
@@ -114,8 +114,7 @@ void UStructProperty::SerializeItem(FStructuredArchive::FSlot Slot, void* Value,
 	FScopedPlaceholderPropertyTracker ImportPropertyTracker(this);
 #endif
 
-	FArchiveUObjectFromStructuredArchive Ar(Slot);
-	Struct->SerializeItem(Ar, Value, Defaults);
+	Struct->SerializeItem(Slot, Value, Defaults);
 }
 
 bool UStructProperty::NetSerializeItem( FArchive& Ar, UPackageMap* Map, void* Data, TArray<uint8> * MetaData ) const
diff --git a/Engine/Source/Runtime/CoreUObject/Public/UObject/Class.h b/Engine/Source/Runtime/CoreUObject/Public/UObject/Class.h
--- a/Engine/Source/Runtime/CoreUObject/Public/UObject/Class.h
+++ b/Engine/Source/Runtime/CoreUObject/Public/UObject/Class.h
@@ -217,7 +217,7 @@ class COREUOBJECT_API UField : public UObject
  */
 class COREUOBJECT_API UStruct : public UField
 {
-	DECLARE_CASTED_CLASS_INTRINSIC(UStruct, UField, 0, TEXT("/Script/CoreUObject"), CASTCLASS_UStruct)
+	DECLARE_CASTED_CLASS_INTRINSIC(UStruct, UField, CLASS_MatchedSerializers, TEXT("/Script/CoreUObject"), CASTCLASS_UStruct)
 
 	// Variables.
 protected:
@@ -257,6 +257,7 @@ public:
 
 	// UObject interface.
 	virtual void Serialize(FArchive& Ar) override;
+	virtual void Serialize(FStructuredArchive::FRecord Record) override;
 	virtual void FinishDestroy() override;
 	virtual void RegisterDependencies() override;
 	static void AddReferencedObjects(UObject* InThis, FReferenceCollector& Collector);
@@ -290,7 +291,8 @@ public:
 	virtual void Link(FArchive& Ar, bool bRelinkExistingProperties);
 
 	/** Serializes struct properties, does not handle defaults */
-	virtual void SerializeBin( FArchive& Ar, void* Data ) const;
+	virtual void SerializeBin(FStructuredArchive::FSlot Slot, void* Data) const;
+	virtual void SerializeBin(FArchive& Ar, void* Data) const;
 
 	/**
 	 * Serializes the class properties that reside in Data if they differ from the corresponding values in DefaultData
@@ -300,7 +302,7 @@ public:
 	 * @param	DefaultData		pointer to the location of the beginning of the data that should be compared against
 	 * @param	DefaultStruct	the struct corresponding to the block of memory located at DefaultData 
 	 */
-	void SerializeBinEx( FArchive& Ar, void* Data, void const* DefaultData, UStruct* DefaultStruct ) const;
+	void SerializeBinEx( FStructuredArchive::FSlot Slot, void* Data, void const* DefaultData, UStruct* DefaultStruct ) const;
 
 	/** Serializes list of properties, using property tags to handle mismatches */
 	virtual void SerializeTaggedProperties( FArchive& Ar, uint8* Data, UStruct* DefaultsStruct, uint8* Defaults, const UObject* BreakRecursionIfFullyLoad=nullptr) const;
@@ -503,11 +505,14 @@ enum EStructFlags
 	/** If set, this struct can share net serialization state across connections */
 	STRUCT_NetSharedSerialization = 0x00400000,
 
+	/**  */
+	STRUCT_SerializeNativeStructured = 0x00800000,
+
 	/** Struct flags that are automatically inherited */
 	STRUCT_Inherit				= STRUCT_HasInstancedReference|STRUCT_Atomic,
 
 	/** Flags that are always computed, never loaded or done with code generation */
-	STRUCT_ComputedFlags		= STRUCT_NetDeltaSerializeNative | STRUCT_NetSerializeNative | STRUCT_SerializeNative | STRUCT_PostSerializeNative | STRUCT_CopyNative | STRUCT_IsPlainOldData | STRUCT_NoDestructor | STRUCT_ZeroConstructor | STRUCT_IdenticalNative | STRUCT_AddStructReferencedObjects | STRUCT_ExportTextItemNative | STRUCT_ImportTextItemNative | STRUCT_SerializeFromMismatchedTag | STRUCT_PostScriptConstruct | STRUCT_NetSharedSerialization
+	STRUCT_ComputedFlags		= STRUCT_NetDeltaSerializeNative | STRUCT_NetSerializeNative | STRUCT_SerializeNative | STRUCT_SerializeNativeStructured | STRUCT_PostSerializeNative | STRUCT_CopyNative | STRUCT_IsPlainOldData | STRUCT_NoDestructor | STRUCT_ZeroConstructor | STRUCT_IdenticalNative | STRUCT_AddStructReferencedObjects | STRUCT_ExportTextItemNative | STRUCT_ImportTextItemNative | STRUCT_SerializeFromMismatchedTag | STRUCT_PostScriptConstruct | STRUCT_NetSharedSerialization
 };
 
 
@@ -527,6 +532,7 @@ struct TStructOpsTypeTraitsBase2
 		WithImportTextItem             = false,                         // struct has an ImportTextItem function used to deserialize a string into an object of that class.
 		WithAddStructReferencedObjects = false,                         // struct has an AddStructReferencedObjects function which allows it to add references to the garbage collector.
 		WithSerializer                 = false,                         // struct has a Serialize function for serializing its state to an FArchive.
+		WithStructuredSerializer       = false,                         // struct has a Serialize function for serializing its state to an FStructuredArchive.
 		WithPostSerialize              = false,                         // struct has a PostSerialize function which is called after it is serialized
 		WithNetSerializer              = false,                         // struct has a NetSerialize function for serializing its state to an FArchive used for network replication.
 		WithNetDeltaSerializer         = false,                         // struct has a NetDeltaSerialize function for serializing differences in state from a previous NetSerialize operation.
@@ -595,6 +601,21 @@ FORCEINLINE typename TEnableIf<TStructOpsTypeTraits<CPPSTRUCT>::WithSerializer,
 	return Data->Serialize(Ar);
 }
 
+/**
+* Selection of structured Serialize call.
+*/
+template<class CPPSTRUCT>
+FORCEINLINE typename TEnableIf<!TStructOpsTypeTraits<CPPSTRUCT>::WithStructuredSerializer, bool>::Type SerializeOrNot(FStructuredArchive::FSlot Slot, CPPSTRUCT *Data)
+{
+	return false;
+}
+
+template<class CPPSTRUCT>
+FORCEINLINE typename TEnableIf<TStructOpsTypeTraits<CPPSTRUCT>::WithStructuredSerializer, bool>::Type SerializeOrNot(FStructuredArchive::FSlot Slot, CPPSTRUCT *Data)
+{
+	return Data->Serialize(Slot);
+}
+
 
 /**
  * Selection of PostSerialize call.
@@ -831,11 +852,14 @@ public:
 
 		/** return true if this class can serialize **/
 		virtual bool HasSerializer() = 0;
+		/** return true if this class can serialize to a structured archive**/
+		virtual bool HasStructuredSerializer() = 0;
 		/** 
 		 * Serialize this structure 
 		 * @return true if the package is new enough to support this, if false, it will fall back to ordinary script struct serialization
 		 */
 		virtual bool Serialize(FArchive& Ar, void *Data) = 0;
+		virtual bool Serialize(FStructuredArchive::FSlot Slot, void *Data) = 0;
 
 		/** return true if this class implements a post serialize call **/
 		virtual bool HasPostSerialize() = 0;
@@ -976,11 +1000,20 @@ public:
 		{
 			return TTraits::WithSerializer;
 		}
+		virtual bool HasStructuredSerializer() override
+		{
+			return TTraits::WithStructuredSerializer;
+		}
 		virtual bool Serialize(FArchive& Ar, void *Data) override
 		{
 			check(TTraits::WithSerializer); // don't call this if we have indicated it is not necessary
 			return SerializeOrNot(Ar, (CPPSTRUCT*)Data);
 		}
+		virtual bool Serialize(FStructuredArchive::FSlot Slot, void *Data) override
+		{
+			check(TTraits::WithStructuredSerializer); // don't call this if we have indicated it is not necessary
+			return SerializeOrNot(Slot, (CPPSTRUCT*)Data);
+		}
 		virtual bool HasPostSerialize() override
 		{
 			return TTraits::WithPostSerialize;
@@ -1111,7 +1144,7 @@ public:
 	#define IMPLEMENT_STRUCT(BaseName) \
 		static UScriptStruct::TAutoCppStructOps<F##BaseName> BaseName##_Ops(TEXT(#BaseName)); 
 
-	DECLARE_CASTED_CLASS_INTRINSIC_NO_CTOR(UScriptStruct, UStruct, 0, TEXT("/Script/CoreUObject"), CASTCLASS_UScriptStruct, COREUOBJECT_API)
+	DECLARE_CASTED_CLASS_INTRINSIC_NO_CTOR(UScriptStruct, UStruct, CLASS_MatchedSerializers, TEXT("/Script/CoreUObject"), CASTCLASS_UScriptStruct, COREUOBJECT_API)
 
 	COREUOBJECT_API UScriptStruct( EStaticConstructor, int32 InSize, EObjectFlags InFlags );
 	COREUOBJECT_API explicit UScriptStruct(const FObjectInitializer& ObjectInitializer, UScriptStruct* InSuperStruct, ICppStructOps* InCppStructOps = nullptr, EStructFlags InStructFlags = STRUCT_NoFlags, SIZE_T ExplicitSize = 0, SIZE_T ExplicitAlignment = 0);
@@ -1136,7 +1169,8 @@ private:
 public:
 
 	// UObject Interface
-	virtual COREUOBJECT_API void Serialize( FArchive& Ar ) override;
+	virtual COREUOBJECT_API void Serialize(FArchive& Ar) override;
+	virtual COREUOBJECT_API void Serialize(FStructuredArchive::FRecord Record) override;
 
 	// UStruct interface.
 	virtual COREUOBJECT_API void Link(FArchive& Ar, bool bRelinkExistingProperties) override;
@@ -1197,7 +1231,7 @@ public:
 	/** Returns true if this struct has a native serialize function */
 	bool UseNativeSerialization() const
 	{
-		if ((StructFlags&STRUCT_SerializeNative) != 0)
+		if ((StructFlags&(STRUCT_SerializeNative | STRUCT_SerializeNativeStructured)) != 0)
 		{
 			return true;
 		}
@@ -1213,11 +1247,12 @@ public:
 	/** 
 	 * Serializes a specific instance of a struct 
 	 *
-	 * @param	Ar			Archive we are serializing to/from
+	 * @param	Slot		The structured archive slot we are serializing to
 	 * @param	Value		Pointer to memory of struct
 	 * @param	Defaults	Default value for this struct, pass nullptr to not use defaults 
 	 */
 	COREUOBJECT_API void SerializeItem(FArchive& Ar, void* Value, void const* Defaults);
+	COREUOBJECT_API void SerializeItem(FStructuredArchive::FSlot Slot, void* Value, void const* Defaults);
 
 	/**
 	 * Export script struct to a string that can later be imported
@@ -2569,6 +2604,7 @@ public:
 	 * @param Object the object to serialize as default
 	 * @param Ar the archive to serialize from
 	 */
+	virtual void SerializeDefaultObject(UObject* Object, FStructuredArchive::FSlot Slot);
 	virtual void SerializeDefaultObject(UObject* Object, FArchive& Ar);
 
 	/** Wraps the PostLoad() call for the class default object.
diff --git a/Engine/Source/Runtime/CoreUObject/Public/UObject/UnrealType.h b/Engine/Source/Runtime/CoreUObject/Public/UObject/UnrealType.h
--- a/Engine/Source/Runtime/CoreUObject/Public/UObject/UnrealType.h
+++ b/Engine/Source/Runtime/CoreUObject/Public/UObject/UnrealType.h
@@ -262,17 +262,18 @@ public:
 	 * @param	Data			pointer to the location of the beginning of the struct's property data
 	 * @param	ArrayIdx		if not -1 (default), only this array slot will be serialized
 	 */
-	void SerializeBinProperty( FArchive& Ar, void* Data, int32 ArrayIdx = -1 )
+	void SerializeBinProperty( FStructuredArchive::FSlot Slot, void* Data, int32 ArrayIdx = -1 )
 	{
-		if( ShouldSerializeValue(Ar) )
+		FStructuredArchive::FStream Stream = Slot.EnterStream();
+		if( ShouldSerializeValue(Slot.GetUnderlyingArchive()) )
 		{
 			const int32 LoopMin = ArrayIdx < 0 ? 0 : ArrayIdx;
 			const int32 LoopMax = ArrayIdx < 0 ? ArrayDim : ArrayIdx + 1;
 			for (int32 Idx = LoopMin; Idx < LoopMax; Idx++)
 			{
 				// Keep setting the property in case something inside of SerializeItem changes it
-				FSerializedPropertyScope SerializedProperty(Ar, this);
-				SerializeItem( FStructuredArchiveFromArchive(Ar).GetSlot(), ContainerPtrToValuePtr<void>(Data, Idx) );
+				FSerializedPropertyScope SerializedProperty(Slot.GetUnderlyingArchive(), this);
+				SerializeItem(Stream.EnterElement(), ContainerPtrToValuePtr<void>(Data, Idx));
 			}
 		}
 	}
@@ -284,18 +285,21 @@ public:
 	 * @param	DefaultData		pointer to the location of the beginning of the data that should be compared against
 	 * @param	DefaultStruct	struct corresponding to the block of memory located at DefaultData 
 	 */
-	void SerializeNonMatchingBinProperty( FArchive& Ar, void* Data, void const* DefaultData, UStruct* DefaultStruct)
+	void SerializeNonMatchingBinProperty( FStructuredArchive::FSlot Slot, void* Data, void const* DefaultData, UStruct* DefaultStruct)
 	{
-		if( ShouldSerializeValue(Ar) )
+		FArchive& UnderlyingArchive = Slot.GetUnderlyingArchive();
+		FStructuredArchive::FStream Stream = Slot.EnterStream();
+
+		if( ShouldSerializeValue(UnderlyingArchive) )
 		{
 			for (int32 Idx = 0; Idx < ArrayDim; Idx++)
 			{
 				void *Target = ContainerPtrToValuePtr<void>(Data, Idx);
 				void const* Default = ContainerPtrToValuePtrForDefaults<void>(DefaultStruct, DefaultData, Idx);
-				if ( !Identical(Target, Default, Ar.GetPortFlags()) )
+				if ( !Identical(Target, Default, UnderlyingArchive.GetPortFlags()) )
 				{
-					FSerializedPropertyScope SerializedProperty(Ar, this);
-					SerializeItem( FStructuredArchiveFromArchive(Ar).GetSlot(), Target, Default );
+					FSerializedPropertyScope SerializedProperty(UnderlyingArchive, this);
+					SerializeItem( Stream.EnterElement(), Target, Default );
 				}
 			}
 		}
diff --git a/Engine/Source/Runtime/Engine/Classes/Engine/UserDefinedStruct.h b/Engine/Source/Runtime/Engine/Classes/Engine/UserDefinedStruct.h
--- a/Engine/Source/Runtime/Engine/Classes/Engine/UserDefinedStruct.h
+++ b/Engine/Source/Runtime/Engine/Classes/Engine/UserDefinedStruct.h
@@ -99,7 +99,7 @@ public:
 #endif	// WITH_EDITOR
 
 	// UObject interface.
-	virtual void Serialize(FArchive& Ar) override;
+	virtual void Serialize(FStructuredArchive::FRecord Record) override;
 	virtual void SerializeTaggedProperties(FArchive& Ar, uint8* Data, UStruct* DefaultsStruct, uint8* Defaults, const UObject* BreakRecursionIfFullyLoad=NULL) const override;
 	virtual FString PropertyNameToDisplayName(FName InName) const override;
 	// End of UObject interface.
diff --git a/Engine/Source/Runtime/Engine/Private/UserDefinedStruct.cpp b/Engine/Source/Runtime/Engine/Private/UserDefinedStruct.cpp
--- a/Engine/Source/Runtime/Engine/Private/UserDefinedStruct.cpp
+++ b/Engine/Source/Runtime/Engine/Private/UserDefinedStruct.cpp
@@ -100,23 +100,24 @@ UUserDefinedStruct::UUserDefinedStruct(const FObjectInitializer& ObjectInitializ
 #endif
 }
 
-void UUserDefinedStruct::Serialize(FArchive& Ar)
+void UUserDefinedStruct::Serialize(FStructuredArchive::FRecord Record)
 {
-	Super::Serialize( Ar );
+	Super::Serialize(Record);
+	FArchive& UnderlyingArchive = Record.GetUnderlyingArchive();
 
 	if (HasAnyFlags(RF_ClassDefaultObject))
 	{
 		return;
 	}
 
-	Ar.UsingCustomVersion(FFrameworkObjectVersion::GUID);
+	UnderlyingArchive.UsingCustomVersion(FFrameworkObjectVersion::GUID);
 
-	if (Ar.CustomVer(FFrameworkObjectVersion::GUID) >= FFrameworkObjectVersion::UserDefinedStructsStoreDefaultInstance)
+	if (UnderlyingArchive.CustomVer(FFrameworkObjectVersion::GUID) >= FFrameworkObjectVersion::UserDefinedStructsStoreDefaultInstance)
 	{
-		if (EUserDefinedStructureStatus::UDSS_UpToDate == Status && !(Ar.GetPortFlags() & PPF_Duplicate))
+		if (EUserDefinedStructureStatus::UDSS_UpToDate == Status && !(UnderlyingArchive.GetPortFlags() & PPF_Duplicate))
 		{
 			// If we're saving or loading new data, serialize our defaults
-			if (!DefaultStructInstance.IsValid() && Ar.IsLoading())
+			if (!DefaultStructInstance.IsValid() && UnderlyingArchive.IsLoading())
 			{
 				DefaultStructInstance.Recreate(this);
 			}
@@ -124,14 +125,14 @@ void UUserDefinedStruct::Serialize(FArchive& Ar)
 			uint8* StructData = DefaultStructInstance.GetStructMemory();
 
 			FScopedPlaceholderRawContainerTracker TrackStruct(StructData);
-			SerializeItem(Ar, StructData, nullptr);
+			SerializeItem(Record.EnterField(FIELD_NAME_TEXT("Data")), StructData, nullptr);
 		}
 	}
 
 #if WITH_EDITOR
-	if (Ar.IsLoading())
+	if (UnderlyingArchive.IsLoading())
 	{
-		if (Ar.CustomVer(FFrameworkObjectVersion::GUID) < FFrameworkObjectVersion::UserDefinedStructsBlueprintVisible)
+		if (UnderlyingArchive.CustomVer(FFrameworkObjectVersion::GUID) < FFrameworkObjectVersion::UserDefinedStructsBlueprintVisible)
 		{
 			for (TFieldIterator<UProperty> PropIt(this); PropIt; ++PropIt)
 			{
@@ -145,8 +146,8 @@ void UUserDefinedStruct::Serialize(FArchive& Ar)
 			// information at editor time about the user structure.
 			if (EditorData != nullptr)
 			{
-				Ar.Preload(EditorData);
-				if (!(Ar.GetPortFlags() & PPF_Duplicate))
+				UnderlyingArchive.Preload(EditorData);
+				if (!(UnderlyingArchive.GetPortFlags() & PPF_Duplicate))
 				{
 					FStructureEditorUtils::RecreateDefaultInstanceInEditorData(this);
 				}
@@ -428,7 +429,7 @@ void UUserDefinedStruct::AddReferencedObjects(UObject* InThis, FReferenceCollect
 	if (StructData)
 	{
 		FVerySlowReferenceCollectorArchiveScope CollectorScope(Collector.GetVerySlowReferenceCollectorArchive(), This);
-		This->SerializeBin(CollectorScope.GetArchive(), StructData);
+		This->SerializeBin(FStructuredArchiveFromArchive(CollectorScope.GetArchive()).GetSlot(), StructData);
 	}
 
 	Super::AddReferencedObjects(This, Collector);
