TBA: Loading fixed for tagged property serialization

Includes conversion of all UProperty::ConvertFromType() and SerializeFromMismatchedTag() functions to use structured archives

Lazy initialization of FArchiveFromStructruredArchive when loading, to support the possibility of an adapter being create around an object property serialize call to its inner UStruct, which then decides not to do anything and return false. Stops the ArchiveFromStructuredArchive from consuming the slot and getting upset later on when we try to serialize normal tagged properties from it.

Disabled lazy bulk data loading from text assets. Requires a bigger change to make it work.

Added some debug checks to json input formatter which track the current value stack size when a new object is pushed onto the stack, and makes sure that the stack has returned to the same size when the object is popped. Catches cases where we unpack an array/stream to the value stack but then don't consume all the items.

diff --git a/Engine/Plugins/Runtime/GameplayAbilities/Source/GameplayAbilities/Private/AttributeSet.cpp b/Engine/Plugins/Runtime/GameplayAbilities/Source/GameplayAbilities/Private/AttributeSet.cpp
--- a/Engine/Plugins/Runtime/GameplayAbilities/Source/GameplayAbilities/Private/AttributeSet.cpp
+++ b/Engine/Plugins/Runtime/GameplayAbilities/Source/GameplayAbilities/Private/AttributeSet.cpp
@@ -460,12 +460,12 @@ void FScalableFloat::SetScalingValue(float InCoeffecient, FName InRowName, UCurv
 	LocalCachedCurveID = INDEX_NONE;
 }
 
-bool FScalableFloat::SerializeFromMismatchedTag(const FPropertyTag& Tag, FArchive& Ar)
+bool FScalableFloat::SerializeFromMismatchedTag(const FPropertyTag& Tag, FStructuredArchive::FSlot Slot)
 {
 	if (Tag.Type == NAME_FloatProperty)
 	{
 		float OldValue;
-		Ar << OldValue;
+		Slot << OldValue;
 		*this = FScalableFloat(OldValue);
 
 		return true;
@@ -473,7 +473,7 @@ bool FScalableFloat::SerializeFromMismatchedTag(const FPropertyTag& Tag, FArchiv
 	else if (Tag.Type == NAME_IntProperty)
 	{
 		int32 OldValue;
-		Ar << OldValue;
+		Slot << OldValue;
 		*this = FScalableFloat((float)OldValue);
 
 		return true;
@@ -481,7 +481,7 @@ bool FScalableFloat::SerializeFromMismatchedTag(const FPropertyTag& Tag, FArchiv
 	else if (Tag.Type == NAME_Int8Property)
 	{
 		int8 OldValue;
-		Ar << OldValue;
+		Slot << OldValue;
 		*this = FScalableFloat((float)OldValue);
 
 		return true;
@@ -489,7 +489,7 @@ bool FScalableFloat::SerializeFromMismatchedTag(const FPropertyTag& Tag, FArchiv
 	else if (Tag.Type == NAME_Int16Property)
 	{
 		int16 OldValue;
-		Ar << OldValue;
+		Slot << OldValue;
 		*this = FScalableFloat((float)OldValue);
 
 		return true;
diff --git a/Engine/Plugins/Runtime/GameplayAbilities/Source/GameplayAbilities/Public/AttributeSet.h b/Engine/Plugins/Runtime/GameplayAbilities/Source/GameplayAbilities/Public/AttributeSet.h
--- a/Engine/Plugins/Runtime/GameplayAbilities/Source/GameplayAbilities/Public/AttributeSet.h
+++ b/Engine/Plugins/Runtime/GameplayAbilities/Source/GameplayAbilities/Public/AttributeSet.h
@@ -327,7 +327,7 @@ public:
 	void operator=(const FScalableFloat& Src);
 
 	/* Used to upgrade a float or int8/int16/int32 property into an FScalableFloat */
-	bool SerializeFromMismatchedTag(const FPropertyTag& Tag, FArchive& Ar);
+	bool SerializeFromMismatchedTag(const FPropertyTag& Tag, FStructuredArchive::FSlot Slot);
 
 private:
 
diff --git a/Engine/Source/Editor/UMGEditor/Private/WidgetBlueprint.cpp b/Engine/Source/Editor/UMGEditor/Private/WidgetBlueprint.cpp
--- a/Engine/Source/Editor/UMGEditor/Private/WidgetBlueprint.cpp
+++ b/Engine/Source/Editor/UMGEditor/Private/WidgetBlueprint.cpp
@@ -504,13 +504,14 @@ FDelegateRuntimeBinding FDelegateEditorBinding::ToRuntimeBinding(UWidgetBlueprin
 	return Binding;
 }
 
-bool FWidgetAnimation_DEPRECATED::SerializeFromMismatchedTag(struct FPropertyTag const& Tag, FArchive& Ar)
+bool FWidgetAnimation_DEPRECATED::SerializeFromMismatchedTag(struct FPropertyTag const& Tag, FStructuredArchive::FSlot Slot)
 {
 	static FName AnimationDataName("AnimationData");
 	if(Tag.Type == NAME_StructProperty && Tag.Name == AnimationDataName)
 	{
-		Ar << MovieScene;
-		Ar << AnimationBindings;
+		FStructuredArchive::FRecord Record = Slot.EnterRecord();
+		Record << NAMED_FIELD(MovieScene);
+		Record << NAMED_FIELD(AnimationBindings);
 		return true;
 	}
 
diff --git a/Engine/Source/Editor/UMGEditor/Public/WidgetBlueprint.h b/Engine/Source/Editor/UMGEditor/Public/WidgetBlueprint.h
--- a/Engine/Source/Editor/UMGEditor/Public/WidgetBlueprint.h
+++ b/Engine/Source/Editor/UMGEditor/Public/WidgetBlueprint.h
@@ -160,7 +160,7 @@ struct FWidgetAnimation_DEPRECATED
 	UPROPERTY()
 	TArray<FWidgetAnimationBinding> AnimationBindings;
 
-	bool SerializeFromMismatchedTag(struct FPropertyTag const& Tag, FArchive& Ar);
+	bool SerializeFromMismatchedTag(struct FPropertyTag const& Tag, FStructuredArchive::FSlot Slot);
 
 };
 
diff --git a/Engine/Source/Runtime/Core/Private/Serialization/ArchiveFromStructuredArchive.cpp b/Engine/Source/Runtime/Core/Private/Serialization/ArchiveFromStructuredArchive.cpp
--- a/Engine/Source/Runtime/Core/Private/Serialization/ArchiveFromStructuredArchive.cpp
+++ b/Engine/Source/Runtime/Core/Private/Serialization/ArchiveFromStructuredArchive.cpp
@@ -124,35 +124,29 @@ FArchive& FArchiveFromStructuredArchive::operator<<(class UObject*& Value)
 			}
 			else
 			{
-				TOptional<FStructuredArchive::FSlot> ObjectsSlot = Root->TryEnterField(FIELD_NAME_TEXT("Objects"), false);
-				if (ObjectsSlot.IsSet())
+				FStructuredArchive::FStream Stream = Root->EnterStream(FIELD_NAME_TEXT("Objects"));
+
+				// We know exactly which stream index we want to load here, but because of the API we need to read through them
+				// in order, consuming the string name until we reach the entry we want and then load it as a uobject reference.
+				// If we are loading from a text archive, we could easily specify here which index we want, and the internal formatter
+				// can just push that single value by itself onto the value stack, but that same API couldn't be implemented for a
+				// binary archive as we can't skip over entries because we don't know how big they are. Maybe we could specify a stride
+				// or something, but at this point the API is complex and pretty formatter specific. Thought required!
+				// For now, just consume all the string names of the objects up until the one we need, then load that as an object
+				// pointer.
+
+				FString Dummy;
+				for (int32 Index = 0; Index < Objects.Num(); ++Index)
 				{
-					// We know exactly which stream index we want to load here, but because of the API we need to read through them
-					// in order, consuming the string name until we reach the entry we want and then load it as a uobject reference.
-					// If we are loading from a text archive, we could easily specify here which index we want, and the internal formatter
-					// can just push that single value by itself onto the value stack, but that same API couldn't be implemented for a
-					// binary archive as we can't skip over entries because we don't know how big they are. Maybe we could specify a stride
-					// or something, but at this point the API is complex and pretty formatter specific. Thought required!
-					// For now, just consume all the string names of the objects up until the one we need, then load that as an object
-					// pointer.
-					FStructuredArchive::FStream Stream = ObjectsSlot->EnterStream();
-					for (int32 Index = 0; Index <= ObjectIdx; ++Index)
+					if (Index == ObjectIdx)
 					{
-						if (Index == ObjectIdx)
-						{
-							Stream.EnterElement() << Value;
-							Objects[ObjectIdx] = Value;
-						}
-						else
-						{
-							FString Dummy;
-							Stream.EnterElement() << Dummy;
-						}
+						Stream.EnterElement() << Value;
+						Objects[ObjectIdx] = Value;
+					}
+					else
+					{
+						Stream.EnterElement() << Dummy;
 					}
-				}
-				else
-				{
-					Value = nullptr;
 				}
 
 				Objects[ObjectIdx] = Value;
@@ -253,10 +247,11 @@ void FArchiveFromStructuredArchive::SerializeInternal(FStructuredArchive::FRecor
 				// when we enter the array here. We never read them, so I'm assuming they just sit there
 				// until we destroy this archive wrapper. Perhaps we need something in the API here to just access
 				// the size of the array but not preparing to access it's values?
-				int32 NumEntries = 0;
-				ObjectsSlot.GetValue().EnterArray(NumEntries);
-				Objects.AddUninitialized(NumEntries);
-				ObjectsValid.Init(false, NumEntries);
+				ObjectsSlot.GetValue() << ObjectNames;
+				//int32 NumEntries = 0;
+				//ObjectsSlot.GetValue().EnterArray(NumEntries);
+				Objects.AddUninitialized(ObjectNames.Num());
+				ObjectsValid.Init(false, ObjectNames.Num());
 			}
 			else
 			{
@@ -278,20 +273,20 @@ void FArchiveFromStructuredArchive::OpenArchive()
 {
 	if (!bWasOpened)
 	{
+		bWasOpened = true;
+
 		if (InnerArchive.IsTextFormat())
 		{
 			Root = RootSlot.EnterRecord();
 
 			if (IsLoading())
 			{
-				Commit();
+				SerializeInternal(Root.GetValue());
 			}
 		}
 		else
 		{
 			RootSlot.EnterStream();
 		}
-
-		bWasOpened = true;
 	}
 }
\ No newline at end of file
diff --git a/Engine/Source/Runtime/Core/Private/Serialization/StructuredArchive.cpp b/Engine/Source/Runtime/Core/Private/Serialization/StructuredArchive.cpp
--- a/Engine/Source/Runtime/Core/Private/Serialization/StructuredArchive.cpp
+++ b/Engine/Source/Runtime/Core/Private/Serialization/StructuredArchive.cpp
@@ -40,6 +40,7 @@ FStructuredArchiveChildReader::FStructuredArchiveChildReader(FStructuredArchive:
 
 	Archive = new FStructuredArchive(*Formatter);
 	Root.Emplace(Archive->Open());
+	InSlot.EnterRecord();
 }
 
 FStructuredArchiveChildReader::~FStructuredArchiveChildReader()
diff --git a/Engine/Source/Runtime/Core/Public/Math/Color.h b/Engine/Source/Runtime/Core/Public/Math/Color.h
--- a/Engine/Source/Runtime/Core/Public/Math/Color.h
+++ b/Engine/Source/Runtime/Core/Public/Math/Color.h
@@ -71,6 +71,18 @@ struct FLinearColor
 		return true;
 	}
 
+	friend void operator<<(FStructuredArchive::FSlot Slot, FLinearColor& Color)
+	{
+		FStructuredArchive::FRecord Record = Slot.EnterRecord();
+		Record << NAMED_ITEM("R", Color.R) << NAMED_ITEM("G", Color.G) << NAMED_ITEM("B", Color.B) << NAMED_ITEM("A", Color.A);
+	}
+
+	bool Serialize(FStructuredArchive::FSlot Slot)
+	{
+		Slot << *this;
+		return true;
+	}
+
 	// Conversions.
 	CORE_API FColor ToRGBE() const;
 
@@ -453,6 +465,18 @@ public:
 		return true;
 	}
 
+	// Serializer.
+	friend void operator<< (FStructuredArchive::FSlot Slot, FColor &Color)
+	{
+		return Slot << Color.DWColor();
+	}
+
+	bool Serialize(FStructuredArchive::FSlot Slot)
+	{
+		Slot << *this;
+		return true;
+	}
+
 	// Operators.
 	FORCEINLINE bool operator==( const FColor &C ) const
 	{
diff --git a/Engine/Source/Runtime/Core/Public/Serialization/ArchiveFromStructuredArchive.h b/Engine/Source/Runtime/Core/Public/Serialization/ArchiveFromStructuredArchive.h
--- a/Engine/Source/Runtime/Core/Public/Serialization/ArchiveFromStructuredArchive.h
+++ b/Engine/Source/Runtime/Core/Public/Serialization/ArchiveFromStructuredArchive.h
@@ -35,6 +35,11 @@ public:
 		return IsTextFormat() ? nullptr : Root->GetUnderlyingArchive().GetCacheableArchive();
 	}
 
+	FORCEINLINE bool ContainsData() const
+	{
+		return Buffer.Num() > 0;
+	}
+
 protected:
 
 	TOptional<FStructuredArchive::FRecord> Root;
@@ -54,6 +59,7 @@ private:
 	TArray<FName> Names;
 	TMap<FName, int32> NameToIndex;
 
+	TArray<FString> ObjectNames;
 	TArray<UObject*> Objects;
 	TBitArray<> ObjectsValid;
 	TMap<UObject*, int32> ObjectToIndex;
diff --git a/Engine/Source/Runtime/CoreUObject/Private/Misc/PackageName.cpp b/Engine/Source/Runtime/CoreUObject/Private/Misc/PackageName.cpp
--- a/Engine/Source/Runtime/CoreUObject/Private/Misc/PackageName.cpp
+++ b/Engine/Source/Runtime/CoreUObject/Private/Misc/PackageName.cpp
@@ -753,7 +753,7 @@ bool FPackageName::FindPackageFileWithoutExtension(const FString& InPackageFilen
 	};
 
 	// Loop through all known extensions and check if the file exist.
-	const int32 BaseIndex = (false && InAllowTextFormats && WITH_TEXT_ARCHIVE_SUPPORT) ? 0 : 2;
+	const int32 BaseIndex = (/*false &&*/ InAllowTextFormats && WITH_TEXT_ARCHIVE_SUPPORT) ? 0 : 2;
 	for (int32 ExtensionIndex = BaseIndex; ExtensionIndex < ARRAY_COUNT(PackageExtensions); ++ExtensionIndex)
 	{
 		FString   PackageFilename = InPackageFilename + *PackageExtensions[ExtensionIndex];
diff --git a/Engine/Source/Runtime/CoreUObject/Private/Serialization/ArchiveUObjectFromStructuredArchive.cpp b/Engine/Source/Runtime/CoreUObject/Private/Serialization/ArchiveUObjectFromStructuredArchive.cpp
--- a/Engine/Source/Runtime/CoreUObject/Private/Serialization/ArchiveUObjectFromStructuredArchive.cpp
+++ b/Engine/Source/Runtime/CoreUObject/Private/Serialization/ArchiveUObjectFromStructuredArchive.cpp
@@ -93,7 +93,24 @@ FArchive& FArchiveUObjectFromStructuredArchive::operator<<(FSoftObjectPath& Valu
 				IdxPtr = &(SoftObjectPathToIndex.Add(Value));
 				*IdxPtr = SoftObjectPaths.Add(Value);
 			}
+
+#if 1
+			// Temp workaround for emulating the behaviour of soft asset path serialization. Use these thread specific overrides to determine if we should actually save
+			// a reference to the path. If not, we still record the path in our list so that we get the correct behaviour later on when we pass these references down to the underlying
+			// archive
+			FName PackageName, PropertyName;
+			ESoftObjectPathCollectType CollectType = ESoftObjectPathCollectType::AlwaysCollect;
+			ESoftObjectPathSerializeType SerializeType = ESoftObjectPathSerializeType::AlwaysSerialize;
+			FSoftObjectPathThreadContext& ThreadContext = FSoftObjectPathThreadContext::Get();
+			ThreadContext.GetSerializationOptions(PackageName, PropertyName, CollectType, SerializeType);
+
+			if (SerializeType == ESoftObjectPathSerializeType::AlwaysSerialize)
+			{
+				Serialize(IdxPtr, sizeof(*IdxPtr));
+			}
+#else
 			Serialize(IdxPtr, sizeof(*IdxPtr));
+#endif
 		}
 	}
 	else
diff --git a/Engine/Source/Runtime/CoreUObject/Private/Serialization/Formatters/JsonArchiveInputFormatter.cpp b/Engine/Source/Runtime/CoreUObject/Private/Serialization/Formatters/JsonArchiveInputFormatter.cpp
--- a/Engine/Source/Runtime/CoreUObject/Private/Serialization/Formatters/JsonArchiveInputFormatter.cpp
+++ b/Engine/Source/Runtime/CoreUObject/Private/Serialization/Formatters/JsonArchiveInputFormatter.cpp
@@ -18,6 +18,7 @@ FJsonArchiveInputFormatter::FJsonArchiveInputFormatter(FArchive& InInner, TFunct
 	, ResolveObjectName(InResolveObjectName)
 {
 	Inner.ArIsTextFormat = true;
+	Inner.ArAllowLazyLoading = false;
 
 	TSharedPtr< FJsonObject > RootObject;
 	TSharedRef< TJsonReader<char> > Reader = TJsonReaderFactory<char>::Create(&InInner);
@@ -48,13 +49,13 @@ FStructuredArchiveFormatter* FJsonArchiveInputFormatter::CreateSubtreeReader()
 void FJsonArchiveInputFormatter::EnterRecord()
 {
 	TSharedPtr<FJsonValue>& Value = ValueStack.Top();
-	ObjectStack.Add(Value->AsObject());
+	ObjectStack.Add(FObjectRecord(Value->AsObject(), ValueStack.Num()));
 }
 
 void FJsonArchiveInputFormatter::EnterRecord(TArray<FString>& OutKeys)
 {
 	EnterRecord();
-	ObjectStack.Top()->Values.GetKeys(OutKeys);
+	ObjectStack.Top().JsonObject->Values.GetKeys(OutKeys);
 	for(int32 Idx = 0; Idx < OutKeys.Num(); Idx++)
 	{
 		OutKeys[Idx] = UnescapeFieldName(*OutKeys[Idx]);
@@ -63,12 +64,13 @@ void FJsonArchiveInputFormatter::EnterRecord(TArray<FString>& OutKeys)
 
 void FJsonArchiveInputFormatter::LeaveRecord()
 {
+	check(ValueStack.Num() == ObjectStack.Top().ValueCountOnCreation);
 	ObjectStack.Pop();
 }
 
 void FJsonArchiveInputFormatter::EnterField(FArchiveFieldName Name)
 {
-	TSharedPtr<FJsonObject>& NewObject = ObjectStack.Top();
+	TSharedPtr<FJsonObject>& NewObject = ObjectStack.Top().JsonObject;
 	TSharedPtr<FJsonValue> Field = NewObject->TryGetField(EscapeFieldName(Name.Name));
 	check(Field.IsValid());
 	ValueStack.Add(Field);
@@ -87,7 +89,7 @@ void FJsonArchiveInputFormatter::LeaveField()
 
 bool FJsonArchiveInputFormatter::TryEnterField(FArchiveFieldName Name, bool bEnterWhenSaving)
 {
-	TSharedPtr<FJsonObject>& NewObject = ObjectStack.Top();
+	TSharedPtr<FJsonObject>& NewObject = ObjectStack.Top().JsonObject;
 	TSharedPtr<FJsonValue> Field = NewObject->TryGetField(EscapeFieldName(Name.Name));
 	if (Field.IsValid())
 	{
@@ -285,7 +287,9 @@ void FJsonArchiveInputFormatter::Serialize(UObject*& Value)
 
 void FJsonArchiveInputFormatter::Serialize(FText& Value)
 {
-	checkNoEntry();	// Implement me
+	FStructuredArchive ChildArchive(*this);
+	FText::SerializeText(ChildArchive.Open(), Value);
+	ChildArchive.Close();
 }
 
 void FJsonArchiveInputFormatter::Serialize(FWeakObjectPtr& Value)
diff --git a/Engine/Source/Runtime/CoreUObject/Private/UObject/Class.cpp b/Engine/Source/Runtime/CoreUObject/Private/UObject/Class.cpp
--- a/Engine/Source/Runtime/CoreUObject/Private/UObject/Class.cpp
+++ b/Engine/Source/Runtime/CoreUObject/Private/UObject/Class.cpp
@@ -927,8 +927,10 @@ void UStruct::SerializeTaggedProperties(FStructuredArchive::FSlot Slot, uint8* D
 		// Load all stored properties, potentially skipping unknown ones.
 		while (1)
 		{
+			FStructuredArchive::FRecord PropertyRecord = PropertiesStream.EnterElement().EnterRecord();
+
 			FPropertyTag Tag;
-			UnderlyingArchive << Tag;
+			PropertyRecord << NAMED_FIELD(Tag);
 
 			if( Tag.Name == NAME_None )
 			{
@@ -1080,7 +1082,9 @@ void UStruct::SerializeTaggedProperties(FStructuredArchive::FSlot Slot, uint8* D
 			}
 			else
 			{
-				switch (Property->ConvertFromType(Tag, UnderlyingArchive, Data, DefaultsStruct))
+				FStructuredArchive::FSlot ValueSlot = PropertyRecord.EnterField(FIELD_NAME_TEXT("Value"));
+
+				switch (Property->ConvertFromType(Tag, ValueSlot, Data, DefaultsStruct))
 				{
 					case EConvertFromTypeResult::Converted:
 						bAdvanceProperty = true;
@@ -1097,7 +1101,7 @@ void UStruct::SerializeTaggedProperties(FStructuredArchive::FSlot Slot, uint8* D
 							uint8* DefaultsFromParent = Property->ContainerPtrToValuePtrForDefaults<uint8>(DefaultsStruct, Defaults, Tag.ArrayIndex);
 
 							// This property is ok.
-							Tag.SerializeTaggedProperty(UnderlyingArchive, Property, DestAddress, DefaultsFromParent);
+							Tag.SerializeTaggedProperty(ValueSlot, Property, DestAddress, DefaultsFromParent);
 							bAdvanceProperty = !UnderlyingArchive.IsCriticalError();
 						}
 						break;
@@ -2098,7 +2102,8 @@ void UScriptStruct::SerializeItem(FStructuredArchive::FSlot Slot, void* Value, v
 			bItemSerialized = TheCppStructOps->Serialize(Ar, Value);
 			if (bItemSerialized && !Slot.IsFilled())
 			{
-				Slot.EnterStream();
+				// The struct said that serialization succeeded but it didn't actually write anything.
+				Slot.EnterRecord();
 			}
 		}		
 	}
diff --git a/Engine/Source/Runtime/CoreUObject/Private/UObject/EnumProperty.cpp b/Engine/Source/Runtime/CoreUObject/Private/UObject/EnumProperty.cpp
--- a/Engine/Source/Runtime/CoreUObject/Private/UObject/EnumProperty.cpp
+++ b/Engine/Source/Runtime/CoreUObject/Private/UObject/EnumProperty.cpp
@@ -11,10 +11,10 @@
 namespace UE4EnumProperty_Private
 {
 	template <typename OldIntType>
-	void ConvertIntToEnumProperty(FArchive& Ar, UEnumProperty* EnumProp, UNumericProperty* UnderlyingProp, UEnum* Enum, void* Obj)
+	void ConvertIntToEnumProperty(FStructuredArchive::FSlot Slot, UEnumProperty* EnumProp, UNumericProperty* UnderlyingProp, UEnum* Enum, void* Obj)
 	{
 		OldIntType OldValue;
-		Ar << OldValue;
+		Slot << OldValue;
 
 		using LargeIntType = typename TChooseClass<TIsSigned<OldIntType>::Value, int64, uint64>::Result;
 
@@ -337,7 +337,7 @@ bool UEnumProperty::SameType(const UProperty* Other) const
 	return Super::SameType(Other) && static_cast<const UEnumProperty*>(Other)->Enum == Enum;
 }
 
-EConvertFromTypeResult UEnumProperty::ConvertFromType(const FPropertyTag& Tag, FArchive& Ar, uint8* Data, UStruct* DefaultsStruct)
+EConvertFromTypeResult UEnumProperty::ConvertFromType(const FPropertyTag& Tag, FStructuredArchive::FSlot Slot , uint8* Data, UStruct* DefaultsStruct)
 {
 	if ((Enum == nullptr) || (UnderlyingProp == nullptr))
 	{
@@ -359,18 +359,18 @@ EConvertFromTypeResult UEnumProperty::ConvertFromType(const FPropertyTag& Tag, F
 				InnerPropertyTag.EnumName = Enum->GetFName();
 				InnerPropertyTag.ArrayIndex = 0;
 
-				PreviousValue = (uint8)UNumericProperty::ReadEnumAsInt64(Ar, DefaultsStruct, InnerPropertyTag);
+				PreviousValue = (uint8)UNumericProperty::ReadEnumAsInt64(Slot, DefaultsStruct, InnerPropertyTag);
 			}
 			else
 			{
 				// a byte property gained an enum
-				Ar << PreviousValue;
+				Slot << PreviousValue;
 			}
 		}
 		else
 		{
 			// attempt to find the old enum and get the byte value from the serialized enum name
-			PreviousValue = (uint8)UNumericProperty::ReadEnumAsInt64(Ar, DefaultsStruct, Tag);
+			PreviousValue = (uint8)UNumericProperty::ReadEnumAsInt64(Slot, DefaultsStruct, Tag);
 		}
 
 		// now copy the value into the object's address space
@@ -378,31 +378,31 @@ EConvertFromTypeResult UEnumProperty::ConvertFromType(const FPropertyTag& Tag, F
 	}
 	else if (Tag.Type == NAME_Int8Property)
 	{
-		UE4EnumProperty_Private::ConvertIntToEnumProperty<int8>(Ar, this, UnderlyingProp, Enum, ContainerPtrToValuePtr<void>(Data, Tag.ArrayIndex));
+		UE4EnumProperty_Private::ConvertIntToEnumProperty<int8>(Slot, this, UnderlyingProp, Enum, ContainerPtrToValuePtr<void>(Data, Tag.ArrayIndex));
 	}
 	else if (Tag.Type == NAME_Int16Property)
 	{
-		UE4EnumProperty_Private::ConvertIntToEnumProperty<int16>(Ar, this, UnderlyingProp, Enum, ContainerPtrToValuePtr<void>(Data, Tag.ArrayIndex));
+		UE4EnumProperty_Private::ConvertIntToEnumProperty<int16>(Slot, this, UnderlyingProp, Enum, ContainerPtrToValuePtr<void>(Data, Tag.ArrayIndex));
 	}
 	else if (Tag.Type == NAME_IntProperty)
 	{
-		UE4EnumProperty_Private::ConvertIntToEnumProperty<int32>(Ar, this, UnderlyingProp, Enum, ContainerPtrToValuePtr<void>(Data, Tag.ArrayIndex));
+		UE4EnumProperty_Private::ConvertIntToEnumProperty<int32>(Slot, this, UnderlyingProp, Enum, ContainerPtrToValuePtr<void>(Data, Tag.ArrayIndex));
 	}
 	else if (Tag.Type == NAME_Int64Property)
 	{
-		UE4EnumProperty_Private::ConvertIntToEnumProperty<int64>(Ar, this, UnderlyingProp, Enum, ContainerPtrToValuePtr<void>(Data, Tag.ArrayIndex));
+		UE4EnumProperty_Private::ConvertIntToEnumProperty<int64>(Slot, this, UnderlyingProp, Enum, ContainerPtrToValuePtr<void>(Data, Tag.ArrayIndex));
 	}
 	else if (Tag.Type == NAME_UInt16Property)
 	{
-		UE4EnumProperty_Private::ConvertIntToEnumProperty<uint16>(Ar, this, UnderlyingProp, Enum, ContainerPtrToValuePtr<void>(Data, Tag.ArrayIndex));
+		UE4EnumProperty_Private::ConvertIntToEnumProperty<uint16>(Slot, this, UnderlyingProp, Enum, ContainerPtrToValuePtr<void>(Data, Tag.ArrayIndex));
 	}
 	else if (Tag.Type == NAME_UInt32Property)
 	{
-		UE4EnumProperty_Private::ConvertIntToEnumProperty<uint32>(Ar, this, UnderlyingProp, Enum, ContainerPtrToValuePtr<void>(Data, Tag.ArrayIndex));
+		UE4EnumProperty_Private::ConvertIntToEnumProperty<uint32>(Slot, this, UnderlyingProp, Enum, ContainerPtrToValuePtr<void>(Data, Tag.ArrayIndex));
 	}
 	else if (Tag.Type == NAME_UInt64Property)
 	{
-		UE4EnumProperty_Private::ConvertIntToEnumProperty<uint64>(Ar, this, UnderlyingProp, Enum, ContainerPtrToValuePtr<void>(Data, Tag.ArrayIndex));
+		UE4EnumProperty_Private::ConvertIntToEnumProperty<uint64>(Slot, this, UnderlyingProp, Enum, ContainerPtrToValuePtr<void>(Data, Tag.ArrayIndex));
 	}
 	else
 	{
diff --git a/Engine/Source/Runtime/CoreUObject/Private/UObject/LinkerLoad.cpp b/Engine/Source/Runtime/CoreUObject/Private/UObject/LinkerLoad.cpp
--- a/Engine/Source/Runtime/CoreUObject/Private/UObject/LinkerLoad.cpp
+++ b/Engine/Source/Runtime/CoreUObject/Private/UObject/LinkerLoad.cpp
@@ -1306,7 +1306,7 @@ FLinkerLoad::ELinkerStatus FLinkerLoad::SerializePackageFileSummary()
 		}
 		
 		// Propagate fact that package cannot use lazy loading to archive (aka this).
-		if( (Summary.PackageFlags & PKG_DisallowLazyLoading) )
+		if( (Summary.PackageFlags & PKG_DisallowLazyLoading) || IsTextFormat() )
 		{
 			ArAllowLazyLoading = false;
 		}
@@ -3409,7 +3409,7 @@ void FLinkerLoad::Preload( UObject* Object )
 					FScopedPlaceholderContainerTracker SerializingObjTracker(Object);
 #endif // USE_CIRCULAR_DEPENDENCY_LOAD_DEFERRING
 
-#if WITH_EDITOR
+#if WITH_EDITOR && WITH_TEXT_ARCHIVE_SUPPORT
 					bool bClassSupportsTextFormat = UClass::IsSafeToSerializeToStructuredArchives(Object->GetClass());
 #endif
 
@@ -3443,12 +3443,22 @@ void FLinkerLoad::Preload( UObject* Object )
 						UObject* PrevSerializedObject = ThreadContext.SerializedObject;
 						ThreadContext.SerializedObject = Object;
 
-#if WITH_TEXT_ARCHIVE_SUPPORT
+#if WITH_EDITOR && WITH_TEXT_ARCHIVE_SUPPORT
 						if (IsTextFormat())
 						{
 							FString ObjectName = Object->GetPathName(Object->GetOutermost());
 							FStructuredArchive::FSlot ExportSlot = StructuredArchiveRootRecord->EnterField(FIELD_NAME(*ObjectName));
-							Object->GetClass()->SerializeDefaultObject(Object, ExportSlot);
+
+							if (bClassSupportsTextFormat)
+							{
+								Object->GetClass()->SerializeDefaultObject(Object, ExportSlot);
+							}
+							else
+							{
+								FStructuredArchiveChildReader ChildReader(ExportSlot);
+								FArchiveUObjectFromStructuredArchive Adapter(ChildReader.GetRoot());
+								Object->GetClass()->SerializeDefaultObject(Object, Adapter);
+							}
 						}
 						else
 #endif
@@ -3486,6 +3496,7 @@ void FLinkerLoad::Preload( UObject* Object )
 							{
 								ObjectName = ObjectPath.Right(ObjectPath.Len() - PackagePath.Len() - 1);
 							}
+
 							FStructuredArchive::FSlot ExportSlot = StructuredArchiveRootRecord->EnterField(FIELD_NAME(*ObjectName));
 
 							if (bClassSupportsTextFormat)
diff --git a/Engine/Source/Runtime/CoreUObject/Private/UObject/ObjectRedirector.cpp b/Engine/Source/Runtime/CoreUObject/Private/UObject/ObjectRedirector.cpp
--- a/Engine/Source/Runtime/CoreUObject/Private/UObject/ObjectRedirector.cpp
+++ b/Engine/Source/Runtime/CoreUObject/Private/UObject/ObjectRedirector.cpp
@@ -33,11 +33,13 @@ void UObjectRedirector::PreSave(const class ITargetPlatform* TargetPlatform)
 	}
 }
 
-void UObjectRedirector::Serialize( FArchive& Ar )
+IMPLEMENT_FARCHIVE_SERIALIZER(UObjectRedirector);
+
+void UObjectRedirector::Serialize( FStructuredArchive::FRecord Record )
 {
-	Super::Serialize(Ar);
+	Super::Serialize(Record);
 
-	Ar << DestinationObject;
+	Record << NAMED_FIELD(DestinationObject);
 }
 
 bool UObjectRedirector::NeedsLoadForEditorGame() const
diff --git a/Engine/Source/Runtime/CoreUObject/Private/UObject/PrimaryAssetId.cpp b/Engine/Source/Runtime/CoreUObject/Private/UObject/PrimaryAssetId.cpp
--- a/Engine/Source/Runtime/CoreUObject/Private/UObject/PrimaryAssetId.cpp
+++ b/Engine/Source/Runtime/CoreUObject/Private/UObject/PrimaryAssetId.cpp
@@ -40,19 +40,19 @@ bool FPrimaryAssetType::ImportTextItem(const TCHAR*& Buffer, int32 PortFlags, UO
 	return true;
 }
 
-bool FPrimaryAssetType::SerializeFromMismatchedTag(struct FPropertyTag const& Tag, FArchive& Ar)
+bool FPrimaryAssetType::SerializeFromMismatchedTag(struct FPropertyTag const& Tag, FStructuredArchive::FSlot Slot)
 {
 	if (Tag.Type == NAME_NameProperty)
 	{
 		FName InName;
-		Ar << InName;
+		Slot << InName;
 		*this = FPrimaryAssetType(InName);
 		return true;
 	}
 	else if (Tag.Type == NAME_StrProperty)
 	{
 		FString InString;
-		Ar << InString;
+		Slot << InString;
 		*this = FPrimaryAssetType(*InString);
 		return true;
 	}
@@ -95,19 +95,19 @@ bool FPrimaryAssetId::ImportTextItem(const TCHAR*& Buffer, int32 PortFlags, UObj
 	return true;
 }
 
-bool FPrimaryAssetId::SerializeFromMismatchedTag(struct FPropertyTag const& Tag, FArchive& Ar)
+bool FPrimaryAssetId::SerializeFromMismatchedTag(struct FPropertyTag const& Tag, FStructuredArchive::FSlot Slot)
 {
 	if (Tag.Type == NAME_NameProperty)
 	{
 		FName InName;
-		Ar << InName;
+		Slot << InName;
 		*this = FPrimaryAssetId(InName.ToString());
 		return true;
 	}
 	else if (Tag.Type == NAME_StrProperty)
 	{
 		FString InString;
-		Ar << InString;
+		Slot << InString;
 		*this = FPrimaryAssetId(InString);
 		return true;
 	}
diff --git a/Engine/Source/Runtime/CoreUObject/Private/UObject/Property.cpp b/Engine/Source/Runtime/CoreUObject/Private/UObject/Property.cpp
--- a/Engine/Source/Runtime/CoreUObject/Private/UObject/Property.cpp
+++ b/Engine/Source/Runtime/CoreUObject/Private/UObject/Property.cpp
@@ -815,7 +815,7 @@ void UProperty::LinkInternal(FArchive& Ar)
 	check(0); // Link shouldn't call super...and we should never link an abstract property, like this base class
 }
 
-EConvertFromTypeResult UProperty::ConvertFromType(const FPropertyTag& Tag, FArchive& Ar, uint8* Data, UStruct* DefaultsStruct)
+EConvertFromTypeResult UProperty::ConvertFromType(const FPropertyTag& Tag, FStructuredArchive::FSlot Slot, uint8* Data, UStruct* DefaultsStruct)
 {
 	return EConvertFromTypeResult::UseSerializeItem;
 }
diff --git a/Engine/Source/Runtime/CoreUObject/Private/UObject/PropertyArray.cpp b/Engine/Source/Runtime/CoreUObject/Private/UObject/PropertyArray.cpp
--- a/Engine/Source/Runtime/CoreUObject/Private/UObject/PropertyArray.cpp
+++ b/Engine/Source/Runtime/CoreUObject/Private/UObject/PropertyArray.cpp
@@ -538,7 +538,7 @@ bool UArrayProperty::SameType(const UProperty* Other) const
 	return Super::SameType(Other) && Inner && Inner->SameType(((UArrayProperty*)Other)->Inner);
 }
 
-EConvertFromTypeResult UArrayProperty::ConvertFromType(const FPropertyTag& Tag, FArchive& Ar, uint8* Data, UStruct* DefaultsStruct)
+EConvertFromTypeResult UArrayProperty::ConvertFromType(const FPropertyTag& Tag, FStructuredArchive::FSlot Slot, uint8* Data, UStruct* DefaultsStruct)
 {
 	// TODO: The ArrayProperty Tag really doesn't have adequate information for
 	// many types. This should probably all be moved in to ::SerializeItem
@@ -547,8 +547,10 @@ EConvertFromTypeResult UArrayProperty::ConvertFromType(const FPropertyTag& Tag,
 	{
 		void* ArrayPropertyData = ContainerPtrToValuePtr<void>(Data);
 
+		FStructuredArchive::FRecord ArrayRecord = Slot.EnterRecord();
+
 		int32 ElementCount = 0;
-		Ar << ElementCount;
+		ArrayRecord << NAMED_ITEM("Count", ElementCount);
 
 		FScriptArrayHelper ScriptArrayHelper(this, ArrayPropertyData);
 		ScriptArrayHelper.EmptyAndAddValues(ElementCount);
@@ -560,11 +562,13 @@ EConvertFromTypeResult UArrayProperty::ConvertFromType(const FPropertyTag& Tag,
 			InnerPropertyTag.Type = Tag.InnerType;
 			InnerPropertyTag.ArrayIndex = 0;
 
-			if (Inner->ConvertFromType(InnerPropertyTag, Ar, ScriptArrayHelper.GetRawPtr(0), DefaultsStruct) == EConvertFromTypeResult::Converted)
+			FStructuredArchive::FStream ValueStream = ArrayRecord.EnterField(FIELD_NAME_TEXT("Value")).EnterStream();
+
+			if (Inner->ConvertFromType(InnerPropertyTag, ValueStream.EnterElement(), ScriptArrayHelper.GetRawPtr(0), DefaultsStruct) == EConvertFromTypeResult::Converted)
 			{
 				for (int32 i = 1; i < ElementCount; ++i)
 				{
-					verify(Inner->ConvertFromType(InnerPropertyTag, Ar, ScriptArrayHelper.GetRawPtr(i), DefaultsStruct) == EConvertFromTypeResult::Converted);
+					verify(Inner->ConvertFromType(InnerPropertyTag, ValueStream.EnterElement(), ScriptArrayHelper.GetRawPtr(i), DefaultsStruct) == EConvertFromTypeResult::Converted);
 				}
 
 				return EConvertFromTypeResult::Converted;
@@ -572,7 +576,7 @@ EConvertFromTypeResult UArrayProperty::ConvertFromType(const FPropertyTag& Tag,
 			// TODO: Implement SerializeFromMismatchedTag handling for arrays of structs
 			else
 			{
-				UE_LOG(LogClass, Warning, TEXT("Array Inner Type mismatch in %s of %s - Previous (%s) Current(%s) for package:  %s"), *Tag.Name.ToString(), *GetName(), *Tag.InnerType.ToString(), *Inner->GetID().ToString(), *Ar.GetArchiveName() );
+				UE_LOG(LogClass, Warning, TEXT("Array Inner Type mismatch in %s of %s - Previous (%s) Current(%s) for package:  %s"), *Tag.Name.ToString(), *GetName(), *Tag.InnerType.ToString(), *Inner->GetID().ToString(), *Slot.GetUnderlyingArchive().GetArchiveName() );
 				return EConvertFromTypeResult::CannotConvert;
 			}
 		}
diff --git a/Engine/Source/Runtime/CoreUObject/Private/UObject/PropertyBool.cpp b/Engine/Source/Runtime/CoreUObject/Private/UObject/PropertyBool.cpp
--- a/Engine/Source/Runtime/CoreUObject/Private/UObject/PropertyBool.cpp
+++ b/Engine/Source/Runtime/CoreUObject/Private/UObject/PropertyBool.cpp
@@ -191,10 +191,10 @@ FString UBoolProperty::GetCPPMacroType( FString& ExtendedTypeText ) const
 }
 
 template<typename T>
-void LoadFromType(UBoolProperty* Property, const FPropertyTag& Tag, FArchive& Ar, uint8* Data)
+void LoadFromType(UBoolProperty* Property, const FPropertyTag& Tag, FStructuredArchive::FSlot Slot, uint8* Data)
 {
 	T IntValue;
-	Ar << IntValue;
+	Slot << IntValue;
 
 	if (IntValue != 0)
 	{
@@ -212,23 +212,23 @@ void LoadFromType(UBoolProperty* Property, const FPropertyTag& Tag, FArchive& Ar
 	}
 }
 
-EConvertFromTypeResult UBoolProperty::ConvertFromType(const FPropertyTag& Tag, FArchive& Ar, uint8* Data, UStruct* DefaultsStruct)
+EConvertFromTypeResult UBoolProperty::ConvertFromType(const FPropertyTag& Tag, FStructuredArchive::FSlot Slot, uint8* Data, UStruct* DefaultsStruct)
 {
 	if (Tag.Type == NAME_IntProperty)
 	{
-		LoadFromType<int32>(this, Tag, Ar, Data);
+		LoadFromType<int32>(this, Tag, Slot, Data);
 	}
 	else if (Tag.Type == NAME_Int8Property)
 	{
-		LoadFromType<int8>(this, Tag, Ar, Data);
+		LoadFromType<int8>(this, Tag, Slot, Data);
 	}
 	else if (Tag.Type == NAME_Int16Property)
 	{
-		LoadFromType<int16>(this, Tag, Ar, Data);
+		LoadFromType<int16>(this, Tag, Slot, Data);
 	}
 	else if (Tag.Type == NAME_Int64Property)
 	{
-		LoadFromType<int64>(this, Tag, Ar, Data);
+		LoadFromType<int64>(this, Tag, Slot, Data);
 	}
 	else if (Tag.Type == NAME_ByteProperty)
 	{
@@ -241,7 +241,7 @@ EConvertFromTypeResult UBoolProperty::ConvertFromType(const FPropertyTag& Tag, F
 				return EConvertFromTypeResult::UseSerializeItem;
 			}
 
-			LoadFromType<uint8>(this, Tag, Ar, Data);
+			LoadFromType<uint8>(this, Tag, Slot, Data);
 		}
 		else
 		{
@@ -250,15 +250,15 @@ EConvertFromTypeResult UBoolProperty::ConvertFromType(const FPropertyTag& Tag, F
 	}
 	else if (Tag.Type == NAME_UInt16Property)
 	{
-		LoadFromType<uint16>(this, Tag, Ar, Data);
+		LoadFromType<uint16>(this, Tag, Slot, Data);
 	}
 	else if (Tag.Type == NAME_UInt32Property)
 	{
-		LoadFromType<uint32>(this, Tag, Ar, Data);
+		LoadFromType<uint32>(this, Tag, Slot, Data);
 	}
 	else if (Tag.Type == NAME_UInt64Property)
 	{
-		LoadFromType<uint64>(this, Tag, Ar, Data);
+		LoadFromType<uint64>(this, Tag, Slot, Data);
 	}
 	else
 	{
diff --git a/Engine/Source/Runtime/CoreUObject/Private/UObject/PropertyByte.cpp b/Engine/Source/Runtime/CoreUObject/Private/UObject/PropertyByte.cpp
--- a/Engine/Source/Runtime/CoreUObject/Private/UObject/PropertyByte.cpp
+++ b/Engine/Source/Runtime/CoreUObject/Private/UObject/PropertyByte.cpp
@@ -144,10 +144,10 @@ FString UByteProperty::GetCPPType( FString* ExtendedTypeText/*=NULL*/, uint32 CP
 template <typename OldIntType>
 struct TConvertIntToEnumProperty
 {
-	static void Convert(FArchive& Ar, UByteProperty* Property, UEnum* Enum, void* Obj, const FPropertyTag& Tag)
+	static void Convert(FStructuredArchive::FSlot Slot, UByteProperty* Property, UEnum* Enum, void* Obj, const FPropertyTag& Tag)
 	{
 		OldIntType OldValue;
-		Ar << OldValue;
+		Slot << OldValue;
 
 		uint8 NewValue = OldValue;
 		if (OldValue > (OldIntType)TNumericLimits<uint8>::Max() || !Enum->IsValidEnumValue(NewValue))
@@ -169,7 +169,7 @@ struct TConvertIntToEnumProperty
 	}
 };
 
-EConvertFromTypeResult UByteProperty::ConvertFromType(const FPropertyTag& Tag, FArchive& Ar, uint8* Data, UStruct* DefaultsStruct)
+EConvertFromTypeResult UByteProperty::ConvertFromType(const FPropertyTag& Tag, FStructuredArchive::FSlot Slot, uint8* Data, UStruct* DefaultsStruct)
 {
 	if (Tag.Type == NAME_ByteProperty  && ((Tag.EnumName == NAME_None) != (Enum == nullptr)))
 	{
@@ -187,12 +187,12 @@ EConvertFromTypeResult UByteProperty::ConvertFromType(const FPropertyTag& Tag, F
 			}
 
 			// simply pretend the property still doesn't have an enum and serialize the single byte
-			Ar << PreviousValue;
+			Slot << PreviousValue;
 		}
 		else
 		{
 			// attempt to find the old enum and get the byte value from the serialized enum name
-			PreviousValue = (uint8)ReadEnumAsInt64(Ar, DefaultsStruct, Tag);
+			PreviousValue = (uint8)ReadEnumAsInt64(Slot, DefaultsStruct, Tag);
 		}
 
 		// now copy the value into the object's address space
@@ -202,7 +202,7 @@ EConvertFromTypeResult UByteProperty::ConvertFromType(const FPropertyTag& Tag, F
 	{
 		// an enum property became a byte
 		// attempt to find the old enum and get the byte value from the serialized enum name
-		uint8 PreviousValue = (uint8)ReadEnumAsInt64(Ar, DefaultsStruct, Tag);
+		uint8 PreviousValue = (uint8)ReadEnumAsInt64(Slot, DefaultsStruct, Tag);
 
 		// now copy the value into the object's address space
 		SetPropertyValue_InContainer(Data, PreviousValue, Tag.ArrayIndex);
@@ -211,77 +211,77 @@ EConvertFromTypeResult UByteProperty::ConvertFromType(const FPropertyTag& Tag, F
 	{
 		if (Enum)
 		{
-			TConvertIntToEnumProperty<int8>::Convert(Ar, this, Enum, Data, Tag);
+			TConvertIntToEnumProperty<int8>::Convert(Slot, this, Enum, Data, Tag);
 		}
 		else
 		{
-			ConvertFromArithmeticValue<int8>(Ar, Data, Tag);
+			ConvertFromArithmeticValue<int8>(Slot, Data, Tag);
 		}
 	}
 	else if (Tag.Type == NAME_Int16Property)
 	{
 		if (Enum)
 		{
-			TConvertIntToEnumProperty<int16>::Convert(Ar, this, Enum, Data, Tag);
+			TConvertIntToEnumProperty<int16>::Convert(Slot, this, Enum, Data, Tag);
 		}
 		else
 		{
-			ConvertFromArithmeticValue<int16>(Ar, Data, Tag);
+			ConvertFromArithmeticValue<int16>(Slot, Data, Tag);
 		}
 	}
 	else if (Tag.Type == NAME_IntProperty)
 	{
 		if (Enum)
 		{
-			TConvertIntToEnumProperty<int32>::Convert(Ar, this, Enum, Data, Tag);
+			TConvertIntToEnumProperty<int32>::Convert(Slot, this, Enum, Data, Tag);
 		}
 		else
 		{
-			ConvertFromArithmeticValue<int32>(Ar, Data, Tag);
+			ConvertFromArithmeticValue<int32>(Slot, Data, Tag);
 		}
 	}
 	else if (Tag.Type == NAME_Int64Property)
 	{
 		if (Enum)
 		{
-			TConvertIntToEnumProperty<int64>::Convert(Ar, this, Enum, Data, Tag);
+			TConvertIntToEnumProperty<int64>::Convert(Slot, this, Enum, Data, Tag);
 		}
 		else
 		{
-			ConvertFromArithmeticValue<int64>(Ar, Data, Tag);
+			ConvertFromArithmeticValue<int64>(Slot, Data, Tag);
 		}
 	}
 	else if (Tag.Type == NAME_UInt16Property)
 	{
 		if (Enum)
 		{
-			TConvertIntToEnumProperty<uint16>::Convert(Ar, this, Enum, Data, Tag);
+			TConvertIntToEnumProperty<uint16>::Convert(Slot, this, Enum, Data, Tag);
 		}
 		else
 		{
-			ConvertFromArithmeticValue<uint16>(Ar, Data, Tag);
+			ConvertFromArithmeticValue<uint16>(Slot, Data, Tag);
 		}
 	}
 	else if (Tag.Type == NAME_UInt32Property)
 	{
 		if (Enum)
 		{
-			TConvertIntToEnumProperty<uint32>::Convert(Ar, this, Enum, Data, Tag);
+			TConvertIntToEnumProperty<uint32>::Convert(Slot, this, Enum, Data, Tag);
 		}
 		else
 		{
-			ConvertFromArithmeticValue<uint32>(Ar, Data, Tag);
+			ConvertFromArithmeticValue<uint32>(Slot, Data, Tag);
 		}
 	}
 	else if (Tag.Type == NAME_UInt64Property)
 	{
 		if (Enum)
 		{
-			TConvertIntToEnumProperty<uint64>::Convert(Ar, this, Enum, Data, Tag);
+			TConvertIntToEnumProperty<uint64>::Convert(Slot, this, Enum, Data, Tag);
 		}
 		else
 		{
-			ConvertFromArithmeticValue<uint64>(Ar, Data, Tag);
+			ConvertFromArithmeticValue<uint64>(Slot, Data, Tag);
 		}
 	}
 	else
diff --git a/Engine/Source/Runtime/CoreUObject/Private/UObject/PropertyMap.cpp b/Engine/Source/Runtime/CoreUObject/Private/UObject/PropertyMap.cpp
--- a/Engine/Source/Runtime/CoreUObject/Private/UObject/PropertyMap.cpp
+++ b/Engine/Source/Runtime/CoreUObject/Private/UObject/PropertyMap.cpp
@@ -894,26 +894,28 @@ bool UMapProperty::SameType(const UProperty* Other) const
 	return Super::SameType(Other) && KeyProp && ValueProp && KeyProp->SameType(MapProp->KeyProp) && ValueProp->SameType(MapProp->ValueProp);
 }
 
-EConvertFromTypeResult UMapProperty::ConvertFromType(const FPropertyTag& Tag, FArchive& Ar, uint8* Data, UStruct* DefaultsStruct)
+EConvertFromTypeResult UMapProperty::ConvertFromType(const FPropertyTag& Tag, FStructuredArchive::FSlot Slot, uint8* Data, UStruct* DefaultsStruct)
 {
+	FArchive& UnderlyingArchive = Slot.GetUnderlyingArchive();
+
 	// Ar related calls in this function must be mirrored in UMapProperty::SerializeItem
 	checkSlow(KeyProp);
 	checkSlow(ValueProp);
 
 	// Ensure that the key/value properties have been loaded before calling SerializeItem() on them
-	Ar.Preload(KeyProp);
-	Ar.Preload(ValueProp);
+	UnderlyingArchive.Preload(KeyProp);
+	UnderlyingArchive.Preload(ValueProp);
 
-	const auto SerializeOrConvert = [](UProperty* CurrentType, const FPropertyTag& InTag, FArchive& InAr, uint8* InData, UStruct* InDefaultsStruct) -> bool
+	const auto SerializeOrConvert = [](UProperty* CurrentType, const FPropertyTag& InTag, FStructuredArchive::FSlot Slot, uint8* InData, UStruct* InDefaultsStruct) -> bool
 	{
 		// Serialize wants the property address, while convert wants the container address. InData is the container address
 		if(CurrentType->GetID() == InTag.Type)
 		{
 			uint8* DestAddress = CurrentType->ContainerPtrToValuePtr<uint8>(InData, InTag.ArrayIndex);
-			CurrentType->SerializeItem(FStructuredArchiveFromArchive(InAr).GetSlot(), DestAddress, nullptr);
+			CurrentType->SerializeItem(Slot, DestAddress, nullptr);
 			return true;
 		}
-		else if( CurrentType->ConvertFromType(InTag, InAr, InData, InDefaultsStruct) == EConvertFromTypeResult::Converted )
+		else if( CurrentType->ConvertFromType(InTag, Slot, InData, InDefaultsStruct) == EConvertFromTypeResult::Converted )
 		{
 			return true;
 		}
@@ -946,18 +948,20 @@ EConvertFromTypeResult UMapProperty::ConvertFromType(const FPropertyTag& Tag, FA
 
 			bool bConversionSucceeded = true;
 
+			FStructuredArchive::FRecord ValueRecord = Slot.EnterRecord();
+
 			// When we saved this instance we wrote out any elements that were in the 'Default' instance but not in the 
 			// instance that was being written. Presumably we were constructed from our defaults and must now remove 
 			// any of the elements that were not present when we saved this Map:
 			int32 NumKeysToRemove = 0;
-			Ar << NumKeysToRemove;
+			FStructuredArchive::FArray KeysToRemoveArray = ValueRecord.EnterArray(FIELD_NAME_TEXT("KeysToRemove"), NumKeysToRemove);
 
 			if( NumKeysToRemove != 0 )
 			{
 				TempKeyStorage = (uint8*)FMemory::Malloc(MapLayout.SetLayout.Size);
 				KeyProp->InitializeValue(TempKeyStorage);
 
-				if (SerializeOrConvert( KeyProp, KeyPropertyTag, Ar, TempKeyStorage, DefaultsStruct))
+				if (SerializeOrConvert( KeyProp, KeyPropertyTag, KeysToRemoveArray.EnterElement(), TempKeyStorage, DefaultsStruct))
 				{
 					// If the key is in the map, remove it
 					int32 Found = MapHelper.FindMapIndexWithKey(TempKeyStorage);
@@ -969,7 +973,7 @@ EConvertFromTypeResult UMapProperty::ConvertFromType(const FPropertyTag& Tag, FA
 					// things are going fine, remove the rest of the keys:
 					for(int32 I = 1; I < NumKeysToRemove; ++I)
 					{
-						verify(SerializeOrConvert( KeyProp, KeyPropertyTag, Ar, TempKeyStorage, DefaultsStruct));
+						verify(SerializeOrConvert( KeyProp, KeyPropertyTag, KeysToRemoveArray.EnterElement(), TempKeyStorage, DefaultsStruct));
 						Found = MapHelper.FindMapIndexWithKey(TempKeyStorage);
 						if (Found != INDEX_NONE)
 						{
@@ -984,7 +988,7 @@ EConvertFromTypeResult UMapProperty::ConvertFromType(const FPropertyTag& Tag, FA
 			}
 
 			int32 NumEntries = 0;
-			Ar << NumEntries;
+			FStructuredArchive::FArray EntriesArray = ValueRecord.EnterArray(FIELD_NAME_TEXT("Entries"), NumEntries);
 
 			if( bConversionSucceeded )
 			{
@@ -996,7 +1000,7 @@ EConvertFromTypeResult UMapProperty::ConvertFromType(const FPropertyTag& Tag, FA
 						KeyProp->InitializeValue(TempKeyStorage);
 					}
 
-					if( SerializeOrConvert( KeyProp, KeyPropertyTag, Ar, TempKeyStorage, DefaultsStruct ) )
+					if( SerializeOrConvert( KeyProp, KeyPropertyTag, EntriesArray.EnterElement(), TempKeyStorage, DefaultsStruct ) )
 					{
 						// Add a new default value if the key doesn't currently exist in the map
 						bool bKeyAlreadyPresent = true;
@@ -1012,12 +1016,12 @@ EConvertFromTypeResult UMapProperty::ConvertFromType(const FPropertyTag& Tag, FA
 						KeyProp->CopyCompleteValue_InContainer(NextPairPtr, TempKeyStorage);
 
 						// Deserialize value
-						if( SerializeOrConvert( ValueProp, ValuePropertyTag, Ar, NextPairPtr, DefaultsStruct ) )
+						if( SerializeOrConvert( ValueProp, ValuePropertyTag, EntriesArray.EnterElement(), NextPairPtr, DefaultsStruct ) )
 						{
 							// first entry went fine, convert the rest:
 							for(int32 I = 1; I < NumEntries; ++I)
 							{
-								verify( SerializeOrConvert( KeyProp, KeyPropertyTag, Ar, TempKeyStorage, DefaultsStruct ) );
+								verify( SerializeOrConvert( KeyProp, KeyPropertyTag, EntriesArray.EnterElement(), TempKeyStorage, DefaultsStruct ) );
 								NextPairIndex = MapHelper.FindMapIndexWithKey(TempKeyStorage);
 								if (NextPairIndex == INDEX_NONE)
 								{
@@ -1027,7 +1031,7 @@ EConvertFromTypeResult UMapProperty::ConvertFromType(const FPropertyTag& Tag, FA
 								NextPairPtr = MapHelper.GetPairPtrWithoutCheck(NextPairIndex);
 								// This copy is unnecessary when the key was already in the map:
 								KeyProp->CopyCompleteValue_InContainer(NextPairPtr, TempKeyStorage);
-								verify( SerializeOrConvert( ValueProp, ValuePropertyTag, Ar, NextPairPtr, DefaultsStruct ) );
+								verify( SerializeOrConvert( ValueProp, ValuePropertyTag, EntriesArray.EnterElement(), NextPairPtr, DefaultsStruct ) );
 							}
 						}
 						else
@@ -1062,7 +1066,7 @@ EConvertFromTypeResult UMapProperty::ConvertFromType(const FPropertyTag& Tag, FA
 					*Tag.ValueType.ToString(),
 					*KeyProp->GetID().ToString(),
 					*ValueProp->GetID().ToString(),
-					*Ar.GetArchiveName()
+					*UnderlyingArchive.GetArchiveName()
 				);
 			}
 
diff --git a/Engine/Source/Runtime/CoreUObject/Private/UObject/PropertyName.cpp b/Engine/Source/Runtime/CoreUObject/Private/UObject/PropertyName.cpp
--- a/Engine/Source/Runtime/CoreUObject/Private/UObject/PropertyName.cpp
+++ b/Engine/Source/Runtime/CoreUObject/Private/UObject/PropertyName.cpp
@@ -52,12 +52,12 @@ const TCHAR* UNameProperty::ImportText_Internal( const TCHAR* Buffer, void* Data
 	return Buffer;
 }
 
-EConvertFromTypeResult UNameProperty::ConvertFromType(const FPropertyTag& Tag, FArchive& Ar, uint8* Data, UStruct* DefaultsStruct)
+EConvertFromTypeResult UNameProperty::ConvertFromType(const FPropertyTag& Tag, FStructuredArchive::FSlot Slot, uint8* Data, UStruct* DefaultsStruct)
 {
 	if (Tag.Type == NAME_StrProperty)
 	{
 		FString str;
-		Ar << str;
+		Slot << str;
 		SetPropertyValue_InContainer(Data, FName(*str), Tag.ArrayIndex);
 		return EConvertFromTypeResult::Converted;
 	}
@@ -66,7 +66,7 @@ EConvertFromTypeResult UNameProperty::ConvertFromType(const FPropertyTag& Tag, F
 	if (Tag.Type == NAME_TextProperty)
 	{ 
 		FText Text;  
-		Ar << Text;
+		Slot << Text;
 		const FName Name = FName(*Text.ToString());
 		SetPropertyValue_InContainer(Data, Name, Tag.ArrayIndex);
 		return EConvertFromTypeResult::Converted;
diff --git a/Engine/Source/Runtime/CoreUObject/Private/UObject/PropertyNumeric.cpp b/Engine/Source/Runtime/CoreUObject/Private/UObject/PropertyNumeric.cpp
--- a/Engine/Source/Runtime/CoreUObject/Private/UObject/PropertyNumeric.cpp
+++ b/Engine/Source/Runtime/CoreUObject/Private/UObject/PropertyNumeric.cpp
@@ -7,11 +7,11 @@
 #include "Templates/Casts.h"
 #include "UObject/UnrealType.h"
 
-int64 UNumericProperty::ReadEnumAsInt64(FArchive& Ar, UStruct* DefaultsStruct, const FPropertyTag& Tag)
+int64 UNumericProperty::ReadEnumAsInt64(FStructuredArchive::FSlot Slot, UStruct* DefaultsStruct, const FPropertyTag& Tag)
 {
 	//@warning: mirrors loading code in UByteProperty::SerializeItem() and UEnumProperty::SerializeItem()
 	FName EnumName;
-	Ar << EnumName;
+	Slot << EnumName;
 
 	UEnum* Enum = FindField<UEnum>(dynamic_cast<UClass*>(DefaultsStruct) ? static_cast<UClass*>(DefaultsStruct) : DefaultsStruct->GetTypedOuter<UClass>(), Tag.EnumName);
 	if (!Enum)
@@ -25,7 +25,7 @@ int64 UNumericProperty::ReadEnumAsInt64(FArchive& Ar, UStruct* DefaultsStruct, c
 		return 0;
 	}
 
-	Ar.Preload(Enum);
+	Slot.GetUnderlyingArchive().Preload(Enum);
 
 	// This handles redirects internally
 	int64 Result = Enum->GetValueByName(EnumName);
diff --git a/Engine/Source/Runtime/CoreUObject/Private/UObject/PropertyObject.cpp b/Engine/Source/Runtime/CoreUObject/Private/UObject/PropertyObject.cpp
--- a/Engine/Source/Runtime/CoreUObject/Private/UObject/PropertyObject.cpp
+++ b/Engine/Source/Runtime/CoreUObject/Private/UObject/PropertyObject.cpp
@@ -30,7 +30,7 @@ FString UObjectProperty::GetCPPMacroType( FString& ExtendedTypeText ) const
 	return TEXT("OBJECT");
 }
 
-EConvertFromTypeResult UObjectProperty::ConvertFromType(const FPropertyTag& Tag, FArchive& Ar, uint8* Data, UStruct* DefaultsStruct)
+EConvertFromTypeResult UObjectProperty::ConvertFromType(const FPropertyTag& Tag, FStructuredArchive::FSlot Slot, uint8* Data, UStruct* DefaultsStruct)
 {
 	static FName NAME_AssetObjectProperty = "AssetObjectProperty"; // old name of soft object property
 
@@ -38,7 +38,7 @@ EConvertFromTypeResult UObjectProperty::ConvertFromType(const FPropertyTag& Tag,
 	{
 		// This property used to be a TSoftObjectPtr<Foo> but is now a raw UObjectProperty Foo*, we can convert without loss of data
 		FSoftObjectPtr PreviousValue;
-		Ar << PreviousValue;
+		Slot << PreviousValue;
 
 		// now copy the value into the object's address space
 		UObject* PreviousValueObj = PreviousValue.LoadSynchronous();
diff --git a/Engine/Source/Runtime/CoreUObject/Private/UObject/PropertySet.cpp b/Engine/Source/Runtime/CoreUObject/Private/UObject/PropertySet.cpp
--- a/Engine/Source/Runtime/CoreUObject/Private/UObject/PropertySet.cpp
+++ b/Engine/Source/Runtime/CoreUObject/Private/UObject/PropertySet.cpp
@@ -756,13 +756,15 @@ bool USetProperty::SameType(const UProperty* Other) const
 	return Super::SameType(Other) && ElementProp && ElementProp->SameType(SetProp->ElementProp);
 }
 
-EConvertFromTypeResult USetProperty::ConvertFromType(const FPropertyTag& Tag, FArchive& Ar, uint8* Data, UStruct* DefaultsStruct)
+EConvertFromTypeResult USetProperty::ConvertFromType(const FPropertyTag& Tag, FStructuredArchive::FSlot Slot, uint8* Data, UStruct* DefaultsStruct)
 {
+	FArchive& UnderlyingArchive = Slot.GetUnderlyingArchive();
+
 	// Ar related calls in this function must be mirrored in USetProperty::ConvertFromType
 	checkSlow(ElementProp);
 
 	// Ensure that the element property has been loaded before calling SerializeItem() on it
-	Ar.Preload(ElementProp);
+	UnderlyingArchive.Preload(ElementProp);
 
 	// Ar related calls in this function must be mirrored in USetProperty::SerializeItem
 	if (Tag.Type == NAME_SetProperty)
@@ -787,18 +789,20 @@ EConvertFromTypeResult USetProperty::ConvertFromType(const FPropertyTag& Tag, FA
 
 			bool bConversionSucceeded = true;
 
+			FStructuredArchive::FRecord ValueRecord = Slot.EnterRecord();
+
 			// When we saved this instance we wrote out any elements that were in the 'Default' instance but not in the 
 			// instance that was being written. Presumably we were constructed from our defaults and must now remove 
 			// any of the elements that were not present when we saved this Set:
 			int32 NumElementsToRemove = 0;
-			Ar << NumElementsToRemove;
+			FStructuredArchive::FArray ElementsToRemoveArray = ValueRecord.EnterArray(FIELD_NAME_TEXT("ElementsToRemove"), NumElementsToRemove);
 
 			if(NumElementsToRemove)
 			{
 				TempElementStorage = (uint8*)FMemory::Malloc(SetLayout.Size);
 				ElementProp->InitializeValue(TempElementStorage);
 
-				if (ElementProp->ConvertFromType(InnerPropertyTag, Ar, TempElementStorage, DefaultsStruct) == EConvertFromTypeResult::Converted)
+				if (ElementProp->ConvertFromType(InnerPropertyTag, ElementsToRemoveArray.EnterElement(), TempElementStorage, DefaultsStruct) == EConvertFromTypeResult::Converted)
 				{
 					int32 Found = ScriptSetHelper.FindElementIndex(TempElementStorage);
 					if (Found != INDEX_NONE)
@@ -808,7 +812,7 @@ EConvertFromTypeResult USetProperty::ConvertFromType(const FPropertyTag& Tag, FA
 
 					for (int32 I = 1; I < NumElementsToRemove; ++I)
 					{
-						verify(ElementProp->ConvertFromType(InnerPropertyTag, Ar, TempElementStorage, DefaultsStruct) == EConvertFromTypeResult::Converted);
+						verify(ElementProp->ConvertFromType(InnerPropertyTag, ElementsToRemoveArray.EnterElement(), TempElementStorage, DefaultsStruct) == EConvertFromTypeResult::Converted);
 
 						Found = ScriptSetHelper.FindElementIndex(TempElementStorage);
 						if (Found != INDEX_NONE)
@@ -824,7 +828,7 @@ EConvertFromTypeResult USetProperty::ConvertFromType(const FPropertyTag& Tag, FA
 			}
 
 			int32 Num = 0;
-			Ar << Num;
+			FStructuredArchive::FArray ElementsArray = ValueRecord.EnterArray(FIELD_NAME_TEXT("Elements"), Num);
 
 			if(bConversionSucceeded)
 			{
@@ -839,7 +843,7 @@ EConvertFromTypeResult USetProperty::ConvertFromType(const FPropertyTag& Tag, FA
 
 					// and read the first entry, we have to check for conversion possibility again because 
 					// NumElementsToRemove may not have run (in fact, it likely did not):
-					if (ElementProp->ConvertFromType(InnerPropertyTag, Ar, TempElementStorage, DefaultsStruct) == EConvertFromTypeResult::Converted)
+					if (ElementProp->ConvertFromType(InnerPropertyTag, ElementsArray.EnterElement(), TempElementStorage, DefaultsStruct) == EConvertFromTypeResult::Converted)
 					{
 						if (ScriptSetHelper.FindElementIndex(TempElementStorage) == INDEX_NONE)
 						{
@@ -854,7 +858,7 @@ EConvertFromTypeResult USetProperty::ConvertFromType(const FPropertyTag& Tag, FA
 						for (int32 I = 1; I < Num; ++I)
 						{
 							// Read key into temporary storage
-							verify(ElementProp->ConvertFromType(InnerPropertyTag, Ar, TempElementStorage, DefaultsStruct) == EConvertFromTypeResult::Converted);
+							verify(ElementProp->ConvertFromType(InnerPropertyTag, ElementsArray.EnterElement(), TempElementStorage, DefaultsStruct) == EConvertFromTypeResult::Converted);
 
 							// Add a new entry if the element doesn't currently exist in the set
 							if (ScriptSetHelper.FindElementIndex(TempElementStorage) == INDEX_NONE)
@@ -879,7 +883,7 @@ EConvertFromTypeResult USetProperty::ConvertFromType(const FPropertyTag& Tag, FA
 			// if we could not convert the property ourself, then indicate that calling code needs to advance the property
 			if(!bConversionSucceeded)
 			{
-				UE_LOG(LogClass, Warning, TEXT("Set Element Type mismatch in %s of %s - Previous (%s) Current (%s) for package: %s"), *Tag.Name.ToString(), *GetName(), *Tag.InnerType.ToString(), *ElementProp->GetID().ToString(), *Ar.GetArchiveName() );
+				UE_LOG(LogClass, Warning, TEXT("Set Element Type mismatch in %s of %s - Previous (%s) Current (%s) for package: %s"), *Tag.Name.ToString(), *GetName(), *Tag.InnerType.ToString(), *ElementProp->GetID().ToString(), *UnderlyingArchive.GetArchiveName() );
 			}
 
 			return bConversionSucceeded ? EConvertFromTypeResult::Converted : EConvertFromTypeResult::CannotConvert;
diff --git a/Engine/Source/Runtime/CoreUObject/Private/UObject/PropertySoftObjectPtr.cpp b/Engine/Source/Runtime/CoreUObject/Private/UObject/PropertySoftObjectPtr.cpp
--- a/Engine/Source/Runtime/CoreUObject/Private/UObject/PropertySoftObjectPtr.cpp
+++ b/Engine/Source/Runtime/CoreUObject/Private/UObject/PropertySoftObjectPtr.cpp
@@ -114,7 +114,7 @@ const TCHAR* USoftObjectProperty::ImportText_Internal( const TCHAR* InBuffer, vo
 	}
 }
 
-EConvertFromTypeResult USoftObjectProperty::ConvertFromType(const FPropertyTag& Tag, FArchive& Ar, uint8* Data, UStruct* DefaultsStruct)
+EConvertFromTypeResult USoftObjectProperty::ConvertFromType(const FPropertyTag& Tag, FStructuredArchive::FSlot Slot, uint8* Data, UStruct* DefaultsStruct)
 {
 	static FName NAME_AssetObjectProperty = "AssetObjectProperty";
 	static FName NAME_SoftObjectPath = "SoftObjectPath";
@@ -122,14 +122,16 @@ EConvertFromTypeResult USoftObjectProperty::ConvertFromType(const FPropertyTag&
 	static FName NAME_StringAssetReference = "StringAssetReference";
 	static FName NAME_StringClassReference = "StringClassReference";
 
+	FArchive& Archive = Slot.GetUnderlyingArchive();
+
 	if (Tag.Type == NAME_AssetObjectProperty)
 	{
 		// Old name of soft object property, serialize normally
 		uint8* DestAddress = ContainerPtrToValuePtr<uint8>(Data, Tag.ArrayIndex);
 
-		Tag.SerializeTaggedProperty(Ar, this, DestAddress, nullptr);
+		Tag.SerializeTaggedProperty(Slot, this, DestAddress, nullptr);
 
-		if (Ar.IsCriticalError())
+		if (Archive.IsCriticalError())
 		{
 			return EConvertFromTypeResult::CannotConvert;
 		}
@@ -144,14 +146,14 @@ EConvertFromTypeResult USoftObjectProperty::ConvertFromType(const FPropertyTag&
 		FSoftObjectPtr* PropertyValue = GetPropertyValuePtr_InContainer(Data, Tag.ArrayIndex);
 		check(PropertyValue);
 
-		return PropertyValue->GetUniqueID().SerializeFromMismatchedTag(Tag, Ar) ? EConvertFromTypeResult::Converted : EConvertFromTypeResult::UseSerializeItem;
+		return PropertyValue->GetUniqueID().SerializeFromMismatchedTag(Tag, Slot) ? EConvertFromTypeResult::Converted : EConvertFromTypeResult::UseSerializeItem;
 	}
 	else if (Tag.Type == NAME_StructProperty && (Tag.StructName == NAME_SoftObjectPath || Tag.StructName == NAME_SoftClassPath || Tag.StructName == NAME_StringAssetReference || Tag.StructName == NAME_StringClassReference))
 	{
 		// This property used to be a FSoftObjectPath but is now a TSoftObjectPtr<Foo>
 		FSoftObjectPath PreviousValue;
 		// explicitly call Serialize to ensure that the various delegates needed for cooking are fired
-		PreviousValue.Serialize(Ar);
+		PreviousValue.Serialize(Slot);
 
 		// now copy the value into the object's address space
 		FSoftObjectPtr PreviousValueSoftObjectPtr;
diff --git a/Engine/Source/Runtime/CoreUObject/Private/UObject/PropertyStr.cpp b/Engine/Source/Runtime/CoreUObject/Private/UObject/PropertyStr.cpp
--- a/Engine/Source/Runtime/CoreUObject/Private/UObject/PropertyStr.cpp
+++ b/Engine/Source/Runtime/CoreUObject/Private/UObject/PropertyStr.cpp
@@ -9,13 +9,13 @@
 	UStrProperty.
 -----------------------------------------------------------------------------*/
 
-EConvertFromTypeResult UStrProperty::ConvertFromType(const FPropertyTag& Tag, FArchive& Ar, uint8* Data, UStruct* DefaultsStruct)
+EConvertFromTypeResult UStrProperty::ConvertFromType(const FPropertyTag& Tag, FStructuredArchive::FSlot Slot, uint8* Data, UStruct* DefaultsStruct)
 {
 	// Convert serialized text to string.
 	if (Tag.Type==NAME_TextProperty)
 	{ 
 		FText Text;
-		Ar << Text;
+		Slot << Text;
 		const FString String = FTextInspector::GetSourceString(Text) ? *FTextInspector::GetSourceString(Text) : TEXT("");
 		SetPropertyValue_InContainer(Data, String, Tag.ArrayIndex);
 
diff --git a/Engine/Source/Runtime/CoreUObject/Private/UObject/PropertyStruct.cpp b/Engine/Source/Runtime/CoreUObject/Private/UObject/PropertyStruct.cpp
--- a/Engine/Source/Runtime/CoreUObject/Private/UObject/PropertyStruct.cpp
+++ b/Engine/Source/Runtime/CoreUObject/Private/UObject/PropertyStruct.cpp
@@ -313,8 +313,10 @@ bool UStructProperty::SameType(const UProperty* Other) const
 	return Super::SameType(Other) && (Struct == ((UStructProperty*)Other)->Struct);
 }
 
-EConvertFromTypeResult UStructProperty::ConvertFromType(const FPropertyTag& Tag, FArchive& Ar, uint8* Data, UStruct* DefaultsStruct)
+EConvertFromTypeResult UStructProperty::ConvertFromType(const FPropertyTag& Tag, FStructuredArchive::FSlot Slot, uint8* Data, UStruct* DefaultsStruct)
 {
+	FArchive& UnderlyingArchive = Slot.GetUnderlyingArchive();
+
 	auto CanSerializeFromStructWithDifferentName = [](const FArchive& InAr, const FPropertyTag& PropertyTag, const UStructProperty* StructProperty)
 	{
 		if (InAr.UE4Ver() < VER_UE4_STRUCT_GUID_IN_PROPERTY_TAG)
@@ -332,25 +334,25 @@ EConvertFromTypeResult UStructProperty::ConvertFromType(const FPropertyTag& Tag,
 			UScriptStruct::ICppStructOps* CppStructOps = Struct->GetCppStructOps();
 			check(CppStructOps && CppStructOps->HasSerializeFromMismatchedTag()); // else should not have STRUCT_SerializeFromMismatchedTag
 			void* DestAddress = ContainerPtrToValuePtr<void>(Data, Tag.ArrayIndex);
-			if (CppStructOps->SerializeFromMismatchedTag(Tag, Ar, DestAddress))
+			if (CppStructOps->SerializeFromMismatchedTag(Tag, Slot, DestAddress))
 			{
 				return EConvertFromTypeResult::Converted;
 			}
 			else
 			{
-				UE_LOG(LogClass, Warning, TEXT("SerializeFromMismatchedTag failed: Type mismatch in %s of %s - Previous (%s) Current(StructProperty) for package:  %s"), *Tag.Name.ToString(), *GetName(), *Tag.Type.ToString(), *Ar.GetArchiveName());
+				UE_LOG(LogClass, Warning, TEXT("SerializeFromMismatchedTag failed: Type mismatch in %s of %s - Previous (%s) Current(StructProperty) for package:  %s"), *Tag.Name.ToString(), *GetName(), *Tag.Type.ToString(), *UnderlyingArchive.GetArchiveName());
 				return EConvertFromTypeResult::CannotConvert;
 			}
 		}
 
-		if (Tag.Type == NAME_StructProperty && Tag.StructName != Struct->GetFName() && !CanSerializeFromStructWithDifferentName(Ar, Tag, this))
+		if (Tag.Type == NAME_StructProperty && Tag.StructName != Struct->GetFName() && !CanSerializeFromStructWithDifferentName(UnderlyingArchive, Tag, this))
 		{
 			//handle Vector -> Vector4 upgrades here because using the SerializeFromMismatchedTag system would cause a dependency from Core -> CoreUObject
 			if (Tag.StructName == NAME_Vector && Struct->GetFName() == NAME_Vector4)
 			{
 				void* DestAddress = ContainerPtrToValuePtr<void>(Data, Tag.ArrayIndex);
 				FVector OldValue;
-				Ar << OldValue;
+				Slot << OldValue;
 
 				//only set X/Y/Z.  The W should already have been set to the property specific default and we don't want to trash it by forcing 0 or 1.
 				FVector4* DestValue = (FVector4*)DestAddress;
@@ -362,7 +364,7 @@ EConvertFromTypeResult UStructProperty::ConvertFromType(const FPropertyTag& Tag,
 			}
 
 			UE_LOG(LogClass, Warning, TEXT("Property %s of %s has a struct type mismatch (tag %s != prop %s) in package:  %s. If that struct got renamed, add an entry to ActiveStructRedirects."),
-				*Tag.Name.ToString(), *GetName(), *Tag.StructName.ToString(), *Struct->GetName(), *Ar.GetArchiveName());
+				*Tag.Name.ToString(), *GetName(), *Tag.StructName.ToString(), *Struct->GetName(), *UnderlyingArchive.GetArchiveName());
 			return EConvertFromTypeResult::CannotConvert;
 		}
 	}
diff --git a/Engine/Source/Runtime/CoreUObject/Private/UObject/SavePackage.cpp b/Engine/Source/Runtime/CoreUObject/Private/UObject/SavePackage.cpp
--- a/Engine/Source/Runtime/CoreUObject/Private/UObject/SavePackage.cpp
+++ b/Engine/Source/Runtime/CoreUObject/Private/UObject/SavePackage.cpp
@@ -5239,8 +5239,15 @@ FSavePackageResultStruct UPackage::Save(UPackage* InOuter, UObject* Base, EObjec
 
 							if ( Export.Object->HasAnyFlags(RF_ClassDefaultObject) )
 							{
-								FArchiveUObjectFromStructuredArchive Adapter(ExportSlot);
-								Export.Object->GetClass()->SerializeDefaultObject(Export.Object, Adapter);
+								if (bSupportsText)
+								{
+									Export.Object->GetClass()->SerializeDefaultObject(Export.Object, ExportSlot);
+								}
+								else
+								{
+									FArchiveUObjectFromStructuredArchive Adapter(ExportSlot);
+									Export.Object->GetClass()->SerializeDefaultObject(Export.Object, Adapter);
+								}
 							}
 							else
 							{
diff --git a/Engine/Source/Runtime/CoreUObject/Private/UObject/SoftObjectPath.cpp b/Engine/Source/Runtime/CoreUObject/Private/UObject/SoftObjectPath.cpp
--- a/Engine/Source/Runtime/CoreUObject/Private/UObject/SoftObjectPath.cpp
+++ b/Engine/Source/Runtime/CoreUObject/Private/UObject/SoftObjectPath.cpp
@@ -108,6 +108,14 @@ bool FSoftObjectPath::Serialize(FArchive& Ar)
 	return true;
 }
 
+bool FSoftObjectPath::Serialize(FStructuredArchive::FSlot Slot)
+{
+	// Archivers will call back into SerializePath for the various fixups
+	Slot << *this;
+
+	return true;
+}
+
 void FSoftObjectPath::SerializePath(FArchive& Ar)
 {
 	bool bSerializeInternals = true;
@@ -271,12 +279,12 @@ bool FSoftObjectPath::ImportTextItem(const TCHAR*& Buffer, int32 PortFlags, UObj
  * @param Ar Archive to serialize from.
  */
 template <class TypePolicy>
-bool SerializeFromMismatchedTagTemplate(FString& Output, const FPropertyTag& Tag, FArchive& Ar)
+bool SerializeFromMismatchedTagTemplate(FString& Output, const FPropertyTag& Tag, FStructuredArchive::FSlot Slot)
 {
 	if (Tag.Type == TypePolicy::GetTypeName())
 	{
 		typename TypePolicy::Type* ObjPtr = nullptr;
-		Ar << ObjPtr;
+		Slot << ObjPtr;
 		if (ObjPtr)
 		{
 			Output = ObjPtr->GetPathName();
@@ -290,7 +298,7 @@ bool SerializeFromMismatchedTagTemplate(FString& Output, const FPropertyTag& Tag
 	else if (Tag.Type == NAME_StrProperty)
 	{
 		FString String;
-		Ar << String;
+		Slot << String;
 
 		Output = String;
 		return true;
@@ -298,7 +306,7 @@ bool SerializeFromMismatchedTagTemplate(FString& Output, const FPropertyTag& Tag
 	return false;
 }
 
-bool FSoftObjectPath::SerializeFromMismatchedTag(struct FPropertyTag const& Tag, FArchive& Ar)
+bool FSoftObjectPath::SerializeFromMismatchedTag(struct FPropertyTag const& Tag, FStructuredArchive::FSlot Slot)
 {
 	struct UObjectTypePolicy
 	{
@@ -308,9 +316,9 @@ bool FSoftObjectPath::SerializeFromMismatchedTag(struct FPropertyTag const& Tag,
 
 	FString Path = ToString();
 
-	bool bReturn = SerializeFromMismatchedTagTemplate<UObjectTypePolicy>(Path, Tag, Ar);
+	bool bReturn = SerializeFromMismatchedTagTemplate<UObjectTypePolicy>(Path, Tag, Slot);
 
-	if (Ar.IsLoading())
+	if (Slot.GetUnderlyingArchive().IsLoading())
 	{
 		SetPath(MoveTemp(Path));
 		PostLoadPath();
@@ -461,7 +469,7 @@ bool FSoftObjectPath::FixupCoreRedirects()
 	return false;
 }
 
-bool FSoftClassPath::SerializeFromMismatchedTag(struct FPropertyTag const& Tag, FArchive& Ar)
+bool FSoftClassPath::SerializeFromMismatchedTag(struct FPropertyTag const& Tag, FStructuredArchive::FSlot Slot)
 {
 	struct UClassTypePolicy
 	{
@@ -472,9 +480,9 @@ bool FSoftClassPath::SerializeFromMismatchedTag(struct FPropertyTag const& Tag,
 
 	FString Path = ToString();
 
-	bool bReturn = SerializeFromMismatchedTagTemplate<UClassTypePolicy>(Path, Tag, Ar);
+	bool bReturn = SerializeFromMismatchedTagTemplate<UClassTypePolicy>(Path, Tag, Slot);
 
-	if (Ar.IsLoading())
+	if (Slot.GetUnderlyingArchive().IsLoading())
 	{
 		SetPath(MoveTemp(Path));
 		PostLoadPath();
diff --git a/Engine/Source/Runtime/CoreUObject/Private/UObject/TextProperty.cpp b/Engine/Source/Runtime/CoreUObject/Private/UObject/TextProperty.cpp
--- a/Engine/Source/Runtime/CoreUObject/Private/UObject/TextProperty.cpp
+++ b/Engine/Source/Runtime/CoreUObject/Private/UObject/TextProperty.cpp
@@ -9,13 +9,13 @@
 #include "Internationalization/StringTableRegistry.h"
 #include "Serialization/ArchiveUObjectFromStructuredArchive.h"
 
-EConvertFromTypeResult UTextProperty::ConvertFromType(const FPropertyTag& Tag, FArchive& Ar, uint8* Data, UStruct* DefaultsStruct)
+EConvertFromTypeResult UTextProperty::ConvertFromType(const FPropertyTag& Tag, FStructuredArchive::FSlot Slot, uint8* Data, UStruct* DefaultsStruct)
 {
 	// Convert serialized string to text.
 	if (Tag.Type==NAME_StrProperty)
 	{
 		FString str;
-		Ar << str;
+		Slot << str;
 		FText Text = FText::FromString(str);
 		Text.TextData->PersistText();
 		Text.Flags |= ETextFlag::ConvertedProperty;
@@ -27,7 +27,7 @@ EConvertFromTypeResult UTextProperty::ConvertFromType(const FPropertyTag& Tag, F
 	if (Tag.Type==NAME_NameProperty)
 	{
 		FName Name;
-		Ar << Name;
+		Slot << Name;
 		FText Text = FText::FromName(Name);
 		Text.Flags |= ETextFlag::ConvertedProperty;
 		SetPropertyValue_InContainer(Data, Text, Tag.ArrayIndex);
diff --git a/Engine/Source/Runtime/CoreUObject/Public/Serialization/Formatters/JsonArchiveInputFormatter.h b/Engine/Source/Runtime/CoreUObject/Public/Serialization/Formatters/JsonArchiveInputFormatter.h
--- a/Engine/Source/Runtime/CoreUObject/Public/Serialization/Formatters/JsonArchiveInputFormatter.h
+++ b/Engine/Source/Runtime/CoreUObject/Public/Serialization/Formatters/JsonArchiveInputFormatter.h
@@ -71,8 +71,21 @@ public:
 private:
 	FArchive& Inner;
 
+	struct FObjectRecord
+	{
+		FObjectRecord(TSharedPtr<FJsonObject> InJsonObject, int64 InValueCount)
+			: JsonObject(InJsonObject)
+			, ValueCountOnCreation(InValueCount)
+		{
+
+		}
+
+		TSharedPtr<FJsonObject> JsonObject;
+		int64 ValueCountOnCreation;	// For debugging purposes, so we can ensure all values have been consumed
+	};
+
 	TFunction<UObject* (const FString&)> ResolveObjectName;
-	TArray<TSharedPtr<FJsonObject>> ObjectStack;
+	TArray<FObjectRecord> ObjectStack;
 	TArray<TSharedPtr<FJsonValue>> ValueStack;
 	TArray<TMap<FString, TSharedPtr<FJsonValue>>::TIterator> MapIteratorStack;
 
diff --git a/Engine/Source/Runtime/CoreUObject/Public/UObject/Class.h b/Engine/Source/Runtime/CoreUObject/Public/UObject/Class.h
--- a/Engine/Source/Runtime/CoreUObject/Public/UObject/Class.h
+++ b/Engine/Source/Runtime/CoreUObject/Public/UObject/Class.h
@@ -782,15 +782,15 @@ FORCEINLINE typename TEnableIf<TStructOpsTypeTraits<CPPSTRUCT>::WithImportTextIt
  * Selection of SerializeFromMismatchedTag call.
  */
 template<class CPPSTRUCT>
-FORCEINLINE typename TEnableIf<!TStructOpsTypeTraits<CPPSTRUCT>::WithSerializeFromMismatchedTag, bool>::Type SerializeFromMismatchedTagOrNot(FPropertyTag const& Tag, FArchive& Ar, CPPSTRUCT *Data)
+FORCEINLINE typename TEnableIf<!TStructOpsTypeTraits<CPPSTRUCT>::WithSerializeFromMismatchedTag, bool>::Type SerializeFromMismatchedTagOrNot(FPropertyTag const& Tag, FStructuredArchive::FSlot Slot, CPPSTRUCT *Data)
 {
 	return false;
 }
 
 template<class CPPSTRUCT>
-FORCEINLINE typename TEnableIf<TStructOpsTypeTraits<CPPSTRUCT>::WithSerializeFromMismatchedTag, bool>::Type SerializeFromMismatchedTagOrNot(FPropertyTag const& Tag, FArchive& Ar, CPPSTRUCT *Data)
+FORCEINLINE typename TEnableIf<TStructOpsTypeTraits<CPPSTRUCT>::WithSerializeFromMismatchedTag, bool>::Type SerializeFromMismatchedTagOrNot(FPropertyTag const& Tag, FStructuredArchive::FSlot Slot, CPPSTRUCT *Data)
 {
-	return Data->SerializeFromMismatchedTag(Tag, Ar);
+	return Data->SerializeFromMismatchedTag(Tag, Slot);
 }
 
 
@@ -940,7 +940,7 @@ public:
 		 * Serialize this structure, from some other tag
 		 * @return true if this succeeded, false will trigger a warning and not serialize at all
 		 */
-		virtual bool SerializeFromMismatchedTag(struct FPropertyTag const& Tag, FArchive& Ar, void *Data) = 0;
+		virtual bool SerializeFromMismatchedTag(struct FPropertyTag const& Tag, FStructuredArchive::FSlot Slot, void *Data) = 0;
 
 		/** return true if this struct has a GetTypeHash */
 		virtual bool HasGetTypeHash() = 0;
@@ -1104,10 +1104,10 @@ public:
 		{
 			return TTraits::WithSerializeFromMismatchedTag;
 		}
-		virtual bool SerializeFromMismatchedTag(struct FPropertyTag const& Tag, FArchive& Ar, void *Data) override
+		virtual bool SerializeFromMismatchedTag(struct FPropertyTag const& Tag, FStructuredArchive::FSlot Slot, void *Data) override
 		{
 			check(TTraits::WithSerializeFromMismatchedTag); // don't call this if we have indicated it is not allowed
-			return SerializeFromMismatchedTagOrNot(Tag, Ar, (CPPSTRUCT*)Data);
+			return SerializeFromMismatchedTagOrNot(Tag, Slot, (CPPSTRUCT*)Data);
 		}
 		virtual bool HasGetTypeHash() override
 		{
diff --git a/Engine/Source/Runtime/CoreUObject/Public/UObject/EnumProperty.h b/Engine/Source/Runtime/CoreUObject/Public/UObject/EnumProperty.h
--- a/Engine/Source/Runtime/CoreUObject/Public/UObject/EnumProperty.h
+++ b/Engine/Source/Runtime/CoreUObject/Public/UObject/EnumProperty.h
@@ -44,7 +44,7 @@ public:
 	virtual const TCHAR* ImportText_Internal( const TCHAR* Buffer, void* Data, int32 PortFlags, UObject* OwnerObject, FOutputDevice* ErrorText ) const override;
 	virtual int32 GetMinAlignment() const override;
 	virtual bool SameType(const UProperty* Other) const override;
-	virtual EConvertFromTypeResult ConvertFromType(const FPropertyTag& Tag, FArchive& Ar, uint8* Data, UStruct* DefaultsStruct) override;
+	virtual EConvertFromTypeResult ConvertFromType(const FPropertyTag& Tag, FStructuredArchive::FSlot Slot, uint8* Data, UStruct* DefaultsStruct) override;
 	// End of UProperty interface
 
 	/**
diff --git a/Engine/Source/Runtime/CoreUObject/Public/UObject/ObjectRedirector.h b/Engine/Source/Runtime/CoreUObject/Public/UObject/ObjectRedirector.h
--- a/Engine/Source/Runtime/CoreUObject/Public/UObject/ObjectRedirector.h
+++ b/Engine/Source/Runtime/CoreUObject/Public/UObject/ObjectRedirector.h
@@ -16,13 +16,14 @@
  */
 class UObjectRedirector : public UObject
 {
-	DECLARE_CASTED_CLASS_INTRINSIC_WITH_API(UObjectRedirector, UObject, 0, TEXT("/Script/CoreUObject"), CASTCLASS_None, COREUOBJECT_API)
+	DECLARE_CASTED_CLASS_INTRINSIC_WITH_API(UObjectRedirector, UObject, CLASS_MatchedSerializers, TEXT("/Script/CoreUObject"), CASTCLASS_None, COREUOBJECT_API)
 
 	// Variables.
 	UObject*		DestinationObject;
 	// UObject interface.
 	virtual void PreSave(const class ITargetPlatform* TargetPlatform) override;
-	void Serialize( FArchive& Ar ) override;
+	void Serialize(FArchive& Ar) override;
+	void Serialize(FStructuredArchive::FRecord Record) override;
 	virtual bool NeedsLoadForEditorGame() const override;
 	virtual void GetAssetRegistryTags(TArray<FAssetRegistryTag>& OutTags) const override;
 
diff --git a/Engine/Source/Runtime/CoreUObject/Public/UObject/PrimaryAssetId.h b/Engine/Source/Runtime/CoreUObject/Public/UObject/PrimaryAssetId.h
--- a/Engine/Source/Runtime/CoreUObject/Public/UObject/PrimaryAssetId.h
+++ b/Engine/Source/Runtime/CoreUObject/Public/UObject/PrimaryAssetId.h
@@ -58,7 +58,7 @@ struct FPrimaryAssetType
 	/** UStruct Overrides */
 	bool ExportTextItem(FString& ValueStr, FPrimaryAssetType const& DefaultValue, UObject* Parent, int32 PortFlags, UObject* ExportRootScope) const;
 	bool ImportTextItem(const TCHAR*& Buffer, int32 PortFlags, UObject* Parent, FOutputDevice* ErrorText);
-	bool SerializeFromMismatchedTag(struct FPropertyTag const& Tag, FArchive& Ar);
+	bool SerializeFromMismatchedTag(struct FPropertyTag const& Tag, FStructuredArchive::FSlot Slot);
 
 	friend inline uint32 GetTypeHash(const FPrimaryAssetType& Key)
 	{
@@ -153,7 +153,7 @@ struct FPrimaryAssetId
 	/** UStruct Overrides */
 	bool ExportTextItem(FString& ValueStr, FPrimaryAssetId const& DefaultValue, UObject* Parent, int32 PortFlags, UObject* ExportRootScope) const;
 	bool ImportTextItem(const TCHAR*& Buffer, int32 PortFlags, UObject* Parent, FOutputDevice* ErrorText);
-	bool SerializeFromMismatchedTag(struct FPropertyTag const& Tag, FArchive& Ar);
+	bool SerializeFromMismatchedTag(struct FPropertyTag const& Tag, FStructuredArchive::FSlot Slot);
 
 	friend inline uint32 GetTypeHash(const FPrimaryAssetId& Key)
 	{
diff --git a/Engine/Source/Runtime/CoreUObject/Public/UObject/SoftObjectPath.h b/Engine/Source/Runtime/CoreUObject/Public/UObject/SoftObjectPath.h
--- a/Engine/Source/Runtime/CoreUObject/Public/UObject/SoftObjectPath.h
+++ b/Engine/Source/Runtime/CoreUObject/Public/UObject/SoftObjectPath.h
@@ -133,6 +133,7 @@ struct COREUOBJECT_API FSoftObjectPath
 
 	/** Struct overrides */
 	bool Serialize(FArchive& Ar);
+	bool Serialize(FStructuredArchive::FSlot Slot);
 	bool operator==(FSoftObjectPath const& Other) const;
 	bool operator!=(FSoftObjectPath const& Other) const
 	{
@@ -141,7 +142,7 @@ struct COREUOBJECT_API FSoftObjectPath
 	FSoftObjectPath& operator=(FSoftObjectPath const& Other);
 	bool ExportTextItem(FString& ValueStr, FSoftObjectPath const& DefaultValue, UObject* Parent, int32 PortFlags, UObject* ExportRootScope) const;
 	bool ImportTextItem( const TCHAR*& Buffer, int32 PortFlags, UObject* Parent, FOutputDevice* ErrorText );
-	bool SerializeFromMismatchedTag(struct FPropertyTag const& Tag, FArchive& Ar);
+	bool SerializeFromMismatchedTag(struct FPropertyTag const& Tag, FStructuredArchive::FSlot Slot);
 
 	/** Serializes the internal path and also handles save/PIE fixups. Call this from the archiver overrides */
 	void SerializePath(FArchive& Ar);
@@ -247,7 +248,7 @@ struct COREUOBJECT_API FSoftClassPath : public FSoftObjectPath
 	 */
 	UClass* ResolveClass() const;
 
-	bool SerializeFromMismatchedTag(const FPropertyTag& Tag, FArchive& Ar);
+	bool SerializeFromMismatchedTag(const FPropertyTag& Tag, FStructuredArchive::FSlot Slot);
 
 	static FSoftClassPath GetOrCreateIDForClass(const UClass *InClass);
 
diff --git a/Engine/Source/Runtime/CoreUObject/Public/UObject/TextProperty.h b/Engine/Source/Runtime/CoreUObject/Public/UObject/TextProperty.h
--- a/Engine/Source/Runtime/CoreUObject/Public/UObject/TextProperty.h
+++ b/Engine/Source/Runtime/CoreUObject/Public/UObject/TextProperty.h
@@ -29,7 +29,7 @@ public:
 	}
 
 	// UProperty interface
-	virtual EConvertFromTypeResult ConvertFromType(const FPropertyTag& Tag, FArchive& Ar, uint8* Data, UStruct* DefaultsStruct) override;
+	virtual EConvertFromTypeResult ConvertFromType(const FPropertyTag& Tag, FStructuredArchive::FSlot Slot, uint8* Data, UStruct* DefaultsStruct) override;
 	virtual bool Identical( const void* A, const void* B, uint32 PortFlags ) const override;
 	virtual void SerializeItem(FStructuredArchive::FSlot Slot, void* Value, void const* Defaults) const override;
 	virtual void ExportTextItem( FString& ValueStr, const void* PropertyValue, const void* DefaultValue, UObject* Parent, int32 PortFlags, UObject* ExportRootScope ) const override;
diff --git a/Engine/Source/Runtime/CoreUObject/Public/UObject/UnrealType.h b/Engine/Source/Runtime/CoreUObject/Public/UObject/UnrealType.h
--- a/Engine/Source/Runtime/CoreUObject/Public/UObject/UnrealType.h
+++ b/Engine/Source/Runtime/CoreUObject/Public/UObject/UnrealType.h
@@ -228,7 +228,7 @@ public:
 	*			CannotConvert:    the tag is not something that the property can convert.
 	*			UseSerializeItem: no conversion was done on the property - this can mean that the tag is correct and normal serialization applies or that the tag is incompatible.
 	*/
-	virtual EConvertFromTypeResult ConvertFromType(const FPropertyTag& Tag, FArchive& Ar, uint8* Data, UStruct* DefaultsStruct);
+	virtual EConvertFromTypeResult ConvertFromType(const FPropertyTag& Tag, FStructuredArchive::FSlot Slot, uint8* Data, UStruct* DefaultsStruct);
 
 	/**
 	 * Determines whether the property values are identical.
@@ -1231,7 +1231,7 @@ class COREUOBJECT_API UNumericProperty : public UProperty
 	virtual FString GetNumericPropertyValueToString(void const* Data) const;
 	// End of UNumericProperty interface
 
-	static int64 ReadEnumAsInt64(FArchive& Ar, UStruct* DefaultsStruct, const FPropertyTag& Tag);
+	static int64 ReadEnumAsInt64(FStructuredArchive::FSlot Slot, UStruct* DefaultsStruct, const FPropertyTag& Tag);
 
 private:
 	virtual bool CanHoldDoubleValueInternal  (double Value) const PURE_VIRTUAL(UNumericProperty::CanHoldDoubleValueInternal,   return false;);
@@ -1278,10 +1278,10 @@ public:
 
 protected:
 	template <typename OldNumericType>
-	void ConvertFromArithmeticValue(FArchive& Ar, void* Obj, const FPropertyTag& Tag)
+	void ConvertFromArithmeticValue(FStructuredArchive::FSlot Slot, void* Obj, const FPropertyTag& Tag)
 	{
 		OldNumericType OldValue;
-		Ar << OldValue;
+		Slot << OldValue;
 		TCppType NewValue = (TCppType)OldValue;
 		this->SetPropertyValue_InContainer(Obj, NewValue, Tag.ArrayIndex);
 
@@ -1291,68 +1291,68 @@ protected:
 			Warning,
 			TEXT("Potential data loss during conversion of integer property %s of %s - was (%s) now (%s) - for package: %s"),
 			*this->GetName(),
-			*Ar.GetArchiveName(),
+			*Slot.GetUnderlyingArchive().GetArchiveName(),
 			*Lex::ToString(OldValue),
 			*Lex::ToString(NewValue),
-			*Ar.GetArchiveName()
+			*Slot.GetUnderlyingArchive().GetArchiveName()
 			);
 	}
 
 public:
-	virtual EConvertFromTypeResult ConvertFromType(const FPropertyTag& Tag, FArchive& Ar, uint8* Data, UStruct* DefaultsStruct) override
+	virtual EConvertFromTypeResult ConvertFromType(const FPropertyTag& Tag, FStructuredArchive::FSlot Slot, uint8* Data, UStruct* DefaultsStruct) override
 	{
 		if (Tag.Type == NAME_Int8Property)
 		{
-			ConvertFromArithmeticValue<int8>(Ar, Data, Tag);
+			ConvertFromArithmeticValue<int8>(Slot, Data, Tag);
 		}
 		else if (Tag.Type == NAME_Int16Property)
 		{
-			ConvertFromArithmeticValue<int16>(Ar, Data, Tag);
+			ConvertFromArithmeticValue<int16>(Slot, Data, Tag);
 		}
 		else if (Tag.Type == NAME_IntProperty)
 		{
-			ConvertFromArithmeticValue<int32>(Ar, Data, Tag);
+			ConvertFromArithmeticValue<int32>(Slot, Data, Tag);
 		}
 		else if (Tag.Type == NAME_Int64Property)
 		{
-			ConvertFromArithmeticValue<int64>(Ar, Data, Tag);
+			ConvertFromArithmeticValue<int64>(Slot, Data, Tag);
 		}
 		else if (Tag.Type == NAME_ByteProperty)
 		{
 			if (Tag.EnumName != NAME_None)
 			{
-				int64 PreviousValue = this->ReadEnumAsInt64(Ar, DefaultsStruct, Tag);
+				int64 PreviousValue = this->ReadEnumAsInt64(Slot, DefaultsStruct, Tag);
 				this->SetPropertyValue_InContainer(Data, PreviousValue, Tag.ArrayIndex);
 			}
 			else
 			{
-				ConvertFromArithmeticValue<int8>(Ar, Data, Tag);
+				ConvertFromArithmeticValue<int8>(Slot, Data, Tag);
 			}
 		}
 		else if (Tag.Type == NAME_EnumProperty)
 		{
-			int64 PreviousValue = this->ReadEnumAsInt64(Ar, DefaultsStruct, Tag);
+			int64 PreviousValue = this->ReadEnumAsInt64(Slot, DefaultsStruct, Tag);
 			this->SetPropertyValue_InContainer(Data, (TCppType)PreviousValue, Tag.ArrayIndex);
 		}
 		else if (Tag.Type == NAME_UInt16Property)
 		{
-			ConvertFromArithmeticValue<uint16>(Ar, Data, Tag);
+			ConvertFromArithmeticValue<uint16>(Slot, Data, Tag);
 		}
 		else if (Tag.Type == NAME_UInt32Property)
 		{
-			ConvertFromArithmeticValue<uint32>(Ar, Data, Tag);
+			ConvertFromArithmeticValue<uint32>(Slot, Data, Tag);
 		}
 		else if (Tag.Type == NAME_UInt64Property)
 		{
-			ConvertFromArithmeticValue<uint64>(Ar, Data, Tag);
+			ConvertFromArithmeticValue<uint64>(Slot, Data, Tag);
 		}
 		else if (Tag.Type == NAME_FloatProperty)
 		{
-			ConvertFromArithmeticValue<float>(Ar, Data, Tag);
+			ConvertFromArithmeticValue<float>(Slot, Data, Tag);
 		}
 		else if (Tag.Type == NAME_DoubleProperty)
 		{
-			ConvertFromArithmeticValue<double>(Ar, Data, Tag);
+			ConvertFromArithmeticValue<double>(Slot, Data, Tag);
 		}
 		else
 		{
@@ -1472,7 +1472,7 @@ class COREUOBJECT_API UByteProperty : public TProperty_Numeric<uint8>
 	virtual bool NetSerializeItem( FArchive& Ar, UPackageMap* Map, void* Data, TArray<uint8> * MetaData = NULL ) const override;
 	virtual void ExportTextItem( FString& ValueStr, const void* PropertyValue, const void* DefaultValue, UObject* Parent, int32 PortFlags, UObject* ExportRootScope ) const override;
 	virtual const TCHAR* ImportText_Internal( const TCHAR* Buffer, void* Data, int32 PortFlags, UObject* Parent, FOutputDevice* ErrorText ) const override;
-	virtual EConvertFromTypeResult ConvertFromType(const FPropertyTag& Tag, FArchive& Ar, uint8* Data, UStruct* DefaultsStruct) override;
+	virtual EConvertFromTypeResult ConvertFromType(const FPropertyTag& Tag, FStructuredArchive::FSlot Slot, uint8* Data, UStruct* DefaultsStruct) override;
 	// End of UProperty interface
 
 	// UNumericProperty interface.
@@ -1778,7 +1778,7 @@ public:
 	virtual void ClearValueInternal( void* Data ) const override;
 	virtual void InitializeValueInternal( void* Dest ) const override;
 	virtual int32 GetMinAlignment() const override;
-	virtual EConvertFromTypeResult ConvertFromType(const FPropertyTag& Tag, FArchive& Ar, uint8* Data, UStruct* DefaultsStruct) override;
+	virtual EConvertFromTypeResult ConvertFromType(const FPropertyTag& Tag, FStructuredArchive::FSlot Slot, uint8* Data, UStruct* DefaultsStruct) override;
 	// End of UProperty interface
 
 	// Emulate the CPP type API, see TPropertyTypeFundamentals
@@ -2060,7 +2060,7 @@ class COREUOBJECT_API UObjectProperty : public TUObjectPropertyBase<UObject*>
 	virtual void SerializeItem(FStructuredArchive::FSlot Slot, void* Value, void const* Defaults) const override;
 	virtual void EmitReferenceInfo(UClass& OwnerClass, int32 BaseOffset, TArray<const UStructProperty*>& EncounteredStructProps) override;
 	virtual const TCHAR* ImportText_Internal(const TCHAR* Buffer, void* Data, int32 PortFlags, UObject* OwnerObject, FOutputDevice* ErrorText) const override;
-	virtual EConvertFromTypeResult ConvertFromType(const FPropertyTag& Tag, FArchive& Ar, uint8* Data, UStruct* DefaultsStruct) override;
+	virtual EConvertFromTypeResult ConvertFromType(const FPropertyTag& Tag, FStructuredArchive::FSlot Slot, uint8* Data, UStruct* DefaultsStruct) override;
 
 private:
 	virtual uint32 GetValueTypeHashInternal(const void* Src) const override;
@@ -2171,7 +2171,7 @@ class COREUOBJECT_API USoftObjectProperty : public TUObjectPropertyBase<FSoftObj
 	virtual void SerializeItem( FStructuredArchive::FSlot Slot, void* Value, void const* Defaults) const override;
 	virtual void ExportTextItem( FString& ValueStr, const void* PropertyValue, const void* DefaultValue, UObject* Parent, int32 PortFlags, UObject* ExportRootScope ) const override;
 	virtual const TCHAR* ImportText_Internal( const TCHAR* Buffer, void* Data, int32 PortFlags, UObject* OwnerObject, FOutputDevice* ErrorText ) const override;
-	virtual EConvertFromTypeResult ConvertFromType(const FPropertyTag& Tag, FArchive& Ar, uint8* Data, UStruct* DefaultsStruct) override;
+	virtual EConvertFromTypeResult ConvertFromType(const FPropertyTag& Tag, FStructuredArchive::FSlot Slot, uint8* Data, UStruct* DefaultsStruct) override;
 	// End of UProperty interface
 
 	// UObjectProperty interface
@@ -2418,7 +2418,7 @@ public:
 	// UProperty interface
 	virtual void ExportTextItem( FString& ValueStr, const void* PropertyValue, const void* DefaultValue, UObject* Parent, int32 PortFlags, UObject* ExportRootScope ) const override;
 	virtual const TCHAR* ImportText_Internal( const TCHAR* Buffer, void* Data, int32 PortFlags, UObject* OwnerObject, FOutputDevice* ErrorText ) const override;
-	virtual EConvertFromTypeResult ConvertFromType(const FPropertyTag& Tag, FArchive& Ar, uint8* Data, UStruct* DefaultsStruct) override;
+	virtual EConvertFromTypeResult ConvertFromType(const FPropertyTag& Tag, FStructuredArchive::FSlot Slot, uint8* Data, UStruct* DefaultsStruct) override;
 	virtual FString GetCPPTypeForwardDeclaration() const override;
 	uint32 GetValueTypeHashInternal(const void* Src) const override;
 	// End of UProperty interface
@@ -2455,7 +2455,7 @@ public:
 	// UProperty interface
 	virtual void ExportTextItem( FString& ValueStr, const void* PropertyValue, const void* DefaultValue, UObject* Parent, int32 PortFlags, UObject* ExportRootScope ) const override;
 	virtual const TCHAR* ImportText_Internal( const TCHAR* Buffer, void* Data, int32 PortFlags, UObject* OwnerObject, FOutputDevice* ErrorText ) const override;
-	virtual EConvertFromTypeResult ConvertFromType(const FPropertyTag& Tag, FArchive& Ar, uint8* Data, UStruct* DefaultsStruct) override;
+	virtual EConvertFromTypeResult ConvertFromType(const FPropertyTag& Tag, FStructuredArchive::FSlot Slot, uint8* Data, UStruct* DefaultsStruct) override;
 	virtual FString GetCPPTypeForwardDeclaration() const override;
 	uint32 GetValueTypeHashInternal(const void* Src) const override;
 	// End of UProperty interface
@@ -2525,7 +2525,7 @@ public:
 	virtual bool ContainsWeakObjectReference() const override;
 	virtual void EmitReferenceInfo(UClass& OwnerClass, int32 BaseOffset, TArray<const UStructProperty*>& EncounteredStructProps) override;
 	virtual bool SameType(const UProperty* Other) const override;
-	virtual EConvertFromTypeResult ConvertFromType(const FPropertyTag& Tag, FArchive& Ar, uint8* Data, UStruct* DefaultsStruct) override;
+	virtual EConvertFromTypeResult ConvertFromType(const FPropertyTag& Tag, FStructuredArchive::FSlot Slot, uint8* Data, UStruct* DefaultsStruct) override;
 	// End of UProperty interface
 
 	FString GetCPPTypeCustom(FString* ExtendedTypeText, uint32 CPPExportFlags, const FString& InnerTypeText, const FString& InInnerExtendedTypeText) const;
@@ -2578,7 +2578,7 @@ public:
 	virtual bool ContainsWeakObjectReference() const override;
 	virtual void EmitReferenceInfo(UClass& OwnerClass, int32 BaseOffset, TArray<const UStructProperty*>& EncounteredStructProps) override;
 	virtual bool SameType(const UProperty* Other) const override;
-	virtual EConvertFromTypeResult ConvertFromType(const FPropertyTag& Tag, FArchive& Ar, uint8* Data, UStruct* DefaultsStruct) override;
+	virtual EConvertFromTypeResult ConvertFromType(const FPropertyTag& Tag, FStructuredArchive::FSlot Slot, uint8* Data, UStruct* DefaultsStruct) override;
 	// End of UProperty interface
 
 	FString GetCPPTypeCustom(FString* ExtendedTypeText, uint32 CPPExportFlags, const FString& KeyTypeText, const FString& InKeyExtendedTypeText, const FString& ValueTypeText, const FString& InValueExtendedTypeText) const;
@@ -2630,7 +2630,7 @@ public:
 	virtual bool ContainsWeakObjectReference() const override;
 	virtual void EmitReferenceInfo(UClass& OwnerClass, int32 BaseOffset, TArray<const UStructProperty*>& EncounteredStructProps) override;
 	virtual bool SameType(const UProperty* Other) const override;
-	virtual EConvertFromTypeResult ConvertFromType(const FPropertyTag& Tag, FArchive& Ar, uint8* Data, UStruct* DefaultsStruct) override;
+	virtual EConvertFromTypeResult ConvertFromType(const FPropertyTag& Tag, FStructuredArchive::FSlot Slot, uint8* Data, UStruct* DefaultsStruct) override;
 	// End of UProperty interface
 
 	FString GetCPPTypeCustom(FString* ExtendedTypeText, uint32 CPPExportFlags, const FString& ElementTypeText, const FString& InElementExtendedTypeText) const;
@@ -4062,7 +4062,7 @@ public:
 	virtual bool ContainsWeakObjectReference() const override;
 	virtual void EmitReferenceInfo(UClass& OwnerClass, int32 BaseOffset, TArray<const UStructProperty*>& EncounteredStructProps) override;
 	virtual bool SameType(const UProperty* Other) const override;
-	virtual EConvertFromTypeResult ConvertFromType(const FPropertyTag& Tag, FArchive& Ar, uint8* Data, UStruct* DefaultsStruct) override;
+	virtual EConvertFromTypeResult ConvertFromType(const FPropertyTag& Tag, FStructuredArchive::FSlot Slot, uint8* Data, UStruct* DefaultsStruct) override;
 	// End of UProperty interface
 
 	DEPRECATED(4.14, "Use UScriptStruct::ImportText instead")
diff --git a/Engine/Source/Runtime/Engine/Classes/Materials/MaterialExpression.h b/Engine/Source/Runtime/Engine/Classes/Materials/MaterialExpression.h
--- a/Engine/Source/Runtime/Engine/Classes/Materials/MaterialExpression.h
+++ b/Engine/Source/Runtime/Engine/Classes/Materials/MaterialExpression.h
@@ -202,7 +202,7 @@ class ENGINE_API UMaterialExpression : public UObject
 	
 	virtual bool Modify( bool bAlwaysMarkDirty=true ) override;
 #endif // WITH_EDITOR
-	virtual void Serialize( FArchive& Ar ) override;
+	virtual void Serialize( FStructuredArchive::FRecord Record ) override;
 	virtual bool NeedsLoadForClient() const override;
 	virtual bool NeedsLoadForEditorGame() const override
 	{
diff --git a/Engine/Source/Runtime/Engine/Classes/Materials/MaterialExpressionBreakMaterialAttributes.h b/Engine/Source/Runtime/Engine/Classes/Materials/MaterialExpressionBreakMaterialAttributes.h
--- a/Engine/Source/Runtime/Engine/Classes/Materials/MaterialExpressionBreakMaterialAttributes.h
+++ b/Engine/Source/Runtime/Engine/Classes/Materials/MaterialExpressionBreakMaterialAttributes.h
@@ -18,7 +18,7 @@ class UMaterialExpressionBreakMaterialAttributes : public UMaterialExpression
  	FMaterialAttributesInput MaterialAttributes;
 
 	//~ Begin UObject Interface
-	virtual void Serialize(FArchive& Ar) override;
+	virtual void Serialize(FStructuredArchive::FRecord Record) override;
 	//~ End UObject Interface
  
 	//~ Begin UMaterialExpression Interface
diff --git a/Engine/Source/Runtime/Engine/Classes/Materials/MaterialExpressionClamp.h b/Engine/Source/Runtime/Engine/Classes/Materials/MaterialExpressionClamp.h
--- a/Engine/Source/Runtime/Engine/Classes/Materials/MaterialExpressionClamp.h
+++ b/Engine/Source/Runtime/Engine/Classes/Materials/MaterialExpressionClamp.h
@@ -42,7 +42,7 @@ class UMaterialExpressionClamp : public UMaterialExpression
 
 
 	//~ Begin UObject Interface
-	virtual void Serialize( FArchive& Ar ) override;
+	virtual void Serialize(FStructuredArchive::FRecord Record) override;
 	//~ End UObject Interface
 
 	//~ Begin UMaterialExpression Interface
diff --git a/Engine/Source/Runtime/Engine/Classes/Materials/MaterialExpressionCustom.h b/Engine/Source/Runtime/Engine/Classes/Materials/MaterialExpressionCustom.h
--- a/Engine/Source/Runtime/Engine/Classes/Materials/MaterialExpressionCustom.h
+++ b/Engine/Source/Runtime/Engine/Classes/Materials/MaterialExpressionCustom.h
@@ -56,7 +56,7 @@ class UMaterialExpressionCustom : public UMaterialExpression
 #if WITH_EDITOR
 	virtual void PostEditChangeProperty(FPropertyChangedEvent& PropertyChangedEvent) override;
 #endif // WITH_EDITOR
-	virtual void Serialize(FArchive& Ar) override;
+	virtual void Serialize(FStructuredArchive::FRecord Record) override;
 	//~ End UObject Interface.
 	
 	//~ Begin UMaterialExpression Interface
diff --git a/Engine/Source/Runtime/Engine/Classes/Materials/MaterialExpressionFeatureLevelSwitch.h b/Engine/Source/Runtime/Engine/Classes/Materials/MaterialExpressionFeatureLevelSwitch.h
--- a/Engine/Source/Runtime/Engine/Classes/Materials/MaterialExpressionFeatureLevelSwitch.h
+++ b/Engine/Source/Runtime/Engine/Classes/Materials/MaterialExpressionFeatureLevelSwitch.h
@@ -22,7 +22,7 @@ class UMaterialExpressionFeatureLevelSwitch : public UMaterialExpression
 	FExpressionInput Inputs[ERHIFeatureLevel::Num];
 
 	//~ Begin UObject Interface.
-	virtual void Serialize(FArchive& Ar) override;
+	virtual void Serialize(FStructuredArchive::FRecord Record) override;
 	virtual bool NeedsLoadForClient() const override;
 	//~ End UObject Interface.
 
diff --git a/Engine/Source/Runtime/Engine/Classes/Materials/MaterialExpressionMakeMaterialAttributes.h b/Engine/Source/Runtime/Engine/Classes/Materials/MaterialExpressionMakeMaterialAttributes.h
--- a/Engine/Source/Runtime/Engine/Classes/Materials/MaterialExpressionMakeMaterialAttributes.h
+++ b/Engine/Source/Runtime/Engine/Classes/Materials/MaterialExpressionMakeMaterialAttributes.h
@@ -69,7 +69,7 @@ class UMaterialExpressionMakeMaterialAttributes : public UMaterialExpression
 	FExpressionInput PixelDepthOffset;
 
 	//~ Begin UObject Interface
-	virtual void Serialize(FArchive& Ar) override;
+	virtual void Serialize(FStructuredArchive::FRecord Record) override;
 	//~ End UObject Interface
 
 	//~ Begin UMaterialExpression Interface
diff --git a/Engine/Source/Runtime/Engine/Classes/Materials/MaterialExpressionSpeedTree.h b/Engine/Source/Runtime/Engine/Classes/Materials/MaterialExpressionSpeedTree.h
--- a/Engine/Source/Runtime/Engine/Classes/Materials/MaterialExpressionSpeedTree.h
+++ b/Engine/Source/Runtime/Engine/Classes/Materials/MaterialExpressionSpeedTree.h
@@ -70,7 +70,7 @@ class UMaterialExpressionSpeedTree : public UMaterialExpression
 	bool bAccurateWindVelocities;
 
 	//~ Begin UObject Interface
-	virtual void Serialize(FArchive& Ar) override;
+	virtual void Serialize(FStructuredArchive::FRecord Record) override;
 #if WITH_EDITOR
 	virtual bool CanEditChange(const UProperty* InProperty) const override;
 #endif // WITH_EDITOR
diff --git a/Engine/Source/Runtime/Engine/Private/Materials/MaterialExpressions.cpp b/Engine/Source/Runtime/Engine/Private/Materials/MaterialExpressions.cpp
--- a/Engine/Source/Runtime/Engine/Private/Materials/MaterialExpressions.cpp
+++ b/Engine/Source/Runtime/Engine/Private/Materials/MaterialExpressions.cpp
@@ -580,16 +580,16 @@ void UMaterialExpression::CopyMaterialExpressions(const TArray<UMaterialExpressi
 #endif // WITH_EDITOR
 
 
-void UMaterialExpression::Serialize( FArchive& Ar )
+void UMaterialExpression::Serialize(FStructuredArchive::FRecord Record)
 {
-	Super::Serialize(Ar);
+	Super::Serialize(Record);
 
 #if WITH_EDITORONLY_DATA
 	const TArray<FExpressionInput*> Inputs = GetInputs();
 	for (int32 InputIndex = 0; InputIndex < Inputs.Num(); ++InputIndex)
 	{
 		FExpressionInput* Input = Inputs[InputIndex];
-		DoMaterialAttributeReorder(Input, Ar.UE4Ver());
+		DoMaterialAttributeReorder(Input, Record.GetUnderlyingArchive().UE4Ver());
 	}
 #endif // WITH_EDITORONLY_DATA
 }
@@ -2910,11 +2910,12 @@ UMaterialExpressionClamp::UMaterialExpressionClamp(const FObjectInitializer& Obj
 #endif
 }
 
-void UMaterialExpressionClamp::Serialize(FArchive& Ar)
+void UMaterialExpressionClamp::Serialize(FStructuredArchive::FRecord Record)
 {
-	Super::Serialize(Ar);
+	Super::Serialize(Record);
+	FArchive& UnderlyingArchive = Record.GetUnderlyingArchive();
 
-	if (Ar.IsLoading() && Ar.UE4Ver() < VER_UE4_RETROFIT_CLAMP_EXPRESSIONS_SWAP)
+	if (UnderlyingArchive.IsLoading() && UnderlyingArchive.UE4Ver() < VER_UE4_RETROFIT_CLAMP_EXPRESSIONS_SWAP)
 	{
 		if (ClampMode == CMODE_ClampMin)
 		{
@@ -4317,13 +4318,14 @@ UMaterialExpressionMakeMaterialAttributes::UMaterialExpressionMakeMaterialAttrib
 #endif
 }
 
-void UMaterialExpressionMakeMaterialAttributes::Serialize(FArchive& Ar)
+void UMaterialExpressionMakeMaterialAttributes::Serialize(FStructuredArchive::FRecord Record)
 {
-	Super::Serialize(Ar);
+	Super::Serialize(Record);
+	FArchive& UnderlyingArchive = Record.GetUnderlyingArchive();
 
-	Ar.UsingCustomVersion(FRenderingObjectVersion::GUID);
+	UnderlyingArchive.UsingCustomVersion(FRenderingObjectVersion::GUID);
 	
-	if (Ar.CustomVer(FRenderingObjectVersion::GUID) < FRenderingObjectVersion::FixedLegacyMaterialAttributeNodeTypes)
+	if (UnderlyingArchive.CustomVer(FRenderingObjectVersion::GUID) < FRenderingObjectVersion::FixedLegacyMaterialAttributeNodeTypes)
 	{
 		// Update the legacy masks else fail on vec3 to vec2 conversion
 		Refraction.SetMask(1, 1, 1, 0, 0);
@@ -4441,14 +4443,15 @@ UMaterialExpressionBreakMaterialAttributes::UMaterialExpressionBreakMaterialAttr
 #endif
 }
 
-void UMaterialExpressionBreakMaterialAttributes::Serialize(FArchive& Ar)
+void UMaterialExpressionBreakMaterialAttributes::Serialize(FStructuredArchive::FRecord Record)
 {
-	Super::Serialize(Ar);
+	Super::Serialize(Record);
+	FArchive& UnderlyingArchive = Record.GetUnderlyingArchive();
 
-	Ar.UsingCustomVersion(FRenderingObjectVersion::GUID);
+	UnderlyingArchive.UsingCustomVersion(FRenderingObjectVersion::GUID);
 
 #if WITH_EDITOR
-	if (Ar.CustomVer(FRenderingObjectVersion::GUID) < FRenderingObjectVersion::FixedLegacyMaterialAttributeNodeTypes)
+	if (UnderlyingArchive.CustomVer(FRenderingObjectVersion::GUID) < FRenderingObjectVersion::FixedLegacyMaterialAttributeNodeTypes)
 	{
 		// Update the masks for legacy content
 		int32 OutputIndex = 0;
@@ -6831,10 +6834,12 @@ bool UMaterialExpressionFeatureLevelSwitch::IsResultMaterialAttributes(int32 Out
 }
 #endif // WITH_EDITOR
 
-void UMaterialExpressionFeatureLevelSwitch::Serialize(FArchive& Ar)
+void UMaterialExpressionFeatureLevelSwitch::Serialize(FStructuredArchive::FRecord Record)
 {
-	Super::Serialize(Ar);
-	if (Ar.IsLoading() && Ar.UE4Ver() < VER_UE4_RENAME_SM3_TO_ES3_1)
+	Super::Serialize(Record);
+	FArchive& UnderlyingArchive = Record.GetUnderlyingArchive();
+
+	if (UnderlyingArchive.IsLoading() && UnderlyingArchive.UE4Ver() < VER_UE4_RENAME_SM3_TO_ES3_1)
 	{
 		// Copy the ES2 input to SM3 (since SM3 will now become ES3_1 and we don't want broken content)
 		Inputs[ERHIFeatureLevel::ES3_1] = Inputs[ERHIFeatureLevel::ES2];
@@ -9219,18 +9224,19 @@ uint32 UMaterialExpressionCustom::GetOutputType(int32 OutputIndex)
 }
 #endif // WITH_EDITOR
 
-void UMaterialExpressionCustom::Serialize(FArchive& Ar)
+void UMaterialExpressionCustom::Serialize(FStructuredArchive::FRecord Record)
 {
-	Super::Serialize(Ar);
+	Super::Serialize(Record);
+	FArchive& UnderlyingArchive = Record.GetUnderlyingArchive();
 
-	Ar.UsingCustomVersion(FRenderingObjectVersion::GUID);
+	UnderlyingArchive.UsingCustomVersion(FRenderingObjectVersion::GUID);
 
 	// Make a copy of the current code before we change it
 	const FString PreFixUp = Code;
 
 	bool bDidUpdate = false;
 
-	if (Ar.UE4Ver() < VER_UE4_INSTANCED_STEREO_UNIFORM_UPDATE)
+	if (UnderlyingArchive.UE4Ver() < VER_UE4_INSTANCED_STEREO_UNIFORM_UPDATE)
 	{
 		// Look for WorldPosition rename
 		if (Code.ReplaceInline(TEXT("Parameters.WorldPosition"), TEXT("Parameters.AbsoluteWorldPosition"), ESearchCase::CaseSensitive) > 0)
@@ -9239,7 +9245,7 @@ void UMaterialExpressionCustom::Serialize(FArchive& Ar)
 		}
 	}
 	// Fix up uniform references that were moved from View to Frame as part of the instanced stereo implementation
-	else if (Ar.UE4Ver() < VER_UE4_INSTANCED_STEREO_UNIFORM_REFACTOR)
+	else if (UnderlyingArchive.UE4Ver() < VER_UE4_INSTANCED_STEREO_UNIFORM_REFACTOR)
 	{
 		// Uniform members that were moved from View to Frame
 		static const FString UniformMembers[] = {
@@ -9331,7 +9337,7 @@ void UMaterialExpressionCustom::Serialize(FArchive& Ar)
 		}
 	}
 
-	if (Ar.CustomVer(FRenderingObjectVersion::GUID) < FRenderingObjectVersion::RemovedRenderTargetSize)
+	if (UnderlyingArchive.CustomVer(FRenderingObjectVersion::GUID) < FRenderingObjectVersion::RemovedRenderTargetSize)
 	{
 		if (Code.ReplaceInline(TEXT("View.RenderTargetSize"), TEXT("View.BufferSizeAndInvSize.xy"), ESearchCase::CaseSensitive) > 0)
 		{
@@ -13852,11 +13858,11 @@ void UMaterialExpressionSpeedTree::GetCaption(TArray<FString>& OutCaptions) cons
 }
 #endif // WITH_EDITOR
 
-void UMaterialExpressionSpeedTree::Serialize(FArchive& Ar)
+void UMaterialExpressionSpeedTree::Serialize(FStructuredArchive::FRecord Record)
 {
-	Super::Serialize(Ar);
+	Super::Serialize(Record);
 
-	if (Ar.UE4Ver() < VER_UE4_SPEEDTREE_WIND_V7)
+	if (Record.GetUnderlyingArchive().UE4Ver() < VER_UE4_SPEEDTREE_WIND_V7)
 	{
 		// update wind presets for speedtree v7
 		switch (WindType)
diff --git a/Engine/Source/Runtime/Engine/Private/PerPlatformProperties.cpp b/Engine/Source/Runtime/Engine/Private/PerPlatformProperties.cpp
--- a/Engine/Source/Runtime/Engine/Private/PerPlatformProperties.cpp
+++ b/Engine/Source/Runtime/Engine/Private/PerPlatformProperties.cpp
@@ -38,5 +38,39 @@ ENGINE_API FArchive& operator<<(FArchive& Ar, TPerPlatformProperty<StructType, V
 	return Ar;
 }
 
+/** Serializer to cook out the most appropriate platform override */
+template<typename StructType, typename ValueType, EName BasePropertyName>
+ENGINE_API void operator<<(FStructuredArchive::FSlot Slot, TPerPlatformProperty<StructType, ValueType, BasePropertyName>& Property)
+{
+	FArchive& UnderlyingArchive = Slot.GetUnderlyingArchive();
+	FStructuredArchive::FRecord Record = Slot.EnterRecord();
+
+	bool bCooked = false;
+#if WITH_EDITOR
+	if (UnderlyingArchive.IsCooking())
+	{
+		bCooked = true;
+		Record << NAMED_FIELD(bCooked);
+		// Save out platform override if it exists and Default otherwise
+		ValueType Value = Property.GetValueForPlatformGroup(UnderlyingArchive.CookingTarget()->GetPlatformInfo().PlatformGroupName);
+		Record << NAMED_FIELD(Value);
+	}
+	else
+#endif
+	{
+		StructType* This = StaticCast<StructType*>(&Property);
+		Record << NAMED_FIELD(bCooked);
+		Record << NAMED_ITEM("Value", This->Default);
+#if WITH_EDITORONLY_DATA
+		if (!bCooked)
+		{
+			Record << NAMED_ITEM("PerPlatform", This->PerPlatform);
+		}
+#endif
+	}
+}
+
 template ENGINE_API FArchive& operator<<(FArchive&, TPerPlatformProperty<FPerPlatformInt, int32, NAME_IntProperty>&);
 template ENGINE_API FArchive& operator<<(FArchive&, TPerPlatformProperty<FPerPlatformFloat, float, NAME_FloatProperty>&);
+template ENGINE_API void operator<<(FStructuredArchive::FSlot Slot, TPerPlatformProperty<FPerPlatformInt, int32, NAME_IntProperty>&);
+template ENGINE_API void operator<<(FStructuredArchive::FSlot Slot, TPerPlatformProperty<FPerPlatformFloat, float, NAME_FloatProperty>&);
diff --git a/Engine/Source/Runtime/Engine/Public/PerPlatformProperties.h b/Engine/Source/Runtime/Engine/Public/PerPlatformProperties.h
--- a/Engine/Source/Runtime/Engine/Public/PerPlatformProperties.h
+++ b/Engine/Source/Runtime/Engine/Public/PerPlatformProperties.h
@@ -74,10 +74,33 @@ struct ENGINE_API TPerPlatformProperty
 		Ar << *this;
 		return true;
 	}
+
+	/* Load old properties that have been converted to FPerPlatformX */
+	bool SerializeFromMismatchedTag(const FPropertyTag& Tag, FStructuredArchive::FSlot Slot)
+	{
+		if (Tag.Type == _BasePropertyName)
+		{
+			_StructType* This = StaticCast<_StructType*>(this);
+			_ValueType OldValue;
+			Slot << OldValue;
+			*This = _StructType(OldValue);
+			return true;
+		}
+		return false;
+	}
+
+	/* Serialization */
+	bool Serialize(FStructuredArchive::FSlot Slot)
+	{
+		Slot << *this;
+		return true;
+	}
 };
 
 template<typename _StructType, typename _ValueType, EName _BasePropertyName>
 ENGINE_API FArchive& operator<<(FArchive& Ar, TPerPlatformProperty<_StructType, _ValueType, _BasePropertyName>& P);
+template<typename _StructType, typename _ValueType, EName _BasePropertyName>
+ENGINE_API void operator<<(FStructuredArchive::FSlot Slot, TPerPlatformProperty<_StructType, _ValueType, _BasePropertyName>& P);
 
 /** FPerPlatformInt - int32 property with per-platform overrides */
 USTRUCT()
@@ -107,6 +130,7 @@ struct ENGINE_API FPerPlatformInt
 	}
 };
 extern template ENGINE_API FArchive& operator<<(FArchive&, TPerPlatformProperty<FPerPlatformInt, int32, NAME_IntProperty>&);
+extern template ENGINE_API void operator<<(FStructuredArchive::FSlot Slot, TPerPlatformProperty<FPerPlatformInt, int32, NAME_IntProperty>&);
 
 template<>
 struct TStructOpsTypeTraits<FPerPlatformInt>
diff --git a/Engine/Source/Runtime/GameplayTags/Classes/GameplayTagContainer.h b/Engine/Source/Runtime/GameplayTags/Classes/GameplayTagContainer.h
--- a/Engine/Source/Runtime/GameplayTags/Classes/GameplayTagContainer.h
+++ b/Engine/Source/Runtime/GameplayTags/Classes/GameplayTagContainer.h
@@ -182,7 +182,7 @@ struct GAMEPLAYTAGS_API FGameplayTag
 	bool NetSerialize_Packed(FArchive& Ar, class UPackageMap* Map, bool& bOutSuccess);
 
 	/** Used to upgrade a Name property to a GameplayTag struct property */
-	bool SerializeFromMismatchedTag(const FPropertyTag& Tag, FArchive& Ar);
+	bool SerializeFromMismatchedTag(const FPropertyTag& Tag, FStructuredArchive::FSlot Slot);
 
 	/** Sets from a ImportText string, used in asset registry */
 	void FromExportString(const FString& ExportString);
diff --git a/Engine/Source/Runtime/GameplayTags/Private/GameplayTagContainer.cpp b/Engine/Source/Runtime/GameplayTags/Private/GameplayTagContainer.cpp
--- a/Engine/Source/Runtime/GameplayTags/Private/GameplayTagContainer.cpp
+++ b/Engine/Source/Runtime/GameplayTags/Private/GameplayTagContainer.cpp
@@ -1134,11 +1134,11 @@ FGameplayTag::FGameplayTag(FName Name)
 	// This constructor is used to bypass the table check and is only usable by GameplayTagManager
 }
 
-bool FGameplayTag::SerializeFromMismatchedTag(const FPropertyTag& Tag, FArchive& Ar)
+bool FGameplayTag::SerializeFromMismatchedTag(const FPropertyTag& Tag, FStructuredArchive::FSlot Slot)
 {
 	if (Tag.Type == NAME_NameProperty)
 	{
-		Ar << TagName;
+		Slot << TagName;
 		return true;
 	}
 	return false;
diff --git a/Engine/Source/Runtime/InputCore/Classes/InputCoreTypes.h b/Engine/Source/Runtime/InputCore/Classes/InputCoreTypes.h
--- a/Engine/Source/Runtime/InputCore/Classes/InputCoreTypes.h
+++ b/Engine/Source/Runtime/InputCore/Classes/InputCoreTypes.h
@@ -74,7 +74,7 @@ struct INPUTCORE_API FKey
 	FName GetFName() const;
 	FName GetMenuCategory() const;
 
-	bool SerializeFromMismatchedTag(struct FPropertyTag const& Tag, FArchive& Ar);
+	bool SerializeFromMismatchedTag(struct FPropertyTag const& Tag, FStructuredArchive::FSlot Slot);
 	bool ExportTextItem(FString& ValueStr, FKey const& DefaultValue, UObject* Parent, int32 PortFlags, UObject* ExportRootScope) const;
 	bool ImportTextItem(const TCHAR*& Buffer, int32 PortFlags, UObject* Parent, FOutputDevice* ErrorText);
 	void PostSerialize(const FArchive& Ar);
diff --git a/Engine/Source/Runtime/InputCore/Private/InputCoreTypes.cpp b/Engine/Source/Runtime/InputCore/Private/InputCoreTypes.cpp
--- a/Engine/Source/Runtime/InputCore/Private/InputCoreTypes.cpp
+++ b/Engine/Source/Runtime/InputCore/Private/InputCoreTypes.cpp
@@ -1022,11 +1022,11 @@ void FKey::ConditionalLookupKeyDetails() const
 	}
 }
 
-bool FKey::SerializeFromMismatchedTag(struct FPropertyTag const& Tag, FArchive& Ar)
+bool FKey::SerializeFromMismatchedTag(struct FPropertyTag const& Tag, FStructuredArchive::FSlot Slot)
 {
 	if (Tag.Type == NAME_ByteProperty && Tag.EnumName == TEXT("EKeys"))
 	{
-		Ar << KeyName;
+		Slot << KeyName;
 		const FString KeyNameString(KeyName.ToString());
 		const int32 FindIndex(KeyNameString.Find(TEXT("EKeys::")));
 		if (FindIndex != INDEX_NONE)
@@ -1037,7 +1037,7 @@ bool FKey::SerializeFromMismatchedTag(struct FPropertyTag const& Tag, FArchive&
 	}
 	else if (Tag.Type == NAME_NameProperty)
 	{
-		Ar << KeyName;
+		Slot << KeyName;
 	}
 
 	return false;
diff --git a/Engine/Source/Runtime/MovieScene/Private/Channels/MovieSceneBoolChannel.cpp b/Engine/Source/Runtime/MovieScene/Private/Channels/MovieSceneBoolChannel.cpp
--- a/Engine/Source/Runtime/MovieScene/Private/Channels/MovieSceneBoolChannel.cpp
+++ b/Engine/Source/Runtime/MovieScene/Private/Channels/MovieSceneBoolChannel.cpp
@@ -7,13 +7,13 @@
 #include "MovieSceneFwd.h"
 #include "MovieSceneFrameMigration.h"
 
-bool FMovieSceneBoolChannel::SerializeFromMismatchedTag(const FPropertyTag& Tag, FArchive& Ar)
+bool FMovieSceneBoolChannel::SerializeFromMismatchedTag(const FPropertyTag& Tag, FStructuredArchive::FSlot Slot)
 {
 	static const FName IntegralCurveName("IntegralCurve");
 	if (Tag.Type == NAME_StructProperty && Tag.StructName == IntegralCurveName)
 	{
 		FIntegralCurve IntegralCurve;
-		FIntegralCurve::StaticStruct()->SerializeItem(Ar, &IntegralCurve, nullptr);
+		FIntegralCurve::StaticStruct()->SerializeItem(Slot, &IntegralCurve, nullptr);
 
 		if (IntegralCurve.GetDefaultValue() != MAX_int32)
 		{
diff --git a/Engine/Source/Runtime/MovieScene/Private/Channels/MovieSceneByteChannel.cpp b/Engine/Source/Runtime/MovieScene/Private/Channels/MovieSceneByteChannel.cpp
--- a/Engine/Source/Runtime/MovieScene/Private/Channels/MovieSceneByteChannel.cpp
+++ b/Engine/Source/Runtime/MovieScene/Private/Channels/MovieSceneByteChannel.cpp
@@ -7,13 +7,13 @@
 #include "Misc/FrameRate.h"
 #include "MovieSceneFrameMigration.h"
 
-bool FMovieSceneByteChannel::SerializeFromMismatchedTag(const FPropertyTag& Tag, FArchive& Ar)
+bool FMovieSceneByteChannel::SerializeFromMismatchedTag(const FPropertyTag& Tag, FStructuredArchive::FSlot Slot)
 {
 	static const FName IntegralCurveName("IntegralCurve");
 	if (Tag.Type == NAME_StructProperty && Tag.StructName == IntegralCurveName)
 	{
 		FIntegralCurve IntegralCurve;
-		FIntegralCurve::StaticStruct()->SerializeItem(Ar, &IntegralCurve, nullptr);
+		FIntegralCurve::StaticStruct()->SerializeItem(Slot, &IntegralCurve, nullptr);
 
 		if (IntegralCurve.GetDefaultValue() != MAX_int32)
 		{
diff --git a/Engine/Source/Runtime/MovieScene/Private/Channels/MovieSceneFloatChannel.cpp b/Engine/Source/Runtime/MovieScene/Private/Channels/MovieSceneFloatChannel.cpp
--- a/Engine/Source/Runtime/MovieScene/Private/Channels/MovieSceneFloatChannel.cpp
+++ b/Engine/Source/Runtime/MovieScene/Private/Channels/MovieSceneFloatChannel.cpp
@@ -6,13 +6,13 @@
 #include "MovieSceneFrameMigration.h"
 
 
-bool FMovieSceneFloatChannel::SerializeFromMismatchedTag(const FPropertyTag& Tag, FArchive& Ar)
+bool FMovieSceneFloatChannel::SerializeFromMismatchedTag(const FPropertyTag& Tag, FStructuredArchive::FSlot Slot)
 {
 	static const FName RichCurveName("RichCurve");
 	if (Tag.Type == NAME_StructProperty && Tag.StructName == RichCurveName)
 	{
 		FRichCurve RichCurve;
-		FRichCurve::StaticStruct()->SerializeItem(Ar, &RichCurve, nullptr);
+		FRichCurve::StaticStruct()->SerializeItem(Slot, &RichCurve, nullptr);
 
 		if (RichCurve.GetDefaultValue() != MAX_flt)
 		{
diff --git a/Engine/Source/Runtime/MovieScene/Private/Channels/MovieSceneIntegerChannel.cpp b/Engine/Source/Runtime/MovieScene/Private/Channels/MovieSceneIntegerChannel.cpp
--- a/Engine/Source/Runtime/MovieScene/Private/Channels/MovieSceneIntegerChannel.cpp
+++ b/Engine/Source/Runtime/MovieScene/Private/Channels/MovieSceneIntegerChannel.cpp
@@ -8,13 +8,13 @@
 #include "MovieSceneFrameMigration.h"
 
 
-bool FMovieSceneIntegerChannel::SerializeFromMismatchedTag(const FPropertyTag& Tag, FArchive& Ar)
+bool FMovieSceneIntegerChannel::SerializeFromMismatchedTag(const FPropertyTag& Tag, FStructuredArchive::FSlot Slot)
 {
 	static const FName IntegralCurveName("IntegralCurve");
 	if (Tag.Type == NAME_StructProperty && Tag.StructName == IntegralCurveName)
 	{
 		FIntegralCurve IntegralCurve;
-		FIntegralCurve::StaticStruct()->SerializeItem(Ar, &IntegralCurve, nullptr);
+		FIntegralCurve::StaticStruct()->SerializeItem(Slot, &IntegralCurve, nullptr);
 
 		if (IntegralCurve.GetDefaultValue() != MAX_int32)
 		{
diff --git a/Engine/Source/Runtime/MovieScene/Private/MovieSceneFrameMigration.cpp b/Engine/Source/Runtime/MovieScene/Private/MovieSceneFrameMigration.cpp
--- a/Engine/Source/Runtime/MovieScene/Private/MovieSceneFrameMigration.cpp
+++ b/Engine/Source/Runtime/MovieScene/Private/MovieSceneFrameMigration.cpp
@@ -31,7 +31,7 @@ TRange<FFrameNumber> FMovieSceneFrameRange::FromFloatRange(const TRange<float>&
 	return NewRange;
 }
 
-bool FMovieSceneFrameRange::SerializeFromMismatchedTag(FPropertyTag const& Tag, FArchive& Ar)
+bool FMovieSceneFrameRange::SerializeFromMismatchedTag(FPropertyTag const& Tag, FStructuredArchive::FSlot Slot)
 {
 	static const FName NAME_FloatRange("FloatRange");
 	if (Tag.Type == NAME_StructProperty && Tag.StructName == NAME_FloatRange)
@@ -44,7 +44,7 @@ bool FMovieSceneFrameRange::SerializeFromMismatchedTag(FPropertyTag const& Tag,
 #endif
 			;
 
-		FloatStruct->SerializeItem(Ar, &FloatRange, nullptr);
+		FloatStruct->SerializeItem(Slot, &FloatRange, nullptr);
 
 		Value = FromFloatRange(FloatRange);
 		return true;
diff --git a/Engine/Source/Runtime/MovieScene/Private/MovieSceneSequencePlayer.cpp b/Engine/Source/Runtime/MovieScene/Private/MovieSceneSequencePlayer.cpp
--- a/Engine/Source/Runtime/MovieScene/Private/MovieSceneSequencePlayer.cpp
+++ b/Engine/Source/Runtime/MovieScene/Private/MovieSceneSequencePlayer.cpp
@@ -8,11 +8,11 @@
 #include "GameFramework/WorldSettings.h"
 #include "Misc/RuntimeErrors.h"
 
-bool FMovieSceneSequencePlaybackSettings::SerializeFromMismatchedTag( const FPropertyTag& Tag, FArchive& Ar )
+bool FMovieSceneSequencePlaybackSettings::SerializeFromMismatchedTag( const FPropertyTag& Tag, FStructuredArchive::FSlot Slot )
 {
 	if (Tag.Type == NAME_StructProperty && Tag.StructName == "LevelSequencePlaybackSettings")
 	{
-		StaticStruct()->SerializeItem(Ar, this, nullptr);
+		StaticStruct()->SerializeItem(Slot, this, nullptr);
 		return true;
 	}
 
diff --git a/Engine/Source/Runtime/MovieScene/Public/Channels/MovieSceneBoolChannel.h b/Engine/Source/Runtime/MovieScene/Public/Channels/MovieSceneBoolChannel.h
--- a/Engine/Source/Runtime/MovieScene/Public/Channels/MovieSceneBoolChannel.h
+++ b/Engine/Source/Runtime/MovieScene/Public/Channels/MovieSceneBoolChannel.h
@@ -22,7 +22,7 @@ struct MOVIESCENE_API FMovieSceneBoolChannel : public FMovieSceneChannel
 	/**
 	 * Serialize this type from another
 	 */
-	bool SerializeFromMismatchedTag(const FPropertyTag& Tag, FArchive& Ar);
+	bool SerializeFromMismatchedTag(const FPropertyTag& Tag, FStructuredArchive::FSlot Slot);
 
 	/**
 	 * Access a mutable interface for this channel's data
diff --git a/Engine/Source/Runtime/MovieScene/Public/Channels/MovieSceneByteChannel.h b/Engine/Source/Runtime/MovieScene/Public/Channels/MovieSceneByteChannel.h
--- a/Engine/Source/Runtime/MovieScene/Public/Channels/MovieSceneByteChannel.h
+++ b/Engine/Source/Runtime/MovieScene/Public/Channels/MovieSceneByteChannel.h
@@ -22,7 +22,7 @@ struct MOVIESCENE_API FMovieSceneByteChannel : public FMovieSceneChannel
 	/**
 	 * Serialize this type from another
 	 */
-	bool SerializeFromMismatchedTag(const FPropertyTag& Tag, FArchive& Ar);
+	bool SerializeFromMismatchedTag(const FPropertyTag& Tag, FStructuredArchive::FSlot Slot);
 
 	/**
 	 * Access a mutable interface for this channel's data
diff --git a/Engine/Source/Runtime/MovieScene/Public/Channels/MovieSceneFloatChannel.h b/Engine/Source/Runtime/MovieScene/Public/Channels/MovieSceneFloatChannel.h
--- a/Engine/Source/Runtime/MovieScene/Public/Channels/MovieSceneFloatChannel.h
+++ b/Engine/Source/Runtime/MovieScene/Public/Channels/MovieSceneFloatChannel.h
@@ -202,7 +202,7 @@ public:
 public:
 
 	/** Serialize this float function from a mismatching property tag (FRichCurve) */
-	bool SerializeFromMismatchedTag(const FPropertyTag& Tag, FArchive& Ar);
+	bool SerializeFromMismatchedTag(const FPropertyTag& Tag, FStructuredArchive::FSlot Slot);
 
 	int32 AddConstantKey(FFrameNumber InTime, float InValue);
 
diff --git a/Engine/Source/Runtime/MovieScene/Public/Channels/MovieSceneIntegerChannel.h b/Engine/Source/Runtime/MovieScene/Public/Channels/MovieSceneIntegerChannel.h
--- a/Engine/Source/Runtime/MovieScene/Public/Channels/MovieSceneIntegerChannel.h
+++ b/Engine/Source/Runtime/MovieScene/Public/Channels/MovieSceneIntegerChannel.h
@@ -22,7 +22,7 @@ struct MOVIESCENE_API FMovieSceneIntegerChannel : public FMovieSceneChannel
 	/**
 	 * Serialize this type from another
 	 */
-	bool SerializeFromMismatchedTag(const FPropertyTag& Tag, FArchive& Ar);
+	bool SerializeFromMismatchedTag(const FPropertyTag& Tag, FStructuredArchive::FSlot Slot);
 
 	/**
 	 * Access a mutable interface for this channel's data
diff --git a/Engine/Source/Runtime/MovieScene/Public/MovieSceneFrameMigration.h b/Engine/Source/Runtime/MovieScene/Public/MovieSceneFrameMigration.h
--- a/Engine/Source/Runtime/MovieScene/Public/MovieSceneFrameMigration.h
+++ b/Engine/Source/Runtime/MovieScene/Public/MovieSceneFrameMigration.h
@@ -58,7 +58,7 @@ struct FMovieSceneFrameRange
 	/**
 	 * Serialize this frame range from a mismatched type (only FFloatRange supported)
 	 */
-	MOVIESCENE_API bool SerializeFromMismatchedTag(FPropertyTag const& Tag, FArchive& Ar);
+	MOVIESCENE_API bool SerializeFromMismatchedTag(FPropertyTag const& Tag, FStructuredArchive::FSlot Slot);
 
 	/**
 	 * Get this range's lower bound
diff --git a/Engine/Source/Runtime/MovieScene/Public/MovieSceneSequencePlayer.h b/Engine/Source/Runtime/MovieScene/Public/MovieSceneSequencePlayer.h
--- a/Engine/Source/Runtime/MovieScene/Public/MovieSceneSequencePlayer.h
+++ b/Engine/Source/Runtime/MovieScene/Public/MovieSceneSequencePlayer.h
@@ -106,7 +106,7 @@ struct FMovieSceneSequencePlaybackSettings
 	/** (Optional) Externally supplied time controller */
 	TSharedPtr<FMovieSceneTimeController> TimeController;
 
-	MOVIESCENE_API bool SerializeFromMismatchedTag(const FPropertyTag& Tag, FArchive& Ar);
+	MOVIESCENE_API bool SerializeFromMismatchedTag(const FPropertyTag& Tag, FStructuredArchive::FSlot Slot);
 };
 
 template<> struct TStructOpsTypeTraits<FMovieSceneSequencePlaybackSettings> : public TStructOpsTypeTraitsBase2<FMovieSceneSequencePlaybackSettings>
diff --git a/Engine/Source/Runtime/MovieSceneTracks/Private/Sections/MovieSceneStringChannel.cpp b/Engine/Source/Runtime/MovieSceneTracks/Private/Sections/MovieSceneStringChannel.cpp
--- a/Engine/Source/Runtime/MovieSceneTracks/Private/Sections/MovieSceneStringChannel.cpp
+++ b/Engine/Source/Runtime/MovieSceneTracks/Private/Sections/MovieSceneStringChannel.cpp
@@ -17,13 +17,13 @@ const FString* FMovieSceneStringChannel::Evaluate(FFrameTime InTime) const
 	return bHasDefaultValue ? &DefaultValue : nullptr;
 }
 
-bool FMovieSceneStringChannel::SerializeFromMismatchedTag(const FPropertyTag& Tag, FArchive& Ar)
+bool FMovieSceneStringChannel::SerializeFromMismatchedTag(const FPropertyTag& Tag, FStructuredArchive::FSlot Slot)
 {
 	static const FName StringCurveName("StringCurve");
 	if (Tag.Type == NAME_StructProperty && Tag.StructName == StringCurveName)
 	{
 		FStringCurve StringCurve;
-		FStringCurve::StaticStruct()->SerializeItem(Ar, &StringCurve, nullptr);
+		FStringCurve::StaticStruct()->SerializeItem(Slot, &StringCurve, nullptr);
 
 		FString NewDefault = StringCurve.GetDefaultValue();
 		if (NewDefault.Len())
diff --git a/Engine/Source/Runtime/MovieSceneTracks/Public/Sections/MovieSceneStringChannel.h b/Engine/Source/Runtime/MovieSceneTracks/Public/Sections/MovieSceneStringChannel.h
--- a/Engine/Source/Runtime/MovieSceneTracks/Public/Sections/MovieSceneStringChannel.h
+++ b/Engine/Source/Runtime/MovieSceneTracks/Public/Sections/MovieSceneStringChannel.h
@@ -22,7 +22,7 @@ struct FMovieSceneStringChannel : public FMovieSceneChannel
 	/**
 	 * Serialize this type from another
 	 */
-	MOVIESCENETRACKS_API bool SerializeFromMismatchedTag(const FPropertyTag& Tag, FArchive& Ar);
+	MOVIESCENETRACKS_API bool SerializeFromMismatchedTag(const FPropertyTag& Tag, FStructuredArchive::FSlot Slot);
 
 	/**
 	 * Access a mutable interface for this channel's data
diff --git a/Engine/Source/Runtime/SlateCore/Private/Sound/SlateSound.cpp b/Engine/Source/Runtime/SlateCore/Private/Sound/SlateSound.cpp
--- a/Engine/Source/Runtime/SlateCore/Private/Sound/SlateSound.cpp
+++ b/Engine/Source/Runtime/SlateCore/Private/Sound/SlateSound.cpp
@@ -44,13 +44,13 @@ UObject* FSlateSound::GetResourceObject( ) const
 }
 
 
-bool FSlateSound::SerializeFromMismatchedTag( FPropertyTag const& Tag, FArchive& Ar )
+bool FSlateSound::SerializeFromMismatchedTag(FPropertyTag const& Tag, FStructuredArchive::FSlot Slot)
 {
 	// Sounds in Slate used to be stored as FName properties, so allow them to be upgraded in-place
 	if (Tag.Type == NAME_NameProperty)
 	{
 		FName SoundName;
-		Ar << SoundName;
+		Slot << SoundName;
 		*this = FromName_DEPRECATED(SoundName);
 
 		return true;
diff --git a/Engine/Source/Runtime/SlateCore/Private/Styling/SlateColor.cpp b/Engine/Source/Runtime/SlateCore/Private/Styling/SlateColor.cpp
--- a/Engine/Source/Runtime/SlateCore/Private/Styling/SlateColor.cpp
+++ b/Engine/Source/Runtime/SlateCore/Private/Styling/SlateColor.cpp
@@ -3,14 +3,14 @@
 #include "Styling/SlateColor.h"
 #include "UObject/PropertyTag.h"
 
-bool FSlateColor::SerializeFromMismatchedTag( const FPropertyTag& Tag, FArchive& Ar )
+bool FSlateColor::SerializeFromMismatchedTag(const FPropertyTag& Tag, FStructuredArchive::FSlot Slot)
 {
 	if (Tag.Type == NAME_StructProperty)
 	{
 		if (Tag.StructName == NAME_Color)
 		{
 			FColor OldColor;
-			Ar << OldColor;
+			Slot << OldColor;
 			*this = FSlateColor(FLinearColor(OldColor));
 
 			return true;
@@ -18,7 +18,7 @@ bool FSlateColor::SerializeFromMismatchedTag( const FPropertyTag& Tag, FArchive&
 		else if(Tag.StructName == NAME_LinearColor)
 		{
 			FLinearColor OldColor;
-			Ar << OldColor;
+			Slot << OldColor;
 			*this = FSlateColor(OldColor);
 
 			return true;
diff --git a/Engine/Source/Runtime/SlateCore/Public/Sound/SlateSound.h b/Engine/Source/Runtime/SlateCore/Public/Sound/SlateSound.h
--- a/Engine/Source/Runtime/SlateCore/Public/Sound/SlateSound.h
+++ b/Engine/Source/Runtime/SlateCore/Public/Sound/SlateSound.h
@@ -48,7 +48,7 @@ public:
 	/**
 	 * Used to upgrade a FName property to a FSlateSound property
 	 */
-	bool SerializeFromMismatchedTag( struct FPropertyTag const& Tag, FArchive& Ar );
+	bool SerializeFromMismatchedTag(struct FPropertyTag const& Tag, FStructuredArchive::FSlot Slot);
 
 public:
 
diff --git a/Engine/Source/Runtime/SlateCore/Public/Styling/SlateColor.h b/Engine/Source/Runtime/SlateCore/Public/Styling/SlateColor.h
--- a/Engine/Source/Runtime/SlateCore/Public/Styling/SlateColor.h
+++ b/Engine/Source/Runtime/SlateCore/Public/Styling/SlateColor.h
@@ -177,7 +177,7 @@ public:
 	}
 
 	/** Used to upgrade an FColor or FLinearColor property to an FSlateColor property */
-	bool SerializeFromMismatchedTag( const struct FPropertyTag& Tag, FArchive& Ar );
+	bool SerializeFromMismatchedTag(const struct FPropertyTag& Tag, FStructuredArchive::FSlot Slot);
 
 protected:
 	
diff --git a/Engine/Source/Runtime/UMG/Public/Animation/WidgetAnimationBinding.h b/Engine/Source/Runtime/UMG/Public/Animation/WidgetAnimationBinding.h
--- a/Engine/Source/Runtime/UMG/Public/Animation/WidgetAnimationBinding.h
+++ b/Engine/Source/Runtime/UMG/Public/Animation/WidgetAnimationBinding.h
@@ -64,4 +64,13 @@ public:
 		Ar << Binding.bIsRootWidget;
 		return Ar;
 	}
+
+	friend void operator<<(FStructuredArchive::FSlot Slot, FWidgetAnimationBinding& Binding)
+	{
+		FStructuredArchive::FRecord Record = Slot.EnterRecord();
+		Record << NAMED_ITEM("WidgetName", Binding.WidgetName);
+		Record << NAMED_ITEM("SlotWidgetName", Binding.SlotWidgetName);
+		Record << NAMED_ITEM("AnimationGuid", Binding.AnimationGuid);
+		Record << NAMED_ITEM("bIsRootWidget", Binding.bIsRootWidget);
+	}
 };
