CIS fix for variable shadowing warning

diff --git a/Engine/Source/Runtime/CoreUObject/Private/UObject/PropertyMap.cpp b/Engine/Source/Runtime/CoreUObject/Private/UObject/PropertyMap.cpp
--- a/Engine/Source/Runtime/CoreUObject/Private/UObject/PropertyMap.cpp
+++ b/Engine/Source/Runtime/CoreUObject/Private/UObject/PropertyMap.cpp
@@ -906,16 +906,16 @@ EConvertFromTypeResult UMapProperty::ConvertFromType(const FPropertyTag& Tag, FS
 	UnderlyingArchive.Preload(KeyProp);
 	UnderlyingArchive.Preload(ValueProp);
 
-	const auto SerializeOrConvert = [](UProperty* CurrentType, const FPropertyTag& InTag, FStructuredArchive::FSlot Slot, uint8* InData, UStruct* InDefaultsStruct) -> bool
+	const auto SerializeOrConvert = [](UProperty* CurrentType, const FPropertyTag& InTag, FStructuredArchive::FSlot InnerSlot, uint8* InData, UStruct* InDefaultsStruct) -> bool
 	{
 		// Serialize wants the property address, while convert wants the container address. InData is the container address
 		if(CurrentType->GetID() == InTag.Type)
 		{
 			uint8* DestAddress = CurrentType->ContainerPtrToValuePtr<uint8>(InData, InTag.ArrayIndex);
-			CurrentType->SerializeItem(Slot, DestAddress, nullptr);
+			CurrentType->SerializeItem(InnerSlot, DestAddress, nullptr);
 			return true;
 		}
-		else if( CurrentType->ConvertFromType(InTag, Slot, InData, InDefaultsStruct) == EConvertFromTypeResult::Converted )
+		else if( CurrentType->ConvertFromType(InTag, InnerSlot, InData, InDefaultsStruct) == EConvertFromTypeResult::Converted )
 		{
 			return true;
 		}
