TBA: In the Json output formatter, write float and double values out with enough precision for successful roundtripping. Added some debug only code which will immediately reconvert the string back to its original value and compare the the input

diff --git a/Engine/Source/Runtime/CoreUObject/Private/Serialization/Formatters/JsonArchiveOutputFormatter.cpp b/Engine/Source/Runtime/CoreUObject/Private/Serialization/Formatters/JsonArchiveOutputFormatter.cpp
--- a/Engine/Source/Runtime/CoreUObject/Private/Serialization/Formatters/JsonArchiveOutputFormatter.cpp
+++ b/Engine/Source/Runtime/CoreUObject/Private/Serialization/Formatters/JsonArchiveOutputFormatter.cpp
@@ -189,7 +189,13 @@ void FJsonArchiveOutputFormatter::Serialize(float& Value)
 	}
 	else
 	{
-		WriteValue(Lex::ToString(Value));
+		FString String = FString::Printf(TEXT("%.9g"), Value);
+#if DO_GUARD_SLOW
+		float RoundTripped;
+		Lex::FromString(RoundTripped, *String);
+		check(RoundTripped == Value);
+#endif
+		WriteValue(String);
 	}
 }
 
@@ -201,7 +207,13 @@ void FJsonArchiveOutputFormatter::Serialize(double& Value)
 	}
 	else
 	{
-		WriteValue(Lex::ToString(Value));
+		FString String = FString::Printf(TEXT("%.17g"), Value);
+#if DO_GUARD_SLOW
+		double RoundTripped;
+		Lex::FromString(RoundTripped, *String);
+		check(RoundTripped == Value);
+#endif
+		WriteValue(FString::Printf(TEXT("%.17g"), Value));
 	}
 }
 
