TBA: Remove only explicit usage of DECLARE_FSTRUCTUREDARCHIVE_SERIALIZER. Should only be used from UHT generated code.

diff --git a/Engine/Source/Runtime/CoreUObject/Public/Misc/TextBuffer.h b/Engine/Source/Runtime/CoreUObject/Public/Misc/TextBuffer.h
--- a/Engine/Source/Runtime/CoreUObject/Public/Misc/TextBuffer.h
+++ b/Engine/Source/Runtime/CoreUObject/Public/Misc/TextBuffer.h
@@ -48,8 +48,8 @@ public:
 
 public:
 
-	virtual void Serialize (FArchive& Ar) override;
-	DECLARE_FSTRUCTUREDARCHIVE_SERIALIZER(UTextBuffer, NO_API);
+	virtual void Serialize(FArchive& Ar) override;
+	virtual void Serialize(FStructuredArchive::FRecord Record) override;
 	virtual void Serialize (const TCHAR* Data, ELogVerbosity::Type Verbosity, const class FName& Category ) override;
 	virtual void Serialize( const TCHAR* V, ELogVerbosity::Type Verbosity, const class FName& Category, const double Time ) override;
 
