Fix for a potential race condition when two threads want to acquire GC lock

diff --git a/Engine/Source/Runtime/CoreUObject/Private/UObject/GCScopeLock.h b/Engine/Source/Runtime/CoreUObject/Private/UObject/GCScopeLock.h
--- a/Engine/Source/Runtime/CoreUObject/Private/UObject/GCScopeLock.h
+++ b/Engine/Source/Runtime/CoreUObject/Private/UObject/GCScopeLock.h
@@ -86,7 +86,7 @@ public:
 		if (!IsInGameThread())
 		{
 			// Wait until GC is done if it was running when entering this function
-			int32 OldAsyncCounterValue = 0;
+			bool bLocked = false;
 			do
 			{
 				if (GCCounter.GetValue() > 0)
@@ -95,13 +95,13 @@ public:
 				}
 				{
 					FScopeLock CriticalLock(&Critical);
-					OldAsyncCounterValue = AsyncCounter.GetValue();
 					if (GCCounter.GetValue() == 0)
 					{
 						AsyncCounter.Increment();
+						bLocked = true;
 					}
 				}
-			} while (AsyncCounter.GetValue() == OldAsyncCounterValue);
+			} while (!bLocked);
 		}
 	}
 	/** Release lock from non-game thread */
@@ -119,7 +119,7 @@ public:
 		SetGCIsWaiting();
 
 		// Wait until all other threads are done if they're currently holding the lock
-		int32 OldGCCounterValue = 0;
+		bool bLocked = false;
 		do
 		{
 			FPlatformProcess::ConditionalSleep([&]()
@@ -128,7 +128,6 @@ public:
 			});
 			{
 				FScopeLock CriticalLock(&Critical);
-				OldGCCounterValue = GCCounter.GetValue();
 				if (AsyncCounter.GetValue() == 0)
 				{
 					GCUnlockedEvent->Reset();
@@ -137,9 +136,10 @@ public:
 					// At this point GC can run so remove the signal that it's waiting
 					FPlatformMisc::MemoryBarrier();
 					ResetGCIsWaiting();
+					bLocked = true;
 				}
 			}
-		} while (GCCounter.GetValue() == OldGCCounterValue);
+		} while (!bLocked);
 	}
 	/** Checks if any async thread has a lock */
 	bool IsAsyncLocked() const
