Build fix

diff --git a/Engine/Source/Runtime/CoreUObject/Private/UObject/PropertyTempVal.cpp b/Engine/Source/Runtime/CoreUObject/Private/UObject/PropertyTempVal.cpp
--- a/Engine/Source/Runtime/CoreUObject/Private/UObject/PropertyTempVal.cpp
+++ b/Engine/Source/Runtime/CoreUObject/Private/UObject/PropertyTempVal.cpp
@@ -24,7 +24,7 @@ FPropertyTempVal::~FPropertyTempVal()
 
 void FPropertyTempVal::Serialize(FArchive& Ar, const void* Defaults)
 {
-	Prop->SerializeItem(Ar, Value, Defaults);
+	Prop->SerializeItem(FStructuredArchiveFromArchive(Ar).GetSlot(), Value, Defaults);
 }
 
 void FPropertyTempVal::ExportText(FString& ValueStr, const void* Defaults, UObject* Parent, int32 PortFlags, UObject* ExportRootScope)
