Fixing CIS

diff --git a/Engine/Source/Runtime/CoreUObject/Private/UObject/PropertyObject.cpp b/Engine/Source/Runtime/CoreUObject/Private/UObject/PropertyObject.cpp
--- a/Engine/Source/Runtime/CoreUObject/Private/UObject/PropertyObject.cpp
+++ b/Engine/Source/Runtime/CoreUObject/Private/UObject/PropertyObject.cpp
@@ -49,7 +49,7 @@ EConvertFromTypeResult UObjectProperty::ConvertFromType(const FPropertyTag& Tag,
 
 			if (!PreviousValueObj && !PreviousValue.IsNull())
 			{
-				UE_LOG(LogClass, Error, TEXT("Failed to convert soft path %s to unloaded object as this is not safe during async loading. Load and resave %s in the editor to fix!"), *PreviousValue.ToString(), *Ar.GetArchiveName());
+				UE_LOG(LogClass, Error, TEXT("Failed to convert soft path %s to unloaded object as this is not safe during async loading. Load and resave %s in the editor to fix!"), *PreviousValue.ToString(), *Slot.GetUnderlyingArchive().GetArchiveName());
 			}
 		}
 		else
