Fixing async loading code asserts when exiting game very early due to an error

diff --git a/Engine/Source/Runtime/CoreUObject/Private/Serialization/AsyncLoading.cpp b/Engine/Source/Runtime/CoreUObject/Private/Serialization/AsyncLoading.cpp
--- a/Engine/Source/Runtime/CoreUObject/Private/Serialization/AsyncLoading.cpp
+++ b/Engine/Source/Runtime/CoreUObject/Private/Serialization/AsyncLoading.cpp
@@ -4857,13 +4857,16 @@ FMaxPackageSummarySize::FMaxPackageSummarySize()
 	// the editor packages may not have the AdditionalPackagesToCook array stripped so we need to allocate more memory
 #if WITH_EDITORONLY_DATA
 	const int32 MinimumPackageSummarySize = 1024;
-	check(GConfig);
+	check(GConfig || GIsRequestingExit);
 	Value = 16384;
-	GConfig->GetInt(TEXT("/Script/Engine.StreamingSettings"), TEXT("s.MaxPackageSummarySize"), Value, GEngineIni);
-	if (Value <= MinimumPackageSummarySize)
+	if (GConfig)
 	{
-		UE_LOG(LogStreaming, Warning, TEXT("Invalid minimum package file summary size (s.MaxPackageSummarySize=%d), %d is min."), Value, MinimumPackageSummarySize);
-		Value = MinimumPackageSummarySize;
+		GConfig->GetInt(TEXT("/Script/Engine.StreamingSettings"), TEXT("s.MaxPackageSummarySize"), Value, GEngineIni);
+		if (Value <= MinimumPackageSummarySize)
+		{
+			UE_LOG(LogStreaming, Warning, TEXT("Invalid minimum package file summary size (s.MaxPackageSummarySize=%d), %d is min."), Value, MinimumPackageSummarySize);
+			Value = MinimumPackageSummarySize;
+		}
 	}
 #else
 	Value = 8192;
@@ -7127,14 +7130,17 @@ bool IsEventDrivenLoaderEnabledInCookedBuilds()
 		FEventDrivenLoaderEnabledInCookedBuildsInit()
 			: bEventDrivenLoaderEnabled(false)
 		{
-			check(GConfig);
-			GConfig->GetBool(TEXT("/Script/Engine.StreamingSettings"), TEXT("s.EventDrivenLoaderEnabled"), bEventDrivenLoaderEnabled, GEngineIni);
-#if !UE_BUILD_SHIPPING
-			if (FParse::Param(FCommandLine::Get(), TEXT("NOEDL")))
+			check(GConfig || GIsRequestingExit);
+			if (GConfig)
 			{
-				bEventDrivenLoaderEnabled = false;
-			}
+				GConfig->GetBool(TEXT("/Script/Engine.StreamingSettings"), TEXT("s.EventDrivenLoaderEnabled"), bEventDrivenLoaderEnabled, GEngineIni);
+#if !UE_BUILD_SHIPPING
+				if (FParse::Param(FCommandLine::Get(), TEXT("NOEDL")))
+				{
+					bEventDrivenLoaderEnabled = false;
+				}
 #endif
+			}
 		}
 	} EventDrivenLoaderEnabledInCookedBuilds;
 	return EventDrivenLoaderEnabledInCookedBuilds.bEventDrivenLoaderEnabled;
