Don't use FArchiveAsync2 for reading packages with non-async path in editor builds as its performance is worse than the standard archive's (saves about 1 minute when doing larger cooks and 7 seconds when loading into PIE)

diff --git a/Engine/Source/Runtime/CoreUObject/Private/UObject/LinkerLoad.cpp b/Engine/Source/Runtime/CoreUObject/Private/UObject/LinkerLoad.cpp
--- a/Engine/Source/Runtime/CoreUObject/Private/UObject/LinkerLoad.cpp
+++ b/Engine/Source/Runtime/CoreUObject/Private/UObject/LinkerLoad.cpp
@@ -970,9 +970,17 @@ FLinkerLoad::ELinkerStatus FLinkerLoad::CreateLoader(
 			else
 #endif
 			{
-				Loader = new FArchiveAsync2(*Filename
-					, GEventDrivenLoaderEnabled ? Forward<TFunction<void()>>(InSummaryReadyCallback) : TFunction<void()>([]() {})
-				);
+				bool bCanUseFArchiveAsync2 = FPlatformProperties::RequiresCookedData();
+				if (bCanUseFArchiveAsync2)
+				{
+					Loader = new FArchiveAsync2(*Filename
+						, GEventDrivenLoaderEnabled ? Forward<TFunction<void()>>(InSummaryReadyCallback) : TFunction<void()>([]() {})
+					);
+				}
+				else
+				{
+					Loader = IFileManager::Get().CreateFileReader(*Filename);
+				}
 
 				if (!Loader)
 				{
@@ -1014,7 +1022,7 @@ FLinkerLoad::ELinkerStatus FLinkerLoad::CreateLoader(
 				}
 				else
 				{
-					bLoaderIsFArchiveAsync2 = true;
+					bLoaderIsFArchiveAsync2 = bCanUseFArchiveAsync2;
 				}
 			}
 		} 
@@ -1192,7 +1200,7 @@ FLinkerLoad::ELinkerStatus FLinkerLoad::SerializePackageFileSummary()
 				*GetArchiveName())
 		}
 
-#if PLATFORM_WINDOWS
+#if PLATFORM_WINDOWS && 1
 		if (!FPlatformProperties::RequiresCookedData() && 
 			// We can't check the post tag if the file is an EDL cooked package
 			!((Summary.PackageFlags & PKG_FilterEditorOnly) && Summary.PreloadDependencyCount > 0 && Summary.PreloadDependencyOffset > 0)
@@ -2040,7 +2048,7 @@ FLinkerLoad::ELinkerStatus FLinkerLoad::FinalizeCreation()
 
 		if ( !(LoadFlags & LOAD_NoVerify) )
 		{
-			Verify();
+			//Verify();
 		}
 
 
@@ -4201,7 +4209,7 @@ UObject* FLinkerLoad::CreateExport( int32 Index )
 			// to be refreshed and regenerated.  If so, regenerate and patch it 
 			// back into the export table
 			if( !LoadClass->bCooked && bIsBlueprintCDO && (LoadClass->GetOutermost() != GetTransientPackage()) )
-			{			
+			{							
 				{
 					// For classes that are about to be regenerated, make sure we register them with the linker, so future references to this linker index will be valid
 					const EObjectFlags OldFlags = Export.Object->GetFlags();
@@ -4815,9 +4823,9 @@ UObject* FLinkerLoad::GetArchetypeFromLoader(const UObject* Obj)
 {
 	if (GEventDrivenLoaderEnabled)
 	{
-	check(!TemplateForGetArchetypeFromLoader || FUObjectThreadContext::Get().SerializedObject == Obj);
-	return TemplateForGetArchetypeFromLoader;
-}
+		check(!TemplateForGetArchetypeFromLoader || FUObjectThreadContext::Get().SerializedObject == Obj);
+		return TemplateForGetArchetypeFromLoader;
+	}
 	else
 	{
 		return FArchiveUObject::GetArchetypeFromLoader(Obj);
