TBA: Fix tagged binary formatter

diff --git a/Engine/Source/Runtime/Core/Public/Serialization/StructuredArchiveFormatter.h b/Engine/Source/Runtime/Core/Public/Serialization/StructuredArchiveFormatter.h
--- a/Engine/Source/Runtime/Core/Public/Serialization/StructuredArchiveFormatter.h
+++ b/Engine/Source/Runtime/Core/Public/Serialization/StructuredArchiveFormatter.h
@@ -49,6 +49,11 @@ enum class EArchiveValueType
 	String,
 	Name,
 	Object,
+	Text,
+	WeakObjectPtr,
+	SoftObjectPtr,
+	SoftObjectPath,
+	LazyObjectPtr,
 	RawData,
 };
 
diff --git a/Engine/Source/Runtime/CoreUObject/Private/Serialization/Formatters/TaggedBinaryArchiveInputFormatter.cpp b/Engine/Source/Runtime/CoreUObject/Private/Serialization/Formatters/TaggedBinaryArchiveInputFormatter.cpp
--- a/Engine/Source/Runtime/CoreUObject/Private/Serialization/Formatters/TaggedBinaryArchiveInputFormatter.cpp
+++ b/Engine/Source/Runtime/CoreUObject/Private/Serialization/Formatters/TaggedBinaryArchiveInputFormatter.cpp
@@ -292,6 +292,36 @@ void FTaggedBinaryArchiveInputFormatter::Serialize(UObject*& Value)
 	SerializeObject(Inner, Value);
 }
 
+void FTaggedBinaryArchiveInputFormatter::Serialize(FText& Value)
+{
+	ExpectType(EArchiveValueType::Text);
+	Inner << Value;
+}
+
+void FTaggedBinaryArchiveInputFormatter::Serialize(FWeakObjectPtr& Value)
+{
+	ExpectType(EArchiveValueType::WeakObjectPtr);
+	Inner << Value;
+}
+
+void FTaggedBinaryArchiveInputFormatter::Serialize(FSoftObjectPtr& Value)
+{
+	ExpectType(EArchiveValueType::SoftObjectPtr);
+	Inner << Value;
+}
+
+void FTaggedBinaryArchiveInputFormatter::Serialize(FSoftObjectPath& Value)
+{
+	ExpectType(EArchiveValueType::SoftObjectPath);
+	Inner << Value;
+}
+
+void FTaggedBinaryArchiveInputFormatter::Serialize(FLazyObjectPtr& Value)
+{
+	ExpectType(EArchiveValueType::LazyObjectPtr);
+	Inner << Value;
+}
+
 void FTaggedBinaryArchiveInputFormatter::Serialize(TArray<uint8>& Value) 
 {
 	ExpectType(EArchiveValueType::RawData);
diff --git a/Engine/Source/Runtime/CoreUObject/Private/Serialization/Formatters/TaggedBinaryArchiveOutputFormatter.cpp b/Engine/Source/Runtime/CoreUObject/Private/Serialization/Formatters/TaggedBinaryArchiveOutputFormatter.cpp
--- a/Engine/Source/Runtime/CoreUObject/Private/Serialization/Formatters/TaggedBinaryArchiveOutputFormatter.cpp
+++ b/Engine/Source/Runtime/CoreUObject/Private/Serialization/Formatters/TaggedBinaryArchiveOutputFormatter.cpp
@@ -255,6 +255,36 @@ void FTaggedBinaryArchiveOutputFormatter::Serialize(UObject*& Value)
 	SerializeObject(Inner, Value);
 }
 
+void FTaggedBinaryArchiveOutputFormatter::Serialize(FText& Value)
+{
+	WriteType(EArchiveValueType::Text);
+	Inner << Value;
+}
+
+void FTaggedBinaryArchiveOutputFormatter::Serialize(FWeakObjectPtr& Value)
+{
+	WriteType(EArchiveValueType::WeakObjectPtr);
+	Inner << Value;
+}
+
+void FTaggedBinaryArchiveOutputFormatter::Serialize(FSoftObjectPtr& Value)
+{
+	WriteType(EArchiveValueType::SoftObjectPtr);
+	Inner << Value;
+}
+
+void FTaggedBinaryArchiveOutputFormatter::Serialize(FSoftObjectPath& Value)
+{
+	WriteType(EArchiveValueType::SoftObjectPath);
+	Inner << Value;
+}
+
+void FTaggedBinaryArchiveOutputFormatter::Serialize(FLazyObjectPtr& Value)
+{
+	WriteType(EArchiveValueType::LazyObjectPtr);
+	Inner << Value;
+}
+
 void FTaggedBinaryArchiveOutputFormatter::Serialize(TArray<uint8>& Value) 
 {
 	WriteType(EArchiveValueType::RawData);
diff --git a/Engine/Source/Runtime/CoreUObject/Public/Serialization/Formatters/TaggedBinaryArchiveInputFormatter.h b/Engine/Source/Runtime/CoreUObject/Public/Serialization/Formatters/TaggedBinaryArchiveInputFormatter.h
--- a/Engine/Source/Runtime/CoreUObject/Public/Serialization/Formatters/TaggedBinaryArchiveInputFormatter.h
+++ b/Engine/Source/Runtime/CoreUObject/Public/Serialization/Formatters/TaggedBinaryArchiveInputFormatter.h
@@ -56,6 +56,11 @@ public:
 	virtual void Serialize(FString& Value) override;
 	virtual void Serialize(FName& Value) override;
 	virtual void Serialize(UObject*& Value) override;
+	virtual void Serialize(FText& Value) override;
+	virtual void Serialize(FWeakObjectPtr& Value) override;
+	virtual void Serialize(FSoftObjectPtr& Value) override;
+	virtual void Serialize(FSoftObjectPath& Value) override;
+	virtual void Serialize(FLazyObjectPtr& Value) override;
 	virtual void Serialize(TArray<uint8>& Value) override;
 	virtual void Serialize(void* Data, uint64 DataSize) override;
 
diff --git a/Engine/Source/Runtime/CoreUObject/Public/Serialization/Formatters/TaggedBinaryArchiveOutputFormatter.h b/Engine/Source/Runtime/CoreUObject/Public/Serialization/Formatters/TaggedBinaryArchiveOutputFormatter.h
--- a/Engine/Source/Runtime/CoreUObject/Public/Serialization/Formatters/TaggedBinaryArchiveOutputFormatter.h
+++ b/Engine/Source/Runtime/CoreUObject/Public/Serialization/Formatters/TaggedBinaryArchiveOutputFormatter.h
@@ -50,6 +50,11 @@ public:
 	virtual void Serialize(FString& Value) override;
 	virtual void Serialize(FName& Value) override;
 	virtual void Serialize(UObject*& Value) override;
+	virtual void Serialize(FText& Value) override;
+	virtual void Serialize(FWeakObjectPtr& Value) override;
+	virtual void Serialize(FSoftObjectPtr& Value) override;
+	virtual void Serialize(FSoftObjectPath& Value) override;
+	virtual void Serialize(FLazyObjectPtr& Value) override;
 	virtual void Serialize(TArray<uint8>& Value) override;
 	virtual void Serialize(void* Data, uint64 DataSize) override;
 
