TBA: Tagged properties are written out as named fields on the "Properties" record, rather than as a stream with a null tag at the end

diff --git a/Engine/Source/Runtime/CoreUObject/Private/UObject/Class.cpp b/Engine/Source/Runtime/CoreUObject/Private/UObject/Class.cpp
--- a/Engine/Source/Runtime/CoreUObject/Private/UObject/Class.cpp
+++ b/Engine/Source/Runtime/CoreUObject/Private/UObject/Class.cpp
@@ -913,11 +913,24 @@ void UStruct::SerializeTaggedProperties(FStructuredArchive::FSlot Slot, uint8* D
 	// Determine if this struct supports optional property guid's (UBlueprintGeneratedClasses Only)
 	const bool bArePropertyGuidsAvailable = (UnderlyingArchive.UE4Ver() >= VER_UE4_PROPERTY_GUID_IN_PROPERTY_TAG) && !FPlatformProperties::RequiresCookedData() && ArePropertyGuidsAvailable();
 
-	FStructuredArchive::FStream PropertiesStream = Slot.EnterStream();
-
 	if( UnderlyingArchive.IsLoading() )
 	{
 		// Load tagged properties.
+		TArray<FString> FieldNames;
+
+		TOptional<FStructuredArchive::FRecord> PropertiesRecord;
+		TOptional<FStructuredArchive::FStream> PropertiesStream;
+
+		if (UnderlyingArchive.IsTextFormat())
+		{
+			PropertiesRecord.Emplace(Slot.EnterRecord(FieldNames));
+		}
+		else
+		{
+			PropertiesStream.Emplace(Slot.EnterStream());
+		}
+
+		int32 CurrentFieldNameIdx = UnderlyingArchive.IsTextFormat() ? 0 : -1;
 
 		// This code assumes that properties are loaded in the same order they are saved in. This removes a n^2 search 
 		// and makes it an O(n) when properties are saved in the same order as they are loaded (default case). In the 
@@ -927,9 +940,9 @@ void UStruct::SerializeTaggedProperties(FStructuredArchive::FSlot Slot, uint8* D
 		int32		RemainingArrayDim	= Property ? Property->ArrayDim : 0;
 
 		// Load all stored properties, potentially skipping unknown ones.
-		while (1)
+		while (CurrentFieldNameIdx < FieldNames.Num())
 		{
-			FStructuredArchive::FRecord PropertyRecord = PropertiesStream.EnterElement().EnterRecord();
+			FStructuredArchive::FRecord PropertyRecord = UnderlyingArchive.IsTextFormat() ? PropertiesRecord->EnterRecord(FIELD_NAME(*FieldNames[CurrentFieldNameIdx++])) : PropertiesStream->EnterElement().EnterRecord();
 
 			FPropertyTag Tag;
 			PropertyRecord << NAMED_FIELD(Tag);
@@ -1124,6 +1137,8 @@ void UStruct::SerializeTaggedProperties(FStructuredArchive::FSlot Slot, uint8* D
 	}
 	else
 	{
+		FStructuredArchive::FRecord PropertiesRecord = Slot.EnterRecord();
+
 		check(UnderlyingArchive.IsSaving() || UnderlyingArchive.IsCountingMemory());
 
 		UScriptStruct* DefaultsScriptStruct = dynamic_cast<UScriptStruct*>(DefaultsStruct);
@@ -1170,7 +1185,7 @@ void UStruct::SerializeTaggedProperties(FStructuredArchive::FSlot Slot, uint8* D
 							Tag.SetPropertyGuid(PropertyGuid);
 						}
 
-						FStructuredArchive::FRecord PropertyRecord = PropertiesStream.EnterElement().EnterRecord();
+						FStructuredArchive::FRecord PropertyRecord = PropertiesRecord.EnterRecord(FIELD_NAME(*Tag.Name.ToString()));
 
 						PropertyRecord << NAMED_FIELD(Tag);
 
@@ -1218,12 +1233,7 @@ void UStruct::SerializeTaggedProperties(FStructuredArchive::FSlot Slot, uint8* D
 			}
 		}
 
-		if (UnderlyingArchive.IsTextFormat())
-		{
-			FPropertyTag EmptyTag;
-			PropertiesStream.EnterElement().EnterRecord() << NAMED_ITEM("Tag", EmptyTag);
-		}
-		else
+		if (!UnderlyingArchive.IsTextFormat())
 		{
 			static FName Temp(NAME_None);
 			UnderlyingArchive << Temp;
