TBA: Added a local const bool to allow easy hacking out of text asset loading support

diff --git a/Engine/Source/Runtime/CoreUObject/Private/Misc/PackageName.cpp b/Engine/Source/Runtime/CoreUObject/Private/Misc/PackageName.cpp
--- a/Engine/Source/Runtime/CoreUObject/Private/Misc/PackageName.cpp
+++ b/Engine/Source/Runtime/CoreUObject/Private/Misc/PackageName.cpp
@@ -753,7 +753,8 @@ bool FPackageName::FindPackageFileWithoutExtension(const FString& InPackageFilen
 	};
 
 	// Loop through all known extensions and check if the file exist.
-	const int32 BaseIndex = (/*false &&*/ InAllowTextFormats && WITH_TEXT_ARCHIVE_SUPPORT) ? 0 : 2;
+	const bool bDebuggingTextFormatEnable = true;
+	const int32 BaseIndex = (bDebuggingTextFormatEnable && InAllowTextFormats && WITH_TEXT_ARCHIVE_SUPPORT) ? 0 : 2;
 	for (int32 ExtensionIndex = BaseIndex; ExtensionIndex < ARRAY_COUNT(PackageExtensions); ++ExtensionIndex)
 	{
 		FString   PackageFilename = InPackageFilename + *PackageExtensions[ExtensionIndex];
