Removing some WIP code accidentally checked in with CL #4121943

diff --git a/Engine/Source/Runtime/CoreUObject/Private/UObject/LinkerLoad.cpp b/Engine/Source/Runtime/CoreUObject/Private/UObject/LinkerLoad.cpp
--- a/Engine/Source/Runtime/CoreUObject/Private/UObject/LinkerLoad.cpp
+++ b/Engine/Source/Runtime/CoreUObject/Private/UObject/LinkerLoad.cpp
@@ -1200,7 +1200,7 @@ FLinkerLoad::ELinkerStatus FLinkerLoad::SerializePackageFileSummary()
 				*GetArchiveName())
 		}
 
-#if PLATFORM_WINDOWS && 1
+#if PLATFORM_WINDOWS
 		if (!FPlatformProperties::RequiresCookedData() && 
 			// We can't check the post tag if the file is an EDL cooked package
 			!((Summary.PackageFlags & PKG_FilterEditorOnly) && Summary.PreloadDependencyCount > 0 && Summary.PreloadDependencyOffset > 0)
@@ -2048,7 +2048,7 @@ FLinkerLoad::ELinkerStatus FLinkerLoad::FinalizeCreation()
 
 		if ( !(LoadFlags & LOAD_NoVerify) )
 		{
-			//Verify();
+			Verify();
 		}
 
 
