TBA: Fix incorrect map property conversion code when converting an old property that contains a map with different key/value types

diff --git a/Engine/Source/Runtime/CoreUObject/Private/UObject/PropertyMap.cpp b/Engine/Source/Runtime/CoreUObject/Private/UObject/PropertyMap.cpp
--- a/Engine/Source/Runtime/CoreUObject/Private/UObject/PropertyMap.cpp
+++ b/Engine/Source/Runtime/CoreUObject/Private/UObject/PropertyMap.cpp
@@ -1000,7 +1000,9 @@ EConvertFromTypeResult UMapProperty::ConvertFromType(const FPropertyTag& Tag, FS
 						KeyProp->InitializeValue(TempKeyStorage);
 					}
 
-					if( SerializeOrConvert( KeyProp, KeyPropertyTag, EntriesArray.EnterElement(), TempKeyStorage, DefaultsStruct ) )
+					FStructuredArchive::FRecord FirstPropertyRecord = EntriesArray.EnterElement().EnterRecord();
+
+					if( SerializeOrConvert( KeyProp, KeyPropertyTag, FirstPropertyRecord.EnterField(FIELD_NAME_TEXT("Key")), TempKeyStorage, DefaultsStruct ) )
 					{
 						// Add a new default value if the key doesn't currently exist in the map
 						bool bKeyAlreadyPresent = true;
@@ -1016,12 +1018,14 @@ EConvertFromTypeResult UMapProperty::ConvertFromType(const FPropertyTag& Tag, FS
 						KeyProp->CopyCompleteValue_InContainer(NextPairPtr, TempKeyStorage);
 
 						// Deserialize value
-						if( SerializeOrConvert( ValueProp, ValuePropertyTag, EntriesArray.EnterElement(), NextPairPtr, DefaultsStruct ) )
+						if( SerializeOrConvert( ValueProp, ValuePropertyTag, FirstPropertyRecord.EnterField(FIELD_NAME_TEXT("Value")), NextPairPtr, DefaultsStruct ) )
 						{
 							// first entry went fine, convert the rest:
 							for(int32 I = 1; I < NumEntries; ++I)
 							{
-								verify( SerializeOrConvert( KeyProp, KeyPropertyTag, EntriesArray.EnterElement(), TempKeyStorage, DefaultsStruct ) );
+								FStructuredArchive::FRecord PropertyRecord = EntriesArray.EnterElement().EnterRecord();
+
+								verify( SerializeOrConvert( KeyProp, KeyPropertyTag, PropertyRecord.EnterField(FIELD_NAME_TEXT("Key")), TempKeyStorage, DefaultsStruct ) );
 								NextPairIndex = MapHelper.FindMapIndexWithKey(TempKeyStorage);
 								if (NextPairIndex == INDEX_NONE)
 								{
@@ -1031,7 +1035,7 @@ EConvertFromTypeResult UMapProperty::ConvertFromType(const FPropertyTag& Tag, FS
 								NextPairPtr = MapHelper.GetPairPtrWithoutCheck(NextPairIndex);
 								// This copy is unnecessary when the key was already in the map:
 								KeyProp->CopyCompleteValue_InContainer(NextPairPtr, TempKeyStorage);
-								verify( SerializeOrConvert( ValueProp, ValuePropertyTag, EntriesArray.EnterElement(), NextPairPtr, DefaultsStruct ) );
+								verify( SerializeOrConvert( ValueProp, ValuePropertyTag, PropertyRecord.EnterField(FIELD_NAME_TEXT("Value")), NextPairPtr, DefaultsStruct ) );
 							}
 						}
 						else
